Anim layers are editor only asset user data attached to a level sequence that's shown from the animation toolbar. Consist of a set of sections and their corresponding owners(control rig/object) with various properties(mute, layer weight, etc.). Note since some of these values may get set in sequencer we poll some of these properties, using caching as needed.

For control rig we have an idea of multiple sections to key, since we need to have that per control e.g. left arm in left arm section/layer, right arm in right arm section/layer, everything else in base layer.

Weight track for the anim layers shows if it's keyed or not so I reuse control rig editor to cache that data which it also does for anim details proxies.

Move SCollapseControlsWidget to track editor header so it can be shared.

Sections in layers/with overrides can have a Blending Order, this will talked about more in the override code review.

Have a couple todo's in particurly when recompiling control rigs that are in layers

diff --git a/Engine/Plugins/Animation/ControlRig/Source/ControlRig/Private/Sequencer/MovieSceneControlRigParameterTrack.cpp b/Engine/Plugins/Animation/ControlRig/Source/ControlRig/Private/Sequencer/MovieSceneControlRigParameterTrack.cpp
--- a/Engine/Plugins/Animation/ControlRig/Source/ControlRig/Private/Sequencer/MovieSceneControlRigParameterTrack.cpp
+++ b/Engine/Plugins/Animation/ControlRig/Source/ControlRig/Private/Sequencer/MovieSceneControlRigParameterTrack.cpp
@@ -403,9 +403,76 @@ UMovieSceneSection* UMovieSceneControlRigParameterTrack::FindOrAddSection(FFrame
 	return NewSection;
 }
 
+TArray<TWeakObjectPtr<UMovieSceneSection>> UMovieSceneControlRigParameterTrack::GetSectionsToKey() const
+{
+	TArray<TWeakObjectPtr<UMovieSceneSection>> SectionsToKey;
+	if (SectionToKeyPerControl.Num() > 0)
+	{
+		SectionToKeyPerControl.GenerateValueArray(SectionsToKey);
+	}
+	else
+	{
+		SectionsToKey.Add(SectionToKey);
+	}
+	return SectionsToKey;
+}
+
+UMovieSceneSection* UMovieSceneControlRigParameterTrack::GetSectionToKey(const FName& InControlName) const
+{
+	if (const TWeakObjectPtr<UMovieSceneSection>* Section = SectionToKeyPerControl.Find(InControlName))
+	{
+		if (Section->IsValid())
+		{
+			return Section->Get();
+		}
+	}
+	return GetSectionToKey();
+}
+
+void UMovieSceneControlRigParameterTrack::SetSectionToKey(UMovieSceneSection* InSection, const FName& InControlName)
+{
+	if (Sections.Num() < 1 || InSection == nullptr)
+	{
+		return;
+	}
+	Modify();
+	SectionToKeyPerControl.Add(InControlName, InSection);
+	SectionToKey = Sections[0];
+}
+
 void UMovieSceneControlRigParameterTrack::SetSectionToKey(UMovieSceneSection* InSection)
 {
-	SectionToKey = InSection;
+	if (Sections.Num() < 1 || InSection == nullptr)
+	{
+		return;
+	}
+	if (UMovieSceneControlRigParameterSection* CRSection = Cast<UMovieSceneControlRigParameterSection>(InSection))
+	{
+		Modify();
+		if (SectionToKeyPerControl.Num() > 0) //we have sections that are in layers so need to respect them
+		{
+			if (bSetSectionToKeyPerControl)
+			{
+				for (TPair < FName, TWeakObjectPtr<UMovieSceneSection>>& SectionToKeyItem : SectionToKeyPerControl)
+				{
+					//only set it as the section to key if it's in that section, otherwise leave it alone
+					FChannelMapInfo* pChannelIndex = CRSection->ControlChannelMap.Find(SectionToKeyItem.Key);
+					if (pChannelIndex)
+					{
+						if (CRSection->GetControlsMask(pChannelIndex->MaskIndex))
+						{
+							SectionToKeyItem.Value = InSection;
+						}
+					}
+				}
+			}
+			SectionToKey = Sections[0];
+		}
+		else
+		{
+			SectionToKey = InSection;
+		}
+	}
 }
 
 UMovieSceneSection* UMovieSceneControlRigParameterTrack::GetSectionToKey() const
diff --git a/Engine/Plugins/Animation/ControlRig/Source/ControlRig/Public/Sequencer/MovieSceneControlRigParameterTrack.h b/Engine/Plugins/Animation/ControlRig/Source/ControlRig/Public/Sequencer/MovieSceneControlRigParameterTrack.h
--- a/Engine/Plugins/Animation/ControlRig/Source/ControlRig/Public/Sequencer/MovieSceneControlRigParameterTrack.h
+++ b/Engine/Plugins/Animation/ControlRig/Source/ControlRig/Public/Sequencer/MovieSceneControlRigParameterTrack.h
@@ -11,6 +11,7 @@
 #include "INodeAndChannelMappings.h"
 #include "MovieSceneControlRigParameterSection.h"
 #include "EulerTransform.h"
+#include "Tracks/IMovieSceneSectionsToKey.h"
 #include "MovieSceneControlRigParameterTrack.generated.h"
 
 struct FEndLoadPackageContext;
@@ -41,6 +42,8 @@ class UMovieSceneControlRigParameterTrack
 	: public UMovieSceneNameableTrack
 	, public IMovieSceneTrackTemplateProducer
 	, public INodeAndChannelMappings
+	, public IMovieSceneSectionsToKey
+
 {
 	GENERATED_UCLASS_BODY()
 
@@ -144,6 +147,18 @@ public:
 	 */
 	virtual UMovieSceneSection* GetSectionToKey() const override;
 
+	/**
+	* Get multiple sections to key
+	* @return multiple sections to key
+	*/
+	virtual TArray<TWeakObjectPtr<UMovieSceneSection>> GetSectionsToKey() const override;
+
+	/**
+	*  Control rig supports per control sections to key needed for anim layer workflows
+	*/
+	CONTROLRIG_API UMovieSceneSection* GetSectionToKey(const FName& InControlName) const;
+	CONTROLRIG_API void SetSectionToKey(UMovieSceneSection* InSection, const FName& inControlName);
+
 	CONTROLRIG_API void SetTrackName(FName InName) { TrackName = InName; }
 
 	UMovieSceneControlRigParameterSection::FSpaceChannelAddedEvent& SpaceChannelAdded() { return OnSpaceChannelAdded; }
@@ -159,6 +174,11 @@ public:
 
 	CONTROLRIG_API int32 GetPriorityOrder() const;
 	CONTROLRIG_API void SetPriorityOrder(int32 InPriorityIndex);
+private:
+	UPROPERTY()
+	TMap<FName, TWeakObjectPtr<UMovieSceneSection>> SectionToKeyPerControl;
+public:
+	bool bSetSectionToKeyPerControl = true;// set up the map when calling the UMovieSceneTrack::SetSectionToKey()
 private:
 	//get the rotation orderm will not be set if it's default FRotator order. If bCurrent is true, it uses what's set
 	//if false it uses the default setting from the current control rig.
diff --git a/Engine/Plugins/Animation/ControlRig/Source/ControlRigEditor/Private/ControlRigEditorStyle.h b/Engine/Plugins/Animation/ControlRig/Source/ControlRigEditor/Private/ControlRigEditorStyle.h
--- a/Engine/Plugins/Animation/ControlRig/Source/ControlRigEditor/Private/ControlRigEditorStyle.h
+++ b/Engine/Plugins/Animation/ControlRig/Source/ControlRigEditor/Private/ControlRigEditorStyle.h
@@ -90,8 +90,11 @@ public:
 			Set("ControlRig.TweenTool", new IMAGE_BRUSH_SVG("Slate/AnimationTweens", Icon16x16));
 			Set("ControlRig.EditableMotionTrails", new IMAGE_BRUSH_SVG("Slate/EditableMotionTrails", Icon16x16));
 			Set("ControlRig.TemporaryPivot", new IMAGE_BRUSH_SVG("Slate/TemporaryPivot", Icon16x16));
-		}
+			Set("ControlRig.AnimLayerSelected", new IMAGE_BRUSH_SVG("Slate/AnimLayerSelected", Icon16x16));
+			Set("ControlRig.FilterAnimLayerSelected", new IMAGE_BRUSH_SVG("Slate/FilterAnimLayerSelected", Icon16x16));
+			Set("ControlRig.AnimLayers", new IMAGE_BRUSH_SVG("Slate/AnimLayers", Icon16x16));
 
+		}
 		// Control Rig Editor styles
 		{
 			// tab icons
diff --git a/Engine/Plugins/Animation/ControlRig/Source/ControlRigEditor/Private/ControlRigSequencerEditorLibrary.cpp b/Engine/Plugins/Animation/ControlRig/Source/ControlRigEditor/Private/ControlRigSequencerEditorLibrary.cpp
--- a/Engine/Plugins/Animation/ControlRig/Source/ControlRigEditor/Private/ControlRigSequencerEditorLibrary.cpp
+++ b/Engine/Plugins/Animation/ControlRig/Source/ControlRigEditor/Private/ControlRigSequencerEditorLibrary.cpp
@@ -49,6 +49,8 @@
 #include "Sections/MovieSceneConstrainedSection.h"
 #include "BakingAnimationKeySettings.h"
 #include "EditMode/ControlRigEditMode.h"
+#include "Sequencer/AnimLayers/AnimLayers.h"
+
 
 #include UE_INLINE_GENERATED_CPP_BY_NAME(ControlRigSequencerEditorLibrary)
 
@@ -3287,5 +3289,133 @@ bool UControlRigSequencerEditorLibrary::SetControlRigApplyMode(UControlRig* InCo
 	return false;
 }
 
+bool UControlRigSequencerEditorLibrary::DeleteAnimLayer(int32 Index)
+{
+	TWeakPtr<ISequencer> WeakSequencer = GetSequencerFromAsset();
+	if (WeakSequencer.IsValid() == false)
+	{
+		UE_LOG(LogControlRig, Error, TEXT("DeleteAnimLayer: Need open Sequencer"));
+		return false;
+	}
+	TSharedPtr<ISequencer> SequencerPtr = WeakSequencer.Pin();
+	if (UAnimLayers* AnimLayers = UAnimLayers::GetAnimLayers(SequencerPtr.Get()))
+	{
+		if (Index != INDEX_NONE)
+		{
+			return AnimLayers->DeleteAnimLayer(SequencerPtr.Get(), Index);
+		}
+	}
+	else
+	{
+		UE_LOG(LogControlRig, Error, TEXT("DeleteAnimLayer: No Anim Layers on Level Sequence"));
+	}
+	return false;
+}
+
+int32 UControlRigSequencerEditorLibrary::DuplicateAnimLayer(int32 Index)
+{
+	TWeakPtr<ISequencer> WeakSequencer = GetSequencerFromAsset();
+	if (WeakSequencer.IsValid() == false)
+	{
+		UE_LOG(LogControlRig, Error, TEXT("DuplicateAnimLayer: Need open Sequencer"));
+		return INDEX_NONE;
+	}
+	TSharedPtr<ISequencer> SequencerPtr = WeakSequencer.Pin();
+	if (UAnimLayers* AnimLayers = UAnimLayers::GetAnimLayers(SequencerPtr.Get()))
+	{
+		if (Index != INDEX_NONE)
+		{
+			return AnimLayers->DuplicateAnimLayer(SequencerPtr.Get(), Index);
+		}
+	}
+	else
+	{
+		UE_LOG(LogControlRig, Error, TEXT("DuplicateAnimLayer: No Anim Layers on Level Sequence"));
+	}
+	return INDEX_NONE;
+}
+
+int32 UControlRigSequencerEditorLibrary::AddAnimLayerFromSelection()
+{
+	TWeakPtr<ISequencer> WeakSequencer = GetSequencerFromAsset();
+	if (WeakSequencer.IsValid() == false)
+	{
+		UE_LOG(LogControlRig, Error, TEXT("AddAnimLayerFromSelection: Need open Sequencer"));
+		return INDEX_NONE;
+	}
+	TSharedPtr<ISequencer> SequencerPtr = WeakSequencer.Pin();
+	if (UAnimLayers* AnimLayers = UAnimLayers::GetAnimLayers(SequencerPtr.Get()))
+	{
+		return AnimLayers->AddAnimLayerFromSelection(SequencerPtr.Get());
+	}
+	else
+	{
+		UE_LOG(LogControlRig, Error, TEXT("AddAnimLayerFromSelection: No Anim Layers on Level Sequence"));
+	}
+	return INDEX_NONE;
+
+}
+
+bool UControlRigSequencerEditorLibrary::MergeAnimLayers(const TArray<int32>& Indices)
+{
+	TWeakPtr<ISequencer> WeakSequencer = GetSequencerFromAsset();
+	if (WeakSequencer.IsValid() == false)
+	{
+		UE_LOG(LogControlRig, Error, TEXT("MergeAnimLayers: Need open Sequencer"));
+		return false;
+	}
+	TSharedPtr<ISequencer> SequencerPtr = WeakSequencer.Pin();
+	if (UAnimLayers* AnimLayers = UAnimLayers::GetAnimLayers(SequencerPtr.Get()))
+	{
+		return AnimLayers->MergeAnimLayers(SequencerPtr.Get(), Indices, nullptr);
+	}
+	else
+	{
+		UE_LOG(LogControlRig, Error, TEXT("MergeAnimLayers: No Anim Layers on Level Sequence"));
+	}
+	return false;
+}
+
+TArray<UAnimLayer*> UControlRigSequencerEditorLibrary::GetAnimLayers()
+{
+	TArray<UAnimLayer*> AnimLayersArray;
+	TWeakPtr<ISequencer> WeakSequencer = GetSequencerFromAsset();
+	if (WeakSequencer.IsValid() == false)
+	{
+		UE_LOG(LogControlRig, Error, TEXT("GetAnimLayers: Need open Sequencer"));
+		return AnimLayersArray;
+	}
+	TSharedPtr<ISequencer> SequencerPtr = WeakSequencer.Pin();
+	if (UAnimLayers* AnimLayers = UAnimLayers::GetAnimLayers(SequencerPtr.Get()))
+	{
+		return AnimLayers->AnimLayers;
+	}
+	else
+	{
+		UE_LOG(LogControlRig, Error, TEXT("GetAnimLayers: No Anim Layers on Level Sequence"));
+	}
+	return AnimLayersArray;
+}
+
+int32 UControlRigSequencerEditorLibrary::GetAnimLayerIndex(UAnimLayer* AnimLayer)
+{
+	TWeakPtr<ISequencer> WeakSequencer = GetSequencerFromAsset();
+	if (WeakSequencer.IsValid() == false)
+	{
+		UE_LOG(LogControlRig, Error, TEXT("GetAnimLayerIndex: Need open Sequencer"));
+		return INDEX_NONE;
+	}
+	TSharedPtr<ISequencer> SequencerPtr = WeakSequencer.Pin();
+	if (UAnimLayers* AnimLayers = UAnimLayers::GetAnimLayers(SequencerPtr.Get()))
+	{
+		return AnimLayers->GetAnimLayerIndex(AnimLayer);
+	}
+	else
+	{
+		UE_LOG(LogControlRig, Error, TEXT("GetAnimLayerIndex: No Anim Layers on Level Sequence"));
+	}
+	return INDEX_NONE;
+}
+
 #undef LOCTEXT_NAMESPACE
 
diff --git a/Engine/Plugins/Animation/ControlRig/Source/ControlRigEditor/Private/EditMode/AnimDetailsProxy.cpp b/Engine/Plugins/Animation/ControlRig/Source/ControlRigEditor/Private/EditMode/AnimDetailsProxy.cpp
--- a/Engine/Plugins/Animation/ControlRig/Source/ControlRigEditor/Private/EditMode/AnimDetailsProxy.cpp
+++ b/Engine/Plugins/Animation/ControlRig/Source/ControlRigEditor/Private/EditMode/AnimDetailsProxy.cpp
@@ -133,7 +133,7 @@ static EAnimDetailSelectionState CachePropertySelection(TWeakPtr<FCurveEditor>&
 									}
 									if (UMovieSceneControlRigParameterTrack* Track = Section->GetTypedOuter< UMovieSceneControlRigParameterTrack>())
 									{
-										if (Track->GetSectionToKey() != Section)
+										if (Track->GetSectionToKey(ControlElement->GetFName()) != Section)
 										{
 											continue;
 										}
@@ -3145,21 +3145,29 @@ UControlRigControlsProxy* UControlRigDetailPanelControlProxies::AddProxy(UObject
 				SequencerOnlyProxies.Add(InObject, NewProxies);
 			}
 		}
-		if (Proxy)
+		if (!Proxy)
+		{
+			return Proxy;
+		}
+		Proxy->Type = Type;
+		Proxy->OwnerObject = InObject;
+		Proxy->OwnerBindingAndTrack.WeakTrack = Track;
+		Proxy->OwnerBindingAndTrack.Binding = Binding;
+		Proxy->AddSequencerProxyItem(InObject, Proxy->OwnerBindingAndTrack.WeakTrack, Binding);
+		if (Type == ERigControlType::Transform || Type == ERigControlType::TransformNoScale || Type == ERigControlType::EulerTransform)
 		{
-			Proxy->Type = Type;
-			Proxy->OwnerObject = InObject;
-			Proxy->OwnerBindingAndTrack.WeakTrack = Track;
-			Proxy->OwnerBindingAndTrack.Binding = Binding;
-			Proxy->AddSequencerProxyItem(InObject, Proxy->OwnerBindingAndTrack.WeakTrack, Binding);
 			Proxy->bIsIndividual = false;
-			Proxy->SetFlags(RF_Transactional);
-			Proxy->Modify();
-			UWorld* World = GCurrentLevelEditingViewportClient ? GCurrentLevelEditingViewportClient->GetWorld() : nullptr;
-			const FConstraintsManagerController& Controller = FConstraintsManagerController::Get(World);
-			Controller.EvaluateAllConstraints();
-			Proxy->ValueChanged();
 		}
+		else
+		{
+			Proxy->bIsIndividual = true;
+		}			
+		Proxy->SetFlags(RF_Transactional);
+		Proxy->Modify();
+		UWorld* World = GCurrentLevelEditingViewportClient ? GCurrentLevelEditingViewportClient->GetWorld() : nullptr;
+		const FConstraintsManagerController& Controller = FConstraintsManagerController::Get(World);
+		Controller.EvaluateAllConstraints();
+		Proxy->ValueChanged();
 	}
 	return Proxy;
 }
@@ -3699,7 +3707,7 @@ bool UControlRigDetailPanelControlProxies::SelectPropertyInternal(UControlRigCon
 									{
 										continue;
 									}
-									UMovieSceneSection* SectionToKey = Track->GetSectionToKey() ? Track->GetSectionToKey() : Track->GetAllSections()[0];
+									UMovieSceneSection* SectionToKey = Track->GetSectionToKey(ControlElement->GetFName()) ? Track->GetSectionToKey(ControlElement->GetFName()) : Track->GetAllSections()[0];
 									if (TViewModelPtr<FChannelGroupOutlinerModel> ChannelModel = CastViewModel<FChannelGroupOutlinerModel>(OutlinerExtenstionIt.GetCurrentItem()))
 									{
 										if (ChannelModel->GetChannel(SectionToKey) == nullptr) //if not section to key we also don't select it.
diff --git a/Engine/Plugins/Animation/ControlRig/Source/ControlRigEditor/Private/EditMode/ControlRigEditMode.cpp b/Engine/Plugins/Animation/ControlRig/Source/ControlRigEditor/Private/EditMode/ControlRigEditMode.cpp
--- a/Engine/Plugins/Animation/ControlRig/Source/ControlRigEditor/Private/EditMode/ControlRigEditMode.cpp
+++ b/Engine/Plugins/Animation/ControlRig/Source/ControlRigEditor/Private/EditMode/ControlRigEditMode.cpp
@@ -81,6 +81,7 @@
 #include "Editor/Sequencer/Private/SSequencer.h"
 #include "Slate/SceneViewport.h"
 #include "Tools/BakingHelper.h"
+#include "Sequencer/AnimLayers/AnimLayers.h"
 
 #include UE_INLINE_GENERATED_CPP_BY_NAME(ControlRigEditMode)
 
@@ -811,17 +812,15 @@ TSet<FName> FControlRigEditMode::GetActiveControlsFromSequencer(UControlRig* Con
 				{
 					if (ControlRigParameterTrack->GetControlRig() == ControlRig)
 					{
-						UMovieSceneControlRigParameterSection* ActiveSection = Cast<UMovieSceneControlRigParameterSection>(ControlRigParameterTrack->GetSectionToKey());
-						if (ActiveSection)
+						TArray<FRigControlElement*> Controls;
+						ControlRig->GetControlsInOrder(Controls);
+						int Index = 0;
+						for (FRigControlElement* ControlElement : Controls)
 						{
-							TArray<FRigControlElement*> Controls;
-							ControlRig->GetControlsInOrder(Controls);
-							TArray<bool> Mask = ActiveSection->GetControlsMask();
-
-							TArray<FName> Names;
-							int Index = 0;
-							for (FRigControlElement* ControlElement : Controls)
+							UMovieSceneControlRigParameterSection* ActiveSection = Cast<UMovieSceneControlRigParameterSection>(ControlRigParameterTrack->GetSectionToKey(ControlElement->GetFName()));
+							if (ActiveSection)
 							{
+								TArray<bool> Mask = ActiveSection->GetControlsMask();
 								if (Mask[Index])
 								{
 									ActiveControls.Add(ControlElement->GetFName());
@@ -5786,6 +5785,7 @@ bool FDetailKeyFrameCacheAndHandler::IsPropertyKeyable(const UClass* InObjectCla
 			}
 		}
 	}
+	
 
 	if (InObjectClass
 		&& InObjectClass->IsChildOf(UAnimDetailControlsProxyTransform::StaticClass())
@@ -5801,6 +5801,11 @@ bool FDetailKeyFrameCacheAndHandler::IsPropertyKeyable(const UClass* InObjectCla
 		return true;
 	}
 
+	if (InObjectClass && InObjectClass->IsChildOf(UAnimLayer::StaticClass()))
+	{
+		return true;
+	}
+
 	if ((InObjectClass && InObjectClass->IsChildOf(UAnimDetailControlsProxyTransform::StaticClass()) && InPropertyHandle.GetProperty())
 		&& (InPropertyHandle.GetProperty()->GetFName() == GET_MEMBER_NAME_CHECKED(UAnimDetailControlsProxyTransform, Location) ||
 		InPropertyHandle.GetProperty()->GetFName() == GET_MEMBER_NAME_CHECKED(UAnimDetailControlsProxyTransform, Rotation) ||
@@ -5904,11 +5909,14 @@ void FDetailKeyFrameCacheAndHandler::OnKeyPropertyClicked(const IPropertyHandle&
 	KeyedPropertyHandle.GetOuterObjects(Objects);
 	for (UObject* Object : Objects)
 	{
-		UControlRigControlsProxy* Proxy = Cast< UControlRigControlsProxy>(Object);
-		if (Proxy)
+		if (UControlRigControlsProxy* Proxy = Cast< UControlRigControlsProxy>(Object))
 		{
 			Proxy->SetKey(SequencerPtr, KeyedPropertyHandle);
 		}
+		else if (UAnimLayer* AnimLayer = (Object->GetTypedOuter<UAnimLayer>()))
+		{
+			AnimLayer->SetKey(SequencerPtr, KeyedPropertyHandle);
+		}
 	}
 }
 
@@ -5952,6 +5960,10 @@ EPropertyKeyedStatus FDetailKeyFrameCacheAndHandler::GetPropertyKeyedStatus(cons
 		{
 			KeyedStatus = Proxy->GetPropertyKeyedStatus(SequencerPtr,PropertyHandle);
 		}
+		else if (UAnimLayer* AnimLayer = (Object->GetTypedOuter<UAnimLayer>()))
+		{
+			KeyedStatus = AnimLayer->GetPropertyKeyedStatus(SequencerPtr, PropertyHandle);
+		}
 		//else check to see if it's in sequencer
 	}
 	CachedPropertyKeyedStatusMap.Add(&PropertyHandle, KeyedStatus);
diff --git a/Engine/Plugins/Animation/ControlRig/Source/ControlRigEditor/Private/EditMode/ControlRigEditModeToolkit.cpp b/Engine/Plugins/Animation/ControlRig/Source/ControlRigEditor/Private/EditMode/ControlRigEditModeToolkit.cpp
--- a/Engine/Plugins/Animation/ControlRig/Source/ControlRigEditor/Private/EditMode/ControlRigEditModeToolkit.cpp
+++ b/Engine/Plugins/Animation/ControlRig/Source/ControlRigEditor/Private/EditMode/ControlRigEditModeToolkit.cpp
@@ -24,6 +24,7 @@
 #include "EditMode/SControlRigSpacePicker.h"
 #include "LevelEditor.h"
 #include "SLevelViewport.h"
+#include "Sequencer/AnimLayers/SAnimLayers.h"
 
 #define LOCTEXT_NAMESPACE "FControlRigEditModeToolkit"
 
@@ -34,6 +35,7 @@ namespace
 }
 
 bool FControlRigEditModeToolkit::bMotionTrailsTabOpen = false;
+bool FControlRigEditModeToolkit::bAnimLayerTabOpen = false;
 bool FControlRigEditModeToolkit::bPoseTabOpen = false;
 bool FControlRigEditModeToolkit::bSnapperTabOpen = false;
 bool FControlRigEditModeToolkit::bTweenOpen = false;
@@ -41,6 +43,7 @@ bool FControlRigEditModeToolkit::bTweenOpen = false;
 const FName FControlRigEditModeToolkit::PoseTabName = FName(TEXT("PoseTab"));
 const FName FControlRigEditModeToolkit::MotionTrailTabName = FName(TEXT("MotionTrailTab"));
 const FName FControlRigEditModeToolkit::SnapperTabName = FName(TEXT("SnapperTab"));
+const FName FControlRigEditModeToolkit::AnimLayerTabName = FName(TEXT("AnimLayerTab"));
 const FName FControlRigEditModeToolkit::TweenOverlayName = FName(TEXT("TweenOverlay"));
 const FName FControlRigEditModeToolkit::OutlinerTabName = FName(TEXT("ControlRigOutlinerTab"));
 const FName FControlRigEditModeToolkit::DetailsTabName = FName(TEXT("ControlRigDetailsTab"));
@@ -115,6 +118,11 @@ void FControlRigEditModeToolkit::TryInvokeToolkitUI(const FName InName)
 		FTabId TabID(MotionTrailTabName);
 		ModeUILayerPtr->GetTabManager()->TryInvokeTab(TabID, false /*bIsActive*/);
 	}
+	else if (InName == AnimLayerTabName)
+	{
+		FTabId TabID(AnimLayerTabName);
+		ModeUILayerPtr->GetTabManager()->TryInvokeTab(TabID, false /*bIsActive*/);
+	}
 	else if (InName == PoseTabName)
 	{
 		FTabId TabID(PoseTabName);
@@ -197,6 +205,14 @@ TSharedRef<SDockTab> SpawnMotionTrailTab(const FSpawnTabArgs& Args)
 		];
 }
 
+TSharedRef<SDockTab> SpawnAnimLayerTab(const FSpawnTabArgs& Args, FControlRigEditMode* InEditorMode)
+{
+	return SNew(SDockTab)
+		[
+			SNew(SAnimLayers, *InEditorMode)
+		];
+}
+
 TSharedRef<SDockTab> SpawnOutlinerTab(const FSpawnTabArgs& Args, FControlRigEditMode* InEditorMode)
 {
 	return SNew(SDockTab)
@@ -410,6 +426,16 @@ void FControlRigEditModeToolkit::RequestModeUITabs()
 			.SetIcon(FSlateIcon(TEXT("ControlRigEditorStyle"), TEXT("ControlRig.EditableMotionTrails")));
 		ModeUILayerPtr->GetTabManager()->RegisterDefaultTabWindowSize(MotionTrailTabName, FVector2D(425, 575));
 
+		ModeUILayerPtr->GetTabManager()->UnregisterTabSpawner(AnimLayerTabName);
+		ModeUILayerPtr->GetTabManager()->RegisterTabSpawner(AnimLayerTabName, FOnSpawnTab::CreateStatic(&SpawnAnimLayerTab,&EditMode))
+			.SetDisplayName(LOCTEXT("AnimLayerTab", "Anim Layers"))
+			.SetTooltipText(LOCTEXT("AnimationLayerTabTooltip", "Animation layers"))
+			.SetGroup(MenuGroup)
+			.SetIcon(FSlateIcon(TEXT("ControlRigEditorStyle"), TEXT("ControlRig.AnimLayers")));
+		ModeUILayerPtr->GetTabManager()->RegisterDefaultTabWindowSize(AnimLayerTabName, FVector2D(425, 200));
+
+
+
 		ModeUILayer.Pin()->ToolkitHostShutdownUI().BindSP(this, &FControlRigEditModeToolkit::UnregisterAndRemoveFloatingTabs);
 
 
@@ -434,6 +460,10 @@ void FControlRigEditModeToolkit::InvokeUI()
 		{
 			TryInvokeToolkitUI(MotionTrailTabName);
 		}
+		if (bAnimLayerTabOpen)
+		{
+			TryInvokeToolkitUI(AnimLayerTabName);
+		}
 		if (bSnapperTabOpen)
 		{
 			TryInvokeToolkitUI(SnapperTabName);
@@ -474,6 +504,18 @@ void FControlRigEditModeToolkit::UnregisterAndRemoveFloatingTabs()
 			}
 			ModeUILayerPtr->GetTabManager()->UnregisterTabSpawner(MotionTrailTabName);
 
+			TSharedPtr<SDockTab> AnimLayerTab = ModeUILayerPtr->GetTabManager()->FindExistingLiveTab(FTabId(AnimLayerTabName));
+			if (AnimLayerTab)
+			{
+				bAnimLayerTabOpen = true;
+				AnimLayerTab->RequestCloseTab();
+			}
+			else
+			{
+				bAnimLayerTabOpen = false;
+			}
+			ModeUILayerPtr->GetTabManager()->UnregisterTabSpawner(AnimLayerTabName);
+
 			TSharedPtr<SDockTab> SnapperTab = ModeUILayerPtr->GetTabManager()->FindExistingLiveTab(FTabId(SnapperTabName));
 			if (SnapperTab)
 			{
diff --git a/Engine/Plugins/Animation/ControlRig/Source/ControlRigEditor/Private/EditMode/ControlRigEditModeToolkit.h b/Engine/Plugins/Animation/ControlRig/Source/ControlRigEditor/Private/EditMode/ControlRigEditModeToolkit.h
--- a/Engine/Plugins/Animation/ControlRig/Source/ControlRigEditor/Private/EditMode/ControlRigEditModeToolkit.h
+++ b/Engine/Plugins/Animation/ControlRig/Source/ControlRigEditor/Private/EditMode/ControlRigEditModeToolkit.h
@@ -65,6 +65,7 @@ public:
 public:
 	static const FName PoseTabName;
 	static const FName MotionTrailTabName;
+	static const FName AnimLayerTabName;
 	static const FName TweenOverlayName;
 	static const FName SnapperTabName;
 	static const FName DetailsTabName;
@@ -93,6 +94,7 @@ protected:
 	void UnregisterAndRemoveFloatingTabs();
 
 	static bool bMotionTrailsTabOpen;
+	static bool bAnimLayerTabOpen;
 	static bool bPoseTabOpen;
 	static bool bSnapperTabOpen;
 	static bool bTweenOpen;
diff --git a/Engine/Plugins/Animation/ControlRig/Source/ControlRigEditor/Private/EditMode/SControlRigEditModeTools.cpp b/Engine/Plugins/Animation/ControlRig/Source/ControlRigEditor/Private/EditMode/SControlRigEditModeTools.cpp
--- a/Engine/Plugins/Animation/ControlRig/Source/ControlRigEditor/Private/EditMode/SControlRigEditModeTools.cpp
+++ b/Engine/Plugins/Animation/ControlRig/Source/ControlRigEditor/Private/EditMode/SControlRigEditModeTools.cpp
@@ -1302,6 +1302,16 @@ void SControlRigEditModeTools::CustomizeToolBarPalette(FToolBarBuilder& ToolBarB
 		EUserInterfaceActionType::Button
 	);
 
+	// Anim Layer
+	ToolBarBuilder.AddToolBarButton(
+		FExecuteAction::CreateRaw(OwningToolkit.Pin().Get(), &FControlRigEditModeToolkit::TryInvokeToolkitUI, FControlRigEditModeToolkit::AnimLayerTabName),
+		NAME_None,
+		LOCTEXT("AnimLayers", "Anim Layers"),
+		LOCTEXT("AnimLayersTooltip", "Display anim layers"),
+		FSlateIcon(TEXT("ControlRigEditorStyle"), TEXT("ControlRig.AnimLayers")),
+		EUserInterfaceActionType::Button
+	);
+
 	//Pivot
 	ToolBarBuilder.AddToolBarButton(
 		FUIAction(
diff --git a/Engine/Plugins/Animation/ControlRig/Source/ControlRigEditor/Private/Sequencer/AnimLayerSequencerFilter.cpp b/Engine/Plugins/Animation/ControlRig/Source/ControlRigEditor/Private/Sequencer/AnimLayerSequencerFilter.cpp
new file mode 100644
--- /dev/null
+++ b/Engine/Plugins/Animation/ControlRig/Source/ControlRigEditor/Private/Sequencer/AnimLayerSequencerFilter.cpp
@@ -0,0 +1,144 @@
+// Copyright Epic Games, Inc. All Rights Reserved.
+
+#include "Sequencer/AnimLayerSequencerFilter.h"
+#include "AnimLayers/AnimLayers.h"
+#include "Filters/ISequencerTrackFilters.h"
+#include "Filters/SequencerTrackFilterBase.h"
+#include "Styling/AppStyle.h"
+#include "Styling/SlateIconFinder.h"
+#include "Framework/Commands/Commands.h"
+#include "Framework/Commands/UICommandInfo.h"
+#include "MVVM/Extensions/IOutlinerExtension.h"
+#include "MVVM/ViewModels/SectionModel.h"
+#include "MVVM/Extensions/ITrackExtension.h"
+#include "MVVM/Extensions/ITrackLaneExtension.h"
+#include "ISequencer.h"
+#include "LevelSequence.h"
+#include "MVVM/ViewModels/ChannelModel.h"
+
+
+#include UE_INLINE_GENERATED_CPP_BY_NAME(AnimLayerSequencerFilter)
+
+#define LOCTEXT_NAMESPACE "UnimLayerSequencerFilter"
+
+using namespace UE::Sequencer;
+//////////////////////////////////////////////////////////////////////////
+//
+
+class FSequencerTrackFilter_AnimLayerSequencerFilterCommands
+	: public TCommands<FSequencerTrackFilter_AnimLayerSequencerFilterCommands>
+{
+public:
+
+	FSequencerTrackFilter_AnimLayerSequencerFilterCommands()
+		: TCommands<FSequencerTrackFilter_AnimLayerSequencerFilterCommands>
+		(
+			"FSequencerTrackFilter_AnimLayerSequencerFilter",
+			NSLOCTEXT("Contexts", "FSequencerTrackFilter_AnimLayerSequencer", "FSequencerTrackFilter_AnimLayerSequencer"),
+			NAME_None,
+			FAppStyle::GetAppStyleSetName() // Icon Style Set
+			)
+	{ }
+
+	/** Toggle the control rig controls filter */
+	TSharedPtr< FUICommandInfo > ToggleSelectedAnimLayer;
+
+	/** Initialize commands */
+	virtual void RegisterCommands() override
+	{
+		UI_COMMAND(ToggleSelectedAnimLayer, "Selected Anim layer", "Toggle the filter for elected Anim Layer", EUserInterfaceActionType::ToggleButton, FInputChord(EKeys::F8));
+	}
+};
+
+class FSequencerTrackFilter_ToggleSelectedAnimLayer : public FSequencerTrackFilter
+{
+public:
+
+	FSequencerTrackFilter_ToggleSelectedAnimLayer(ISequencerTrackFilters& InOutFilterInterface, TSharedPtr<FFilterCategory> InCategory = nullptr)
+		: FSequencerTrackFilter(InOutFilterInterface, MoveTemp(InCategory))
+		, BindingCount(0)
+	{
+		FSequencerTrackFilter_AnimLayerSequencerFilterCommands::Register();
+	}
+
+	~FSequencerTrackFilter_ToggleSelectedAnimLayer()
+	{
+		BindingCount--;
+
+		if (BindingCount < 1)
+		{
+			FSequencerTrackFilter_AnimLayerSequencerFilterCommands::Unregister();
+		}
+	}
+
+	virtual FString GetName() const override { return TEXT("SelectedAnimLayersFilter"); }
+	virtual FText GetDisplayName() const override { return LOCTEXT("SequenceTrackFilter_SelectedAnimLayers", "Selected Anim Layers"); }
+	virtual FSlateIcon GetIcon() const { return FSlateIconFinder::FindIconForClass(UControlRig::StaticClass()); } //UanimLayer
+
+	virtual bool PassesFilter(FSequencerTrackFilterType InItem) const override
+	{
+
+		FSequencerFilterData& FilterData = FilterInterface.GetFilterData();
+
+		const TViewModelPtr<IOutlinerExtension> OutlinerExtension = InItem.AsModel()->FindAncestorOfType<IOutlinerExtension>();
+		if (!OutlinerExtension.IsValid())
+		{
+			return false;
+		}
+		UMovieSceneTrack* Track = nullptr;
+
+		if (const TViewModelPtr<ITrackExtension> AncestorTrackModel = InItem->FindAncestorOfType<ITrackExtension>(true))
+		{
+			Track = AncestorTrackModel->GetTrack();
+		}
+
+		if (Track)
+		{
+			ULevelSequence* LevelSequence = Track->GetTypedOuter<ULevelSequence>();
+			if (LevelSequence)
+			{
+				if (const UAnimLayers* AnimLayers = UAnimLayers::GetAnimLayers(LevelSequence))
+				{
+					TArray<UMovieSceneSection*> Sections = AnimLayers->GetSelectedLayerSections();
+					if (TViewModelPtr<FChannelGroupOutlinerModel> ChannelModel = CastViewModel<FChannelGroupOutlinerModel>(InItem))
+					{
+						for (UMovieSceneSection* Section : Sections)
+						{
+							if (ChannelModel->GetChannel(Section))
+							{
+								return true;
+							}
+						}
+					}
+				}
+			}
+		}
+		return false;
+	}
+	virtual FText GetDefaultToolTipText() const override
+	{
+		return LOCTEXT("SequenceTrackFilter_SelectedAnimLayersTip", "Show Selected Anim Layers");
+	}
+
+	virtual TSharedPtr<FUICommandInfo> GetToggleCommand() const override
+	{
+		return FSequencerTrackFilter_AnimLayerSequencerFilterCommands::Get().ToggleSelectedAnimLayer;
+	}
+
+
+private:
+	mutable uint32 BindingCount;
+};
+
+//////////////////////////////////////////////////////////////////////////
+//
+
+
+void UAnimLayerSequencerFilter::AddTrackFilterExtensions(ISequencerTrackFilters& InOutFilterInterface
+	, const TSharedRef<FFilterCategory>& InPreferredCategory
+	, TArray<TSharedRef<FSequencerTrackFilter>>& InOutFilterList) const
+{
+	InOutFilterList.Add(MakeShared<FSequencerTrackFilter_ToggleSelectedAnimLayer>(InOutFilterInterface, InPreferredCategory));
+}
+#undef LOCTEXT_NAMESPACE
+
diff --git a/Engine/Plugins/Animation/ControlRig/Source/ControlRigEditor/Private/Sequencer/AnimLayerSequencerFilter.h b/Engine/Plugins/Animation/ControlRig/Source/ControlRigEditor/Private/Sequencer/AnimLayerSequencerFilter.h
new file mode 100644
--- /dev/null
+++ b/Engine/Plugins/Animation/ControlRig/Source/ControlRigEditor/Private/Sequencer/AnimLayerSequencerFilter.h
@@ -0,0 +1,19 @@
+// Copyright Epic Games, Inc. All Rights Reserved.
+
+#pragma once
+
+#include "Filters/SequencerTrackFilterExtension.h"
+#include "AnimLayerSequencerFilter.generated.h"
+
+UCLASS()
+class UAnimLayerSequencerFilter : public USequencerTrackFilterExtension
+{
+public:
+	GENERATED_BODY()
+
+	//~ Begin USequencerTrackFilterExtension
+	virtual void AddTrackFilterExtensions(ISequencerTrackFilters& InOutFilterInterface
+		, const TSharedRef<FFilterCategory>& InPreferredCategory
+		, TArray<TSharedRef<FSequencerTrackFilter>>& InOutFilterList) const override;
+	//~ End USequencerTrackFilterExtension
+};
diff --git a/Engine/Plugins/Animation/ControlRig/Source/ControlRigEditor/Private/Sequencer/AnimLayers/AnimLayers.cpp b/Engine/Plugins/Animation/ControlRig/Source/ControlRigEditor/Private/Sequencer/AnimLayers/AnimLayers.cpp
new file mode 100644
--- /dev/null
+++ b/Engine/Plugins/Animation/ControlRig/Source/ControlRigEditor/Private/Sequencer/AnimLayers/AnimLayers.cpp
@@ -0,0 +1,2316 @@
+// Copyright Epic Games, Inc. All Rights Reserved.
+#include "AnimLayers.h"
+#include "ControlRig.h"
+#include "ISequencer.h"
+#include "MVVM/Selection/Selection.h"
+#include "MVVM/ViewModels/SequencerEditorViewModel.h"
+#include "MVVM/ViewModels/SequencerOutlinerViewModel.h"
+#include "MVVM/ViewModels/SequencerTrackAreaViewModel.h"
+#include "MVVM/ViewModels/SequencerEditorViewModel.h"
+#include "MVVM/ViewModels/ObjectBindingModel.h"
+#include "LevelSequence.h"
+#include "ControlRigSequencerEditorLibrary.h"
+#include "Sequencer/MovieSceneControlRigParameterTrack.h"
+#include "Sequencer/MovieSceneControlRigParameterSection.h"
+#include "Tracks/MovieScene3DTransformTrack.h"
+#include "Tracks/MovieScenePropertyTrack.h"
+#include "ILevelSequenceEditorToolkit.h"
+#include "LevelSequencePlayer.h"
+#include "EditMode/ControlRigEditMode.h"
+#include "LevelSequenceEditorBlueprintLibrary.h"
+#include "MovieScene.h"
+#include "Editor.h"
+#include "EditorModeManager.h"
+#include "Editor/EditorEngine.h"
+#include "Editor/UnrealEdEngine.h"
+#include "Engine/Selection.h"
+#include "ScopedTransaction.h"
+#include "EditMode/ControlRigEditMode.h"
+#include "EditorModeManager.h"
+#include "Channels/MovieSceneFloatChannel.h"
+#include "Channels/MovieSceneBoolChannel.h"
+#include "Channels/MovieSceneChannelProxy.h"
+#include "Channels/MovieSceneDoubleChannel.h"
+#include "Channels/MovieSceneIntegerChannel.h"
+#include "Channels/MovieSceneByteChannel.h"
+#include "ISequencerPropertyKeyedStatus.h"
+#include "BakingAnimationKeySettings.h"
+#include "Algo/Accumulate.h"
+
+#define LOCTEXT_NAMESPACE "AnimLayers"
+
+
+static TArray<FGuid> GetSelectedOutlinerGuids(ISequencer* SequencerPtr)
+{
+	using namespace UE::Sequencer;
+	TArray<FGuid> SelectedObjects;// = SequencerPtr->GetViewModel()->GetSelection()->GetBoundObjectsGuids(); //this is not exposed :(
+
+	for (TViewModelPtr<FObjectBindingModel> ObjectBindingNode : SequencerPtr->GetViewModel()->GetSelection()->Outliner.Filter<FObjectBindingModel>())
+	{
+		FGuid Guid = ObjectBindingNode->GetObjectGuid();
+		SelectedObjects.Add(Guid);
+	}
+	return SelectedObjects;
+}
+
+struct FControlRigAndControlsAndTrack
+{
+	UMovieSceneControlRigParameterTrack* Track;
+	UControlRig* ControlRig;
+	TArray<FName> Controls;
+};
+
+struct FObjectAndTrack
+{
+	UMovieScenePropertyTrack* Track;
+	UObject* BoundObject;
+};
+
+static void GetSelectedControlRigsAndBoundObjects(ISequencer* SequencerPtr, TArray<FControlRigAndControlsAndTrack>& OutSelectedCRs, TArray<FObjectAndTrack>& OutBoundObjects)
+{
+	if (SequencerPtr == nullptr || SequencerPtr->GetViewModel() == nullptr)
+	{
+		return;
+	}
+	ULevelSequence* LevelSequence = Cast<ULevelSequence>(SequencerPtr->GetFocusedMovieSceneSequence());
+	if (LevelSequence == nullptr)
+	{
+		return;
+	}
+
+	using namespace UE::Sequencer;
+	TArray<FGuid> SelectedObjects = GetSelectedOutlinerGuids(SequencerPtr);;// = SequencerPtr->GetViewModel()->GetSelection()->GetBoundObjectsGuids();
+
+	UMovieScene* MovieScene = LevelSequence->GetMovieScene();
+	if (MovieScene)
+	{
+		const TArray<FMovieSceneBinding>& Bindings = MovieScene->GetBindings();
+		for (const FMovieSceneBinding& Binding : Bindings)
+		{
+			bool bHaveControlRig = false;
+			TArray<UMovieSceneTrack*> CRTracks = MovieScene->FindTracks(UMovieSceneControlRigParameterTrack::StaticClass(), Binding.GetObjectGuid(), NAME_None);
+			for (UMovieSceneTrack* AnyOleTrack : CRTracks)
+			{
+				if (UMovieSceneControlRigParameterTrack* Track = Cast<UMovieSceneControlRigParameterTrack>(AnyOleTrack))
+				{
+					if (UControlRig* ControlRig = Track->GetControlRig())
+					{
+						FControlRigAndControlsAndTrack CRControls;
+						CRControls.ControlRig = ControlRig;
+						CRControls.Track = Track;
+						CRControls.Controls = ControlRig->CurrentControlSelection();
+						if (CRControls.Controls.Num() > 0)
+						{
+							bHaveControlRig = true;
+							OutSelectedCRs.Add(CRControls);
+						}
+					}
+				}
+			}
+			//if we have control rig controls don't add the base skel mesh for now
+			if (bHaveControlRig == false && SelectedObjects.Contains(Binding.GetObjectGuid()))
+			{
+				TArray<UMovieSceneTrack*> Tracks = MovieScene->FindTracks(UMovieScenePropertyTrack::StaticClass(), Binding.GetObjectGuid(), NAME_None);
+				if (Tracks.Num() < 1)
+				{
+					continue;
+				}
+				const FMovieSceneSequenceID& SequenceID = SequencerPtr->State.FindSequenceId(LevelSequence);
+
+				for (UMovieSceneTrack* AnyOleTrack : Tracks)
+				{
+					if (UMovieScenePropertyTrack* Track = Cast<UMovieScenePropertyTrack>(AnyOleTrack))
+					{
+						const FMovieSceneBlendTypeField SupportedBlendTypes = Track->GetSupportedBlendTypes();
+						if (SupportedBlendTypes.Num() == 0)
+						{
+							continue;
+						}
+						const TArrayView<TWeakObjectPtr<UObject>>& BoundObjects = SequencerPtr->FindBoundObjects(Binding.GetObjectGuid(), SequenceID);
+						for (const TWeakObjectPtr<UObject> CurrentBoundObject : BoundObjects)
+						{
+							if (UObject* BoundObject = CurrentBoundObject.Get())
+							{
+								FObjectAndTrack ObjectAndTrack;
+								ObjectAndTrack.BoundObject = BoundObject;
+								ObjectAndTrack.Track = Track;
+								OutBoundObjects.Add(ObjectAndTrack);
+							}
+						}
+					}
+				}
+			}
+		}
+	}
+}
+
+bool FAnimLayerSelectionSet::MergeWithAnotherSelection(const FAnimLayerSelectionSet& Selection)
+{
+	if (BoundObject.IsValid() && (BoundObject.Get() == Selection.BoundObject.Get()))
+	{
+		for (const TPair<FName, FAnimLayerPropertyAndChannels>& Incoming : Selection.Names)
+		{
+			FAnimLayerPropertyAndChannels& Channels = Names.FindOrAdd(Incoming.Key);
+			Channels.Channels |= (Incoming.Value.Channels);
+		}
+		return true;
+	}
+	return false;
+}
+
+FAnimLayerState::FAnimLayerState() : bKeyed(ECheckBoxState::Unchecked), bActive(true),  bLock(false), Weight(1.0), Type((uint32)EAnimLayerType::Base)
+{
+	Name = FText(LOCTEXT("BaseLayer", "Base Layer"));
+}
+
+FText FAnimLayerState::AnimLayerTypeToText() const
+{
+	if ((EAnimLayerType)Type == EAnimLayerType::Additive)
+	{
+		return FText(LOCTEXT("Additive", "Additive"));
+	}
+	else if ((EAnimLayerType)Type == EAnimLayerType::Override)
+	{
+		return FText(LOCTEXT("Override", "Override"));
+	}
+	return FText(LOCTEXT("Base", "Base"));
+}
+
+
+UAnimLayer::UAnimLayer(class FObjectInitializer const& ObjectInitializer)
+	: Super(ObjectInitializer)
+{
+	WeightProxy = NewObject<UAnimLayerWeightProxy>(this, TEXT("Weight"), RF_Transactional);
+}
+
+void UAnimLayer::SetKey(TSharedPtr<ISequencer>& Sequencer, const IPropertyHandle& KeyedPropertyHandle)
+{
+	FScopedTransaction PropertyChangedTransaction(LOCTEXT("KeyAnimLayerWeight", "Key Anim Layer Weight"), !GIsTransacting);
+	bool bAnythingKeyed = false;
+
+	for (const TPair<TWeakObjectPtr<UObject>,FAnimLayerItem>& Pair : AnimLayerItems)
+	{
+		if (Pair.Key != nullptr)
+		{
+			for (const FAnimLayerSectionItem& SectionItem : Pair.Value.SectionItems)
+			{
+				if (SectionItem.Section.IsValid() && SectionItem.Section.Get()->TryModify())
+				{
+					FMovieSceneFloatChannel* FloatChannel = nullptr;
+					if (UMovieSceneControlRigParameterSection* CRSection = Cast<UMovieSceneControlRigParameterSection>(SectionItem.Section.Get()))
+					{
+						FloatChannel = &CRSection->Weight;
+					}
+					else if (UMovieScene3DTransformSection* LayerSection = Cast<UMovieScene3DTransformSection>(SectionItem.Section.Get()))
+					{
+						FloatChannel = LayerSection->GetWeightChannel();
+					}
+					if (FloatChannel)
+					{
+						//don't add key if there!
+						FFrameTime LocalTime = Sequencer->GetLocalTime().Time;
+						const TRange<FFrameNumber> FrameRange = TRange<FFrameNumber>(LocalTime.FrameNumber);
+						TArray<FFrameNumber> KeyTimes;
+						FloatChannel->GetKeys(FrameRange, &KeyTimes, nullptr);
+						if (KeyTimes.Num() == 0)
+						{
+							float Value = State.Weight;
+							FloatChannel->AddCubicKey(LocalTime.FrameNumber, Value, ERichCurveTangentMode::RCTM_SmartAuto);
+							bAnythingKeyed = true;
+						}
+					}
+				}
+			}
+		}
+	}
+
+	if(bAnythingKeyed == false)
+	{
+		PropertyChangedTransaction.Cancel();
+	}
+}
+
+EPropertyKeyedStatus UAnimLayer::GetPropertyKeyedStatus(TSharedPtr<ISequencer>& Sequencer, const IPropertyHandle& PropertyHandle)
+{
+	EPropertyKeyedStatus KeyedStatus = EPropertyKeyedStatus::NotKeyed;
+
+	if (Sequencer.IsValid() == false || Sequencer->GetFocusedMovieSceneSequence() == nullptr)
+	{
+		return KeyedStatus;
+	}
+	const TRange<FFrameNumber> FrameRange = TRange<FFrameNumber>(Sequencer->GetLocalTime().Time.FrameNumber);
+	int32 NumKeyed = 0;
+	int32 NumToCheck = 0;
+	for (const TPair<TWeakObjectPtr<UObject>,FAnimLayerItem>& Pair : AnimLayerItems)
+	{
+		if (Pair.Key != nullptr)
+		{
+			for (const FAnimLayerSectionItem& SectionItem : Pair.Value.SectionItems)
+			{
+				if (SectionItem.Section.IsValid())
+				{
+					FMovieSceneFloatChannel* FloatChannel = nullptr;
+					if (UMovieSceneControlRigParameterSection* CRSection = Cast<UMovieSceneControlRigParameterSection>(SectionItem.Section.Get()))
+					{
+						FloatChannel = &CRSection->Weight;
+					}
+					else if (UMovieScene3DTransformSection* LayerSection = Cast<UMovieScene3DTransformSection>(SectionItem.Section.Get()))
+					{
+						FloatChannel = LayerSection->GetWeightChannel();
+					}
+					if (FloatChannel)
+					{
+						++NumToCheck;
+						EPropertyKeyedStatus NewKeyedStatus = EPropertyKeyedStatus::NotKeyed;
+						if (FloatChannel->GetNumKeys() > 0)
+						{
+							TArray<FFrameNumber> KeyTimes;
+							FloatChannel->GetKeys(FrameRange, &KeyTimes, nullptr);
+							if (KeyTimes.Num() > 0)
+							{
+								++NumKeyed;
+								NewKeyedStatus = EPropertyKeyedStatus::PartiallyKeyed;
+							}
+							else
+							{
+								NewKeyedStatus = EPropertyKeyedStatus::KeyedInOtherFrame;
+							}
+						}
+						KeyedStatus = FMath::Max(KeyedStatus, NewKeyedStatus);
+					}
+				}
+			}
+		}
+	}
+	if (KeyedStatus == EPropertyKeyedStatus::PartiallyKeyed && NumToCheck == NumKeyed)
+	{
+		KeyedStatus = EPropertyKeyedStatus::KeyedInFrame;
+	}
+	return KeyedStatus;
+}
+
+ECheckBoxState UAnimLayer::GetKeyed() const
+{
+	TOptional<ECheckBoxState> CurrentVal;
+	for (const TPair<TWeakObjectPtr<UObject>, FAnimLayerItem>& Pair : AnimLayerItems)
+	{
+		if (Pair.Key != nullptr)
+		{
+			for (const FAnimLayerSectionItem& SectionItem : Pair.Value.SectionItems)
+			{
+				if (SectionItem.Section.IsValid())
+				{
+					if (UMovieSceneTrack* Track = Cast<UMovieSceneTrack>(SectionItem.Section.Get()->GetTypedOuter< UMovieSceneTrack>()))
+					{
+						if (UMovieSceneControlRigParameterTrack* ControlRigTrack = Cast<UMovieSceneControlRigParameterTrack>(Track))
+						{
+							for (const TPair<FName, FAnimLayerPropertyAndChannels>& ControlName : SectionItem.AnimLayerSet.Names)
+							{
+								if (ControlRigTrack->GetSectionToKey(ControlName.Key) == SectionItem.Section.Get())
+								{
+									if (CurrentVal.IsSet() && CurrentVal.GetValue() != ECheckBoxState::Checked)
+									{
+										CurrentVal = ECheckBoxState::Undetermined;
+									}
+									if (CurrentVal.IsSet() == false)
+									{
+										CurrentVal = ECheckBoxState::Checked;
+									}
+								}
+								else
+								{
+									if (CurrentVal.IsSet() && CurrentVal.GetValue() != ECheckBoxState::Unchecked)
+									{
+										CurrentVal = ECheckBoxState::Undetermined;
+									}
+									if (CurrentVal.IsSet() == false)
+									{
+										CurrentVal = ECheckBoxState::Unchecked;
+									}
+								}
+							}
+						}
+						else
+						{
+							if (Track->GetSectionToKey() == SectionItem.Section.Get() ||
+								(Track->GetAllSections().Num() == 1 && (Track->GetAllSections()[0] == SectionItem.Section.Get())))
+							{
+								if (CurrentVal.IsSet() && CurrentVal.GetValue() != ECheckBoxState::Checked)
+								{
+									CurrentVal = ECheckBoxState::Undetermined;
+								}
+								if (CurrentVal.IsSet() == false)
+								{
+									CurrentVal = ECheckBoxState::Checked;
+								}
+							}
+							else
+							{
+								if (CurrentVal.IsSet() && CurrentVal.GetValue() != ECheckBoxState::Unchecked)
+								{
+									CurrentVal = ECheckBoxState::Undetermined;
+								}
+								if (CurrentVal.IsSet() == false)
+								{
+									CurrentVal = ECheckBoxState::Unchecked;
+								}
+							}
+						}
+					}
+				}
+			}
+		}
+	}
+	if (CurrentVal.IsSet())
+	{
+		if (State.bKeyed != CurrentVal.GetValue() && CurrentVal.GetValue() == ECheckBoxState::Checked)
+		{
+			SetSectionToKey();
+		}
+		State.bKeyed = CurrentVal.GetValue();
+	}
+	else
+	{
+		//has no sections and is base so it's keyed, else it's not
+		State.bKeyed = State.Type == EAnimLayerType::Base ? ECheckBoxState::Checked : ECheckBoxState::Unchecked;
+	}
+	return State.bKeyed;
+}
+
+void UAnimLayer::SetSectionToKey() const
+{
+
+	for (const TPair<TWeakObjectPtr<UObject>, FAnimLayerItem>& Pair : AnimLayerItems)
+	{
+		if (Pair.Key != nullptr)
+		{
+			for (const FAnimLayerSectionItem& SectionItem : Pair.Value.SectionItems)
+			{
+				if (SectionItem.Section.IsValid())
+				{
+					if (UMovieSceneTrack* Track = Cast<UMovieSceneTrack>(SectionItem.Section.Get()->GetTypedOuter< UMovieSceneTrack>()))
+					{
+						Track->Modify();
+						if (UMovieSceneControlRigParameterTrack* ControlRigTrack = Cast<UMovieSceneControlRigParameterTrack>(Track))
+						{
+							for (const TPair<FName, FAnimLayerPropertyAndChannels>& ControlName : SectionItem.AnimLayerSet.Names)
+							{
+								ControlRigTrack->SetSectionToKey(SectionItem.Section.Get(), ControlName.Key);
+							}
+						}
+						else
+						{
+							Track->SetSectionToKey(SectionItem.Section.Get());
+						}
+					}
+				}
+			}
+		}
+	}
+}
+
+void UAnimLayer::SetKeyed()
+{
+	State.bKeyed = ECheckBoxState::Checked;
+	SetSectionToKey();
+
+}
+
+bool UAnimLayer::GetActive() const
+{
+	TOptional<bool> CurrentVal;
+	for (const TPair<TWeakObjectPtr<UObject>, FAnimLayerItem>& Pair : AnimLayerItems)
+	{
+		if (Pair.Key != nullptr)
+		{
+			for (const FAnimLayerSectionItem& SectionItem : Pair.Value.SectionItems)
+			{
+				if (SectionItem.Section.IsValid())
+				{
+					const bool bActive = SectionItem.Section.Get()->IsActive();
+					if (CurrentVal.IsSet() && CurrentVal.GetValue() != bActive)
+					{
+						SectionItem.Section.Get()->SetIsActive(CurrentVal.GetValue());
+					}
+					if (CurrentVal.IsSet() == false)
+					{
+						CurrentVal = bActive;
+					}
+				}
+			}
+		}
+	}
+	if (CurrentVal.IsSet())
+	{
+		State.bActive = CurrentVal.GetValue();
+	}
+	return State.bActive;
+}
+
+void UAnimLayer::SetActive(bool bInActive)
+{
+	const FScopedTransaction Transaction(LOCTEXT("SetActive_Transaction", "Set Active"), !GIsTransacting);
+	Modify();
+	State.bActive = bInActive;
+	for (const TPair<TWeakObjectPtr<UObject>, FAnimLayerItem>& Pair : AnimLayerItems)
+	{
+		if (Pair.Key != nullptr)
+		{
+			for (const FAnimLayerSectionItem& SectionItem : Pair.Value.SectionItems)
+			{
+				if (SectionItem.Section.IsValid())
+				{
+					SectionItem.Section->Modify();
+					SectionItem.Section.Get()->SetIsActive(State.bActive);
+				}
+			}
+		}
+	}
+}
+
+bool UAnimLayer::AddSelectedInSequencer() 
+{
+	TSharedPtr<ISequencer> SequencerPtr = UAnimLayers::GetSequencerFromAsset();
+	if (SequencerPtr.IsValid() == false)
+	{
+		return false;
+	}
+	TArray<FControlRigAndControlsAndTrack> SelectedCRs;
+	TArray<FObjectAndTrack> SelectedBoundObjects;
+	GetSelectedControlRigsAndBoundObjects(SequencerPtr.Get(), SelectedCRs, SelectedBoundObjects);
+	if (SelectedCRs.Num() <= 0 && SelectedBoundObjects.Num() <= 0)
+	{
+		return false;
+	}
+	bool bAddedSomething = false;
+	const FScopedTransaction Transaction(LOCTEXT("AddSelectedAnimLayer_Transaction", "Add Selected"),!GIsTransacting);
+	Modify();
+	for (FControlRigAndControlsAndTrack& CRControls : SelectedCRs)
+	{
+		if (FAnimLayerItem* ExistingAnimLayerItem = AnimLayerItems.Find(CRControls.ControlRig))
+		{
+			for (FAnimLayerSectionItem& SectionItem : ExistingAnimLayerItem->SectionItems)
+			{
+				if (UMovieSceneControlRigParameterSection* CRSection = Cast<UMovieSceneControlRigParameterSection>(SectionItem.Section))
+				{
+					for (FName& ControlName : CRControls.Controls)
+					{
+						if (SectionItem.AnimLayerSet.Names.Contains(ControlName) == false)
+						{
+							FAnimLayerPropertyAndChannels Channels;
+							Channels.Name = ControlName;
+							Channels.Channels = (uint32)EControlRigContextChannelToKey::AllTransform;
+							SectionItem.AnimLayerSet.Names.Add(ControlName, Channels);
+						}
+					}
+					TArray<FName> AllControls;
+					SectionItem.AnimLayerSet.Names.GenerateKeyArray(AllControls);
+					UAnimLayers::SetUpControlRigSection(CRSection, AllControls);
+					bAddedSomething = true;
+				}
+			}
+		}
+		else
+		{
+			//add new section
+			FAnimLayerItem AnimLayerItem;
+			FAnimLayerSectionItem SectionItem;
+			SectionItem.AnimLayerSet.BoundObject = CRControls.ControlRig;
+			for (FName& ControlName : CRControls.Controls)
+			{
+				FAnimLayerPropertyAndChannels Channels;
+				Channels.Name = ControlName;
+				Channels.Channels = (uint32)EControlRigContextChannelToKey::AllTransform;
+				SectionItem.AnimLayerSet.Names.Add(ControlName, Channels);
+			}
+			// Add a new section that starts and ends at the same time
+			TGuardValue<bool> GuardSetSection(CRControls.Track->bSetSectionToKeyPerControl, false);
+			if (UMovieSceneControlRigParameterSection* NewSection = Cast<UMovieSceneControlRigParameterSection>(CRControls.Track->CreateNewSection()))
+			{
+				ensureAlwaysMsgf(NewSection->HasAnyFlags(RF_Transactional), TEXT("CreateNewSection must return an instance with RF_Transactional set! (pass RF_Transactional to NewObject)"));
+				NewSection->SetFlags(RF_Transactional);
+				NewSection->SetTransformMask(FMovieSceneTransformMask{ EMovieSceneTransformChannel::All });
+				FMovieSceneFloatChannel* FloatChannel = &NewSection->Weight;
+				SectionItem.Section = NewSection;
+				AnimLayerItem.SectionItems.Add(SectionItem);
+				AnimLayerItems.Add(CRControls.ControlRig, AnimLayerItem);
+				UAnimLayers::SetUpSectionDefaults(SequencerPtr.Get(), this, CRControls.Track, NewSection,FloatChannel);
+				UAnimLayers::SetUpControlRigSection(NewSection, CRControls.Controls);
+				bAddedSomething = true;
+			}
+		}
+	}
+
+	for (FObjectAndTrack& ObjectAndTrack : SelectedBoundObjects)
+	{
+		FAnimLayerItem& AnimLayerItem = AnimLayerItems.FindOrAdd(ObjectAndTrack.BoundObject);
+		FAnimLayerSectionItem SectionItem;
+		SectionItem.AnimLayerSet.BoundObject = ObjectAndTrack.BoundObject;
+		/* if we ever support channels
+		for (FName& ControlName : SelectedControls)
+		{
+			FAnimLayerPropertyAndChannels Channels;
+			Channels.Name = ControlName;
+			Channels.Channels = (uint32)EControlRigContextChannelToKey::AllTransform;
+			AnimLayerItem.AnimLayerSet.Names.Add(ControlName, Channels);
+
+		}*/
+		// Add a new section that starts and ends at the same time
+		ObjectAndTrack.Track->Modify();
+		if (UMovieSceneSection* NewSection = ObjectAndTrack.Track->CreateNewSection())
+		{
+			ensureAlwaysMsgf(NewSection->HasAnyFlags(RF_Transactional), TEXT("CreateNewSection must return an instance with RF_Transactional set! (pass RF_Transactional to NewObject)"));
+			NewSection->SetFlags(RF_Transactional);
+			FMovieSceneFloatChannel* FloatChannel = nullptr;
+			if (UMovieScene3DTransformSection* TransformSection = Cast<UMovieScene3DTransformSection>(NewSection))
+			{
+				TransformSection->SetMask(FMovieSceneTransformMask{ EMovieSceneTransformChannel::All });
+				FloatChannel = TransformSection->GetWeightChannel();
+			}
+			SectionItem.Section = NewSection;
+			AnimLayerItem.SectionItems.Add(SectionItem);
+			UAnimLayers::SetUpSectionDefaults(SequencerPtr.Get(), this, ObjectAndTrack.Track, NewSection, FloatChannel);
+			bAddedSomething = true;
+		}
+	}
+	if (bAddedSomething)
+	{
+		if (UAnimLayers* AnimLayers = UAnimLayers::GetAnimLayers(SequencerPtr.Get()))
+		{
+			AnimLayers->SetUpBaseLayerSections();
+		}
+		SetKeyed();
+	}
+	SequencerPtr->NotifyMovieSceneDataChanged(EMovieSceneDataChangeType::MovieSceneStructureItemAdded);
+	return true;
+}
+
+bool UAnimLayer::RemoveAnimLayerItem(UObject* InObject)
+{
+	if (FAnimLayerItem* Item =  AnimLayerItems.Find(InObject))
+	{
+		for (const FAnimLayerSectionItem& SectionItem : Item->SectionItems)
+		{
+			if (SectionItem.Section.IsValid())
+			{
+				if (UMovieSceneTrack* Track = SectionItem.Section->GetTypedOuter<UMovieSceneTrack>())
+				{
+					if (Track->GetAllSections().Find(SectionItem.Section.Get()) != 0)
+					{
+						Track->Modify();
+						Track->RemoveSection(*(SectionItem.Section.Get()));
+					}
+				}
+			}
+		}
+		AnimLayerItems.Remove(InObject);
+		return true;
+	}
+	return false;
+	
+}
+bool UAnimLayer::RemoveSelectedInSequencer() 
+{
+	TSharedPtr<ISequencer> SequencerPtr = UAnimLayers::GetSequencerFromAsset();
+	if (SequencerPtr.IsValid() == false)
+	{
+		return false;
+	}
+
+	if (UAnimLayers* AnimLayers = UAnimLayers::GetAnimLayers(SequencerPtr.Get()))
+	{
+		if (AnimLayers->AnimLayers[0] == this)
+		{
+			return false;
+		}
+	}
+
+	TArray<FControlRigAndControlsAndTrack> SelectedCRs;
+	TArray<FObjectAndTrack> SelectedBoundObjects;
+	GetSelectedControlRigsAndBoundObjects(SequencerPtr.Get(), SelectedCRs, SelectedBoundObjects);
+	if (SelectedCRs.Num() <= 0 && SelectedBoundObjects.Num() < 0)
+	{
+		return false;
+	}
+	bool bRemovedSomething = false;
+	UAnimLayer& AnimLayer = *this;
+	const FScopedTransaction Transaction(LOCTEXT("RemoveSelected_Transaction", "Remove Selected"), !GIsTransacting);
+	Modify();
+	for (FControlRigAndControlsAndTrack& CRControls : SelectedCRs)
+	{
+		if (FAnimLayerItem* ExistingAnimLayerItem = AnimLayer.AnimLayerItems.Find(CRControls.ControlRig))
+		{
+			for (FAnimLayerSectionItem& SectionItem : ExistingAnimLayerItem->SectionItems)
+			{
+				if (SectionItem.Section.IsValid())
+				{
+					if (UMovieSceneControlRigParameterSection* CRSection = Cast<UMovieSceneControlRigParameterSection>(SectionItem.Section))
+					{
+						for (FName& ControlName : CRControls.Controls)
+						{
+							if (SectionItem.AnimLayerSet.Names.Contains(ControlName) == true)
+							{
+								SectionItem.AnimLayerSet.Names.Remove(ControlName);
+							}
+							TArray<FName> ControlNames;
+							SectionItem.AnimLayerSet.Names.GetKeys(ControlNames);
+							UAnimLayers::SetUpControlRigSection(CRSection, ControlNames);
+							bRemovedSomething = true;
+						}
+					}
+					if (SectionItem.AnimLayerSet.Names.Num() == 0)
+					{
+						AnimLayer.RemoveAnimLayerItem(CRControls.ControlRig);
+						bRemovedSomething = true;
+						break;
+					}
+				}
+			}
+		}
+	}
+	for (FObjectAndTrack& ObjectAndTrack : SelectedBoundObjects)
+	{
+		if (AnimLayer.AnimLayerItems.Find(ObjectAndTrack.BoundObject))
+		{
+			AnimLayer.RemoveAnimLayerItem(ObjectAndTrack.BoundObject);
+			bRemovedSomething = true;
+		}
+	}
+	if (bRemovedSomething)
+	{
+		if (UAnimLayers* AnimLayers = UAnimLayers::GetAnimLayers(SequencerPtr.Get()))
+		{
+			AnimLayers->SetUpBaseLayerSections();
+		}
+	}
+	if (AnimLayer.AnimLayerItems.Num() == 0)
+	{
+		if (UAnimLayers* AnimLayers = UAnimLayers::GetAnimLayers(SequencerPtr.Get()))
+		{
+			int32 Index = AnimLayers->GetAnimLayerIndex(this);
+			if (Index != INDEX_NONE)
+			{
+				AnimLayers->DeleteAnimLayer(SequencerPtr.Get(), Index);
+			}
+		}
+	}
+	
+	SequencerPtr->NotifyMovieSceneDataChanged(EMovieSceneDataChangeType::MovieSceneStructureItemRemoved);
+	return true;
+
+}
+
+void UAnimLayer::GetAnimLayerObjects(FAnimLayerObjects& InLayerObjects) const
+{
+	for (const TPair<TWeakObjectPtr<UObject>,FAnimLayerItem>& Pair : AnimLayerItems)
+	{
+		if (Pair.Key != nullptr)
+		{
+			for (const FAnimLayerSectionItem& SectionItem : Pair.Value.SectionItems)
+			{
+				if (SectionItem.Section.IsValid())
+				{
+					if (Pair.Key->IsA<UControlRig>())
+					{
+						if (UMovieSceneControlRigParameterSection* Section = Cast<UMovieSceneControlRigParameterSection>(SectionItem.Section.Get()))
+						{
+							FAnimLayerControlRigObject ControlRigObject;
+							ControlRigObject.ControlRig = Cast<UControlRig>(Pair.Key.Get());
+							for (const TPair<FName, FAnimLayerPropertyAndChannels>& ControlName : SectionItem.AnimLayerSet.Names)
+							{
+								ControlRigObject.ControlNames.Add(ControlName.Key);
+							}
+							InLayerObjects.ControlRigObjects.Add(ControlRigObject);
+						}
+					}
+					else if (IsAccepableNonControlRigObject(Pair.Key.Get()))
+					{
+						FAnimLayerSceneObject SceneObject;
+						SceneObject.SceneObjectOrComponent = Pair.Key;
+						InLayerObjects.SceneObjects.Add(SceneObject);
+					}
+				}
+			}
+		}
+	}
+}
+
+bool UAnimLayer::IsAccepableNonControlRigObject(UObject* InObject)
+{
+	return (InObject->IsA<AActor>() || InObject->IsA<USceneComponent>());
+}
+
+ECheckBoxState UAnimLayer::GetSelected() const
+{
+	FAnimLayerObjects LayerObjects;
+	GetAnimLayerObjects(LayerObjects);
+	TOptional<ECheckBoxState>  SelectionState;
+	for (const FAnimLayerControlRigObject& ControlRigObject : LayerObjects.ControlRigObjects)
+	{
+		if (ControlRigObject.ControlRig.IsValid())
+		{
+			TArray<FName> SelectedControls = ControlRigObject.ControlRig->CurrentControlSelection();
+			for (const FName& ControlName : ControlRigObject.ControlNames)
+			{
+				if (SelectedControls.Contains(ControlName))
+				{
+					if (SelectionState.IsSet() == false)
+					{
+						SelectionState = ECheckBoxState::Checked;
+					}
+					else if (SelectionState.GetValue() != ECheckBoxState::Checked)
+					{
+						SelectionState = ECheckBoxState::Undetermined;
+					}
+				}
+				else
+				{
+					if (SelectionState.IsSet() == false)
+					{
+						SelectionState = ECheckBoxState::Unchecked;
+					}
+					else if (SelectionState.GetValue() != ECheckBoxState::Unchecked)
+					{
+						SelectionState = ECheckBoxState::Undetermined;
+					}
+				}
+			}
+		}
+		else
+		{
+			if (SelectionState.IsSet())
+			{
+				if (SelectionState.GetValue() == ECheckBoxState::Checked)
+				{
+					SelectionState = ECheckBoxState::Undetermined;
+				}
+			}
+			else
+			{
+				SelectionState = ECheckBoxState::Unchecked;
+			}
+		}
+	}
+	USelection* ComponentSelection = GEditor->GetSelectedComponents();
+	TArray<TWeakObjectPtr<UObject>> SelectedComponents;
+	ComponentSelection->GetSelectedObjects(SelectedComponents);
+	USelection* ActorSelection = GEditor->GetSelectedActors();
+	TArray<TWeakObjectPtr<UObject>> SelectedActors;
+	ActorSelection->GetSelectedObjects(SelectedActors);
+
+	for (const FAnimLayerSceneObject& SceneObject : LayerObjects.SceneObjects)
+	{
+		if (SceneObject.SceneObjectOrComponent.IsValid() &&
+			SceneObject.SceneObjectOrComponent->IsA<AActor>())
+		{
+			if (SelectedActors.Contains(SceneObject.SceneObjectOrComponent))
+			{
+				if (SelectionState.IsSet() == false)
+				{
+					SelectionState = ECheckBoxState::Checked;
+				}
+				else if (SelectionState.GetValue() != ECheckBoxState::Checked)
+				{
+					SelectionState = ECheckBoxState::Undetermined;
+				}
+			}
+			else
+			{
+				if (SelectionState.IsSet() == false)
+				{
+					SelectionState = ECheckBoxState::Unchecked;
+				}
+				else if (SelectionState.GetValue() != ECheckBoxState::Unchecked)
+				{
+					SelectionState = ECheckBoxState::Undetermined;
+				}
+			}
+		}
+		else if (SceneObject.SceneObjectOrComponent.IsValid() &&
+			SceneObject.SceneObjectOrComponent->IsA<USceneComponent>())
+		{
+			if (SelectedComponents.Contains(SceneObject.SceneObjectOrComponent))
+			{
+				if (SelectionState.IsSet() == false)
+				{
+					SelectionState = ECheckBoxState::Checked;
+				}
+				else if (SelectionState.GetValue() != ECheckBoxState::Checked)
+				{
+					SelectionState = ECheckBoxState::Undetermined;
+				}
+			}
+			else
+			{
+				if (SelectionState.IsSet() == false)
+				{
+					SelectionState = ECheckBoxState::Unchecked;
+				}
+				else if (SelectionState.GetValue() != ECheckBoxState::Unchecked)
+				{
+					SelectionState = ECheckBoxState::Undetermined;
+				}
+			}
+
+		}
+	}
+	if (SelectionState.IsSet())
+	{
+		return SelectionState.GetValue();
+	}
+	return ECheckBoxState::Unchecked;
+}
+
+void UAnimLayer::SetSelected(bool bInSelected, bool bClearSelection)
+{
+	if (!GEditor)
+	{
+		return;
+	}
+	FAnimLayerObjects LayerObjects;
+	GetAnimLayerObjects(LayerObjects);
+	if (LayerObjects.ControlRigObjects.Num() <= 0 && LayerObjects.SceneObjects.Num() <= 0)
+	{
+		return;
+	}
+	const FScopedTransaction Transaction(LOCTEXT("SetSelected_Transaction", "Set Selection"), !GIsTransacting);
+	Modify();
+	State.bSelected = bInSelected ? ECheckBoxState::Checked : ECheckBoxState::Unchecked;
+	if (bClearSelection)
+	{
+		if ((GEditor->GetSelectedActorCount() || GEditor->GetSelectedComponentCount()))
+		{
+			GEditor->SelectNone(false, true);
+			GEditor->NoteSelectionChange();
+		}
+
+		TArray<UControlRig*> ControlRigs;
+		FControlRigEditMode* ControlRigEditMode = static_cast<FControlRigEditMode*>(GLevelEditorModeTools().GetActiveMode(FControlRigEditMode::ModeName));
+		if (ControlRigEditMode)
+		{
+			ControlRigs = ControlRigEditMode->GetControlRigsArray(false /*bIsVisible*/);
+			for (UControlRig* ControlRig : ControlRigs)
+			{
+				ControlRig->ClearControlSelection();
+			}
+		}
+		if (bInSelected == false)
+		{
+			return; //clearing and not selecting so we are done.
+		}
+	}
+
+	for (const FAnimLayerSceneObject& SceneObject : LayerObjects.SceneObjects)
+	{
+		if (AActor* Actor = Cast<AActor>(SceneObject.SceneObjectOrComponent.Get()))
+		{
+			GEditor->SelectActor(Actor, bInSelected, true);
+		}
+		else if (UActorComponent* Component = Cast<UActorComponent>(SceneObject.SceneObjectOrComponent.Get()))
+		{
+			GEditor->SelectComponent(Component, bInSelected, true);
+		}
+	}
+
+	for (const FAnimLayerControlRigObject& ControlRigObject : LayerObjects.ControlRigObjects)
+	{
+		if (ControlRigObject.ControlRig.IsValid())
+		{
+			for (const FName& ControlName : ControlRigObject.ControlNames)
+			{
+				ControlRigObject.ControlRig->SelectControl(ControlName, bInSelected);
+			}
+		}
+	}	
+}
+
+bool UAnimLayer::GetLock() const
+{
+	TOptional<bool> CurrentVal;
+	for (const TPair<TWeakObjectPtr<UObject>,FAnimLayerItem>& Pair : AnimLayerItems)
+	{
+		if (Pair.Key != nullptr)
+		{
+			for (const FAnimLayerSectionItem& SectionItem : Pair.Value.SectionItems)
+			{
+				if (SectionItem.Section.IsValid())
+				{
+					const bool bIsLocked = SectionItem.Section.Get()->IsLocked();
+					if (CurrentVal.IsSet() && CurrentVal.GetValue() != bIsLocked)
+					{
+						SectionItem.Section.Get()->SetIsLocked(CurrentVal.GetValue());
+					}
+					if (CurrentVal.IsSet() == false)
+					{
+						CurrentVal = bIsLocked;
+					}
+				}
+			}
+		}
+	}
+	if (CurrentVal.IsSet())
+	{
+		State.bLock = CurrentVal.GetValue();
+	}
+	return State.bLock;
+}
+
+void UAnimLayer::SetLock(bool bInLock)
+{
+	const FScopedTransaction Transaction(LOCTEXT("SetLock_Transaction", "Set Lock"), !GIsTransacting);
+	Modify();
+	State.bLock = bInLock;
+	for (const TPair<TWeakObjectPtr<UObject>,FAnimLayerItem>& Pair : AnimLayerItems)
+	{
+		if (Pair.Key != nullptr)
+		{
+			for (const FAnimLayerSectionItem& SectionItem : Pair.Value.SectionItems)
+			{
+				if (SectionItem.Section.IsValid())
+				{
+
+					SectionItem.Section->Modify();
+					SectionItem.Section.Get()->SetIsLocked(State.bLock);
+				}
+			}
+		}
+	}
+}
+
+FText UAnimLayer::GetName() const
+{
+	return State.Name;
+}
+
+void UAnimLayer::SetName(const FText& InName)
+{
+	const FScopedTransaction Transaction(LOCTEXT("SetName_Transaction", "Set Name"), !GIsTransacting);
+	Modify();
+	State.Name = InName;
+
+	for (const TPair<TWeakObjectPtr<UObject>,FAnimLayerItem>& Pair : AnimLayerItems)
+	{
+		if (Pair.Key != nullptr)
+		{
+			for (const FAnimLayerSectionItem& SectionItem : Pair.Value.SectionItems)
+			{
+				if (SectionItem.Section.IsValid())
+				{
+					if (UMovieSceneNameableTrack* NameableTrack = Cast<UMovieSceneNameableTrack>(SectionItem.Section.Get()->GetTypedOuter< UMovieSceneNameableTrack>()))
+					{
+						NameableTrack->Modify();
+						NameableTrack->SetTrackRowDisplayName(State.Name, SectionItem.Section.Get()->GetRowIndex());
+					}
+				}
+			}
+		}
+	}
+}
+
+double UAnimLayer::GetWeight() const
+{
+	TSharedPtr<ISequencer> SequencerPtr = UAnimLayers::GetSequencerFromAsset();
+	if (SequencerPtr)
+	{
+		TOptional<float> DifferentWeightValue;
+		for (const TPair<TWeakObjectPtr<UObject>,FAnimLayerItem>& Pair : AnimLayerItems)
+		{
+			if (Pair.Key != nullptr)
+			{
+				for (const FAnimLayerSectionItem& SectionItem : Pair.Value.SectionItems)
+				{
+					if (SectionItem.Section.IsValid())
+					{
+						FMovieSceneFloatChannel* FloatChannel = nullptr;
+						if (UMovieSceneControlRigParameterSection* CRSection = Cast<UMovieSceneControlRigParameterSection>(SectionItem.Section.Get()))
+						{
+							FloatChannel = &CRSection->Weight;
+						}
+						else if (UMovieScene3DTransformSection* LayerSection = Cast<UMovieScene3DTransformSection>(SectionItem.Section.Get()))
+						{
+							FloatChannel = LayerSection->GetWeightChannel();
+						}
+						if (FloatChannel)
+						{
+							const FFrameNumber CurrentTime = SequencerPtr->GetLocalTime().Time.FloorToFrame();
+							float Value = 0.f;
+							FloatChannel->Evaluate(CurrentTime, Value);
+							if (State.Weight != Value)
+							{
+								DifferentWeightValue = Value;
+							}
+
+						}
+					}
+				}
+			}
+		}
+		if (DifferentWeightValue.IsSet())
+		{
+			State.Weight = DifferentWeightValue.GetValue();
+			WeightProxy->Weight = State.Weight;
+		}
+	}
+	return State.Weight;
+}
+
+static void SetFloatWeightValue(float InValue, ISequencer * Sequencer, UMovieSceneSection* OwningSection,FMovieSceneFloatChannel* Channel)
+{
+	using namespace UE::MovieScene;
+	using namespace UE::Sequencer;
+
+
+	if (!OwningSection)
+	{
+		return;
+	}
+	OwningSection->SetFlags(RF_Transactional);;
+
+	if (!OwningSection->TryModify() || !Channel || !Sequencer)
+	{
+		return;
+	}
+
+	const bool  bAutoSetTrackDefaults = Sequencer->GetAutoSetTrackDefaults();	
+	const FFrameNumber CurrentTime = Sequencer->GetLocalTime().Time.FloorToFrame();
+
+	EMovieSceneKeyInterpolation Interpolation = GetInterpolationMode(Channel, CurrentTime, Sequencer->GetKeyInterpolation());
+
+	TArray<FKeyHandle> KeysAtCurrentTime;
+	Channel->GetKeys(TRange<FFrameNumber>(CurrentTime), nullptr, &KeysAtCurrentTime);
+
+	if (KeysAtCurrentTime.Num() > 0)
+	{
+		AssignValue(Channel, KeysAtCurrentTime[0], InValue);
+	}
+	else
+	{
+		bool bHasAnyKeys = Channel->GetNumKeys() != 0;
+
+		if (bHasAnyKeys || bAutoSetTrackDefaults == false)
+		{
+			// When auto setting track defaults are disabled, add a key even when it's empty so that the changed
+			// value is saved and is propagated to the property.
+			AddKeyToChannel(Channel, CurrentTime, InValue, Interpolation);
+			bHasAnyKeys = Channel->GetNumKeys() != 0;
+		}
+
+		if (bHasAnyKeys)
+		{
+			TRange<FFrameNumber> KeyRange = TRange<FFrameNumber>(CurrentTime);
+			TRange<FFrameNumber> SectionRange = OwningSection->GetRange();
+
+			if (!SectionRange.Contains(KeyRange))
+			{
+				OwningSection->SetRange(TRange<FFrameNumber>::Hull(KeyRange, SectionRange));
+			}
+		}
+	}
+
+	// Always update the default value when auto-set default values is enabled so that the last changes
+	// are always saved to the track.
+	if (bAutoSetTrackDefaults)
+	{
+		SetChannelDefault(Channel, InValue);
+	}
+
+	Channel->AutoSetTangents();
+	Sequencer->NotifyMovieSceneDataChanged(EMovieSceneDataChangeType::TrackValueChanged);
+}
+
+void UAnimLayer::SetWeight(double InWeight)
+{
+	State.Weight = InWeight;
+	TSharedPtr<ISequencer> SequencerPtr = UAnimLayers::GetSequencerFromAsset();
+	if (SequencerPtr)
+	{
+		for (const TPair<TWeakObjectPtr<UObject>,FAnimLayerItem>& Pair : AnimLayerItems)
+		{
+			if (Pair.Key != nullptr)
+			{
+				for (const FAnimLayerSectionItem& SectionItem : Pair.Value.SectionItems)
+				{
+					if (SectionItem.Section.IsValid())
+					{
+						if (SectionItem.Section.Get()->TryModify())
+						{
+							FMovieSceneFloatChannel* FloatChannel = nullptr;
+							if (UMovieSceneControlRigParameterSection* CRSection = Cast<UMovieSceneControlRigParameterSection>(SectionItem.Section.Get()))
+							{
+								FloatChannel = &CRSection->Weight;
+							}
+							else if (UMovieScene3DTransformSection* LayerSection = Cast<UMovieScene3DTransformSection>(SectionItem.Section.Get()))
+							{
+								FloatChannel = LayerSection->GetWeightChannel();
+							}
+							if (FloatChannel)
+							{
+								float WeightValue = (float)InWeight;
+								SetFloatWeightValue(WeightValue, SequencerPtr.Get(), SectionItem.Section.Get(), FloatChannel);
+							}
+						}
+					}
+				}
+			}
+		}
+	}
+}
+
+void UAnimLayerWeightProxy::PostEditChangeChainProperty(struct FPropertyChangedChainEvent& PropertyChangedEvent)
+{
+	Super::PostEditChangeChainProperty(PropertyChangedEvent);
+#if WITH_EDITOR
+	if (UAnimLayer* AnimLayer = GetTypedOuter<UAnimLayer>())
+	{
+		if (PropertyChangedEvent.Property && (
+			PropertyChangedEvent.ChangeType == EPropertyChangeType::ValueSet ||
+			PropertyChangedEvent.ChangeType == EPropertyChangeType::Interactive ||
+			PropertyChangedEvent.ChangeType == EPropertyChangeType::Unspecified)
+			)
+		{
+			//set values
+			FProperty* Property = PropertyChangedEvent.Property;
+			FProperty* MemberProperty = nullptr;
+			if ((Property && Property->GetFName() == GET_MEMBER_NAME_CHECKED(UAnimLayerWeightProxy, Weight)) ||
+				(MemberProperty && MemberProperty->GetFName() == GET_MEMBER_NAME_CHECKED(UAnimLayerWeightProxy, Weight)))
+			{
+				Modify();
+				AnimLayer->SetWeight(Weight);
+			}
+		}
+	}
+#endif
+}
+
+#if WITH_EDITOR
+void UAnimLayerWeightProxy::PostEditUndo()
+{
+	if (UAnimLayer* AnimLayer = GetTypedOuter<UAnimLayer>())
+	{
+		AnimLayer->SetWeight(Weight);
+	}
+}
+#endif
+
+EAnimLayerType UAnimLayer::GetType() const
+{
+	TOptional<EMovieSceneBlendType> CurrentVal;
+	for (const TPair<TWeakObjectPtr<UObject>,FAnimLayerItem>& Pair : AnimLayerItems)
+	{
+		if (Pair.Key != nullptr)
+		{
+			for (const FAnimLayerSectionItem& SectionItem : Pair.Value.SectionItems)
+			{
+				if (SectionItem.Section.IsValid())
+				{
+					const EMovieSceneBlendType BlendType = SectionItem.Section.Get()->GetBlendType().BlendType;
+					if (CurrentVal.IsSet() && CurrentVal.GetValue() != BlendType)
+					{
+						SectionItem.Section.Get()->SetBlendType(CurrentVal.GetValue());
+					}
+					if (CurrentVal.IsSet() == false)
+					{
+						CurrentVal = BlendType;
+					}
+				}
+			}
+		}
+	}
+	if (CurrentVal.IsSet())
+	{
+		switch (CurrentVal.GetValue())
+		{
+		case EMovieSceneBlendType::Additive:
+			State.Type = (int32)EAnimLayerType::Additive;
+			break;
+		case EMovieSceneBlendType::Override:
+			State.Type = (int32)EAnimLayerType::Override;
+			break;
+		case EMovieSceneBlendType::Absolute:
+			State.Type = (int32)EAnimLayerType::Base;
+			break;
+		}
+	}
+
+	return (EAnimLayerType)(State.Type);
+}
+
+void UAnimLayer::SetType(EAnimLayerType LayerType)
+{
+	const FScopedTransaction Transaction(LOCTEXT("SetType_Transaction", "Set Type"), !GIsTransacting);
+	Modify();
+
+	State.Type = (uint32)(LayerType);
+	for (const TPair<TWeakObjectPtr<UObject>,FAnimLayerItem>& Pair : AnimLayerItems)
+	{
+		if (Pair.Key != nullptr)
+		{
+			for (const FAnimLayerSectionItem& SectionItem : Pair.Value.SectionItems)
+			{
+				if (SectionItem.Section.IsValid())
+				{
+					switch (LayerType)
+					{
+					case EAnimLayerType::Additive:
+						SectionItem.Section.Get()->SetBlendType(EMovieSceneBlendType::Additive);
+						break;
+					case EAnimLayerType::Override:
+						SectionItem.Section.Get()->SetBlendType(EMovieSceneBlendType::Override);
+						break;
+					case EAnimLayerType::Base:
+						SectionItem.Section.Get()->SetBlendType(EMovieSceneBlendType::Absolute);
+						break;
+					}
+				}
+			}
+
+		}
+	}
+}
+
+UAnimLayers::UAnimLayers(class FObjectInitializer const& ObjectInitializer)
+	: Super(ObjectInitializer)
+{
+
+}
+
+#if WITH_EDITOR
+void UAnimLayers::PostEditUndo()
+{
+	AnimLayerListChangedBroadcast();
+}
+#endif
+
+TSharedPtr<ISequencer> UAnimLayers::GetSequencerFromAsset()
+{
+	ULevelSequence* LevelSequence = ULevelSequenceEditorBlueprintLibrary::GetCurrentLevelSequence();
+	IAssetEditorInstance* AssetEditor = GEditor->GetEditorSubsystem<UAssetEditorSubsystem>()->FindEditorForAsset(LevelSequence, false);
+	ILevelSequenceEditorToolkit* LevelSequenceEditor = static_cast<ILevelSequenceEditorToolkit*>(AssetEditor);
+	TSharedPtr<ISequencer> SequencerPtr = LevelSequenceEditor ? LevelSequenceEditor->GetSequencer() : nullptr;
+	return SequencerPtr;
+}
+
+void UAnimLayers::AddBaseLayer()
+{
+	UAnimLayer* AnimLayer = NewObject<UAnimLayer>(this, TEXT("BaseLayer"), RF_Transactional);
+	AnimLayer->State.Type = EAnimLayerType::Base;
+	AnimLayer->State.bKeyed = ECheckBoxState::Checked;
+	AnimLayers.Add(AnimLayer);
+}
+
+UAnimLayers* UAnimLayers::GetAnimLayers(ISequencer* SequencerPtr)
+{
+	if (!SequencerPtr)
+	{
+		return nullptr;
+	}
+	ULevelSequence* LevelSequence = Cast<ULevelSequence>(SequencerPtr->GetFocusedMovieSceneSequence());
+	return UAnimLayers::GetAnimLayers(LevelSequence);
+}
+
+UAnimLayers* UAnimLayers::GetAnimLayers(ULevelSequence* LevelSequence)
+{
+	if(!LevelSequence)
+	{
+		return nullptr;
+	}
+	if (LevelSequence && LevelSequence->GetClass()->ImplementsInterface(UInterface_AssetUserData::StaticClass()))
+	{
+		if (IInterface_AssetUserData* AssetUserDataInterface = Cast< IInterface_AssetUserData >(LevelSequence))
+		{
+			UAnimLayers* AnimLayers = AssetUserDataInterface->GetAssetUserData< UAnimLayers >();
+			if (!AnimLayers)
+			{
+				AnimLayers = NewObject<UAnimLayers>(LevelSequence, NAME_None, RF_Public | RF_Transactional);
+				AssetUserDataInterface->AddAssetUserData(AnimLayers);
+			}
+			return AnimLayers;
+		}
+	}
+	return nullptr;
+}
+
+int32 UAnimLayers::GetAnimLayerIndex(UAnimLayer* InAnimLayer) const
+{
+	if (InAnimLayer != nullptr)
+	{
+		return AnimLayers.Find(InAnimLayer);
+	}
+	return INDEX_NONE;
+}
+
+bool UAnimLayers::DeleteAnimLayer(ISequencer* SequencerPtr, int32 Index)
+{
+	if (Index >= 1 && Index < AnimLayers.Num()) //don't delete base
+	{
+		if (UAnimLayer* AnimLayer = AnimLayers[Index])
+		{
+			const FScopedTransaction Transaction(LOCTEXT("DeleteAnimLayer_Transaction", "Delete Anim Layer"), !GIsTransacting);
+			Modify();
+			for (const TPair<TWeakObjectPtr<UObject>, FAnimLayerItem>& Pair : AnimLayer->AnimLayerItems)
+			{
+				if (Pair.Key != nullptr)
+				{
+					for (const FAnimLayerSectionItem& SectionItem : Pair.Value.SectionItems)
+					{
+						if (SectionItem.Section.IsValid())
+						{
+							if (UMovieSceneTrack* Track = SectionItem.Section->GetTypedOuter<UMovieSceneTrack>())
+							{
+								if (Track->GetAllSections().Find(SectionItem.Section.Get()) != 0)
+								{
+									Track->Modify();
+									Track->RemoveSection(*(SectionItem.Section.Get()));
+								}
+							}
+						}
+					}
+				}
+			}
+			AnimLayers.RemoveAt(Index);
+			SequencerPtr->NotifyMovieSceneDataChanged(EMovieSceneDataChangeType::MovieSceneStructureItemRemoved);
+		}
+		AnimLayerListChangedBroadcast();
+	}
+	else
+	{
+		return false;
+	}
+	return true;
+}
+
+static void CopySectionIntoAnother(UMovieSceneSection* ToSection, UMovieSceneSection* FromSection)
+{
+	const FFrameNumber Min = TNumericLimits<FFrameNumber>::Lowest();
+	const FFrameNumber Max = TNumericLimits<FFrameNumber>::Max();
+	TRange<FFrameNumber> Range(Min, Max);
+	TArray<UMovieSceneSection*> AbsoluteSections;
+	TArray<UMovieSceneSection*> AdditiveSections;
+	if (ToSection->GetBlendType().Get() == EMovieSceneBlendType::Additive)
+	{
+		AdditiveSections.Add(ToSection);
+	}
+	else if (ToSection->GetBlendType().Get() == EMovieSceneBlendType::Absolute)
+	{
+		AbsoluteSections.Add(ToSection);
+	}
+	if (FromSection->GetBlendType().Get() == EMovieSceneBlendType::Additive)
+	{
+		AdditiveSections.Add(FromSection);
+	}
+	else if (FromSection->GetBlendType().Get() == EMovieSceneBlendType::Absolute)
+	{
+		AbsoluteSections.Add(FromSection);
+	}
+	//if CR or transform we just do certain channels so we avoid baking in the weight channel
+	//if not we do it by property type
+	if (ToSection->IsA<UMovieSceneControlRigParameterSection>())
+	{
+		MovieSceneToolHelpers::MergeSections<FMovieSceneFloatChannel>(ToSection, AbsoluteSections, AdditiveSections, Range, true /*bSkipLastChannel*/);
+	}
+	else if (ToSection->IsA<UMovieScene3DTransformSection>())
+	{
+		MovieSceneToolHelpers::MergeSections<FMovieSceneDoubleChannel>(ToSection, AbsoluteSections, AdditiveSections, Range, false /*bSkipLastChannel*/);
+	}
+	else
+	{
+		FMovieSceneChannelProxy& ChannelProxy = ToSection->GetChannelProxy();
+		for (const FMovieSceneChannelEntry& Entry : ToSection->GetChannelProxy().GetAllEntries())
+		{
+			const FName ChannelTypeName = Entry.GetChannelTypeName();
+			if (ChannelTypeName == FMovieSceneFloatChannel::StaticStruct()->GetFName())
+			{
+				MovieSceneToolHelpers::MergeSections<FMovieSceneFloatChannel>(ToSection, AbsoluteSections, AdditiveSections, Range, false /*bSkipLastChannel*/);
+			}
+			else if (ChannelTypeName == FMovieSceneDoubleChannel::StaticStruct()->GetFName())
+			{
+				MovieSceneToolHelpers::MergeSections<FMovieSceneDoubleChannel>(ToSection, AbsoluteSections, AdditiveSections, Range, false /*bSkipLastChannel*/);
+			}
+		}
+	}
+}
+
+int32 UAnimLayers::DuplicateAnimLayer(ISequencer* SequencerPtr, int32 Index)
+{
+	int32 NewIndex = INDEX_NONE;
+	if (Index >= 1 && Index < AnimLayers.Num()) //don't duplicate base
+	{
+		if (UAnimLayer* ExistingAnimLayer = AnimLayers[Index])
+		{
+			if (ExistingAnimLayer->AnimLayerItems.Num() == 0)
+			{
+				UE_LOG(LogControlRig, Error, TEXT("Anim Layers: Can not duplicate empty layer"));
+				return INDEX_NONE;
+			}
+			const FScopedTransaction Transaction(LOCTEXT("DuplicateAnimLayer_Transaction", "Duplicate Anim Layer"), !GIsTransacting);
+			Modify();
+			UAnimLayer* NewAnimLayer = NewObject<UAnimLayer>(this, NAME_None, RF_Transactional);
+			bool bItemAdded = false;
+			for (const TPair<TWeakObjectPtr<UObject>,FAnimLayerItem>& Pair : ExistingAnimLayer->AnimLayerItems)
+			{
+				if (Pair.Key != nullptr)
+				{
+					for (const FAnimLayerSectionItem& SectionItem : Pair.Value.SectionItems)
+					{
+						if (SectionItem.Section.IsValid())
+						{
+							if (UMovieSceneTrack* Track = SectionItem.Section->GetTypedOuter<UMovieSceneTrack>())
+							{
+								Track->Modify();
+								if (UMovieSceneControlRigParameterSection* Section = Cast<UMovieSceneControlRigParameterSection>(SectionItem.Section.Get()))
+								{
+									if (UControlRig* ControlRig = Cast<UControlRig>(Pair.Key))
+									{
+										FAnimLayerItem AnimLayerItem;
+										FAnimLayerSectionItem NewSectionItem;
+										NewSectionItem.AnimLayerSet.BoundObject = ControlRig;
+										NewSectionItem.AnimLayerSet = SectionItem.AnimLayerSet;
+										// Add a new section that starts and ends at the same time
+										if (UMovieSceneControlRigParameterTrack* CRTack = Cast<UMovieSceneControlRigParameterTrack>(Track))
+										{
+											TGuardValue<bool> GuardSetSection(CRTack->bSetSectionToKeyPerControl, false);
+											if (UMovieSceneControlRigParameterSection* NewSection = Cast<UMovieSceneControlRigParameterSection>(Track->CreateNewSection()))
+											{
+												if (bItemAdded == false)
+												{
+													NewAnimLayer->State.Weight = 1.0;
+													NewAnimLayer->State.Type = (EAnimLayerType::Additive);
+													bItemAdded = true;
+												}
+												ensureAlwaysMsgf(NewSection->HasAnyFlags(RF_Transactional), TEXT("CreateNewSection must return an instance with RF_Transactional set! (pass RF_Transactional to NewObject)"));
+												NewSection->SetFlags(RF_Transactional);
+												NewSection->SetTransformMask(FMovieSceneTransformMask{ EMovieSceneTransformChannel::All });
+												FMovieSceneFloatChannel* FloatChannel = &NewSection->Weight;
+												NewSectionItem.Section = NewSection;
+												AnimLayerItem.SectionItems.Add(NewSectionItem);
+												NewAnimLayer->AnimLayerItems.Add(ControlRig, AnimLayerItem);
+												SetUpSectionDefaults(SequencerPtr, NewAnimLayer, Track, NewSection, FloatChannel);
+												TArray<FName> ControlNames;
+												NewSectionItem.AnimLayerSet.Names.GetKeys(ControlNames);
+												SetUpControlRigSection(NewSection, ControlNames);
+												//current copy keys
+												CopySectionIntoAnother(NewSection, Section);
+											}
+										}
+									}
+								}
+								else
+								{
+									FAnimLayerItem AnimLayerItem;
+									FAnimLayerSectionItem NewSectionItem;
+									NewSectionItem.AnimLayerSet.BoundObject = Pair.Key;
+									if (UMovieSceneSection* NewSection = Track->CreateNewSection())
+									{
+										if (bItemAdded == false)
+										{
+											NewAnimLayer->State.Weight = 1.0;
+											NewAnimLayer->State.Type = (EAnimLayerType::Additive);
+											bItemAdded = true;
+										}
+
+										ensureAlwaysMsgf(NewSection->HasAnyFlags(RF_Transactional), TEXT("CreateNewSection must return an instance with RF_Transactional set! (pass RF_Transactional to NewObject)"));
+										NewSection->SetFlags(RF_Transactional);
+										FMovieSceneFloatChannel* FloatChannel = nullptr;
+										if (UMovieScene3DTransformSection* TransformSection = Cast<UMovieScene3DTransformSection>(NewSection))
+										{
+											TransformSection->SetMask(FMovieSceneTransformMask{ EMovieSceneTransformChannel::All });
+											FloatChannel = TransformSection->GetWeightChannel();
+										}
+										NewSectionItem.Section = NewSection;
+										AnimLayerItem.SectionItems.Add(NewSectionItem);
+										NewAnimLayer->AnimLayerItems.Add(Pair.Key, AnimLayerItem);
+										SetUpSectionDefaults(SequencerPtr, NewAnimLayer, Track, NewSection, FloatChannel);
+										//current copy keys
+										CopySectionIntoAnother(NewSection, Section);
+									}
+								}
+							}
+						}
+					}
+				}
+			}
+			if (bItemAdded)
+			{
+				FString LayerName = FString::Printf(TEXT("Anim Layer %d"), AnimLayers.Num());
+				FText LayerText;
+				LayerText = LayerText.FromString(LayerName);
+				NewAnimLayer->SetName(LayerText); //need items/sections to be added so we can change their track row names
+				NewIndex = AnimLayers.Add(NewAnimLayer);
+				NewAnimLayer->SetKeyed();
+				AnimLayerListChangedBroadcast();
+				// no need to since it's a dup SetUpBaseLayerSections();
+			}
+		}
+	}
+	SequencerPtr->NotifyMovieSceneDataChanged(EMovieSceneDataChangeType::MovieSceneStructureItemAdded);
+	return NewIndex;
+}
+
+static void AddNamesToMask(FAnimLayerSectionItem* Owner, UMovieSceneControlRigParameterSection* CRSection,const FAnimLayerSelectionSet& NewSet)
+{
+	bool bNameAdded = false;
+	for (TPair<FName, FAnimLayerPropertyAndChannels> NameAndChannels : NewSet.Names)
+	{
+		if (Owner->AnimLayerSet.Names.Contains(NameAndChannels.Key) == false)
+		{
+			bNameAdded = true;
+			Owner->AnimLayerSet.Names.Add(NameAndChannels);
+		}
+	}
+	if (bNameAdded)
+	{
+		TArray<FName> AllControls;
+		Owner->AnimLayerSet.Names.GenerateKeyArray(AllControls);
+		UAnimLayers::SetUpControlRigSection(CRSection, AllControls);
+	}
+}
+
+FAnimLayerSectionItem* FAnimLayerItem::FindMatchingSectionItem(UMovieSceneSection* InMovieSceneSection)
+{
+	if (InMovieSceneSection)
+	{
+		for (FAnimLayerSectionItem& CurrentItem : SectionItems)
+		{
+			if (CurrentItem.Section.IsValid())
+			{
+				UMovieSceneTrack* InTrack = InMovieSceneSection->GetTypedOuter<UMovieSceneTrack>();
+				UMovieSceneTrack* CurrentTrack = CurrentItem.Section->GetTypedOuter<UMovieSceneTrack>();
+				if (CurrentTrack && InTrack && CurrentTrack == InTrack)
+				{
+					return &CurrentItem;
+				}
+			}
+		}
+	}
+	return nullptr;
+}
+
+bool UAnimLayers::MergeAnimLayers(ISequencer* InSequencer, const TArray<int32>& Indices, const FBakingAnimationKeySettings* InSettings)
+{
+	if (InSequencer == nullptr)
+	{
+		return false;
+	}
+	UAnimLayer* BaseLayer = nullptr;
+	TArray<UAnimLayer*> LayersToMerge;
+	const FFrameNumber Min = InSettings ? InSettings->StartFrame : TNumericLimits<FFrameNumber>::Lowest();
+	const FFrameNumber Max = InSettings ? InSettings->EndFrame :  TNumericLimits<FFrameNumber>::Max();
+	TRange<FFrameNumber> Range(Min,Max);
+	FScopedTransaction Transaction(LOCTEXT("Merge Anim Layers", "Merge Anim Layers"), !GIsTransacting);
+
+	TArray<int32> SortedIndices = Indices;
+	SortedIndices.Sort([](const int32& Index1, const int32& Index2) {
+		return  Index1 < Index2; 
+		});
+
+	for (int32 Index : SortedIndices)
+	{
+		if (Index >= 0 && Index < AnimLayers.Num())
+		{
+			if (UAnimLayer* AnimLayer = AnimLayers[Index].Get())
+			{
+				if (BaseLayer == nullptr)
+				{
+					BaseLayer = AnimLayer;
+				}
+				else
+				{
+					LayersToMerge.Add(AnimLayer);
+				}
+			}
+		}
+	}
+	if (BaseLayer == nullptr || LayersToMerge.Num() < 1)
+	{
+		return false; 
+	}
+	Modify();
+	BaseLayer->Modify();
+
+	for (UAnimLayer* AnimLayer : LayersToMerge)
+	{
+		for (TPair<TWeakObjectPtr<UObject>,FAnimLayerItem>& Pair : AnimLayer->AnimLayerItems)
+		{
+			if (Pair.Key != nullptr)
+			{
+				for (FAnimLayerSectionItem& SectionItem : Pair.Value.SectionItems)
+				{
+					if (SectionItem.Section.IsValid())
+					{
+						FAnimLayerSectionItem* BaseSectionItem = nullptr;
+						if (FAnimLayerItem* Owner = BaseLayer->AnimLayerItems.Find(Pair.Key))
+						{
+							BaseSectionItem = Owner->FindMatchingSectionItem(SectionItem.Section.Get());
+
+						}
+						if(BaseSectionItem && BaseSectionItem->Section.IsValid())
+						{
+							if (SectionItem.Section->IsActive())//active sections merge them
+							{
+								TArray<UMovieSceneSection*> AbsoluteSections;
+								TArray<UMovieSceneSection*> AdditiveSections;
+								//base section get's added to both
+								AdditiveSections.Add(BaseSectionItem->Section.Get());
+								AbsoluteSections.Add(BaseSectionItem->Section.Get());
+
+								if (SectionItem.Section->GetBlendType().Get() == EMovieSceneBlendType::Additive)
+								{
+									AdditiveSections.Add(SectionItem.Section.Get());
+								}
+								else if (SectionItem.Section->GetBlendType().Get() == EMovieSceneBlendType::Absolute)
+								{
+									AbsoluteSections.Add(SectionItem.Section.Get());
+								}
+								TArrayView<FMovieSceneFloatChannel*> BaseFloatChannels = BaseSectionItem->Section->GetChannelProxy().GetChannels<FMovieSceneFloatChannel>();
+								TArrayView<FMovieSceneDoubleChannel*> BaseDoubleChannels = BaseSectionItem->Section->GetChannelProxy().GetChannels<FMovieSceneDoubleChannel>();
+								if (BaseDoubleChannels.Num() > 0)
+								{
+									MovieSceneToolHelpers::MergeSections<FMovieSceneDoubleChannel>(BaseSectionItem->Section.Get(), AbsoluteSections, AdditiveSections, Range, false /*bSkipLastChannel*/);
+
+								}
+								else if (BaseFloatChannels.Num() > 0)
+								{
+									MovieSceneToolHelpers::MergeSections<FMovieSceneFloatChannel>(BaseSectionItem->Section.Get(), AbsoluteSections, AdditiveSections, Range, true /*bSkipLastChannel*/);
+								}
+							}
+							if (UMovieSceneControlRigParameterSection* CRSection = Cast<UMovieSceneControlRigParameterSection>(BaseSectionItem->Section))
+							{
+								if (SectionItem.AnimLayerSet.Names.Num() > 0)
+								{
+									AddNamesToMask(BaseSectionItem, CRSection, SectionItem.AnimLayerSet);
+									if (SortedIndices[0] != 0) //if not base then make sure mask is set up
+									{
+										TArray<FName> AllControls;
+										BaseSectionItem->AnimLayerSet.Names.GenerateKeyArray(AllControls);
+										UAnimLayers::SetUpControlRigSection(CRSection, AllControls);
+									}
+								}
+								
+							}
+						}
+						else //okay this object doesn't exist in the first layer we are merging into so we need to move it to the other one 
+						{
+							FAnimLayerItem AnimLayerItem;
+							FAnimLayerSectionItem NewSectionItem;
+							NewSectionItem.AnimLayerSet.BoundObject = Pair.Key;
+							NewSectionItem.AnimLayerSet.Names = SectionItem.AnimLayerSet.Names;
+							NewSectionItem.Section = SectionItem.Section;
+							if (UMovieSceneNameableTrack* NameableTrack = Cast<UMovieSceneNameableTrack>(SectionItem.Section.Get()->GetTypedOuter< UMovieSceneNameableTrack>()))
+							{
+								NameableTrack->Modify();
+								NameableTrack->SetTrackRowDisplayName(BaseLayer->State.Name, SectionItem.Section.Get()->GetRowIndex());
+							}
+							AnimLayerItem.SectionItems.Add(NewSectionItem);
+							BaseLayer->AnimLayerItems.Add(Pair.Key, AnimLayerItem);
+							//since we moved the section over we reset it on the merged so we don't delete it when we remove the layer
+							if (FAnimLayerItem* OldAnimLayeritem = AnimLayer->AnimLayerItems.Find(Pair.Key))
+							{
+								SectionItem.Section.Reset();
+							}
+						}
+					}
+				}
+				int32 Index = GetAnimLayerIndex(AnimLayer);
+				if (Index != INDEX_NONE)
+				{
+					DeleteAnimLayer(InSequencer, Index);
+				}
+			}
+		}
+		//we now delete the layer that was merged, this will delete sections, etc..
+	}
+	InSequencer->NotifyMovieSceneDataChanged(EMovieSceneDataChangeType::MovieSceneStructureItemAdded);
+	return true;
+}
+
+void UAnimLayers::AnimLayerListChangedBroadcast()
+{
+	OnAnimLayerListChanged.Broadcast(this);
+}
+
+
+TArray<UMovieSceneSection*> UAnimLayers::GetSelectedLayerSections() const
+{
+	TArray<UMovieSceneSection*> Sections;
+	for (const UAnimLayer* AnimLayer : AnimLayers)
+	{
+		if (AnimLayer && AnimLayer->bIsSelectedInList)
+		{
+			for (const TPair<TWeakObjectPtr<UObject>, FAnimLayerItem>& Pair : AnimLayer->AnimLayerItems)
+			{
+				if (Pair.Key != nullptr)
+				{
+					for (const FAnimLayerSectionItem& SectionItem : Pair.Value.SectionItems)
+					{
+						if (SectionItem.Section.IsValid() )
+						{
+							Sections.Add(SectionItem.Section.Get());
+						}
+					}
+				}
+			}
+		}
+	}
+	return Sections;
+}
+
+bool UAnimLayers::IsTrackOnSelectedLayer(const UMovieSceneTrack* InTrack)const
+{
+	TSet<UMovieSceneTrack*> CurrentTracks;
+	for (const UAnimLayer* AnimLayer : AnimLayers)
+	{
+		if (AnimLayer && AnimLayer->bIsSelectedInList)
+		{
+			for (const TPair<TWeakObjectPtr<UObject>,FAnimLayerItem>& Pair : AnimLayer->AnimLayerItems)
+			{
+				if (Pair.Key != nullptr)
+				{
+					for (const FAnimLayerSectionItem& SectionItem : Pair.Value.SectionItems)
+					{
+						if (SectionItem.Section.IsValid())
+						{
+							UMovieSceneTrack* OwnerTrack = SectionItem.Section->GetTypedOuter<UMovieSceneTrack>();
+							if (OwnerTrack == InTrack)
+							{
+								return true;
+							}
+						}
+					}
+				}
+			}
+		}
+	}
+	return false;
+}
+
+int32 UAnimLayers::AddAnimLayerFromSelection(ISequencer* SequencerPtr)
+{
+	int32 NewIndex = INDEX_NONE;
+	//wrap scoped transaction since it can deselect control rigs
+	TArray<FControlRigAndControlsAndTrack> SelectedCRs;
+	TArray<FObjectAndTrack> SelectedBoundObjects; 
+	{
+		const FScopedTransaction Transaction(LOCTEXT("AddAnimLayer_Transaction", "Add Anim Layer"), !GIsTransacting);
+		Modify();
+		if (AnimLayers.Num() == 0)
+		{
+			AddBaseLayer();
+		}
+		UAnimLayer* AnimLayer = NewObject<UAnimLayer>(this, NAME_None, RF_Transactional);
+
+		GetSelectedControlRigsAndBoundObjects(SequencerPtr, SelectedCRs, SelectedBoundObjects);
+
+		if (SelectedCRs.Num() <= 0 && SelectedBoundObjects.Num() <= 0)
+		{
+			FString LayerName = FString::Printf(TEXT("Empty Layer %d"), AnimLayers.Num());
+			FText LayerText;
+			LayerText = LayerText.FromString(LayerName);
+			AnimLayer->SetName(LayerText); //need items/sections to be added so we can change their track row names
+			AnimLayer->State.Weight = 1.0;
+			AnimLayer->State.Type = (EAnimLayerType::Additive);
+			int32 Index = AnimLayers.Add(AnimLayer);
+			AnimLayerListChangedBroadcast();
+
+			return Index;
+		}
+
+		bool bItemAdded = false;
+
+		for (FControlRigAndControlsAndTrack& CRControls : SelectedCRs)
+		{
+			Modify();
+			FAnimLayerItem AnimLayerItem;
+			FAnimLayerSectionItem SectionItem;
+			SectionItem.AnimLayerSet.BoundObject = CRControls.ControlRig;
+			for (FName& ControlName : CRControls.Controls)
+			{
+				FAnimLayerPropertyAndChannels Channels;
+				Channels.Name = ControlName;
+				Channels.Channels = (uint32)EControlRigContextChannelToKey::AllTransform;
+				SectionItem.AnimLayerSet.Names.Add(ControlName, Channels);
+
+			}
+			CRControls.Track->Modify();
+			// Add a new section that starts and ends at the same time
+			TGuardValue<bool> GuardSetSection(CRControls.Track->bSetSectionToKeyPerControl, false);
+			if (UMovieSceneControlRigParameterSection* NewSection = Cast<UMovieSceneControlRigParameterSection>(CRControls.Track->CreateNewSection()))
+			{
+				if (bItemAdded == false)
+				{
+					AnimLayer->State.Weight = 1.0;
+					AnimLayer->State.Type = (EAnimLayerType::Additive);
+					bItemAdded = true;
+				}
+
+				ensureAlwaysMsgf(NewSection->HasAnyFlags(RF_Transactional), TEXT("CreateNewSection must return an instance with RF_Transactional set! (pass RF_Transactional to NewObject)"));
+				NewSection->SetFlags(RF_Transactional);
+				NewSection->SetTransformMask(FMovieSceneTransformMask{ EMovieSceneTransformChannel::All });
+				FMovieSceneFloatChannel* FloatChannel = &NewSection->Weight;
+				SectionItem.Section = NewSection;
+				AnimLayerItem.SectionItems.Add(SectionItem);
+				AnimLayer->AnimLayerItems.Add(CRControls.ControlRig, AnimLayerItem);
+				SetUpSectionDefaults(SequencerPtr, AnimLayer, CRControls.Track, NewSection, FloatChannel);
+				SetUpControlRigSection(NewSection, CRControls.Controls);
+
+			}
+		}
+		for (FObjectAndTrack& ObjectAndTrack : SelectedBoundObjects)
+		{
+			Modify();
+
+			FAnimLayerItem& AnimLayerItem = AnimLayer->AnimLayerItems.FindOrAdd(ObjectAndTrack.BoundObject);
+			FAnimLayerSectionItem SectionItem;
+			SectionItem.AnimLayerSet.BoundObject = ObjectAndTrack.BoundObject;
+			/*
+			for (FName& ControlName : SelectedControls)
+			{
+				FAnimLayerPropertyAndChannels Channels;
+				Channels.Name = ControlName;
+				Channels.Channels = (uint32)EControlRigContextChannelToKey::AllTransform;
+				AnimLayerItem.AnimLayerSet.Names.Add(ControlName, Channels);
+
+			}*/
+
+			// Add a new section that starts and ends at the same time
+			ObjectAndTrack.Track->Modify();
+			if (UMovieSceneSection* NewSection = ObjectAndTrack.Track->CreateNewSection())
+			{
+				if (bItemAdded == false)
+				{
+					AnimLayer->State.Weight = 1.0;
+					AnimLayer->State.Type = (EAnimLayerType::Additive);
+					bItemAdded = true;
+				}
+
+				ensureAlwaysMsgf(NewSection->HasAnyFlags(RF_Transactional), TEXT("CreateNewSection must return an instance with RF_Transactional set! (pass RF_Transactional to NewObject)"));
+				NewSection->SetFlags(RF_Transactional);
+				SectionItem.Section = NewSection;
+				FMovieSceneFloatChannel* FloatChannel = nullptr;
+				if (UMovieScene3DTransformSection* TransformSection = Cast<UMovieScene3DTransformSection>(NewSection))
+				{
+					TransformSection->SetMask(FMovieSceneTransformMask{ EMovieSceneTransformChannel::All });
+					FloatChannel = TransformSection->GetWeightChannel();
+				}
+				AnimLayerItem.SectionItems.Add(SectionItem);
+				SetUpSectionDefaults(SequencerPtr, AnimLayer, ObjectAndTrack.Track, NewSection, FloatChannel);
+			}
+		}
+
+		if (bItemAdded)
+		{
+			FString LayerName = FString::Printf(TEXT("Anim Layer %d"), AnimLayers.Num());
+			FText LayerText;
+			LayerText = LayerText.FromString(LayerName);
+			AnimLayer->SetName(LayerText); //need items/sections to be added so we can change their track row names
+			int32 Index = AnimLayers.Add(AnimLayer);
+			SetUpBaseLayerSections();
+			AnimLayer->SetKeyed();
+			SequencerPtr->NotifyMovieSceneDataChanged(EMovieSceneDataChangeType::MovieSceneStructureItemAdded);
+
+			NewIndex = Index;;
+			AnimLayerListChangedBroadcast();
+		}
+	}
+	//may need to reselect controls here
+	for (FControlRigAndControlsAndTrack& CRControls : SelectedCRs)
+	{
+		for (const FName& Control : CRControls.Controls)
+		{
+			if (CRControls.ControlRig->IsControlSelected(Control) == false)
+			{
+				CRControls.ControlRig->SelectControl(Control, true);
+			}
+		}
+	}
+	return NewIndex;
+}
+
+static void AddSectionToAnimLayerItem(const FAnimLayerSelectionSet& CurrentAnimLayerSet, FAnimLayerItem* AnimLayerItem, UObject* BoundObject, UMovieSceneSection* InSection)
+{
+	FAnimLayerSectionItem NewSectionItem;
+	NewSectionItem.AnimLayerSet.BoundObject = BoundObject;
+	NewSectionItem.Section = InSection;
+	if (UMovieSceneControlRigParameterSection* CRSection = Cast<UMovieSceneControlRigParameterSection>(NewSectionItem.Section))
+	{
+		FMovieSceneFloatChannel* FloatChannel = &CRSection->Weight;
+		FloatChannel->SetDefault(1.0f);
+		CRSection->SetTransformMask(CRSection->GetTransformMask().GetChannels() | EMovieSceneTransformChannel::Weight);
+	}
+	else if (UMovieScene3DTransformSection* Section = Cast<UMovieScene3DTransformSection>(NewSectionItem.Section))
+	{
+		FMovieSceneFloatChannel* FloatChannel = Section->GetWeightChannel();
+		FloatChannel->SetDefault(1.0f);
+		Section->SetMask(Section->GetMask().GetChannels() | EMovieSceneTransformChannel::Weight);
+	}
+	for (const TPair<FName, FAnimLayerPropertyAndChannels>& Set : CurrentAnimLayerSet.Names)
+	{
+		FAnimLayerPropertyAndChannels Channels;
+		Channels.Name = Set.Value.Name;
+		Channels.Channels = Set.Value.Channels;
+		NewSectionItem.AnimLayerSet.Names.Add(Set.Key, Channels);
+	}
+	AnimLayerItem->SectionItems.Add(NewSectionItem);
+}
+
+void UAnimLayers::SetUpBaseLayerSections()
+{
+	if (AnimLayers.Num() > 0)
+	{
+		if (UAnimLayer* BaseAnimLayer = AnimLayers[0])
+		{
+			BaseAnimLayer->Modify();
+			BaseAnimLayer->AnimLayerItems.Reset(); //clear it out
+			for (int32 Index = 1; Index < AnimLayers.Num(); ++Index)
+			{
+				UAnimLayer* AnimLayer = AnimLayers[Index];
+				AnimLayer->Modify();
+				for (const TPair<TWeakObjectPtr<UObject>,FAnimLayerItem>& Pair : AnimLayer->AnimLayerItems)
+				{
+					if (Pair.Key != nullptr)
+					{
+						for (const FAnimLayerSectionItem& SectionItem : Pair.Value.SectionItems)
+						{
+							if (SectionItem.Section.IsValid())
+							{
+								if (UMovieSceneTrack* Track = SectionItem.Section->GetTypedOuter<UMovieSceneTrack>())
+								{
+									const TArray<UMovieSceneSection*>& Sections = Track->GetAllSections();
+									if (Sections.Num() > 1 && Sections[0]->GetBlendType().IsValid() && Sections[0]->GetBlendType() == EMovieSceneBlendType::Absolute)
+									{
+										if (FAnimLayerItem* Existing = BaseAnimLayer->AnimLayerItems.Find(Pair.Key))
+										{
+											if (Pair.Key->IsA<UControlRig>()) //if control rig just merge over control names
+											{
+												for (const TPair<FName, FAnimLayerPropertyAndChannels>& Set : SectionItem.AnimLayerSet.Names)
+												{
+													for (FAnimLayerSectionItem& ExistingSectionItem : Existing->SectionItems)
+													{
+														if (ExistingSectionItem.AnimLayerSet.Names.Find(Set.Key) == nullptr)
+														{
+															FAnimLayerPropertyAndChannels Channels;
+															Channels.Name = Set.Value.Name;
+															Channels.Channels = Set.Value.Channels;
+															ExistingSectionItem.AnimLayerSet.Names.Add(Set.Key, Channels);
+														}
+													}
+												}
+											}
+											else
+											{
+												AddSectionToAnimLayerItem(SectionItem.AnimLayerSet,Existing, Pair.Key.Get(), Sections[0]);
+											}
+										}
+										else
+										{
+											FAnimLayerItem AnimLayerItem;
+											AddSectionToAnimLayerItem(SectionItem.AnimLayerSet, &AnimLayerItem, Pair.Key.Get(), Sections[0]);
+											BaseAnimLayer->AnimLayerItems.Add(Pair.Key, AnimLayerItem);
+										}
+									}
+								}
+							}
+						}
+					}
+				}
+			}
+			//set the name will set it on all base sections
+			FText LayerText = BaseAnimLayer->State.Name;
+			BaseAnimLayer->SetName(LayerText);
+		}
+	}
+	else
+	{
+		AddBaseLayer();
+	}
+}
+
+void UAnimLayers::GetAnimLayerStates(TArray<FAnimLayerState>& OutStates)
+{
+	OutStates.Reset();
+	for (UAnimLayer* AnimLayer : AnimLayers)
+	{
+		if (AnimLayer)
+		{
+			OutStates.Emplace(AnimLayer->State);
+		}
+	}
+}
+
+void UAnimLayers::SetUpSectionDefaults(ISequencer* SequencerPtr, UAnimLayer* Layer, UMovieSceneTrack* Track, UMovieSceneSection* NewSection, FMovieSceneFloatChannel* WeightChannel)
+{
+	int32 OverlapPriority = 0;
+	TMap<int32, int32> NewToOldRowIndices;
+	int32 RowIndex = Track->GetMaxRowIndex() + 1;
+	for (UMovieSceneSection* Section : Track->GetAllSections())
+	{
+		OverlapPriority = FMath::Max(Section->GetOverlapPriority() + 1, OverlapPriority);
+
+		// Move existing sections on the same row or beyond so that they don't overlap with the new section
+		if (Section != NewSection && Section->GetRowIndex() >= RowIndex)
+		{
+			int32 OldRowIndex = Section->GetRowIndex();
+			int32 NewRowIndex = Section->GetRowIndex() + 1;
+			NewToOldRowIndices.FindOrAdd(NewRowIndex, OldRowIndex);
+			Section->Modify();
+			Section->SetRowIndex(NewRowIndex);
+		}
+	}
+
+	Track->Modify();
+
+	Track->OnRowIndicesChanged(NewToOldRowIndices);
+	NewSection->SetRange(TRange<FFrameNumber>::All());
+	
+	NewSection->SetOverlapPriority(OverlapPriority);
+	NewSection->SetRowIndex(RowIndex);
+
+	Track->AddSection(*NewSection);
+	Track->UpdateEasing();
+
+	if (UMovieSceneNameableTrack* NameableTrack = Cast<UMovieSceneNameableTrack>(Track))
+	{
+		NameableTrack->SetTrackRowDisplayName(Layer->GetName(), RowIndex);
+	}
+
+	switch (Layer->State.Type)
+	{
+	case EAnimLayerType::Additive:
+		NewSection->SetBlendType(EMovieSceneBlendType::Additive);
+		break;
+	case EAnimLayerType::Override:
+		NewSection->SetBlendType(EMovieSceneBlendType::Override);
+		break;
+	case EAnimLayerType::Base:
+		NewSection->SetBlendType(EMovieSceneBlendType::Absolute);
+		break;
+	}
+	if (WeightChannel)
+	{
+		WeightChannel->SetDefault(1.0f);
+	}
+}
+
+void UAnimLayers::SetUpControlRigSection(UMovieSceneControlRigParameterSection* ParameterSection, TArray<FName>& ControlNames)
+{
+	if (ParameterSection)
+	{
+		UControlRig* ControlRig = ParameterSection ? ParameterSection->GetControlRig() : nullptr;
+
+		if (ParameterSection && ControlRig)
+		{
+			ParameterSection->Modify();
+			ParameterSection->FillControlsMask(false);
+
+			TArray<FRigControlElement*> Controls;
+			ControlRig->GetControlsInOrder(Controls);
+			for (const FName& RigName : ControlNames)
+			{
+				FChannelMapInfo* pChannelIndex = ParameterSection->ControlChannelMap.Find(RigName);
+				if (pChannelIndex)
+				{
+					ParameterSection->SetControlsMask(pChannelIndex->MaskIndex, true);
+				}
+			}
+		}
+	}
+}
+
+struct FKeyInterval
+{
+	FFrameNumber StartFrame;
+	double StartValue;
+	FFrameNumber EndFrame;
+	double EndValue;
+};
+
+template<typename ChannelType>
+void GetPairs(ChannelType* Channel, TArray<FKeyInterval>& OutKeyIntervals)
+{
+	using ChannelValueType = typename ChannelType::ChannelValueType;
+	TMovieSceneChannelData<ChannelValueType> ChannelData = Channel->GetData();
+
+	TArrayView<FFrameNumber> Times = ChannelData.GetTimes();
+	TArrayView <ChannelValueType> Values = ChannelData.GetValues();
+	OutKeyIntervals.SetNum(0);
+	for (int32 Index = 0; Index < Times.Num() -1; ++Index)
+	{
+		FKeyInterval KeyInterval;
+		KeyInterval.StartFrame = Times[Index];
+		KeyInterval.StartValue = Values[Index].Value;
+		KeyInterval.EndFrame = Times[Index + 1];
+		KeyInterval.EndValue = Values[Index + 1].Value;
+		OutKeyIntervals.Add(KeyInterval);
+	}
+}
+
+template<typename ChannelType>
+void  EvaluateCurveOverRange(ChannelType* Channel, const FFrameNumber& StartTime, const FFrameNumber& EndTime,
+	const FFrameNumber& Interval, TArray<TPair<FFrameNumber, double>>& OutKeys)
+{
+	using CurveValueType = typename ChannelType::CurveValueType;
+	CurveValueType Value;
+	FFrameNumber CurrentTime = StartTime;
+	OutKeys.SetNum(0);
+	while (CurrentTime < EndTime)
+	{
+		Channel->Evaluate(CurrentTime, Value);
+		CurrentTime += Interval;
+		TPair<FFrameNumber, double> OutValue;
+		OutValue.Key = CurrentTime;
+		OutValue.Value = (double) Value;
+		OutKeys.Add(OutValue);
+	}
+}
+
+void GetPercentageOfChange(const TArray<TPair<FFrameNumber, double>>& InKeys, TArray<double>& ValueDifferences,
+TArray<TPair<FFrameNumber, double>>& PercentageDifferences)
+{
+	ValueDifferences.SetNum(InKeys.Num());
+	for (int32 Index = 0; Index < InKeys.Num() -1 ; ++Index)
+	{
+		const double ValueDifference = InKeys[Index + 1].Value - InKeys[Index].Value;
+		ValueDifferences[Index] = ValueDifference;
+	}
+	const double TotalChange = Algo::Accumulate(ValueDifferences, 0.);
+	if (FMath::IsNearlyZero(TotalChange) == false)
+	{
+		const double TotalChangePercentage = (100.0 / TotalChange);
+		PercentageDifferences.SetNum(InKeys.Num());
+		for (int32 Index = 0; Index < InKeys.Num(); ++Index)
+		{
+			PercentageDifferences[Index].Key = InKeys[Index].Key;
+			PercentageDifferences[Index].Value = TotalChangePercentage * ValueDifferences[Index];
+		}
+	}
+	else
+	{
+		PercentageDifferences.SetNum(0);
+	}
+}
+template<typename ChannelType>
+void AddKeyToChannel(ChannelType* Channel, EMovieSceneKeyInterpolation DefaultInterpolation, const FFrameNumber& FrameNumber, double Value)
+{
+	switch (DefaultInterpolation)
+	{
+	case EMovieSceneKeyInterpolation::Linear:
+		Channel->AddLinearKey(FrameNumber, Value);
+		break;
+	case EMovieSceneKeyInterpolation::Constant:
+		Channel->AddConstantKey(FrameNumber, Value);
+		break;
+	case  EMovieSceneKeyInterpolation::Auto:
+		Channel->AddCubicKey(FrameNumber, Value, ERichCurveTangentMode::RCTM_Auto);
+		break;
+	case  EMovieSceneKeyInterpolation::SmartAuto:
+	default:
+		Channel->AddCubicKey(FrameNumber, Value, ERichCurveTangentMode::RCTM_SmartAuto);
+		break;
+	}
+}
+template<typename ChannelType>
+void AdjustmentBlend(UMovieSceneSection* Section, TArrayView<ChannelType*>& BaseChannels, 
+	TArrayView<ChannelType*>& LayerChannels,ISequencer* InSequencer)
+{
+	if (BaseChannels.Num() != LayerChannels.Num())
+	{
+		return;
+	}
+	const FFrameRate& FrameRate = InSequencer->GetFocusedDisplayRate();
+	const FFrameRate& TickResolution = InSequencer->GetFocusedTickResolution();
+	const FFrameNumber Interval = TickResolution.AsFrameNumber(FrameRate.AsInterval());
+	EMovieSceneKeyInterpolation DefaultInterpolation = InSequencer->GetKeyInterpolation();
+
+	TArray<FKeyInterval> KeyIntervals;
+	TArray <TPair<FFrameNumber, double>> Keys;
+	TArray<double> ValueDifferences;
+	TArray<TPair<FFrameNumber, double>> PercentageDifferences;
+	for (int32 Index = 0; Index < BaseChannels.Num(); ++Index)
+	{
+		ChannelType* BaseChannel = BaseChannels[Index];
+		ChannelType* LayerChannel = LayerChannels[Index];
+		KeyIntervals.SetNum(0);
+		Keys.SetNum(0);
+		GetPairs(LayerChannel, KeyIntervals);
+		for (FKeyInterval& KeyInterval : KeyIntervals)
+		{
+			EvaluateCurveOverRange(BaseChannel, KeyInterval.StartFrame, KeyInterval.EndFrame,
+				Interval, Keys);
+			GetPercentageOfChange(Keys, ValueDifferences, PercentageDifferences);
+			const double TotalPoseLayerChange = FMath::Abs(KeyInterval.EndValue - KeyInterval.StartValue);
+			double PreviousValue = KeyInterval.StartValue;
+			for (TPair<FFrameNumber, double>& TimeValue : PercentageDifferences)
+			{
+				const double ValueDelta = (TotalPoseLayerChange / 100.0) * TimeValue.Value;
+				const double CurrentValue = (KeyInterval.EndValue > KeyInterval.StartValue) ?
+					PreviousValue + ValueDelta : PreviousValue - ValueDelta;
+				AddKeyToChannel(LayerChannel, DefaultInterpolation, TimeValue.Key, CurrentValue);
+			}
+		}
+	}
+}
+
+bool UAnimLayers::AdjustmentBlendLayers(ISequencer* InSequencer, int32 LayerIndex)
+{
+	if (InSequencer == nullptr)
+	{
+		return false;
+	}
+	if (LayerIndex < 1 || LayerIndex >= AnimLayers.Num())
+	{
+		return false;
+	}
+	UAnimLayer* BaseLayer = AnimLayers[0];
+	UAnimLayer* AnimLayer = AnimLayers[LayerIndex];
+	const FFrameNumber Min =  TNumericLimits<FFrameNumber>::Lowest();
+	const FFrameNumber Max =  TNumericLimits<FFrameNumber>::Max();
+	TRange<FFrameNumber> Range(Min, Max);
+	FScopedTransaction Transaction(LOCTEXT("AdjustmentBlendLayer", "Adjustment Blend layer"), !GIsTransacting);
+	
+	Modify();
+	AnimLayer->Modify();
+
+
+	for (TPair<TWeakObjectPtr<UObject>, FAnimLayerItem>& Pair : AnimLayer->AnimLayerItems)
+	{
+		if (Pair.Key != nullptr)
+		{
+			for (FAnimLayerSectionItem& SectionItem : Pair.Value.SectionItems)
+			{
+				if (SectionItem.Section.IsValid())
+				{
+					FAnimLayerSectionItem* BaseSectionItem = nullptr;
+					if (FAnimLayerItem* Owner = BaseLayer->AnimLayerItems.Find(Pair.Key))
+					{
+						BaseSectionItem = Owner->FindMatchingSectionItem(SectionItem.Section.Get());
+
+					}
+					if (BaseSectionItem && BaseSectionItem->Section.IsValid())
+					{
+						if (SectionItem.Section->IsActive())//active sections merge them
+						{
+							TArrayView<FMovieSceneFloatChannel*> BaseFloatChannels = BaseSectionItem->Section->GetChannelProxy().GetChannels<FMovieSceneFloatChannel>();
+							TArrayView<FMovieSceneDoubleChannel*> BaseDoubleChannels = BaseSectionItem->Section->GetChannelProxy().GetChannels<FMovieSceneDoubleChannel>();
+							TArrayView<FMovieSceneFloatChannel*> LayerFloatChannels = SectionItem.Section->GetChannelProxy().GetChannels<FMovieSceneFloatChannel>();
+							TArrayView<FMovieSceneDoubleChannel*> LayerDoubleChannels = SectionItem.Section->GetChannelProxy().GetChannels<FMovieSceneDoubleChannel>();
+
+							
+							if (BaseDoubleChannels.Num() > 0)
+							{
+								AdjustmentBlend<FMovieSceneDoubleChannel>(SectionItem.Section.Get(), BaseDoubleChannels,LayerDoubleChannels,  InSequencer);
+							}
+							else if (BaseFloatChannels.Num() > 0)
+							{
+								AdjustmentBlend<FMovieSceneFloatChannel>(SectionItem.Section.Get(), BaseFloatChannels,LayerFloatChannels, InSequencer);
+							}
+						}	
+					}
+				}
+			}
+		}
+	}
+
+	InSequencer->NotifyMovieSceneDataChanged(EMovieSceneDataChangeType::MovieSceneStructureItemsChanged);
+	return true;
+}
+
+#undef LOCTEXT_NAMESPACE
diff --git a/Engine/Plugins/Animation/ControlRig/Source/ControlRigEditor/Private/Sequencer/AnimLayers/AnimLayers.h b/Engine/Plugins/Animation/ControlRig/Source/ControlRigEditor/Private/Sequencer/AnimLayers/AnimLayers.h
new file mode 100644
--- /dev/null
+++ b/Engine/Plugins/Animation/ControlRig/Source/ControlRigEditor/Private/Sequencer/AnimLayers/AnimLayers.h
@@ -0,0 +1,313 @@
+// Copyright Epic Games, Inc. All Rights Reserved.
+#pragma once
+
+#include "CoreMinimal.h"
+#include "UObject/Object.h"
+#include "UObject/SparseDelegate.h"
+#include "Engine/EngineBaseTypes.h"
+#include "Engine/AssetUserData.h"
+#include "Rigs/RigHierarchyDefines.h"
+#include "ControlRig.h"
+#include "Styling/SlateTypes.h"
+#include "AnimLayers.generated.h"
+
+
+class ULevelSequence;
+class UControlRig;
+class UMovieSceneTrack;
+class UMovieSceneSection;
+class ISequencer;
+class UMovieSceneControlRigParameterTrack;
+class UMovieSceneControlRigParameterSection;
+enum class EPropertyKeyedStatus : uint8;
+struct FBakingAnimationKeySettings;
+struct FMovieSceneFloatChannel;
+class IPropertyHandle;
+
+USTRUCT(BlueprintType)
+// Name of a property and control and the specific channels that make up the layer
+struct FAnimLayerPropertyAndChannels
+{
+	GENERATED_BODY()
+
+	FAnimLayerPropertyAndChannels() : Name(NAME_None) , Channels((uint32)EControlRigContextChannelToKey::AllTransform) {};
+	
+	//Name of the property or control
+	UPROPERTY(BlueprintReadWrite, EditAnywhere, Category = "Data")
+	FName Name; 
+
+	//Mask of channels for that property/control, currently not used
+	UPROPERTY(BlueprintReadWrite, EditAnywhere, Category = "Data", meta = (Bitmask, BitmaskEnum = EControlRigContextChannelToKey))
+	int32 Channels;
+};
+
+//Bound object/control rig and the properties/channels it is made of
+//A layer will consist of a bunch of these
+USTRUCT(BlueprintType)
+struct FAnimLayerSelectionSet
+{
+	GENERATED_BODY()
+	FAnimLayerSelectionSet() :BoundObject(nullptr) {};
+
+	UPROPERTY(BlueprintReadWrite, EditAnywhere, Category = "Data")
+	TWeakObjectPtr<UObject> BoundObject; //bound object is either a CR or a bound sequencer object
+
+	UPROPERTY(BlueprintReadWrite, EditAnywhere, Category = "Data")
+	TMap<FName, FAnimLayerPropertyAndChannels> Names; //property/control name and channels
+
+	bool MergeWithAnotherSelection(const FAnimLayerSelectionSet& Selection);
+};
+
+//The set with it's section
+USTRUCT(BlueprintType)
+struct FAnimLayerSectionItem
+{
+	GENERATED_BODY()
+	FAnimLayerSectionItem() : Section(nullptr) {};
+
+	UPROPERTY(BlueprintReadWrite, EditAnywhere, Category = "Data")
+	FAnimLayerSelectionSet AnimLayerSet;
+
+	UPROPERTY(BlueprintReadWrite, EditAnywhere, Category = "Data")
+	TWeakObjectPtr<UMovieSceneSection> Section;
+};
+
+//individual layer items that make up the layer
+USTRUCT(BlueprintType)
+struct FAnimLayerItem
+{
+	GENERATED_BODY()
+	FAnimLayerItem()  {};
+
+	UPROPERTY(BlueprintReadWrite, EditAnywhere, Category = "Data")
+	TArray< FAnimLayerSectionItem>  SectionItems;
+
+	//find section that matches based upon incoming section (same type and property)
+	FAnimLayerSectionItem* FindMatchingSectionItem(UMovieSceneSection* InMovieSceneSection);
+};
+
+UENUM()
+enum EAnimLayerType : uint32
+{
+	Base = 0x0,
+	Additive = 0x1,
+	Override = 0x2,
+};
+ENUM_CLASS_FLAGS(EAnimLayerType)
+
+//state and properties of a layer
+USTRUCT(BlueprintType)
+struct FAnimLayerState
+{
+
+	GENERATED_BODY()
+
+	FAnimLayerState();
+	FText AnimLayerTypeToText()const;
+
+	UPROPERTY(BlueprintReadWrite, EditAnywhere, Category = "Data")
+	mutable ECheckBoxState bKeyed;
+
+	UPROPERTY(BlueprintReadWrite, EditAnywhere, Category = "Data")
+	mutable ECheckBoxState bSelected;
+
+	UPROPERTY(BlueprintReadWrite, EditAnywhere, Category = "Data")
+	mutable bool bActive;
+
+	UPROPERTY(BlueprintReadWrite, EditAnywhere, Category = "Data")
+	mutable bool bLock;
+	
+	UPROPERTY(BlueprintReadWrite, EditAnywhere, Category = "Data")
+	FText Name;
+	
+	UPROPERTY(BlueprintReadWrite, EditAnywhere, Category = "Data")
+	mutable double Weight;
+
+	UPROPERTY(BlueprintReadWrite, EditAnywhere, Category = "Data", meta = (Bitmask, BitmaskEnum = EAnimLayerType))
+	mutable int32 Type;
+
+};
+
+USTRUCT(BlueprintType)
+struct FAnimLayerControlRigObject
+{
+	GENERATED_BODY()
+	UPROPERTY(BlueprintReadWrite, EditAnywhere, Category = "Data")
+	TWeakObjectPtr<UControlRig> ControlRig;
+	UPROPERTY(BlueprintReadWrite, EditAnywhere, Category = "Data")
+	TArray<FName>  ControlNames;
+};
+
+USTRUCT(BlueprintType)
+struct FAnimLayerSceneObject
+{
+	GENERATED_BODY()
+	UPROPERTY(BlueprintReadWrite, EditAnywhere, Category = "Data")
+	TWeakObjectPtr<UObject> SceneObjectOrComponent;
+	
+	//just doing transform for now
+	//UPROPERTY(BlueprintReadWrite, EditAnywhere, Category = "Data")
+	//TArray<FName>  PropertyNames;
+};
+
+USTRUCT(BlueprintType)
+struct FAnimLayerObjects
+{
+	GENERATED_BODY()
+
+	UPROPERTY(BlueprintReadWrite, EditAnywhere, Category = "Data")
+	TArray<FAnimLayerControlRigObject> ControlRigObjects;
+
+	UPROPERTY(BlueprintReadWrite, EditAnywhere, Category = "Data")
+	TArray<FAnimLayerSceneObject> SceneObjects;
+};
+
+UCLASS(EditInlineNew, CollapseCategories, HideDropdown)
+class UAnimLayerWeightProxy : public UObject
+{
+	GENERATED_BODY()
+public:
+
+	UPROPERTY(EditAnywhere, Interp, DisplayName = "", Category = "NoCategory", meta = (SliderExponent = "1.0"))
+	double Weight = 1.0;
+
+	//UObect overrides for setting values when weight changes
+	virtual void PostEditChangeChainProperty(struct FPropertyChangedChainEvent& PropertyChangedEvent) override;
+#if WITH_EDITOR
+	virtual void PostEditUndo() override;
+#endif
+};
+
+//for friend access to State variables directly
+class UAnimLayers;
+struct FAnimLayerSourceUIEntry;
+
+UCLASS(BlueprintType, MinimalAPI)
+class UAnimLayer : public UObject
+{
+	GENERATED_UCLASS_BODY()
+public:
+	UAnimLayer() {};
+
+	void GetAnimLayerObjects(FAnimLayerObjects& InLayerObjects) const;
+
+	UFUNCTION(BlueprintCallable, Category = "State")
+	ECheckBoxState GetKeyed() const;
+	UFUNCTION(BlueprintCallable, Category = "State")
+	void SetKeyed();
+
+	UFUNCTION(BlueprintCallable, Category = "State")
+	ECheckBoxState GetSelected() const;
+	UFUNCTION(BlueprintCallable, Category = "State")
+	void SetSelected(bool bInSelected, bool bClearSelection);
+
+	UFUNCTION(BlueprintCallable, Category = "State")
+	bool GetActive() const;
+	UFUNCTION(BlueprintCallable, Category = "State")
+	void SetActive(bool bInActive);
+
+	UFUNCTION(BlueprintCallable, Category = "State")
+	bool GetLock() const;
+	UFUNCTION(BlueprintCallable, Category = "State")
+	void SetLock(bool bInLock);
+
+	UFUNCTION(BlueprintCallable, Category = "State")
+	FText GetName() const;
+	UFUNCTION(BlueprintCallable, Category = "State")
+	void SetName(const FText& InName);
+
+	UFUNCTION(BlueprintCallable, Category = "State")
+	double GetWeight() const;
+	UFUNCTION(BlueprintCallable, Category = "State")
+	void SetWeight(double InWeight);
+
+	UFUNCTION(BlueprintCallable, Category = "State")
+	EAnimLayerType GetType() const;
+	UFUNCTION(BlueprintCallable, Category = "State")
+	void SetType(EAnimLayerType LayerType);
+
+	UFUNCTION(BlueprintCallable, Category = "State")
+	bool AddSelectedInSequencer();
+	UFUNCTION(BlueprintCallable, Category = "State")
+	bool RemoveSelectedInSequencer();
+
+	bool RemoveAnimLayerItem(UObject* InObject);
+	void SetSelectedInList(bool bInValue) { bIsSelectedInList = bInValue; }
+	bool GetSelectedInList() const { return bIsSelectedInList; }
+
+	//for key property status for weights
+	void SetKey(TSharedPtr<ISequencer>& Sequencer, const IPropertyHandle& KeyedPropertyHandle);
+	EPropertyKeyedStatus GetPropertyKeyedStatus(TSharedPtr<ISequencer>& Sequencer, const IPropertyHandle& PropertyHandle);
+private:
+	UPROPERTY(VisibleAnywhere, Category = "Data")
+	TMap<TWeakObjectPtr<UObject>, FAnimLayerItem> AnimLayerItems;
+
+	UPROPERTY(VisibleAnywhere, Category = "Data")
+	FAnimLayerState State;
+
+	UPROPERTY(Transient)
+	TObjectPtr<UAnimLayerWeightProxy> WeightProxy;
+	
+	void SetSectionToKey() const;
+
+	//static helper functions
+	static bool IsAccepableNonControlRigObject(UObject* InObject);
+
+	bool bIsSelectedInList = false;
+
+	friend class UAnimLayers;
+	friend struct FAnimLayerSourceUIEntry;
+};
+
+/** Link To Set of Anim Sequences that we may belinked to.*/
+UCLASS(BlueprintType, MinimalAPI)
+class UAnimLayers : public UAssetUserData
+{
+	GENERATED_UCLASS_BODY()
+
+public:
+#if WITH_EDITOR
+	virtual void PostEditUndo() override;
+#endif
+	virtual bool IsEditorOnly() const override { return true; }
+
+	UPROPERTY(BlueprintReadWrite, Category = Layers)
+	TArray<TObjectPtr<UAnimLayer>> AnimLayers;
+
+	int32 GetAnimLayerIndex(UAnimLayer* InAnimLayer) const;
+	bool DeleteAnimLayer(ISequencer* InSequencer, int32 Index);
+	int32 DuplicateAnimLayer(ISequencer* InSequencer, int32 Index);
+	int32 AddAnimLayerFromSelection(ISequencer* InSequencer);
+	void GetAnimLayerStates(TArray<FAnimLayerState>& OutStates);
+	bool MergeAnimLayers(ISequencer* InSequencer, const TArray<int32>& Indices, const FBakingAnimationKeySettings* InSettings);
+	
+	//will always blend to base fo rnow
+	bool AdjustmentBlendLayers(ISequencer* InSequencer, int32 LayerIndex);
+
+	bool IsTrackOnSelectedLayer(const UMovieSceneTrack* InTrack)const; 
+	TArray<UMovieSceneSection*> GetSelectedLayerSections() const;
+
+	//Get
+	static UAnimLayers* GetAnimLayers(ISequencer* InSequencer);
+	static UAnimLayers* GetAnimLayers(ULevelSequence* InLevelSequence);
+	static 	TSharedPtr<ISequencer> GetSequencerFromAsset();
+public:
+	DECLARE_EVENT_OneParam(UAnimLayers, FAnimLayerListChanged, UAnimLayers*);
+
+	FAnimLayerListChanged& AnimLayerListChanged() { return OnAnimLayerListChanged; };
+
+private:
+	void AnimLayerListChangedBroadcast();
+	FAnimLayerListChanged OnAnimLayerListChanged;
+	void AddBaseLayer();
+public:
+	void SetUpBaseLayerSections();
+
+	static void SetUpSectionDefaults(ISequencer* SequencerPtr, UAnimLayer* Layer, UMovieSceneTrack* InTrack, UMovieSceneSection* InSection, FMovieSceneFloatChannel* FloatChannel);
+	static void SetUpControlRigSection(UMovieSceneControlRigParameterSection* InSection, TArray<FName>& ControlNames);
+
+};
+
+
+
+
diff --git a/Engine/Plugins/Animation/ControlRig/Source/ControlRigEditor/Private/Sequencer/AnimLayers/SAnimLayers.cpp b/Engine/Plugins/Animation/ControlRig/Source/ControlRigEditor/Private/Sequencer/AnimLayers/SAnimLayers.cpp
new file mode 100644
--- /dev/null
+++ b/Engine/Plugins/Animation/ControlRig/Source/ControlRigEditor/Private/Sequencer/AnimLayers/SAnimLayers.cpp
@@ -0,0 +1,1767 @@
+// Copyright Epic Games, Inc. All Rights Reserved.
+#include "SAnimLayers.h"
+#include "AnimLayers.h"
+#include "Templates/SharedPointer.h"
+#include "Widgets/SCompoundWidget.h"
+#include "Widgets/Views/SListView.h"
+#include "Widgets/Views/STreeView.h"
+#include "SPositiveActionButton.h"
+#include "Widgets/Input/STextComboBox.h"
+#include "Widgets/Input/SNumericEntryBox.h"
+#include "Widgets/Text/SInlineEditableTextBlock.h"
+#include "DetailWidgetRow.h"
+#include <DetailsNameWidgetOverrideCustomization.h>
+#include "IDetailCustomization.h"
+#include "IDetailCustomNodeBuilder.h"
+#include "Containers/Set.h"
+#include "PropertyPath.h"
+#include "Widgets/SWidget.h"
+#include "EditMode/ControlRigEditMode.h"
+#include "Toolkits/IToolkitHost.h"
+#include "EditorModeManager.h"
+#include "ControlRig.h"
+#include "LevelSequencePlayer.h"
+#include "LevelSequenceEditorBlueprintLibrary.h"
+#include "MovieScene.h"
+#include "Editor.h"
+#include "EditorModeManager.h"
+#include "ILevelSequenceEditorToolkit.h"
+#include "ISequencer.h"
+#include "SceneOutlinerPublicTypes.h"
+#include "ControlRigEditorStyle.h"
+#include "LevelEditor.h"
+#include "Modules/ModuleManager.h"
+#include "Sequencer/ControlRigParameterTrackEditor.h"
+
+#define LOCTEXT_NAMESPACE "AnimLayers"
+
+struct FAnimLayerSourceUIEntry;
+
+typedef TSharedPtr<FAnimLayerSourceUIEntry> FAnimLayerSourceUIEntryPtr;
+
+namespace AnimLayerSourceListUI
+{
+	static const FName LayerColumnName(TEXT("Layer"));
+	static const FName StatusColumnName(TEXT("Status"));
+	static const FName WeightColumnName(TEXT("Weight"));
+	static const FName TypeColumnName(TEXT("Type"));
+};
+
+// Structure that defines a single entry in the source UI
+struct FAnimLayerSourceUIEntry
+{
+public:
+	FAnimLayerSourceUIEntry(UAnimLayer* InAnimLayer): AnimLayer(InAnimLayer)
+	{}
+
+	void SetSelectedInList(bool bInValue);
+	bool GetSelectedInList() const;
+
+	void SelectObjects() const;
+	void AddSelected() const;
+	void RemoveSelected() const;
+	void Duplicate() const;
+	void DeleteAnimLayer() const;
+
+	ECheckBoxState GetKeyed() const;
+	void SetKeyed();
+	FReply OnKeyedColor();
+	FSlateColor GetKeyedColor() const;
+
+	ECheckBoxState GetSelected() const;
+	void SetSelected(bool bInSelected);
+
+	bool GetActive() const;
+	void SetActive(bool bInActive);
+
+	bool GetLock() const;
+	void SetLock(bool bInLock);
+
+	FText GetName() const;
+	void SetName(const FText& InName);
+
+	double GetWeight() const;
+	void SetWeight(double InWeight);
+
+	EAnimLayerType GetType() const;
+	FText GetTypeToText() const;
+	void SetType(EAnimLayerType InType);
+
+	UObject* GetWeightObject() const;
+
+	int32 GetAnimLayerIndex(UAnimLayers* AnimLayers) const;
+	void ClearCaches() const { bSelectionStateValid = false; bKeyedStateIsValid = false; }
+private:
+	//cached values
+	mutable bool bSelectionStateValid = false;
+	mutable bool bKeyedStateIsValid = false;
+	mutable ECheckBoxState SelectionState = ECheckBoxState::Unchecked;
+	mutable ECheckBoxState KeyedState = ECheckBoxState::Unchecked;
+
+	UAnimLayer* AnimLayer = nullptr;
+};
+
+int32 FAnimLayerSourceUIEntry::GetAnimLayerIndex(UAnimLayers* AnimLayers) const
+{
+	if (AnimLayers && AnimLayer)
+	{
+		return AnimLayers->AnimLayers.Find(AnimLayer);
+	}
+	return INDEX_NONE;
+}
+
+UObject* FAnimLayerSourceUIEntry::GetWeightObject() const
+{
+	if (AnimLayer)
+	{
+		return AnimLayer->WeightProxy.Get();
+	}
+	return nullptr;
+}
+
+void FAnimLayerSourceUIEntry::SelectObjects() const
+{
+	if (AnimLayer)
+	{
+		AnimLayer->SetSelected(true, !(FSlateApplication::Get().GetModifierKeys().IsShiftDown()));
+		ClearCaches();
+	}
+}
+
+void FAnimLayerSourceUIEntry::AddSelected() const
+{
+	if (AnimLayer)
+	{
+
+		AnimLayer->AddSelectedInSequencer();
+		ClearCaches();
+	}
+}
+
+void FAnimLayerSourceUIEntry::RemoveSelected() const
+{
+	if (AnimLayer)
+{
+		AnimLayer->RemoveSelectedInSequencer();
+		ClearCaches();
+	}
+}
+
+void FAnimLayerSourceUIEntry::DeleteAnimLayer() const
+{
+	if (AnimLayer)
+	{
+		if (TSharedPtr<ISequencer> SequencerPtr = UAnimLayers::GetSequencerFromAsset())
+		{
+			if (UAnimLayers* AnimLayers = UAnimLayers::GetAnimLayers(SequencerPtr.Get()))
+			{
+				int32 Index = AnimLayers->GetAnimLayerIndex(AnimLayer);
+				if (Index != INDEX_NONE)
+				{
+					AnimLayers->DeleteAnimLayer(SequencerPtr.Get(), Index);
+				}
+			}
+		}
+		ClearCaches();
+	}
+}
+void FAnimLayerSourceUIEntry::Duplicate() const
+{
+	if (AnimLayer)
+	{
+		if (TSharedPtr<ISequencer> SequencerPtr = UAnimLayers::GetSequencerFromAsset())
+		{
+			if (UAnimLayers* AnimLayers = UAnimLayers::GetAnimLayers(SequencerPtr.Get()))
+			{
+				int32 Index = AnimLayers->GetAnimLayerIndex(AnimLayer);
+				if (Index != INDEX_NONE)
+				{
+					AnimLayers->DuplicateAnimLayer(SequencerPtr.Get(), Index);
+				}
+			}
+		}
+		ClearCaches();
+	}
+}
+
+ECheckBoxState FAnimLayerSourceUIEntry::GetKeyed() const
+{
+	if (bKeyedStateIsValid == false)
+	{
+		bKeyedStateIsValid = true;
+		if (AnimLayer)
+		{
+			KeyedState = AnimLayer->GetKeyed();
+		}
+	}
+	return KeyedState;
+}
+
+void FAnimLayerSourceUIEntry::SetKeyed()
+{
+	if (AnimLayer)
+	{
+		bKeyedStateIsValid = false;
+		AnimLayer->SetKeyed();
+	}
+}
+
+FReply FAnimLayerSourceUIEntry::OnKeyedColor()
+{
+	if (AnimLayer)
+	{
+		const FScopedTransaction Transaction(LOCTEXT("SetKeyed_Transaction", "Set Keyed"), !GIsTransacting);
+		AnimLayer->SetKeyed();
+	}
+	return FReply::Handled();
+}
+
+FSlateColor FAnimLayerSourceUIEntry::GetKeyedColor() const
+{
+	ECheckBoxState State = GetKeyed();
+	switch (State)
+	{
+		case ECheckBoxState::Undetermined:
+		{
+			return (FLinearColor::Green /2);
+		};
+		case ECheckBoxState::Checked:
+		{
+			return FLinearColor::Green;
+		};
+	}
+	return FLinearColor::Transparent;
+}
+
+ECheckBoxState FAnimLayerSourceUIEntry::GetSelected() const
+{
+	if (bSelectionStateValid == false)
+	{
+		bSelectionStateValid = true;
+		if (AnimLayer)
+		{
+			SelectionState = AnimLayer->GetSelected();
+		}
+	}
+	return SelectionState;
+}
+
+void FAnimLayerSourceUIEntry::SetSelected(bool bInSelected)
+{
+	if (AnimLayer)
+	{
+		bSelectionStateValid = false;
+		AnimLayer->SetSelected(bInSelected, !(FSlateApplication::Get().GetModifierKeys().IsShiftDown()));
+	}
+}
+
+void FAnimLayerSourceUIEntry::SetSelectedInList(bool bInValue)
+{ 
+	if (AnimLayer)
+	{
+		AnimLayer->SetSelectedInList(bInValue);
+		if (bInValue)
+		{
+			AnimLayer->SetKeyed();//selection also sets keyed
+		}
+	}
+}
+
+bool FAnimLayerSourceUIEntry::GetSelectedInList() const 
+{ 
+	if (AnimLayer)
+	{
+		return AnimLayer->GetSelectedInList();
+	}
+	return false;
+}
+
+bool FAnimLayerSourceUIEntry::GetActive() const
+{
+	if (AnimLayer)
+	{
+		return AnimLayer->GetActive();
+	}
+	return false;
+}
+
+void FAnimLayerSourceUIEntry::SetActive(bool bInActive)
+{
+	if (AnimLayer)
+	{
+		AnimLayer->SetActive(bInActive);
+	}
+}
+
+bool FAnimLayerSourceUIEntry::GetLock() const
+{
+	if (AnimLayer)
+	{
+		return AnimLayer->GetLock();
+	}
+	return false;
+}
+
+void FAnimLayerSourceUIEntry::SetLock(bool bInLock)
+{
+	if (AnimLayer)
+	{
+		AnimLayer->SetLock(bInLock);
+	}
+}
+
+FText FAnimLayerSourceUIEntry::GetName() const
+{
+	if (AnimLayer)
+	{
+		return AnimLayer->GetName();
+	}
+	return FText();
+}
+
+void FAnimLayerSourceUIEntry::SetName(const FText& InName)
+{
+	if (AnimLayer)
+	{
+		AnimLayer->SetName(InName);
+	}
+}
+
+double FAnimLayerSourceUIEntry::GetWeight() const
+{
+	if (AnimLayer)
+	{
+		return AnimLayer->GetWeight();
+	}
+	return 0.0;
+}
+
+void FAnimLayerSourceUIEntry::SetWeight(double InWeight)
+{
+	if (AnimLayer)
+	{
+		AnimLayer->SetWeight(InWeight);
+	}
+}
+
+EAnimLayerType FAnimLayerSourceUIEntry::GetType() const
+{
+	if (AnimLayer)
+	{
+		return AnimLayer->GetType();
+	}
+	return EAnimLayerType::Base;
+}
+
+
+void FAnimLayerSourceUIEntry::SetType(EAnimLayerType InType)
+{
+	if (AnimLayer)
+	{
+		AnimLayer->SetType(InType);
+	}
+}
+
+FText FAnimLayerSourceUIEntry::GetTypeToText() const
+{
+	if (AnimLayer)
+	{
+		return AnimLayer->State.AnimLayerTypeToText();
+	}
+	return FText();
+}
+
+namespace UE::AnimLayers
+{
+	/** Handles deletion */
+	template <typename ListType, typename ListElementType>
+	class SAnimLayersBaseListView : public ListType
+	{
+	public:
+		virtual FReply OnKeyDown(const FGeometry& MyGeometry, const FKeyEvent& InKeyEvent) override
+		{
+			if (InKeyEvent.GetKey() == EKeys::Delete || InKeyEvent.GetKey() == EKeys::BackSpace)
+			{
+				TArray<ListElementType> SelectedItem = ListType::GetSelectedItems();
+				for (ListElementType Item : SelectedItem)
+				{
+					Item->DeleteAnimLayer();
+				}
+				
+				return FReply::Handled();
+			}
+
+			return FReply::Unhandled();
+		}
+	};
+}
+
+class SAnimLayerListView : public UE::AnimLayers::SAnimLayersBaseListView<STreeView<FAnimLayerSourceUIEntryPtr>, FAnimLayerSourceUIEntryPtr>
+{
+public:
+	void Construct(const FArguments& InArgs)
+	{
+		UE::AnimLayers::SAnimLayersBaseListView<STreeView<FAnimLayerSourceUIEntryPtr>, FAnimLayerSourceUIEntryPtr>::Construct(InArgs);
+	}
+};
+
+class FAnimLayerSourcesView : public TSharedFromThis<FAnimLayerSourcesView>
+{
+public:
+	FAnimLayerSourcesView();
+
+	// Gather information about all sources and update the list view 
+	void RefreshSourceData(bool bRefreshUI);
+	// Handler that creates a widget row for a given ui entry
+	TSharedRef<ITableRow> MakeSourceListViewWidget(FAnimLayerSourceUIEntryPtr Entry, const TSharedRef<STableViewBase>& OwnerTable) const;
+	// Handles constructing a context menu for the sources
+	TSharedPtr<SWidget> OnSourceConstructContextMenu();;
+	// Handle selection change, triggering the OnSourceSelectionChangedDelegate delegate.
+	void OnSourceListSelectionChanged(FAnimLayerSourceUIEntryPtr Entry, ESelectInfo::Type SelectionType) const;
+	// Focus on the added layer name
+	void FocusRenameLayer(int32 Index)
+	{
+		FocusOnIndex = Index;
+	}
+
+	void AddController(FAnimLayerController* InController);
+
+	void RenameItem(int32 Index) const;
+private:
+	// Create the sources list view
+	void CreateSourcesListView();
+
+	mutable int32  FocusOnIndex = INDEX_NONE;
+
+
+	void AddSelected();
+	void RemoveSelected();
+	void SelectObjects();
+	void Duplicate();
+	void MergeLayers();
+	void AdjustmentBlend();
+	void DeleteAnimLayer();
+	void Rename();
+
+public:
+	// Holds the data that will be displayed in the list view
+	TArray<FAnimLayerSourceUIEntryPtr> SourceData;
+	// Holds the sources list view widget
+	TSharedPtr<SAnimLayerListView> SourcesListView;
+	//pointer to controller safe to be raw since it's the parent.
+	FAnimLayerController* Controller = nullptr;
+
+};
+
+struct  FAnimLayerController : public TSharedFromThis<FAnimLayerController>
+{
+public:
+	FAnimLayerController();
+	~FAnimLayerController();
+private:
+	bool bSelectionFilterActive = false;
+
+private:
+
+	void RebuildSourceList();
+
+	// Handles source selection changing.
+	void OnSourceSelectionChangedHandler(FAnimLayerSourceUIEntryPtr Entry, ESelectInfo::Type SelectionType) const;
+public:
+	// Handles the source collection changing.
+	void RefreshSourceData(bool bRefreshUI);
+	void RefreshTimeDependantData();
+	void RefreshSelectionData();
+	void HandleOnAnimLayerListChanged(UAnimLayers*) { RebuildSourceList(); }
+	void FocusRenameLayer(int32 Index) { if (SourcesView.IsValid()) SourcesView->FocusRenameLayer(Index); }
+	void SelectItem(int32 Index, bool bClear = true);
+	void ToggleSelectionFilterActive();
+	bool IsSelectionFilterActive() const;
+public:
+	// Sources view
+	TSharedPtr<FAnimLayerSourcesView> SourcesView;
+	// Guard from reentrant selection
+	mutable bool bSelectionChangedGuard = false;
+};
+
+
+class SAnimLayerSourcesRow : public SMultiColumnTableRow<FAnimLayerSourceUIEntryPtr>
+{
+public:
+	SLATE_BEGIN_ARGS(SAnimLayerSourcesRow) {}
+	/** The list item for this row */
+	SLATE_ARGUMENT(FAnimLayerSourceUIEntryPtr, Entry)
+	SLATE_END_ARGS()
+
+	void Construct(const FArguments& Args, const TSharedRef<STableViewBase>& OwnerTableView)
+	{
+		EntryPtr = Args._Entry;
+		LayerTypeTextList.Emplace(new FText(LOCTEXT("Additive", "Additive")));
+		LayerTypeTextList.Emplace(new FText(LOCTEXT("Override", "Override")));
+		SMultiColumnTableRow<FAnimLayerSourceUIEntryPtr>::Construct(
+			FSuperRowType::FArguments()
+			.Padding(2.f),
+			OwnerTableView
+		);
+	}
+
+	// Static Source UI FNames
+	void BeginEditingName()
+	{
+		if (LayerNameTextBlock.IsValid())
+		{
+			LayerNameTextBlock->EnterEditingMode();
+		}
+	}
+
+	FSlateColor GetKeyedColor() const
+	{
+		if (EntryPtr)
+		{
+			return EntryPtr->GetKeyedColor();
+		}
+		return FLinearColor::Transparent;
+	}
+
+	/** Overridden from SMultiColumnTableRow.  Generates a widget for this column of the list view. */
+	virtual TSharedRef<SWidget> GenerateWidgetForColumn(const FName& ColumnName) override
+	{
+		if (ColumnName == AnimLayerSourceListUI::LayerColumnName)
+		{
+			return SNew(SHorizontalBox)
+				+ SHorizontalBox::Slot()
+				.MaxWidth(6)
+				[
+					SNew(SBox)
+						.WidthOverride(6)
+						[
+							SNew(SButton)
+								.ContentPadding(0)
+								.VAlign(VAlign_Fill)
+								.IsFocusable(true)
+								.IsEnabled(true)
+								.ButtonStyle(FAppStyle::Get(), "Sequencer.AnimationOutliner.ColorStrip")
+								.OnClicked_Lambda([this]()
+									{
+										return EntryPtr->OnKeyedColor();
+									})
+								.Content()
+								[
+									SNew(SImage)
+										.Image(FAppStyle::GetBrush("WhiteBrush"))
+										.ColorAndOpacity(this, &SAnimLayerSourcesRow::GetKeyedColor)
+								]
+						]
+				]
+				+ SHorizontalBox::Slot()
+				.AutoWidth()
+				.Padding(5, 0)
+				.VAlign(VAlign_Center)
+				[
+					SNew(SButton)
+					.IsFocusable(false)
+					.ButtonStyle(FAppStyle::Get(), "NoBorder")
+					.ButtonColorAndOpacity_Lambda([this]()
+						{
+							return FLinearColor::White;
+						})
+					.OnClicked_Lambda([this]()
+						{
+							bool bValue = EntryPtr->GetSelected() == ECheckBoxState::Unchecked ? false : true;
+							EntryPtr->SetSelected(!bValue);
+							return FReply::Handled();
+						})
+					.ContentPadding(1.f)
+					.ToolTipText(LOCTEXT("AnimLayerSelectionTooltip", "Selection In Layer"))
+					[
+						SNew(SImage)
+							.ColorAndOpacity_Lambda([this]()
+								{
+									const FColor Selected(38, 187, 255);
+									const FColor NotSelected(56, 56, 56);
+									//FLinearColor::White;
+									bool bValue = EntryPtr->GetSelected() == ECheckBoxState::Unchecked ? false : true;
+									FSlateColor SlateColor;
+									if (bValue == true)
+									{
+										if (EntryPtr->GetSelectedInList())
+										{
+											SlateColor = FLinearColor::White;
+										}
+										else
+										{
+											SlateColor = FSlateColor(Selected);
+										}										
+									}
+									else
+									{
+										SlateColor = FSlateColor(NotSelected);
+									}
+									return SlateColor;
+								})
+							.Image(FSlateIcon(FControlRigEditorStyle::Get().GetStyleSetName(), "ControlRig.AnimLayerSelected").GetIcon())
+					]
+				]
+				+ SHorizontalBox::Slot()
+				.FillWidth(10.0)
+				.Padding(10.0,0)
+				[	SNew(SBox)
+					.HAlign(EHorizontalAlignment::HAlign_Left)
+					[
+						SAssignNew(LayerNameTextBlock, SInlineEditableTextBlock)
+							.Justification(ETextJustify::Center)
+							.Text_Lambda([this]()
+								{
+									return EntryPtr->GetName();
+								})
+							.OnTextCommitted(this, &SAnimLayerSourcesRow::OnLayerNameCommited)
+					]
+				];
+		}
+		else if (ColumnName == AnimLayerSourceListUI::StatusColumnName)
+		{
+			return SNew(SHorizontalBox)
+				+ SHorizontalBox::Slot()
+					.AutoWidth()
+					.Padding(1, 0)
+					.VAlign(VAlign_Center)
+					[
+						SAssignNew(MuteButton, SButton)
+						.IsFocusable(false)
+						.ButtonStyle(FAppStyle::Get(), "NoBorder")
+						.ButtonColorAndOpacity_Lambda([this] ()
+						{
+							const bool bIsActive = !EntryPtr->GetActive();
+							const bool bIsHovered = MuteButton->IsHovered();
+							return GetStatusImageColorAndOpacity(bIsActive, bIsHovered);		
+						})
+						.OnClicked_Lambda([this]()
+						{
+							bool bValue = EntryPtr->GetActive();
+							EntryPtr->SetActive(!bValue);
+							return FReply::Handled();
+						})
+						.ContentPadding(1.f)
+						.ToolTipText(LOCTEXT("AnimLayerMuteTooltip", "Mute Layer"))
+						[
+							SNew(SImage)
+								.ColorAndOpacity_Lambda([this]() 
+								{ 
+									const bool bIsActive = !EntryPtr->GetActive();
+									const bool bIsHovered = MuteButton->IsHovered();
+									return GetStatusImageColorAndOpacity(bIsActive, bIsHovered);
+								})
+								.Image(FAppStyle::GetBrush("Sequencer.Column.Mute"))
+						]
+					]
+				+ SHorizontalBox::Slot()
+					.AutoWidth()
+					.Padding(1, 0)
+					.VAlign(VAlign_Center)
+					[
+						SAssignNew(LockButton, SButton)
+						.IsFocusable(false)
+						.ButtonStyle(FAppStyle::Get(), "NoBorder")
+						.ButtonColorAndOpacity_Lambda([this] ()
+						{
+							const bool bIsLock = EntryPtr->GetLock();
+							const bool bIsHovered = LockButton->IsHovered();
+							return GetStatusImageColorAndOpacity(bIsLock, bIsHovered);
+						})
+						.OnClicked_Lambda([this]()
+						{
+							bool bValue = EntryPtr->GetLock();
+							EntryPtr->SetLock(!bValue);
+							return FReply::Handled();
+						})
+						.ContentPadding(1.f)
+						.ToolTipText(LOCTEXT("AnimLayerLockTooltip", "Lock Layer"))
+						[
+							SNew(SImage)
+								.ColorAndOpacity_Lambda([this]() 
+								{ 
+									const bool bIsLock = EntryPtr->GetLock();
+									const bool bIsHovered = LockButton->IsHovered();
+									return GetStatusImageColorAndOpacity(bIsLock, bIsHovered);
+								})
+							.Image(FAppStyle::GetBrush("Sequencer.Column.Locked"))
+						]
+					];
+
+
+		}
+		else if (ColumnName == AnimLayerSourceListUI::WeightColumnName)
+		{
+			UObject* WeightObject = nullptr;
+			FControlRigEditMode* EditMode = nullptr;
+			if (TSharedPtr<ISequencer> SequencerPtr = UAnimLayers::GetSequencerFromAsset())
+			{
+				WeightObject = EntryPtr->GetWeightObject();
+
+				TSharedPtr<IToolkitHost> ToolkitHost = SequencerPtr->GetToolkitHost();
+				if (ToolkitHost.IsValid())
+				{
+					FEditorModeTools& EditorModeTools = ToolkitHost->GetEditorModeManager();			
+					if (!EditorModeTools.IsModeActive(FControlRigEditMode::ModeName))
+					{
+						EditorModeTools.ActivateMode(FControlRigEditMode::ModeName);
+
+						EditMode = static_cast<FControlRigEditMode*>(EditorModeTools.GetActiveMode(FControlRigEditMode::ModeName));
+						if (EditMode && EditMode->GetToolkit().IsValid() == false)
+						{
+							EditMode->Enter();
+						}
+					}
+					EditMode = static_cast<FControlRigEditMode*>(EditorModeTools.GetActiveMode(FControlRigEditMode::ModeName));	
+				}
+			}
+			return SAssignNew(WeightDetails, SAnimWeightDetails, EditMode, WeightObject);
+			
+		}
+		else if (ColumnName == AnimLayerSourceListUI::TypeColumnName)
+		{
+			if(EntryPtr->GetType() != (int32)(EAnimLayerType::Base))
+			{ 
+				return SAssignNew(LayerTypeCombo, SComboBox<TSharedPtr<FText>>)
+					.ContentPadding(FMargin(10.0f, 2.0f))
+					.OptionsSource(&LayerTypeTextList)
+					.HasDownArrow(false)
+					.OnGenerateWidget_Lambda([this](TSharedPtr<FText> Item)
+					{
+						return SNew(SBox)
+							.MaxDesiredWidth(100.0f)
+							[
+								SNew(STextBlock)
+								.TextStyle(FAppStyle::Get(), "AnimViewport.MessageText")
+								.Text(*Item)
+							];
+					})
+					.OnSelectionChanged(this, &SAnimLayerSourcesRow::OnLayerTypeChanged)
+					[
+						SNew(STextBlock)
+						.TextStyle(FAppStyle::Get(), "AnimViewport.MessageText")
+						.Text(this, &SAnimLayerSourcesRow::GetLayerTypeText)
+					];
+			}
+			else
+			{
+				return SNullWidget::NullWidget;
+			}
+		}
+
+		return SNullWidget::NullWidget;
+	}
+
+private:
+	
+	FSlateColor GetStatusImageColorAndOpacity(bool bIsActive, bool bIsHovered) const
+	{
+		FLinearColor OutColor = FLinearColor::White;
+		float Opacity = 0.0f;
+
+		if (bIsActive)
+		{
+			// Directly active, full opacity
+			Opacity = 1.0f;
+		}
+		else if (bIsHovered)
+		{
+			// Mouse is over widget and it is not directly active.
+			Opacity = .8f;
+		}
+		else 
+		{
+			// Not active in any way and mouse is not over widget or item.
+			Opacity = 0.1f;
+		}
+		OutColor.A = Opacity;
+		return OutColor;
+	}
+
+
+	void OnLayerNameCommited(const FText& InNewText, ETextCommit::Type InTextCommit)
+	{
+		if (InNewText.IsEmpty())
+		{
+			return;
+		}
+		EntryPtr->SetName(InNewText);
+	}
+
+	FText GetLayerTypeText() const
+	{
+		return EntryPtr->GetTypeToText();
+	}
+	void OnLayerTypeChanged(TSharedPtr<FText> InNewText, ESelectInfo::Type SelectInfo)
+	{
+		if (InNewText.IsValid() == false || InNewText->IsEmpty())
+		{
+			return;
+		}
+		FText Additive(LOCTEXT("Additive", "Additive"));
+		FText Override(LOCTEXT("Override", "Override"));
+		if (InNewText->IdenticalTo(Additive))
+		{
+			EntryPtr->SetType(EAnimLayerType::Additive);
+		}
+		else if (InNewText->IdenticalTo(Override))
+		{
+			EntryPtr->SetType(EAnimLayerType::Override);
+		}
+	}
+
+
+private:
+	FAnimLayerSourceUIEntryPtr EntryPtr;
+	TSharedPtr<SInlineEditableTextBlock> LayerNameTextBlock;
+	mutable TArray<TSharedPtr<FText>> LayerTypeTextList;
+	TSharedPtr<SComboBox<TSharedPtr<FText>>> LayerTypeCombo;
+	//TSharedPtr<SNumericEntryBox<double>> WeightBox;
+	TSharedPtr< SAnimWeightDetails> WeightDetails;
+	TSharedPtr<SButton> MuteButton;
+	TSharedPtr<SButton> LockButton;
+
+};
+
+FAnimLayerSourcesView::FAnimLayerSourcesView()
+{
+	CreateSourcesListView();
+}
+
+void FAnimLayerSourcesView::AddController(FAnimLayerController* InController)
+{
+	Controller = InController;
+}
+
+TSharedRef<ITableRow> FAnimLayerSourcesView::MakeSourceListViewWidget(FAnimLayerSourceUIEntryPtr Entry, const TSharedRef<STableViewBase>& OwnerTable) const
+{
+	if (FocusOnIndex != INDEX_NONE)
+	{
+		int32 NextIndex = FocusOnIndex;
+		FocusOnIndex = INDEX_NONE;
+		GEditor->GetTimerManager()->SetTimerForNextTick([this, NextIndex]()
+			{
+				if (TSharedPtr<ISequencer> SequencerPtr = UAnimLayers::GetSequencerFromAsset())
+				{
+					if (UAnimLayers* AnimLayers = UAnimLayers::GetAnimLayers(SequencerPtr.Get()))
+					{
+						for (int32 Index = 0; Index < SourceData.Num(); ++Index)
+						{
+							const FAnimLayerSourceUIEntryPtr& Item = SourceData[Index];
+							if (Item->GetAnimLayerIndex(AnimLayers) == NextIndex)
+							{
+								RenameItem(Index);
+								FocusOnIndex = INDEX_NONE;
+							}
+						}
+					}
+				}
+			});
+	}
+	return SNew(SAnimLayerSourcesRow, OwnerTable)
+		.Entry(Entry);
+}
+
+void FAnimLayerSourcesView::RenameItem(int32 Index) const
+{
+	if (Index < SourceData.Num())
+	{
+		if (TSharedPtr<ITableRow> Row = SourcesListView->WidgetFromItem(SourceData[Index]))
+		{
+			TWeakPtr<SAnimLayerSourcesRow> Widget = StaticCastSharedRef<SAnimLayerSourcesRow>(Row->AsWidget());
+			if (Widget.IsValid())
+			{
+				Widget.Pin()->BeginEditingName();
+			}
+		}
+	}
+}
+
+void FAnimLayerSourcesView::CreateSourcesListView()
+{
+
+	SAssignNew(SourcesListView, SAnimLayerListView)
+	//	.ListItemsSource(&SourceData)
+		.TreeItemsSource(&SourceData)
+		.OnGetChildren_Lambda([](FAnimLayerSourceUIEntryPtr ChannelItem, TArray<FAnimLayerSourceUIEntryPtr>& OutChildren)
+		{
+		})
+		.SelectionMode(ESelectionMode::Multi)
+		.OnGenerateRow_Raw(this, &FAnimLayerSourcesView::MakeSourceListViewWidget)
+		.OnContextMenuOpening_Raw(this, &FAnimLayerSourcesView::OnSourceConstructContextMenu)
+		.OnSelectionChanged_Raw(this, &FAnimLayerSourcesView::OnSourceListSelectionChanged)
+		.HeaderRow
+		(
+			SNew(SHeaderRow)
+			+ SHeaderRow::Column(AnimLayerSourceListUI::LayerColumnName)
+			.FillWidth(120.f)
+			.DefaultLabel(LOCTEXT("LayerColumnName", "Name"))
+			+ SHeaderRow::Column(AnimLayerSourceListUI::StatusColumnName)
+			.FillWidth(20.f)
+			.DefaultLabel(LOCTEXT("StatusColumnName", "Status"))
+			+ SHeaderRow::Column(AnimLayerSourceListUI::WeightColumnName)
+			.FillWidth(60.f)
+			.DefaultLabel(LOCTEXT("WeightColumnName", "Weight"))
+			+ SHeaderRow::Column(AnimLayerSourceListUI::TypeColumnName)
+			.FillWidth(80.f)
+			.DefaultLabel(LOCTEXT("TypeColumnName", "Type"))
+		);
+}
+
+void FAnimLayerSourcesView::AddSelected()
+{
+	TArray<FAnimLayerSourceUIEntryPtr> Selected;
+	SourcesListView->GetSelectedItems(Selected);
+	for (const FAnimLayerSourceUIEntryPtr& Ptr : Selected)
+	{
+		Ptr->AddSelected();
+	}
+}
+void FAnimLayerSourcesView::RemoveSelected()
+{
+	TArray<FAnimLayerSourceUIEntryPtr> Selected;
+	SourcesListView->GetSelectedItems(Selected);
+	for (const FAnimLayerSourceUIEntryPtr& Ptr : Selected)
+	{
+		Ptr->RemoveSelected();
+	}
+}
+
+void FAnimLayerSourcesView::SelectObjects()
+{
+	TArray<FAnimLayerSourceUIEntryPtr> Selected;
+	SourcesListView->GetSelectedItems(Selected);
+	for(const FAnimLayerSourceUIEntryPtr& Ptr: Selected)
+	{
+		Ptr->SelectObjects();
+	}
+}
+
+void FAnimLayerSourcesView::DeleteAnimLayer()
+{
+	TArray<FAnimLayerSourceUIEntryPtr> Selected;
+	SourcesListView->GetSelectedItems(Selected);
+	for (const FAnimLayerSourceUIEntryPtr& Ptr : Selected)
+	{
+		Ptr->DeleteAnimLayer();
+	}
+}
+void FAnimLayerSourcesView::Duplicate()
+{
+	TArray<FAnimLayerSourceUIEntryPtr> Selected;
+	SourcesListView->GetSelectedItems(Selected);
+	for (const FAnimLayerSourceUIEntryPtr& Ptr : Selected)
+	{
+		Ptr->Duplicate();
+	}
+}
+
+void FAnimLayerSourcesView::AdjustmentBlend()
+{
+	if (TSharedPtr<ISequencer> SequencerPtr = UAnimLayers::GetSequencerFromAsset())
+	{
+		if (UAnimLayers* AnimLayers = UAnimLayers::GetAnimLayers(SequencerPtr.Get()))
+		{
+			TArray<FAnimLayerSourceUIEntryPtr> Selected;
+			SourcesListView->GetSelectedItems(Selected);
+			if (Selected.Num() != 1)
+			{
+				return;
+			}
+			int32 Index = Selected[0]->GetAnimLayerIndex(AnimLayers);
+			if (Index != INDEX_NONE && Index != 0) //not base
+			{
+				AnimLayers->AdjustmentBlendLayers(SequencerPtr.Get(), Index);
+			}
+		}
+	}
+}
+
+void FAnimLayerSourcesView::MergeLayers()
+{
+	if (TSharedPtr<ISequencer> SequencerPtr = UAnimLayers::GetSequencerFromAsset())
+	{
+		if (UAnimLayers* AnimLayers = UAnimLayers::GetAnimLayers(SequencerPtr.Get()))
+		{
+			TArray<FAnimLayerSourceUIEntryPtr> Selected;
+			SourcesListView->GetSelectedItems(Selected);
+			if (Selected.Num() < 2)
+			{
+				return;
+			}
+			TArray<int32> LayersToMerge;
+			for (const FAnimLayerSourceUIEntryPtr& Ptr : Selected)
+			{
+				int32 Index = Ptr->GetAnimLayerIndex(AnimLayers);
+				if (Index != INDEX_NONE)
+				{
+					LayersToMerge.Add(Index);
+				}
+			}
+			if (LayersToMerge.Num() > 1)
+			{
+				AnimLayers->MergeAnimLayers(SequencerPtr.Get(), LayersToMerge, nullptr);
+				//for now we don't use the dialog since it's not doing proper per key frame reductions
+				/*
+				FCollapseControlsCB CollapseCB = FCollapseControlsCB::CreateLambda([AnimLayers,LayersToMerge](TSharedPtr<ISequencer>& InSequencer,const FBakingAnimationKeySettings& InSettings)
+				{
+					AnimLayers->MergeAnimLayers(InSequencer.Get(), LayersToMerge, InSettings);
+
+				});
+
+				TSharedRef<SCollapseControlsWidget> BakeWidget =
+					SNew(SCollapseControlsWidget)
+					.Sequencer(SequencerPtr);
+
+				BakeWidget->SetCollapseCB(CollapseCB);
+				BakeWidget->OpenDialog(false);
+				*/
+			}
+		}
+	}
+}
+
+void FAnimLayerSourcesView::Rename()
+{
+	TArray<FAnimLayerSourceUIEntryPtr> Selected;
+	SourcesListView->GetSelectedItems(Selected);
+	if (Selected.Num() == 1)
+	{
+		int32 Index = INDEX_NONE;
+		if (SourceData.Find(Selected[0], Index))
+		{
+			if (Index != INDEX_NONE)
+			{
+				RenameItem(Index);
+			}
+		}
+	}
+}
+
+TSharedPtr<SWidget> FAnimLayerSourcesView::OnSourceConstructContextMenu()
+{
+	if (TSharedPtr<ISequencer> SequencerPtr = UAnimLayers::GetSequencerFromAsset())
+	{
+		if (UAnimLayers* AnimLayers = UAnimLayers::GetAnimLayers(SequencerPtr.Get()))
+		{
+			const bool bShouldCloseWindowAfterMenuSelection = true;
+			FMenuBuilder MenuBuilder(bShouldCloseWindowAfterMenuSelection, nullptr);
+
+			TArray<FAnimLayerSourceUIEntryPtr> Selected;
+			SourcesListView->GetSelectedItems(Selected);
+			if (Selected.Num() > 0)
+			{
+				MenuBuilder.BeginSection("AnimLayerContextMenuLayer", LOCTEXT("AnimLayerContextMenuLayer", "Layer"));
+				{
+					FUIAction Action = FUIAction(FExecuteAction::CreateRaw((this), &FAnimLayerSourcesView::AddSelected));
+					const FText Label = LOCTEXT("AddSelected", "Add Selected");
+					const FText ToolTipText = LOCTEXT("AddSelectedTooltip", "Add selection to objects to selected layers");
+					MenuBuilder.AddMenuEntry(Label, ToolTipText, FSlateIcon(), Action);
+				}
+				{
+					FUIAction Action = FUIAction(FExecuteAction::CreateRaw((this), &FAnimLayerSourcesView::RemoveSelected));
+					const FText Label = LOCTEXT("RemoveSelected", "Remove Selected");
+					const FText ToolTipText = LOCTEXT("RemoveSelectedtooltip", "Remove selection from selected layers");
+					MenuBuilder.AddMenuEntry(Label, ToolTipText, FSlateIcon(), Action);
+				}
+				{
+					FUIAction Action = FUIAction(FExecuteAction::CreateRaw((this), &FAnimLayerSourcesView::SelectObjects));
+					const FText Label = LOCTEXT("SelectObjects", "Select Objects");
+					const FText ToolTipText = LOCTEXT("SelectObjectsTooltip", "Select all objects in this layer");
+					MenuBuilder.AddMenuEntry(Label, ToolTipText, FSlateIcon(), Action);
+				}
+
+				MenuBuilder.AddMenuSeparator();
+
+				{
+					FUIAction Action = FUIAction(FExecuteAction::CreateRaw((this), &FAnimLayerSourcesView::Duplicate));
+					const FText Label = LOCTEXT("Duplicate", "Duplicate");
+					const FText ToolTipText = LOCTEXT("Duplicatetooltip", "Duplicate to new layer");
+					MenuBuilder.AddMenuEntry(Label, ToolTipText, FSlateIcon(), Action);
+				}
+				if (Selected.Num() > 1)
+				{
+					FUIAction Action = FUIAction(FExecuteAction::CreateRaw((this), &FAnimLayerSourcesView::MergeLayers));
+					const FText Label = LOCTEXT("MergeLayers", "Merge Layers");
+					const FText ToolTipText = LOCTEXT("MergeLayerstooltip", "Merge selected layers");
+					MenuBuilder.AddMenuEntry(Label, ToolTipText, FSlateIcon(), Action);
+				}
+				/* adjustment blending not in yet
+				if (Selected.Num() == 1 && Selected[0]->GetAnimLayerIndex(AnimLayers) != 0)
+				{
+					FUIAction Action = FUIAction(FExecuteAction::CreateRaw((this), &FAnimLayerSourcesView::AdjustmentBlend));
+					const FText Label = LOCTEXT("AdjustmentBlend", "Adjustment blend");
+					const FText ToolTipText = LOCTEXT("AdjustmentBlendtooltip", "AdustmentBlend");
+					MenuBuilder.AddMenuEntry(Label, ToolTipText, FSlateIcon(), Action);
+				}
+				*/
+				MenuBuilder.AddMenuSeparator();
+				if (Selected.Num() == 1)
+				{
+					{
+						FUIAction Action = FUIAction(FExecuteAction::CreateRaw((this), &FAnimLayerSourcesView::Rename));
+						const FText Label = LOCTEXT("Rename", "Rename");
+						const FText ToolTipText = LOCTEXT("RenameLayerTooltip", "Rename selected layer");
+						MenuBuilder.AddMenuEntry(Label, ToolTipText, FSlateIcon(), Action);
+					}
+
+				}
+				{
+					FUIAction Action = FUIAction(FExecuteAction::CreateRaw((this), &FAnimLayerSourcesView::DeleteAnimLayer));
+					const FText Label = LOCTEXT("DeletaLayer", "DeleteLayer");
+					const FText ToolTipText = LOCTEXT("DeleteLayertooltip", "Delete selected layers");
+					MenuBuilder.AddMenuEntry(Label, ToolTipText, FSlateIcon(), Action);
+				}
+
+				MenuBuilder.EndSection();
+
+				return MenuBuilder.MakeWidget();
+			}
+		}
+	}
+	return nullptr;
+}
+
+void FAnimLayerSourcesView::OnSourceListSelectionChanged(FAnimLayerSourceUIEntryPtr Entry, ESelectInfo::Type SelectionType) const
+{
+	TArray<FAnimLayerSourceUIEntryPtr> Selected;
+	SourcesListView->GetSelectedItems(Selected);
+	//first turn off all
+	if (TSharedPtr<ISequencer> SequencerPtr = UAnimLayers::GetSequencerFromAsset())
+	{
+		if (UAnimLayers* AnimLayers = UAnimLayers::GetAnimLayers(SequencerPtr.Get()))
+		{
+			for (UAnimLayer* AnimLayer : AnimLayers->AnimLayers)
+			{
+				AnimLayer->SetSelectedInList(false);
+			}
+		}
+	}
+	for (const FAnimLayerSourceUIEntryPtr& Item : Selected)
+	{
+		Item->SetSelectedInList(true);
+	}	
+	//refresh tree to rerun the filter
+	if (TSharedPtr<ISequencer> SequencerPtr = UAnimLayers::GetSequencerFromAsset())
+	{
+		SequencerPtr->NotifyMovieSceneDataChanged(EMovieSceneDataChangeType::RefreshTree);
+	}
+}
+
+void FAnimLayerSourcesView::RefreshSourceData(bool bRefreshUI)
+{
+	SourceData.Reset();
+	FocusOnIndex = INDEX_NONE;
+	if (TSharedPtr<ISequencer> SequencerPtr = UAnimLayers::GetSequencerFromAsset())
+	{
+		if (UAnimLayers* AnimLayers = UAnimLayers::GetAnimLayers(SequencerPtr.Get()))
+		{
+			for (UAnimLayer* AnimLayer : AnimLayers->AnimLayers)
+			{
+				if (Controller == nullptr || Controller->IsSelectionFilterActive() == false
+					|| AnimLayer->GetSelected() != ECheckBoxState::Unchecked)
+				{
+					SourceData.Add(MakeShared<FAnimLayerSourceUIEntry>(AnimLayer));
+				}
+			}
+			/*
+			for (FGuid SourceGuid : Client->GetDisplayableSources())
+			{
+				SourceData.Add(MakeShared<FAnimLayerSourceUIEntry>(SourceGuid, Client));
+			}
+			*/
+		}
+	}
+	if (bRefreshUI)
+	{
+		SourcesListView->RequestListRefresh();
+		TArray< FAnimLayerSourceUIEntryPtr> Selected;
+		for (const FAnimLayerSourceUIEntryPtr& Item : SourceData) 
+		{
+			if (Item->GetSelectedInList())
+			{
+				Selected.Add(Item);
+			}
+		}
+		if (Selected.Num())
+		{
+			SourcesListView->SetItemSelection(Selected, true);
+		}
+	}
+}
+
+//////////////////////////////////////////////////////////////
+/// FAnimLayerController
+///////////////////////////////////////////////////////////
+
+FAnimLayerController::FAnimLayerController()
+{
+
+	SourcesView = MakeShared<FAnimLayerSourcesView>();
+	RebuildSourceList();
+}
+
+FAnimLayerController::~FAnimLayerController()
+{
+	/*
+	if (Client)
+	{
+		Client->OnLiveLinkSourcesChanged().Remove(OnSourcesChangedHandle);
+		OnSourcesChangedHandle.Reset();
+	}
+	*/
+}
+
+void FAnimLayerController::RefreshTimeDependantData()
+{
+	if (SourcesView.IsValid())
+	{
+		for (const FAnimLayerSourceUIEntryPtr& Item : SourcesView->SourceData)
+		{
+			Item->GetWeight();
+		}
+	}
+}
+void FAnimLayerController::SelectItem(int32 Index, bool bClear)
+{
+	if (SourcesView.IsValid() && SourcesView->SourcesListView.IsValid())
+	{
+		if (bClear)
+		{
+			SourcesView->SourcesListView->ClearSelection();
+		}
+		TArray< FAnimLayerSourceUIEntryPtr> Selected;
+		int32 Count = 0;
+		for (const FAnimLayerSourceUIEntryPtr& Item : SourcesView->SourceData)
+		{
+			if (Count == Index)
+			{
+				Item->SetSelectedInList(true);
+				Selected.Add(Item);
+			}
+			else if (bClear && Item->GetSelectedInList())
+			{
+				Item->SetSelectedInList(false);
+			}
+			++Count;
+		}
+		if (Selected.Num())
+		{
+			SourcesView->SourcesListView->SetItemSelection(Selected, true);
+		}
+	}
+
+
+}
+
+void FAnimLayerController::RefreshSelectionData()
+{
+	if (SourcesView.IsValid())
+	{
+		for (const FAnimLayerSourceUIEntryPtr& Item : SourcesView->SourceData)
+		{
+			Item->ClearCaches();
+		}
+	}
+}
+
+void FAnimLayerController::RefreshSourceData(bool bRefreshUI)
+{
+	if (SourcesView.IsValid())
+	{
+		SourcesView->RefreshSourceData(bRefreshUI);
+	}
+}
+void FAnimLayerController::ToggleSelectionFilterActive()
+{
+	bSelectionFilterActive = !bSelectionFilterActive;
+	RebuildSourceList();
+}
+
+bool FAnimLayerController::IsSelectionFilterActive() const
+{
+	return bSelectionFilterActive;
+}
+
+void FAnimLayerController::OnSourceSelectionChangedHandler(FAnimLayerSourceUIEntryPtr Entry, ESelectInfo::Type SelectionType) const
+{
+	if (bSelectionChangedGuard)
+	{
+		return;
+	}
+}
+
+void FAnimLayerController::RebuildSourceList()
+{
+	SourcesView->RefreshSourceData(true);
+}
+
+//////////////////////////////////////////////////////////////
+/// SAnimLayers
+///////////////////////////////////////////////////////////
+
+void SAnimLayers::Construct(const FArguments& InArgs, FControlRigEditMode& InEditMode)
+{
+	AnimLayerController = MakeShared<FAnimLayerController>();
+	AnimLayerController->SourcesView->AddController(AnimLayerController.Get());
+	LastMovieSceneSig = FGuid();
+	if (TSharedPtr<ISequencer> SequencerPtr = UAnimLayers::GetSequencerFromAsset())
+	{
+		SequencerPtr->OnActivateSequence().AddRaw(this, &SAnimLayers::OnActivateSequence);
+		SequencerPtr->OnMovieSceneDataChanged().AddRaw(this, &SAnimLayers::OnMovieSceneDataChanged);
+		SequencerPtr->OnGlobalTimeChanged().AddRaw(this, &SAnimLayers::OnGlobalTimeChanged);
+		SequencerPtr->OnEndScrubbingEvent().AddRaw(this, &SAnimLayers::OnGlobalTimeChanged);
+		SequencerPtr->OnStopEvent().AddRaw(this, &SAnimLayers::OnGlobalTimeChanged);
+
+		if (UAnimLayers* AnimLayersPtr = UAnimLayers::GetAnimLayers(SequencerPtr.Get()))
+		{
+			AnimLayersPtr->AnimLayerListChanged().AddRaw(AnimLayerController.Get(), &FAnimLayerController::HandleOnAnimLayerListChanged);
+			AnimLayers = AnimLayersPtr;
+		}
+	}
+	ChildSlot
+	[
+		SNew(SVerticalBox)
+			+SVerticalBox::Slot()
+			.Padding(FMargin(4.0f, 0.f,0.f,0.f))
+			.AutoHeight()
+			[
+				SNew(SBorder)
+				.BorderImage(FAppStyle::GetBrush("ToolPanel.GroupBorder"))
+				.Padding(4.0f)
+				[
+					SNew(SHorizontalBox)
+					+ SHorizontalBox::Slot()
+					.Padding(.0f)
+					.HAlign(HAlign_Fill)
+					.FillWidth(1.0f)
+					[
+						SNew(SHorizontalBox)
+						+ SHorizontalBox::Slot()
+						.AutoWidth()
+						.HAlign(EHorizontalAlignment::HAlign_Right)
+						.Padding(.0f)
+						[
+							SNew(SPositiveActionButton)
+							.OnClicked(this, &SAnimLayers::OnAddClicked)
+							.Icon(FAppStyle::Get().GetBrush("Icons.Plus"))
+							.Text(LOCTEXT("AnimLayer", "Layer"))
+							.ToolTipText(LOCTEXT("AnimLayer", "Add a new Animation Layer"))
+						]
+						+ SHorizontalBox::Slot()
+						.FillWidth(10.f)
+						[
+							SNew(SSpacer)
+						]
+						+ SHorizontalBox::Slot()
+						.AutoWidth()
+						.HAlign(EHorizontalAlignment::HAlign_Right)
+						.Padding(5.0f)
+						.VAlign(VAlign_Center)
+						[
+							SNew(SButton)
+								.ButtonStyle(FAppStyle::Get(), "NoBorder")
+								.ButtonColorAndOpacity_Lambda([this]()
+									{
+										return FLinearColor::White;
+									})
+								.OnClicked_Lambda([this]()
+									{
+										AnimLayerController->ToggleSelectionFilterActive();
+										return FReply::Handled();
+									})
+								.ContentPadding(1.f)
+								.ToolTipText(LOCTEXT("AnimLayerSelectionFilerTooltip", "Only show Anim Layers with selected objects"))
+								[
+									SNew(SImage)
+										.ColorAndOpacity_Lambda([this]()
+											{
+
+												const FLinearColor Selected = FLinearColor(FColor::FromHex(TEXT("#EBA30A")));
+												const FColor NotSelected(56, 56, 56);
+												FSlateColor SlateColor;
+												if (AnimLayerController->IsSelectionFilterActive() == true)
+												{
+													SlateColor = FSlateColor(Selected);
+												}
+												else
+												{
+													SlateColor = FSlateColor(NotSelected);
+												}
+												return SlateColor;
+											})
+										.Image(FSlateIcon(FControlRigEditorStyle::Get().GetStyleSetName(), "ControlRig.FilterAnimLayerSelected").GetIcon())
+								]
+						]	
+					]
+				]
+			]
+			+SVerticalBox::Slot()
+			.FillHeight(1.f)
+			.Padding(FMargin(0.0f, 4.0f, 0.0f, 0.0f))
+			[
+				AnimLayerController->SourcesView->SourcesListView.ToSharedRef()
+			]
+	];
+	SetEditMode(InEditMode);
+	RegisterSelectionChanged();
+	SetCanTick(true);
+
+}
+
+void SAnimLayers::Tick(const FGeometry& AllottedGeometry, const double InCurrentTime, const float InDeltaTime)
+{
+	if (ISequencer* Sequencer = GetSequencer())
+	{
+		FGuid CurrentMovieSceneSig = Sequencer->GetFocusedMovieSceneSequence()->GetMovieScene()->GetSignature();
+		if (LastMovieSceneSig != CurrentMovieSceneSig)
+		{
+			LastMovieSceneSig = CurrentMovieSceneSig;
+			AnimLayerController->RefreshSelectionData();
+		}
+	}
+}
+
+
+FReply SAnimLayers::OnSelectionFilterClicked()
+{
+	AnimLayerController->ToggleSelectionFilterActive();
+	return FReply::Handled();
+}
+
+bool SAnimLayers::IsSelectionFilterActive() const
+{
+	return AnimLayerController->IsSelectionFilterActive();
+}
+
+//mz todo, if in layers with control rig's need to replace them.
+void SAnimLayers::OnObjectsReplaced(const TMap<UObject*, UObject*>& OldToNewInstanceMap)
+{
+	//if there's a control rig recreate the tree, controls may have changed
+	bool bNewControlRig = false;
+	for (const TPair<UObject*, UObject*>& Pair : OldToNewInstanceMap)
+	{
+		if (Pair.Key && Pair.Value)
+		{
+			if (Pair.Key->IsA<UControlRig>() && Pair.Value->IsA<UControlRig>())
+			{
+				bNewControlRig = false;
+				break;
+			}
+		}
+	}
+}
+
+SAnimLayers::SAnimLayers()
+{
+	FCoreUObjectDelegates::OnObjectsReplaced.AddRaw(this, &SAnimLayers::OnObjectsReplaced);
+}
+
+SAnimLayers::~SAnimLayers()
+{
+	if (AnimLayers.IsValid())
+	{
+		AnimLayers.Get()->AnimLayerListChanged().RemoveAll(AnimLayerController.Get());
+	}
+
+	if (TSharedPtr<ISequencer> SequencerPtr = UAnimLayers::GetSequencerFromAsset())
+	{
+		SequencerPtr->OnActivateSequence().RemoveAll(this);
+		SequencerPtr->OnMovieSceneDataChanged().RemoveAll(this);
+		SequencerPtr->OnGlobalTimeChanged().RemoveAll(this);
+		SequencerPtr->OnEndScrubbingEvent().RemoveAll(this);
+		SequencerPtr->OnStopEvent().RemoveAll(this);
+
+	}
+	
+	// unregister previous one
+	if (OnSelectionChangedHandle.IsValid())
+	{
+		FLevelEditorModule& LevelEditor = FModuleManager::GetModuleChecked<FLevelEditorModule>("LevelEditor");
+		FLevelEditorModule::FActorSelectionChangedEvent& ActorSelectionChangedEvent = LevelEditor.OnActorSelectionChanged();
+		ActorSelectionChangedEvent.Remove(OnSelectionChangedHandle);
+		OnSelectionChangedHandle.Reset();
+	}	
+	
+	for (TWeakObjectPtr<UControlRig>& ControlRig : BoundControlRigs)
+	{
+		if (ControlRig.IsValid())
+		{
+			ControlRig.Get()->ControlRigBound().RemoveAll(this);
+			const TSharedPtr<IControlRigObjectBinding> Binding = ControlRig.Get()->GetObjectBinding();
+			if (Binding)
+			{
+				Binding->OnControlRigBind().RemoveAll(this);
+			}
+		}
+	}
+	BoundControlRigs.SetNum(0);
+}
+
+FReply SAnimLayers::OnKeyDown(const FGeometry& MyGeometry, const FKeyEvent& InKeyEvent)
+{
+	if (FControlRigEditMode* EditMode = static_cast<FControlRigEditMode*>(ModeTools->GetActiveMode(FControlRigEditMode::ModeName)))
+	{
+		TWeakPtr<ISequencer> Sequencer = EditMode->GetWeakSequencer();
+		if (Sequencer.IsValid())
+		{
+
+			if (Sequencer.Pin()->GetCommandBindings(ESequencerCommandBindings::CurveEditor).Get()->ProcessCommandBindings(InKeyEvent))
+			{
+				return FReply::Handled();
+			}
+		}
+	}
+	return FReply::Unhandled();
+}
+
+void SAnimLayers::HandleControlSelected(UControlRig* Subject, FRigControlElement* ControlElement, bool bSelected)
+{
+	FControlRigBaseDockableView::HandleControlSelected(Subject, ControlElement, bSelected);
+	if (AnimLayerController)
+	{
+		if (AnimLayerController->IsSelectionFilterActive())
+		{
+			AnimLayerController->RefreshSourceData(true);
+		}
+		AnimLayerController->RefreshSelectionData();
+	}
+}
+
+void SAnimLayers::RegisterSelectionChanged()
+{
+	FLevelEditorModule& LevelEditor = FModuleManager::GetModuleChecked<FLevelEditorModule>("LevelEditor");
+	FLevelEditorModule::FActorSelectionChangedEvent& ActorSelectionChangedEvent = LevelEditor.OnActorSelectionChanged();
+
+	// unregister previous one
+	if (OnSelectionChangedHandle.IsValid())
+	{
+		ActorSelectionChangedEvent.Remove(OnSelectionChangedHandle);
+		OnSelectionChangedHandle.Reset();
+	}
+
+	// register
+	OnSelectionChangedHandle = ActorSelectionChangedEvent.AddRaw(this, &SAnimLayers::OnActorSelectionChanged);
+}
+
+void SAnimLayers::OnActorSelectionChanged(const TArray<UObject*>& NewSelection, bool bForceRefresh)
+{
+	if (AnimLayerController)
+	{
+		if (AnimLayerController->IsSelectionFilterActive())
+		{
+			AnimLayerController->RefreshSourceData(true);
+		}
+		AnimLayerController->RefreshSelectionData();
+	}
+}
+
+void SAnimLayers::OnActivateSequence(FMovieSceneSequenceIDRef ID)
+{
+	if (AnimLayerController)
+	{
+		AnimLayerController->RefreshSourceData(true);
+		AnimLayerController->RefreshSelectionData();
+	}
+}
+
+void SAnimLayers::OnGlobalTimeChanged()
+{
+	if (AnimLayerController)
+	{
+		AnimLayerController->RefreshTimeDependantData();
+	}
+}
+
+void SAnimLayers::OnMovieSceneDataChanged(EMovieSceneDataChangeType)
+{
+	if (AnimLayerController)
+	{
+		AnimLayerController->RefreshTimeDependantData();
+		AnimLayerController->RefreshSelectionData();
+	}
+}
+
+FReply SAnimLayers::OnAddClicked()
+{
+	if (ISequencer* SequencerPtr = GetSequencer())
+	{
+		if (UAnimLayers* AnimLayersPtr = UAnimLayers::GetAnimLayers(SequencerPtr))
+		{
+			int32 Index = AnimLayersPtr->AddAnimLayerFromSelection(SequencerPtr);
+			if (Index != INDEX_NONE)
+			{
+				if (AnimLayerController.IsValid())
+				{
+					AnimLayerController->FocusRenameLayer(Index);
+					AnimLayerController->SelectItem(Index, true);
+				}
+			}
+		}
+	}
+
+	return FReply::Handled();
+}
+
+void SAnimLayers::SetEditMode(FControlRigEditMode& InEditMode)
+{
+	FControlRigBaseDockableView::SetEditMode(InEditMode);
+	ModeTools = InEditMode.GetModeManager();
+	if (FControlRigEditMode* EditMode = static_cast<FControlRigEditMode*>(ModeTools->GetActiveMode(FControlRigEditMode::ModeName)))
+	{
+		TArrayView<TWeakObjectPtr<UControlRig>> ControlRigs = EditMode->GetControlRigs();
+		for (TWeakObjectPtr<UControlRig>& ControlRig : ControlRigs)
+		{
+			if (ControlRig.IsValid())
+			{
+				if (!ControlRig.Get()->ControlRigBound().IsBoundToObject(this))
+				{
+					ControlRig.Get()->ControlRigBound().AddRaw(this, &SAnimLayers::HandleOnControlRigBound);
+					BoundControlRigs.Add(ControlRig);
+				}
+				const TSharedPtr<IControlRigObjectBinding> Binding = ControlRig.Get()->GetObjectBinding();
+				if (Binding && !Binding->OnControlRigBind().IsBoundToObject(this))
+				{
+					Binding->OnControlRigBind().AddRaw(this, &SAnimLayers::HandleOnObjectBoundToControlRig);
+				}
+			}
+		}
+	}
+}
+
+void SAnimLayers::HandleControlAdded(UControlRig* ControlRig, bool bIsAdded)
+{
+	FControlRigBaseDockableView::HandleControlAdded(ControlRig, bIsAdded);
+	if (ControlRig)
+	{
+		if (bIsAdded == true)
+		{
+			if (!ControlRig->ControlRigBound().IsBoundToObject(this))
+			{
+				ControlRig->ControlRigBound().AddRaw(this, &SAnimLayers::HandleOnControlRigBound);
+				BoundControlRigs.Add(ControlRig);
+			}
+			const TSharedPtr<IControlRigObjectBinding> Binding = ControlRig->GetObjectBinding();
+			if (Binding && !Binding->OnControlRigBind().IsBoundToObject(this))
+			{
+				Binding->OnControlRigBind().AddRaw(this, &SAnimLayers::HandleOnObjectBoundToControlRig);
+			}
+		}
+		else
+		{
+			BoundControlRigs.Remove(ControlRig);
+			ControlRig->ControlRigBound().RemoveAll(this);
+			const TSharedPtr<IControlRigObjectBinding> Binding = ControlRig->GetObjectBinding();
+			if (Binding)
+			{
+				Binding->OnControlRigBind().RemoveAll(this);
+			}
+		}
+	}
+}
+
+void SAnimLayers::HandleOnControlRigBound(UControlRig* InControlRig)
+{
+	if (!InControlRig)
+	{
+		return;
+	}
+
+	const TSharedPtr<IControlRigObjectBinding> Binding = InControlRig->GetObjectBinding();
+
+	if (Binding && !Binding->OnControlRigBind().IsBoundToObject(this))
+	{
+		Binding->OnControlRigBind().AddRaw(this, &SAnimLayers::HandleOnObjectBoundToControlRig);
+	}
+}
+
+//mz todo need to test recompiling
+void SAnimLayers::HandleOnObjectBoundToControlRig(UObject* InObject)
+{
+	
+}
+
+class SInvalidWeightNameDetailWidget : public SSpacer
+{
+	SLATE_BEGIN_ARGS(SInvalidWeightNameDetailWidget)
+		{}
+	SLATE_END_ARGS()
+
+	void Construct(const FArguments& InArgs)
+	{
+		SetVisibility(EVisibility::Collapsed);
+	}
+
+};
+
+class  FWeightNameOverride : public FDetailsNameWidgetOverrideCustomization
+{
+public:
+	FWeightNameOverride() {};
+	virtual ~FWeightNameOverride() override = default;
+	virtual TSharedRef<SWidget> CustomizeName(TSharedRef<SWidget> InnerNameContent, FPropertyPath& Path) override 
+	{
+		const TSharedRef<SWidget> NameContent = SNew(SInvalidWeightNameDetailWidget);
+		return NameContent;
+	}
+};
+
+void SAnimWeightDetails::Construct(const FArguments& InArgs, FControlRigEditMode* InEditMode, UObject* InWeightObject)
+{
+	if (InEditMode == nullptr || InWeightObject == nullptr)
+	{
+		return;
+	}
+	using namespace UE::Sequencer;
+
+	FDetailsViewArgs DetailsViewArgs;
+	{
+		DetailsViewArgs.bAllowSearch = false;
+		DetailsViewArgs.bHideSelectionTip = true;
+		DetailsViewArgs.bLockable = false;
+		DetailsViewArgs.bSearchInitialKeyFocus = false;
+		DetailsViewArgs.bUpdatesFromSelection = false;
+		DetailsViewArgs.bShowOptions = false;
+		DetailsViewArgs.bShowModifiedPropertiesOption = false;
+		DetailsViewArgs.bCustomNameAreaLocation = false;
+		DetailsViewArgs.bCustomFilterAreaLocation = false;
+		DetailsViewArgs.NameAreaSettings = FDetailsViewArgs::HideNameArea;
+		DetailsViewArgs.bAllowMultipleTopLevelObjects = true;
+		DetailsViewArgs.bShowScrollBar = false; // Don't need to show this, as we are putting it in a scroll box
+		DetailsViewArgs.DetailsNameWidgetOverrideCustomization = MakeShared<FWeightNameOverride>();
+
+	}
+
+	WeightView = FModuleManager::GetModuleChecked<FPropertyEditorModule>("PropertyEditor").CreateDetailView(DetailsViewArgs);
+	WeightView->SetKeyframeHandler(InEditMode->DetailKeyFrameCache);
+
+
+	ChildSlot
+	[
+		SNew(SScrollBox)
+		+ SScrollBox::Slot()
+		[
+			SNew(SVerticalBox)
+			+ SVerticalBox::Slot()
+			.AutoHeight()
+			[
+				WeightView.ToSharedRef()
+			]
+		]
+	];
+	TArray<UObject*> Objects;
+	Objects.Add(InWeightObject);
+	WeightView->SetObjects(Objects, true);
+}
+
+SAnimWeightDetails::~SAnimWeightDetails()
+{
+}
+
+
+
+#undef LOCTEXT_NAMESPACE
+
+
+
+
diff --git a/Engine/Plugins/Animation/ControlRig/Source/ControlRigEditor/Private/Sequencer/AnimLayers/SAnimLayers.h b/Engine/Plugins/Animation/ControlRig/Source/ControlRigEditor/Private/Sequencer/AnimLayers/SAnimLayers.h
new file mode 100644
--- /dev/null
+++ b/Engine/Plugins/Animation/ControlRig/Source/ControlRigEditor/Private/Sequencer/AnimLayers/SAnimLayers.h
@@ -0,0 +1,83 @@
+// Copyright Epic Games, Inc. All Rights Reserved.
+
+
+#include "CoreMinimal.h"
+#include "EditMode/ControlRigBaseDockableView.h"
+#include "Widgets/SWidget.h"
+#include "Widgets/SCompoundWidget.h"
+#include "ISequencer.h"
+#include "Misc/Guid.h"
+
+class UControlRig;
+class FControlRigEditMode;
+struct FRigControlElement;
+struct FRigElementKey;
+struct FAnimLayerController;
+class UAnimLayers;
+enum class EMovieSceneDataChangeType;
+
+
+class SAnimLayers : public FControlRigBaseDockableView, public SCompoundWidget
+{
+	SLATE_BEGIN_ARGS(SAnimLayers) {}
+	SLATE_END_ARGS()
+	SAnimLayers();
+	~SAnimLayers();
+
+	void Construct(const FArguments& InArgs, FControlRigEditMode& InEditMode);
+	//SCompuntWidget
+	void Tick(const FGeometry& AllottedGeometry, const double InCurrentTime, const float InDeltaTime) override;
+
+	//FControlRigBaseDockableView overrides
+	virtual void SetEditMode(FControlRigEditMode& InEditMode) override;
+	virtual void HandleControlAdded(UControlRig* ControlRig, bool bIsAdded) override;
+
+private:
+	void HandleOnControlRigBound(UControlRig* InControlRig);
+	void HandleOnObjectBoundToControlRig(UObject* InObject);
+	virtual void HandleControlSelected(UControlRig* Subject, FRigControlElement* InControl, bool bSelected) override;
+	void OnObjectsReplaced(const TMap<UObject*, UObject*>& OldToNewInstanceMap);
+
+	//set of control rigs we are bound too and need to clear delegates from
+	TArray<TWeakObjectPtr<UControlRig>> BoundControlRigs;
+private:
+	//actor selection changing
+	void RegisterSelectionChanged();
+	void OnActorSelectionChanged(const TArray<UObject*>& NewSelection, bool bForceRefresh);
+	FDelegateHandle OnSelectionChangedHandle;
+
+	//sequencer changing so need to refresh list or update weight
+	void OnActivateSequence(FMovieSceneSequenceIDRef ID);
+	void OnGlobalTimeChanged();
+	void OnMovieSceneDataChanged(EMovieSceneDataChangeType);
+	FGuid LastMovieSceneSig = FGuid();
+
+private:
+	FReply OnAddClicked();
+	FReply OnSelectionFilterClicked();
+	bool IsSelectionFilterActive() const;
+	TSharedPtr<FAnimLayerController> AnimLayerController;
+	TWeakObjectPtr<UAnimLayers> AnimLayers;
+private:
+	//Keyboard interaction for Sequencer hotkeys
+	virtual bool SupportsKeyboardFocus() const override { return true; }
+	virtual FReply OnKeyDown(const FGeometry& MyGeometry, const FKeyEvent& InKeyEvent) override;
+};
+
+class SAnimWeightDetails : public SCompoundWidget
+{
+
+	SLATE_BEGIN_ARGS(SAnimWeightDetails)
+	{}
+	SLATE_END_ARGS()
+	~SAnimWeightDetails();
+
+	void Construct(const FArguments& InArgs, FControlRigEditMode* InEditMode, UObject* InWeightObject);
+
+private:
+
+	TSharedPtr<IDetailsView> WeightView;
+
+};
+
+
diff --git a/Engine/Plugins/Animation/ControlRig/Source/ControlRigEditor/Private/Sequencer/ControlRigParameterTrackEditor.cpp b/Engine/Plugins/Animation/ControlRig/Source/ControlRigEditor/Private/Sequencer/ControlRigParameterTrackEditor.cpp
--- a/Engine/Plugins/Animation/ControlRig/Source/ControlRigEditor/Private/Sequencer/ControlRigParameterTrackEditor.cpp
+++ b/Engine/Plugins/Animation/ControlRig/Source/ControlRigEditor/Private/Sequencer/ControlRigParameterTrackEditor.cpp
@@ -2046,7 +2046,7 @@ void FControlRigParameterTrackEditor::AddControlRigFromComponent(FGuid InGuid)
 
 bool FControlRigParameterTrackEditor::HasTransformKeyOverridePriority() const
 {
-	return false; 
+	return false;
 }
 
 bool FControlRigParameterTrackEditor::CanAddTransformKeysForSelectedObjects() const
@@ -3636,7 +3636,7 @@ void FControlRigParameterTrackEditor::HandleControlSelected(UControlRig* Subject
 			//Just select in section to key, if deselecting makes sure deselected everywhere
 			if (bSelected == true)
 			{
-				UMovieSceneSection* Section = Track->GetSectionToKey();
+				UMovieSceneSection* Section = Track->GetSectionToKey(ControlElement->GetFName());
 				UMovieSceneControlRigParameterSection* ParamSection = Cast<UMovieSceneControlRigParameterSection>(Section);
 				SelectSequencerNodeInSection(ParamSection, ControlElement->GetFName(), bSelected);
 			}
@@ -4133,7 +4133,18 @@ FKeyPropertyResult FControlRigParameterTrackEditor::AddKeysToControlRigHandle(UO
 	{
 		float Weight = 1.0f;
 
-		UMovieSceneSection* SectionToKey = bCreateSection ? Track->FindOrExtendSection(KeyTime, Weight) : Track->FindSection(KeyTime);
+		UMovieSceneSection* SectionToKey = Track->GetSectionToKey(RigControlName);
+		if (SectionToKey)
+		{
+			if (SectionToKey->HasEndFrame() && SectionToKey->GetExclusiveEndFrame() < KeyTime)
+			{
+				SectionToKey->SetEndFrame(KeyTime);
+			}
+			else if(SectionToKey->HasStartFrame() && SectionToKey->GetInclusiveStartFrame() > KeyTime)
+			{
+				SectionToKey->SetStartFrame(KeyTime);
+			}
+		}
 
 		// If there's no overlapping section to key, create one only if a track was newly created. Otherwise, skip keying altogether
 		// so that the user is forced to create a section to key on.
@@ -4165,7 +4176,7 @@ FKeyPropertyResult FControlRigParameterTrackEditor::AddKeysToControlRigHandle(UO
 		//if we create a key then compensate
 		if (KeyPropertyResult.bKeyCreated)
 		{
-			UMovieSceneControlRigParameterSection* ParamSection = Cast<UMovieSceneControlRigParameterSection>(Track->GetSectionToKey());
+			UMovieSceneControlRigParameterSection* ParamSection = Cast<UMovieSceneControlRigParameterSection>(Track->GetSectionToKey(ControlRigName));
 			if (UControlRig* SectionControlRig = ParamSection ? ParamSection->GetControlRig() : nullptr)
 			{
 				TOptional<FFrameNumber> OptionalKeyTime = KeyTime;
@@ -4236,7 +4247,7 @@ void FControlRigParameterTrackEditor::AddControlKeys(
 			BeginKeying(LocalFrameTime.RoundToFrame());
 		}
 		const FFrameNumber FrameTime = GetTimeForKey();
-		UMovieSceneSection* Section = Track->FindSection(FrameTime);
+		UMovieSceneSection* Section = Track->GetSectionToKey(RigControlName);
 		ParamSection = Cast<UMovieSceneControlRigParameterSection>(Section);
 
 		if (ParamSection && ParamSection->GetDoNotKey())
@@ -4929,50 +4940,12 @@ void FControlRigParameterTrackEditor::SelectFKBonesToAnimate(UFKControlRig* Auto
 	//reconstruct all channel proxies TODO or not to do that is the question
 }
 
-//////////////////////////////////////////////////////////////
-/// SCollapseControlsWidget
-///////////////////////////////////////////////////////////
-
-/** Widget allowing collapsing of controls */
-class SCollapseControlsWidget : public SCompoundWidget
-{
-public:
-
-	SLATE_BEGIN_ARGS(SCollapseControlsWidget)
-		: _Sequencer(nullptr), _OwnerTrack(nullptr)
-	{}
-
-	SLATE_ARGUMENT(TWeakPtr<ISequencer>, Sequencer)
-	SLATE_ARGUMENT(UMovieSceneTrack*, OwnerTrack)
-	SLATE_END_ARGS()
-
-	void Construct(const FArguments& InArgs);
-	virtual ~SCollapseControlsWidget() override {}
-
-	FReply OpenDialog(bool bModal = true);
-	void CloseDialog();
-
-
-private:
-	void Collapse();
-
-	TWeakPtr<ISequencer> Sequencer;
-	TWeakObjectPtr<UMovieSceneTrack> OwnerTrack;
-	//static to be reused
-	static TOptional<FBakingAnimationKeySettings> CollapseControlsSettings;
-	//structonscope for details panel
-	TSharedPtr < TStructOnScope<FBakingAnimationKeySettings>> Settings;
-	TWeakPtr<SWindow> DialogWindow;
-	TSharedPtr<IStructureDetailsView> DetailsView;
-};
-
 
 TOptional<FBakingAnimationKeySettings> SCollapseControlsWidget::CollapseControlsSettings;
 
 void SCollapseControlsWidget::Construct(const FArguments& InArgs)
 {
 	Sequencer = InArgs._Sequencer;
-	OwnerTrack = InArgs._OwnerTrack;
 
 	if (CollapseControlsSettings.IsSet() == false)
 	{
@@ -5065,13 +5038,11 @@ void SCollapseControlsWidget::Construct(const FArguments& InArgs)
 		];
 }
 
-
 void  SCollapseControlsWidget::Collapse()
 {
 	FBakingAnimationKeySettings* BakeSettings = Settings->Get();
 	TSharedPtr<ISequencer> SequencerPtr = Sequencer.Pin();
-	FControlRigParameterTrackEditor::CollapseAllLayers(SequencerPtr, OwnerTrack.Get(), *BakeSettings);
-
+	CollapseCB.ExecuteIfBound(SequencerPtr, *BakeSettings);
 	CollapseControlsSettings = *BakeSettings;
 }
 
@@ -5308,16 +5279,20 @@ bool FControlRigParameterTrackEditor::CollapseAllLayers(TSharedPtr<ISequencer>&S
 
 void FControlRigParameterSection::CollapseAllLayers()
 {
-	TSharedPtr<ISequencer> SequencerPtr = WeakSequencer.Pin();
 	if (UMovieSceneControlRigParameterSection* ParameterSection = CastChecked<UMovieSceneControlRigParameterSection>(WeakSection.Get()))
 	{
 		UMovieSceneTrack* OwnerTrack = ParameterSection->GetTypedOuter<UMovieSceneTrack>();
+		FCollapseControlsCB CollapseCB = FCollapseControlsCB::CreateLambda([this,OwnerTrack](TSharedPtr<ISequencer>& InSequencer, const FBakingAnimationKeySettings& InSettings)
+		{
+			FControlRigParameterTrackEditor::CollapseAllLayers(InSequencer, OwnerTrack, InSettings);
+		});
+
 		TSharedPtr<ISequencer> Sequencer = WeakSequencer.Pin();
 		TSharedRef<SCollapseControlsWidget> BakeWidget =
 			SNew(SCollapseControlsWidget)
-			.Sequencer(Sequencer)
-			.OwnerTrack(OwnerTrack);
+			.Sequencer(Sequencer);
 
+		BakeWidget->SetCollapseCB(CollapseCB);
 		BakeWidget->OpenDialog(false);
 	}
 }
diff --git a/Engine/Plugins/Animation/ControlRig/Source/ControlRigEditor/Private/Sequencer/ControlRigParameterTrackEditor.h b/Engine/Plugins/Animation/ControlRig/Source/ControlRigEditor/Private/Sequencer/ControlRigParameterTrackEditor.h
--- a/Engine/Plugins/Animation/ControlRig/Source/ControlRigEditor/Private/Sequencer/ControlRigParameterTrackEditor.h
+++ b/Engine/Plugins/Animation/ControlRig/Source/ControlRigEditor/Private/Sequencer/ControlRigParameterTrackEditor.h
@@ -36,6 +36,46 @@ struct FMovieSceneChannel;
 struct FKeyAddOrDeleteEventItem;
 struct FKeyMoveEventItem;
 struct FBakingAnimationKeySettings;
+class ISequencer;
+class UMovieSceneTrack;
+class IStructureDetailsView;
+
+//////////////////////////////////////////////////////////////
+/// SCollapseControlsWidget
+///////////////////////////////////////////////////////////
+DECLARE_DELEGATE_TwoParams(FCollapseControlsCB, TSharedPtr<ISequencer>& InSequencer, const FBakingAnimationKeySettings& InSettings);
+
+/** Widget allowing collapsing of controls */
+class SCollapseControlsWidget : public SCompoundWidget
+{
+public:
+
+	SLATE_BEGIN_ARGS(SCollapseControlsWidget)
+		: _Sequencer(nullptr)
+		{}
+		SLATE_ARGUMENT(TWeakPtr<ISequencer>, Sequencer)
+	SLATE_END_ARGS()
+
+	void Construct(const FArguments& InArgs);
+	virtual ~SCollapseControlsWidget() override {}
+
+	FReply OpenDialog(bool bModal = true);
+	void CloseDialog();
+
+	void SetCollapseCB(FCollapseControlsCB& InCB) { CollapseCB = InCB; }
+private:
+	void Collapse();
+
+	TWeakPtr<ISequencer> Sequencer;
+	//static to be reused
+	static TOptional<FBakingAnimationKeySettings> CollapseControlsSettings;
+	//structonscope for details panel
+	TSharedPtr < TStructOnScope<FBakingAnimationKeySettings>> Settings;
+	TWeakPtr<SWindow> DialogWindow;
+	TSharedPtr<IStructureDetailsView> DetailsView;
+	FCollapseControlsCB CollapseCB;
+};
+
 /**
  * Tools for animation tracks
  */
diff --git a/Engine/Plugins/Animation/ControlRig/Source/ControlRigEditor/Private/Sequencer/ControlRigSpaceChannelEditors.cpp b/Engine/Plugins/Animation/ControlRig/Source/ControlRigEditor/Private/Sequencer/ControlRigSpaceChannelEditors.cpp
--- a/Engine/Plugins/Animation/ControlRig/Source/ControlRigEditor/Private/Sequencer/ControlRigSpaceChannelEditors.cpp
+++ b/Engine/Plugins/Animation/ControlRig/Source/ControlRigEditor/Private/Sequencer/ControlRigSpaceChannelEditors.cpp
@@ -207,7 +207,7 @@ FSpaceChannelAndSection FControlRigSpaceChannelHelpers::FindSpaceChannelAndSecti
 		UMovieSceneControlRigParameterTrack* ControlRigParameterTrack = Cast<UMovieSceneControlRigParameterTrack>(MovieScene->FindTrack(UMovieSceneControlRigParameterTrack::StaticClass(), Binding.GetObjectGuid(), NAME_None));
 		if (ControlRigParameterTrack && ControlRigParameterTrack->GetControlRig() == ControlRig)
 		{
-			UMovieSceneControlRigParameterSection* ActiveSection = Cast<UMovieSceneControlRigParameterSection>(ControlRigParameterTrack->GetSectionToKey());
+			UMovieSceneControlRigParameterSection* ActiveSection = Cast<UMovieSceneControlRigParameterSection>(ControlRigParameterTrack->GetSectionToKey(ControlName));
 			if (ActiveSection)
 			{
 				ActiveSection->Modify();
diff --git a/Engine/Plugins/Animation/ControlRig/Source/ControlRigEditor/Private/Tools/ControlRigSnapper.cpp b/Engine/Plugins/Animation/ControlRig/Source/ControlRigEditor/Private/Tools/ControlRigSnapper.cpp
--- a/Engine/Plugins/Animation/ControlRig/Source/ControlRigEditor/Private/Tools/ControlRigSnapper.cpp
+++ b/Engine/Plugins/Animation/ControlRig/Source/ControlRigEditor/Private/Tools/ControlRigSnapper.cpp
@@ -414,7 +414,7 @@ static void GetControlRigParents(const FControlRigForWorldTransforms& ControlRig
 	}
 }
 
-static UMovieSceneControlRigParameterSection* GetControlRigSection(ISequencer* Sequencer, const UControlRig* ControlRig)
+static UMovieSceneControlRigParameterSection* GetControlRigSection(ISequencer* Sequencer, const UControlRig* ControlRig, const FName& ControlName)
 {
 
 	if (ControlRig == nullptr || Sequencer == nullptr)
@@ -432,7 +432,7 @@ static UMovieSceneControlRigParameterSection* GetControlRigSection(ISequencer* S
 		UMovieSceneControlRigParameterTrack* ControlRigParameterTrack = Cast<UMovieSceneControlRigParameterTrack>(MovieScene->FindTrack(UMovieSceneControlRigParameterTrack::StaticClass(), Binding.GetObjectGuid(), NAME_None));
 		if (ControlRigParameterTrack && ControlRigParameterTrack->GetControlRig() == ControlRig)
 		{
-			UMovieSceneControlRigParameterSection* ActiveSection = Cast<UMovieSceneControlRigParameterSection>(ControlRigParameterTrack->GetSectionToKey());
+			UMovieSceneControlRigParameterSection* ActiveSection = Cast<UMovieSceneControlRigParameterSection>(ControlRigParameterTrack->GetSectionToKey(ControlName));
 			if (ActiveSection)
 			{
 				return ActiveSection;
@@ -493,10 +493,10 @@ static void GetTransformFrames(TSharedPtr<ISequencer>&  Sequencer, const FContro
 	}
 	for (FControlRigForWorldTransforms& ControlRig : ParentControlRigs)
 	{
-		UMovieSceneControlRigParameterSection* Section = GetControlRigSection(Sequencer.Get(), ControlRig.ControlRig.Get());
-		if (Section)
+		for (FName& ControlName : ControlRig.ControlNames)
 		{
-			for (FName& ControlName : ControlRig.ControlNames)
+			UMovieSceneControlRigParameterSection* Section = GetControlRigSection(Sequencer.Get(), ControlRig.ControlRig.Get(), ControlName);
+			if (Section)
 			{
 				TArrayView<FMovieSceneFloatChannel*> Channels = FControlRigSequencerHelpers::GetFloatChannels(ControlRig.ControlRig.Get(),
 					ControlName, Section);
@@ -510,6 +510,7 @@ static void GetTransformFrames(TSharedPtr<ISequencer>&  Sequencer, const FContro
 		}
 	}
 }
+
 bool FControlRigSnapper::SnapIt(FFrameNumber StartFrame, FFrameNumber EndFrame,const FControlRigSnapperSelection& ActorToSnap,
 	const FControlRigSnapperSelection& ParentToSnap, const UControlRigSnapSettings* SnapSettings)
 {
@@ -617,6 +618,7 @@ bool FControlRigSnapper::SnapIt(FFrameNumber StartFrame, FFrameNumber EndFrame,c
 		FRigControlModifiedContext Context;
 		Context.SetKey = EControlRigSetKey::Always;
 
+		TSet<UMovieSceneControlRigParameterSection*> ControlRigSections;
 		for (const FControlRigForWorldTransforms& ControlRigAndSelection : ActorToSnap.ControlRigs)
 		{
 			//get actor transform...
@@ -720,19 +722,23 @@ bool FControlRigSnapper::SnapIt(FFrameNumber StartFrame, FFrameNumber EndFrame,c
 							Context.LocalTime = TickResolution.AsSeconds(FFrameTime(FrameNumber));
 							FTransform GlobalTransform = WorldTransformToSnap[Index].GetRelativeTransform(ControlRigParentWorldTransforms[Index]);
 							ControlRig->SetControlGlobalTransform(Name, GlobalTransform, true, Context, false /*undo*/, false /*bPrintPython*/, true/* bFixEulerFlips*/);
+							UMovieSceneControlRigParameterSection* ControlRigSection = GetControlRigSection(Sequencer.Get(), ControlRig,Name);
+							ControlRigSections.Add(ControlRigSection);
 						}
 					}
 
-					UMovieSceneControlRigParameterSection* ControlRigSection = GetControlRigSection(Sequencer.Get(), ControlRig);
-					if (ControlRigSection && SnapSettings && SnapSettings->BakingKeySettings != EBakingKeySettings::KeysOnly
+					if (SnapSettings && SnapSettings->BakingKeySettings != EBakingKeySettings::KeysOnly
 						&& SnapSettings->bReduceKeys == true && Frames.Num() > 2)
 					{
-						FKeyDataOptimizationParams Param;
-						Param.bAutoSetInterpolation = true;
-						Param.Tolerance = SnapSettings->Tolerance;
-						TRange<FFrameNumber> Range(Frames[0], Frames[Frames.Num() - 1]);
-						Param.Range = Range;
-						MovieSceneToolHelpers::OptimizeSection(Param, ControlRigSection);
+						for (UMovieSceneControlRigParameterSection* ControlRigSection : ControlRigSections)
+						{
+							FKeyDataOptimizationParams Param;
+							Param.bAutoSetInterpolation = true;
+							Param.Tolerance = SnapSettings->Tolerance;
+							TRange<FFrameNumber> Range(Frames[0], Frames[Frames.Num() - 1]);
+							Param.Range = Range;
+							MovieSceneToolHelpers::OptimizeSection(Param, ControlRigSection);
+						}
 					}
 				}
 			}
diff --git a/Engine/Plugins/Animation/ControlRig/Source/ControlRigEditor/Public/ControlRigSequencerEditorLibrary.h b/Engine/Plugins/Animation/ControlRig/Source/ControlRigEditor/Public/ControlRigSequencerEditorLibrary.h
--- a/Engine/Plugins/Animation/ControlRig/Source/ControlRigEditor/Public/ControlRigSequencerEditorLibrary.h
+++ b/Engine/Plugins/Animation/ControlRig/Source/ControlRigEditor/Public/ControlRigSequencerEditorLibrary.h
@@ -26,6 +26,7 @@ class UTickableTransformConstraint;
 class UTransformableHandle;
 struct FBakingAnimationKeySettings;
 class UControlRig;
+class UAnimLayer;
 
 USTRUCT(BlueprintType)
 struct FControlRigSequencerBindingProxy
@@ -1125,4 +1126,50 @@ public:
 	UFUNCTION(BlueprintCallable, Category = "Editor Scripting | Sequencer Tools | Control Rig")
 	static bool SetControlRigApplyMode(UControlRig* InControlRig, EControlRigFKRigExecuteMode InApplyMode);
 
+	/**
+	* Delete anim layer at specified index
+	* @param Index The index where the anim layer exists
+	* @return Returns true if successful, false otherwise
+	*/
+	UFUNCTION(BlueprintCallable, Category = "Editor Scripting | Sequencer Tools | Control Rig | Animation Layers")
+	bool DeleteAnimLayer(int32 Index);
+
+	/**
+	* Duplicate anim layer at specified index
+	* @param Index The index where the anim layer exists
+	* @return Returns index of new layer, -1 if none created
+	*/
+	UFUNCTION(BlueprintCallable, Category = "Editor Scripting | Sequencer Tools | Control Rig | Animation Layers")
+	int32 DuplicateAnimLayer(int32 Index);
+
+	/**
+	* Add anim layer from objects selected in Sequencer
+	* @return Returns Index of created anim layer
+	*/
+	UFUNCTION(BlueprintCallable, Category = "Editor Scripting | Sequencer Tools | Control Rig | Animation Layers")
+	int32 AddAnimLayerFromSelection();
+
+	/**
+	* Merge specified anim layers into one layer. Will merge onto the anim layer with the lowest index
+	* @param Indices The indices to merge
+	* @return Returns true if successful, false otherwise
+	*/
+	UFUNCTION(BlueprintCallable, Category = "Editor Scripting | Sequencer Tools | Control Rig | Animation Layers")
+	bool MergeAnimLayers(const TArray<int32>& Indices);
+
+	/**
+	* Get the animation layer objects
+	* @return Returns array of anim layer objects if they exist on active Sequencer
+	*/
+	UFUNCTION(BlueprintCallable, Category = "Editor Scripting | Sequencer Tools | Control Rig | Animation Layers")
+	TArray<UAnimLayer*> GetAnimLayers();
+	
+	/**
+	* Helper function to get the index in the anim layer array from the anim layer
+	* @param Anim Layer to get the index for
+	* @return Returns index for the anim layer or INDEX_NONE(-1) if it doesn't exist
+	*/
+	UFUNCTION(BlueprintCallable, Category = "Editor Scripting | Sequencer Tools | Control Rig | Animation Layers")
+	int32 GetAnimLayerIndex(UAnimLayer* AnimLayer);
+
 };
diff --git a/Engine/Source/Editor/MovieSceneTools/Private/MovieSceneToolHelpers.cpp b/Engine/Source/Editor/MovieSceneTools/Private/MovieSceneToolHelpers.cpp
--- a/Engine/Source/Editor/MovieSceneTools/Private/MovieSceneToolHelpers.cpp
+++ b/Engine/Source/Editor/MovieSceneTools/Private/MovieSceneToolHelpers.cpp
@@ -83,6 +83,8 @@
 #include "Channels/MovieSceneChannelTraits.h"
 #include "Channels/MovieSceneIntegerChannel.h"
 #include "Channels/MovieSceneByteChannel.h"
+#include "Channels/MovieSceneFloatChannel.h"
+#include "Channels/MovieSceneDoubleChannel.h"
 #include "Framework/MultiBox/MultiBoxBuilder.h"
 #include "EntitySystem/Interrogation/MovieSceneInterrogationLinker.h"
 #include "ConstraintsManager.h"
@@ -5039,105 +5041,6 @@ void MovieSceneToolHelpers::CalculateFramesBetween(
 	}
 }
 
-UENUM(BlueprintType)
-enum class FChannelMergeAlgorithm : uint8
-{
-	/**Average values together*/
-	Average,
-	/**Add values together*/
-	Add
-};
-
-//merge the channels onto the first one using the passed in algorithm, we may skip the last channel also (since it may be weight).
-template<typename ChannelType>
-static bool MergeChannels(TArray<ChannelType*>& Channels,const  TArray<UMovieSceneSection*>& Sections, const TRange<FFrameNumber>& Range,
-	FChannelMergeAlgorithm MergeAlgorithm)
-{
-	if (Channels.Num() < 2 && Channels.Num() != Sections.Num())
-	{
-		return false;
-	}
-	using ChannelValueType = typename ChannelType::ChannelValueType;
-	using CurveValueType = typename ChannelType::CurveValueType;
-
-	//base channel we set values on
-	ChannelType* BaseChannel = Channels[0];
-	TMovieSceneChannelData<ChannelValueType> BaseChannelData = BaseChannel->GetData();
-	//iterate over each key
-	TArray<FFrameNumber> KeyTimes;
-	TArray<FKeyHandle> Handles;
-	//cached set that we set at the end
-	TArray<TPair< FFrameNumber, ChannelValueType>> KeysToSet;
-	for (int32 ChannelIndex = 0; ChannelIndex < Channels.Num(); ++ChannelIndex)
-	{
-		KeyTimes.Reset();
-		Handles.Reset();
-		ChannelType* Channel = Channels[ChannelIndex];
-		Channel->GetKeys(Range, &KeyTimes, &Handles);
-		for (int32 FrameIndex = 0; FrameIndex < KeyTimes.Num(); ++FrameIndex)
-		{
-			const FFrameNumber& Frame = KeyTimes[FrameIndex];
-			if (Sections[0]->GetRange().Contains(Frame) == false)  //frame is outside base range so skip
-			{
-				continue;
-			}
-			const FFrameTime FrameTime(Frame);
-			int32 KeyIndex = Channel->GetData().GetIndex(Handles[FrameIndex]);
-			ChannelValueType Value = Channel->GetData().GetValues()[KeyIndex];
-			//got value with tangents and times, now we perform the operation
-			Value.Value = 0.0; //zero out the value we calculate it 
-			if (MergeAlgorithm == FChannelMergeAlgorithm::Average)
-			{
-				double DNumChannels = 0.0;
-				for (int32 WeightIndex = 0; WeightIndex < Sections.Num(); ++WeightIndex)
-				{
-					if (Sections[WeightIndex]->GetRange().Contains(Frame))  
-					{
-						float Weight = Sections[WeightIndex]->GetTotalWeightValue(FrameTime);
-						CurveValueType WeightedValue = 0.0;
-						ChannelType* EachChannel = Channels[WeightIndex];
-						EachChannel->Evaluate(FrameTime, WeightedValue);
-						WeightedValue *= ((double)Weight);
-						Value.Value += WeightedValue;
-						DNumChannels += 1.0;
-					}
-				}
-				if (DNumChannels > 0.0) //should always happen since base(0) at least be here
-				{
-					Value.Value /= DNumChannels;
-				}
-			}
-			else if (MergeAlgorithm == FChannelMergeAlgorithm::Add)
-			{
-				for (int32 WeightIndex = 0; WeightIndex < Sections.Num(); ++WeightIndex)
-				{
-					if (Sections[WeightIndex]->GetRange().Contains(Frame))  
-					{
-						float Weight = Sections[WeightIndex]->GetTotalWeightValue(FrameTime);
-						if (Sections[WeightIndex]->GetBlendType().IsValid() == false ||
-							Sections[WeightIndex]->GetBlendType().Get() != EMovieSceneBlendType::Additive)
-						{
-							Weight = 1.0;
-						}
-						CurveValueType WeightedValue = 0.0;
-						ChannelType* EachChannel = Channels[WeightIndex];
-						EachChannel->Evaluate(FrameTime, WeightedValue);
-						WeightedValue *= Weight;
-						Value.Value += WeightedValue;
-					}
-				}
-			}
-			KeysToSet.Add(TPair<FFrameNumber, ChannelValueType>(Frame, Value));
-		}
-	}
-	for (TPair<FFrameNumber, ChannelValueType>& KeyToSet : KeysToSet)
-	{
-		MovieSceneToolHelpers::SetOrAddKey(BaseChannelData, KeyToSet.Key, KeyToSet.Value);
-	}
-
-	return true;
-}
-
 bool MovieSceneToolHelpers::OptimizeSection(const FKeyDataOptimizationParams& InParams, UMovieSceneSection* InSection)
 {
 	TArrayView<FMovieSceneFloatChannel*> FloatChannels = InSection->GetChannelProxy().GetChannels<FMovieSceneFloatChannel>();
@@ -5153,87 +5056,6 @@ bool MovieSceneToolHelpers::OptimizeSection(const FKeyDataOptimizationParams& In
 	return true;
 }
 
-template<typename ChannelType>
-static bool MergeSections(UMovieSceneSection* BaseSection, TArrayView<ChannelType*> BaseDoubleChannels, TArray<UMovieSceneSection*>& AbsoluteSections, TArray<UMovieSceneSection*>& AdditiveSections,
-	const TRange<FFrameNumber>& Range,bool bSkipLastChannel)
-{
-	TArrayView<ChannelType*> BaseChannels = BaseSection->GetChannelProxy().GetChannels<ChannelType>();
-	if (BaseChannels.Num() > 0)
-	{
-		int32 SectionIndex = 0;
-		//sanity check to make sure channels are the same size
-		for (SectionIndex = 0; SectionIndex < AbsoluteSections.Num(); ++SectionIndex)
-		{
-			TArrayView<ChannelType*>Channels = AbsoluteSections[SectionIndex]->GetChannelProxy().GetChannels<ChannelType>();
-			if (Channels.Num() != BaseChannels.Num())
-			{
-				UE_LOG(LogMovieScene, Warning, TEXT("MergeSections:: Invalid number of channels"));
-				return false;
-			}
-		}
-		for (SectionIndex = 0; SectionIndex < AdditiveSections.Num(); ++SectionIndex)
-		{
-			TArrayView<ChannelType*>Channels = AdditiveSections[SectionIndex]->GetChannelProxy().GetChannels<ChannelType>();
-			if (Channels.Num() != BaseChannels.Num())
-			{
-				UE_LOG(LogMovieScene, Warning, TEXT("MergeSections:: Invalid number of channels"));
-				return false;
-			}
-		}
-		BaseSection->Modify();
-
-		TArray<ChannelType*> Channels;
-		int32 ChannelIndex = 0;
-		bSkipLastChannel = false;
-		int32 BaseChannelsNum = bSkipLastChannel ? BaseChannels.Num() - 1 : BaseChannels.Num();
-		for (ChannelIndex = 0; ChannelIndex < BaseChannels.Num(); ++ChannelIndex)
-		{
-			if (AbsoluteSections.Num() > 0)
-			{
-				Channels.Reset();
-				for (SectionIndex = 0; SectionIndex < AbsoluteSections.Num(); ++SectionIndex)
-				{
-					TArrayView<ChannelType*>OurChannels = AbsoluteSections[SectionIndex]->GetChannelProxy().GetChannels<ChannelType>();
-					Channels.Add(OurChannels[ChannelIndex]);
-				}
-				//now blend them
-				if (MergeChannels(Channels, AbsoluteSections, Range,
-					FChannelMergeAlgorithm::Average) == false)
-				{
-					UE_LOG(LogMovieScene, Warning, TEXT("MergeSections:: Could not merge channels"));
-					return false;
-				}
-			}
-			if (AdditiveSections.Num() > 0)
-			{
-				//now do additives
-				Channels.Reset();
-				for (SectionIndex = 0; SectionIndex < AdditiveSections.Num(); ++SectionIndex)
-				{
-					TArrayView<ChannelType*>OurChannels = AdditiveSections[SectionIndex]->GetChannelProxy().GetChannels<ChannelType>();
-					Channels.Add(OurChannels[ChannelIndex]);
-				}
-				//now blend them
-				if (MergeChannels(Channels, AdditiveSections, Range,
-					FChannelMergeAlgorithm::Add) == false)
-				{
-					UE_LOG(LogMovieScene, Warning, TEXT("MergeSections:: Could not merge channels"));
-					return false;
-				}
-			}
-		}
-		for (ChannelType* Channel : BaseChannels)
-		{
-			Channel->AutoSetTangents();
-		}
-	}
-	else
-	{
-		return false;
-	}
-	return true;
-}
-
 bool MovieSceneToolHelpers::CollapseSection(TSharedPtr<ISequencer>& SequencerPtr, UMovieSceneTrack* OwnerTrack, TArray<UMovieSceneSection*> Sections,
 	const FBakingAnimationKeySettings& InSettings)
 {
@@ -5309,7 +5131,7 @@ bool MovieSceneToolHelpers::CollapseSection(TSharedPtr<ISequencer>& SequencerPtr
 		TArrayView<FMovieSceneDoubleChannel*> BaseDoubleChannels = BaseSection->GetChannelProxy().GetChannels<FMovieSceneDoubleChannel>();
 		if (BaseDoubleChannels.Num() > 0) //transforms
 		{
-			if (MergeSections(BaseSection, BaseDoubleChannels, AbsoluteSections, AdditiveSections, Range, false /*bSkipLastChannel*/) == false)
+			if (MovieSceneToolHelpers::MergeSections<FMovieSceneDoubleChannel>(BaseSection, AbsoluteSections, AdditiveSections, Range, false /*bSkipLastChannel*/) == false)
 			{
 				Transaction.Cancel();
 				return false;
@@ -5318,7 +5140,7 @@ bool MovieSceneToolHelpers::CollapseSection(TSharedPtr<ISequencer>& SequencerPtr
 		else if(BaseFloatChannels.Num() > 0) //control rig
 		{
 			//skip weight channel for control rig
-			if (MergeSections(BaseSection, BaseFloatChannels, AbsoluteSections, AdditiveSections, Range, true /*bSkipLastChannel*/) == false)
+			if (MovieSceneToolHelpers::MergeSections<FMovieSceneFloatChannel>(BaseSection,  AbsoluteSections, AdditiveSections, Range, true /*bSkipLastChannel*/) == false)
 			{
 				Transaction.Cancel();
 				return false;
diff --git a/Engine/Source/Editor/MovieSceneTools/Public/MovieSceneToolHelpers.h b/Engine/Source/Editor/MovieSceneTools/Public/MovieSceneToolHelpers.h
--- a/Engine/Source/Editor/MovieSceneTools/Public/MovieSceneToolHelpers.h
+++ b/Engine/Source/Editor/MovieSceneTools/Public/MovieSceneToolHelpers.h
@@ -19,10 +19,11 @@
 #include "SEnumCombo.h"
 #include "Animation/AnimSequence.h"
 #include "INodeAndChannelMappings.h"
+#include "Channels/MovieSceneChannelProxy.h"
+#include "MovieSceneSection.h"
 
 class ISequencer;
 class UMovieScene;
-class UMovieSceneSection;
 class UMovieSceneSequence;
 struct FMovieSceneObjectBindingID;
 class UMovieSceneTrack;
@@ -105,6 +106,14 @@ public:
 	bool bCachedEnableUpdateRateOptimizations;
 };
 
+enum class FChannelMergeAlgorithm : uint8
+{
+	/**Average values together*/
+	Average,
+	/**Add values together*/
+	Add
+};
+
 class MOVIESCENETOOLS_API MovieSceneToolHelpers
 {
 public:
@@ -662,6 +671,267 @@ public:
 	static bool CollapseSection(TSharedPtr<ISequencer>& SequencerPtr, UMovieSceneTrack* OwnerTrack, TArray<UMovieSceneSection*> Sections,
 		const FBakingAnimationKeySettings& InSettings);
 
+	//merge the channels onto the first one using the passed in algorithm, we may skip the last channel also (since it may be weight).
+	template<typename ChannelType>
+	static bool MergeChannels(TArray<ChannelType*>& Channels, const  TArray<UMovieSceneSection*>& Sections, const TRange<FFrameNumber>& Range,
+		FChannelMergeAlgorithm MergeAlgorithm)
+	{
+		if (Channels.Num() < 2 && Channels.Num() != Sections.Num())
+		{
+			return false;
+		}
+		using ChannelValueType = typename ChannelType::ChannelValueType;
+		using CurveValueType = typename ChannelType::CurveValueType;
+
+		//base channel we set values on
+		ChannelType* BaseChannel = Channels[0];
+		TMovieSceneChannelData<ChannelValueType> BaseChannelData = BaseChannel->GetData();
+		//iterate over each key
+		TArray<FFrameNumber> KeyTimes;
+		TArray<FKeyHandle> Handles;
+		//cached set that we set at the end
+		TArray<TPair< FFrameNumber, ChannelValueType>> KeysToSet;
+		for (int32 ChannelIndex = 0; ChannelIndex < Channels.Num(); ++ChannelIndex)
+		{
+			KeyTimes.Reset();
+			Handles.Reset();
+			ChannelType* Channel = Channels[ChannelIndex];
+			Channel->GetKeys(Range, &KeyTimes, &Handles);
+			for (int32 FrameIndex = 0; FrameIndex < KeyTimes.Num(); ++FrameIndex)
+			{
+				const FFrameNumber& Frame = KeyTimes[FrameIndex];
+				if (Sections[0]->GetRange().Contains(Frame) == false)  //frame is outside base range so skip
+				{
+					continue;
+				}
+				const FFrameTime FrameTime(Frame);
+				int32 KeyIndex = Channel->GetData().GetIndex(Handles[FrameIndex]);
+				ChannelValueType Value = Channel->GetData().GetValues()[KeyIndex];
+				//got value with tangents and times, now we perform the operation
+				Value.Value = 0.0; //zero out the value we calculate it 
+				if (MergeAlgorithm == FChannelMergeAlgorithm::Average)
+				{
+					double DNumChannels = 0.0;
+					for (int32 WeightIndex = 0; WeightIndex < Sections.Num(); ++WeightIndex)
+					{
+						if (Sections[WeightIndex]->GetRange().Contains(Frame))
+						{
+							float Weight = Sections[WeightIndex]->GetTotalWeightValue(FrameTime);
+							CurveValueType WeightedValue = 0.0;
+							ChannelType* EachChannel = Channels[WeightIndex];
+							EachChannel->Evaluate(FrameTime, WeightedValue);
+							WeightedValue *= ((double)Weight);
+							Value.Value += WeightedValue;
+							DNumChannels += 1.0;
+						}
+					}
+					if (DNumChannels > 0.0) //should always happen since base(0) at least be here
+					{
+						Value.Value /= DNumChannels;
+					}
+				}
+				else if (MergeAlgorithm == FChannelMergeAlgorithm::Add)
+				{
+					for (int32 WeightIndex = 0; WeightIndex < Sections.Num(); ++WeightIndex)
+					{
+						if (Sections[WeightIndex]->GetRange().Contains(Frame))
+						{
+							float Weight = Sections[WeightIndex]->GetTotalWeightValue(FrameTime);
+							if (Sections[WeightIndex]->GetBlendType().IsValid() == false ||
+								Sections[WeightIndex]->GetBlendType().Get() != EMovieSceneBlendType::Additive)
+							{
+								Weight = 1.0;
+							}
+							CurveValueType WeightedValue = 0.0;
+							ChannelType* EachChannel = Channels[WeightIndex];
+							EachChannel->Evaluate(FrameTime, WeightedValue);
+							WeightedValue *= Weight;
+							Value.Value += WeightedValue;
+						}
+					}
+				}
+				KeysToSet.Add(TPair<FFrameNumber, ChannelValueType>(Frame, Value));
+			}
+		}
+		for (TPair<FFrameNumber, ChannelValueType>& KeyToSet : KeysToSet)
+		{
+			MovieSceneToolHelpers::SetOrAddKey(BaseChannelData, KeyToSet.Key, KeyToSet.Value);
+		}
+
+		return true;
+	}
+
+	template<typename ChannelType>
+	static bool MergeNonWeightedChannels(TArray<ChannelType*>& Channels, const  TArray<UMovieSceneSection*>& Sections, const TRange<FFrameNumber>& Range,
+		FChannelMergeAlgorithm MergeAlgorithm)
+	{
+		if (Channels.Num() < 2 && Channels.Num() != Sections.Num())
+		{
+			return false;
+		}
+		using ChannelValueType = typename ChannelType::ChannelValueType;
+		using CurveValueType = typename ChannelType::CurveValueType;
+
+		//base channel we set values on
+		ChannelType* BaseChannel = Channels[0];
+		TMovieSceneChannelData<ChannelValueType> BaseChannelData = BaseChannel->GetData();
+		//iterate over each key
+		TArray<FFrameNumber> KeyTimes;
+		TArray<FKeyHandle> Handles;
+		//cached set that we set at the end
+		TArray<TPair< FFrameNumber, ChannelValueType>> KeysToSet;
+		for (int32 ChannelIndex = 0; ChannelIndex < Channels.Num(); ++ChannelIndex)
+		{
+			KeyTimes.Reset();
+			Handles.Reset();
+			ChannelType* Channel = Channels[ChannelIndex];
+			Channel->GetKeys(Range, &KeyTimes, &Handles);
+			for (int32 FrameIndex = 0; FrameIndex < KeyTimes.Num(); ++FrameIndex)
+			{
+				const FFrameNumber& Frame = KeyTimes[FrameIndex];
+				if (Sections[0]->GetRange().Contains(Frame) == false)  //frame is outside base range so skip
+				{
+					continue;
+				}
+				const FFrameTime FrameTime(Frame);
+				int32 KeyIndex = Channel->GetData().GetIndex(Handles[FrameIndex]);
+				ChannelValueType Value = Channel->GetData().GetValues()[KeyIndex];
+				//got value with tangents and times, now we perform the operation
+				Value.Value = 0.0; //zero out the value we calculate it 
+				if (MergeAlgorithm == FChannelMergeAlgorithm::Average)
+				{
+					double DNumChannels = 0.0;
+					for (int32 WeightIndex = 0; WeightIndex < Sections.Num(); ++WeightIndex)
+					{
+						if (Sections[WeightIndex]->GetRange().Contains(Frame))
+						{
+							float Weight = Sections[WeightIndex]->GetTotalWeightValue(FrameTime);
+							CurveValueType WeightedValue = 0.0;
+							ChannelType* EachChannel = Channels[WeightIndex];
+							EachChannel->Evaluate(FrameTime, WeightedValue);
+							WeightedValue *= ((double)Weight);
+							Value.Value += WeightedValue;
+							DNumChannels += 1.0;
+						}
+					}
+					if (DNumChannels > 0.0) //should always happen since base(0) at least be here
+					{
+						Value.Value /= DNumChannels;
+					}
+				}
+				else if (MergeAlgorithm == FChannelMergeAlgorithm::Add)
+				{
+					for (int32 WeightIndex = 0; WeightIndex < Sections.Num(); ++WeightIndex)
+					{
+						if (Sections[WeightIndex]->GetRange().Contains(Frame))
+						{
+							float Weight = Sections[WeightIndex]->GetTotalWeightValue(FrameTime);
+							if (Sections[WeightIndex]->GetBlendType().IsValid() == false ||
+								Sections[WeightIndex]->GetBlendType().Get() != EMovieSceneBlendType::Additive)
+							{
+								Weight = 1.0;
+							}
+							CurveValueType WeightedValue = 0.0;
+							ChannelType* EachChannel = Channels[WeightIndex];
+							EachChannel->Evaluate(FrameTime, WeightedValue);
+							WeightedValue *= Weight;
+							Value.Value += WeightedValue;
+						}
+					}
+				}
+				KeysToSet.Add(TPair<FFrameNumber, ChannelValueType>(Frame, Value));
+			}
+		}
+		for (TPair<FFrameNumber, ChannelValueType>& KeyToSet : KeysToSet)
+		{
+			MovieSceneToolHelpers::SetOrAddKey(BaseChannelData, KeyToSet.Key, KeyToSet.Value);
+		}
+
+		return true;
+	}
+
+	template<typename ChannelType>
+	static bool MergeSections(UMovieSceneSection* BaseSection, TArray<UMovieSceneSection*>& AbsoluteSections, TArray<UMovieSceneSection*>& AdditiveSections,
+		const TRange<FFrameNumber>& Range, bool bSkipLastChannel)
+	{
+		TArrayView<ChannelType*> BaseChannels = BaseSection->GetChannelProxy().GetChannels<ChannelType>();
+		if (BaseChannels.Num() > 0)
+		{
+			int32 SectionIndex = 0;
+			//sanity check to make sure channels are the same size
+			for (SectionIndex = 0; SectionIndex < AbsoluteSections.Num(); ++SectionIndex)
+			{
+				TArrayView<ChannelType*>Channels = AbsoluteSections[SectionIndex]->GetChannelProxy().GetChannels<ChannelType>();
+				if (Channels.Num() != BaseChannels.Num())
+				{
+					UE_LOG(LogMovieScene, Warning, TEXT("MergeSections:: Invalid number of channels"));
+					return false;
+				}
+			}
+			for (SectionIndex = 0; SectionIndex < AdditiveSections.Num(); ++SectionIndex)
+			{
+				TArrayView<ChannelType*>Channels = AdditiveSections[SectionIndex]->GetChannelProxy().GetChannels<ChannelType>();
+				if (Channels.Num() != BaseChannels.Num())
+				{
+					UE_LOG(LogMovieScene, Warning, TEXT("MergeSections:: Invalid number of channels"));
+					return false;
+				}
+			}
+			BaseSection->Modify();
+
+			TArray<ChannelType*> Channels;
+			int32 ChannelIndex = 0;
+			bSkipLastChannel = false;
+			int32 BaseChannelsNum = bSkipLastChannel ? BaseChannels.Num() - 1 : BaseChannels.Num();
+			for (ChannelIndex = 0; ChannelIndex < BaseChannels.Num(); ++ChannelIndex)
+			{
+				if (AbsoluteSections.Num() > 0)
+				{
+					Channels.Reset();
+					for (SectionIndex = 0; SectionIndex < AbsoluteSections.Num(); ++SectionIndex)
+					{
+						TArrayView<ChannelType*>OurChannels = AbsoluteSections[SectionIndex]->GetChannelProxy().GetChannels<ChannelType>();
+						Channels.Add(OurChannels[ChannelIndex]);
+					}
+					//now blend them
+					if (MergeChannels(Channels, AbsoluteSections, Range,
+						FChannelMergeAlgorithm::Average) == false)
+					{
+						UE_LOG(LogMovieScene, Warning, TEXT("MergeSections:: Could not merge channels"));
+						return false;
+					}
+				}
+				if (AdditiveSections.Num() > 0)
+				{
+					//now do additives
+					Channels.Reset();
+					for (SectionIndex = 0; SectionIndex < AdditiveSections.Num(); ++SectionIndex)
+					{
+						TArrayView<ChannelType*>OurChannels = AdditiveSections[SectionIndex]->GetChannelProxy().GetChannels<ChannelType>();
+						Channels.Add(OurChannels[ChannelIndex]);
+					}
+					//now blend them
+					if (MergeChannels(Channels, AdditiveSections, Range,
+						FChannelMergeAlgorithm::Add) == false)
+					{
+						UE_LOG(LogMovieScene, Warning, TEXT("MergeSections:: Could not merge channels"));
+						return false;
+					}
+				}
+			}
+			for (ChannelType* Channel : BaseChannels)
+			{
+				Channel->AutoSetTangents();
+			}
+		}
+		else
+		{
+			return false;
+		}
+		return true;
+	}
+
+
 	static bool OptimizeSection(const FKeyDataOptimizationParams& InParams, UMovieSceneSection* InSection);
 
 	/** Returns the frame numbers between start and end. */
diff --git a/Engine/Source/Editor/Sequencer/Private/SSequencerSection.cpp b/Engine/Source/Editor/Sequencer/Private/SSequencerSection.cpp
--- a/Engine/Source/Editor/Sequencer/Private/SSequencerSection.cpp
+++ b/Engine/Source/Editor/Sequencer/Private/SSequencerSection.cpp
@@ -41,6 +41,7 @@
 #include "IKeyArea.h"
 #include "Widgets/SWeakWidget.h"
 #include "Algo/Transform.h"
+#include "Tracks/IMovieSceneSectionsToKey.h"
 
 namespace UE
 {
@@ -73,6 +74,15 @@ FTimeToPixel ConstructTimeConverterForSection(const FGeometry& InSectionGeometry
 	return FTimeToPixel(InSectionGeometry, TRange<double>(LowerTime, UpperTime), TickResolution);
 }
 
+static bool IsSectionToKey(const UMovieSceneTrack* Track, UMovieSceneSection* SectionObject)
+{
+	if (const IMovieSceneSectionsToKey* MultipleSectionsToKey = Cast<IMovieSceneSectionsToKey>(Track))
+	{
+		TArray<TWeakObjectPtr<UMovieSceneSection>> SectionsToKey = MultipleSectionsToKey->GetSectionsToKey();
+		return SectionsToKey.Contains(SectionObject);
+	}
+	return (Track->GetSectionToKey() == SectionObject);
+}
 
 struct FSequencerSectionPainterImpl : FSequencerSectionPainter
 {
@@ -788,7 +798,7 @@ struct FSequencerSectionPainterImpl : FSequencerSectionPainter
 		const bool bLocked = SectionObject->IsLocked() || SectionObject->IsReadOnly();
 
 		// Only show section to key border if we have more than one section
-		const bool bIsSectionToKey = Track && Track->GetAllSections().Num() > 1 && Track->GetSectionToKey() == SectionObject;
+		const bool bIsSectionToKey = Track && Track->GetAllSections().Num() > 1 && IsSectionToKey(Track,SectionObject);
 
 		const ESlateDrawEffect DrawEffects = bParentEnabled ? ESlateDrawEffect::None : ESlateDrawEffect::DisabledEffect;
 
@@ -1675,7 +1685,7 @@ void SSequencerSection::PaintEasingHandles( FSequencerSectionPainter& InPainter,
 		const float MinHandleSize = 8.f;
 
 		const ESlateDrawEffect DrawEffects = InPainter.bParentEnabled ? ESlateDrawEffect::None : ESlateDrawEffect::DisabledEffect;
-		const bool bIsSectionToKey = Track->GetAllSections().Num() > 1 && Track->GetSectionToKey() == UnderlappingSectionObj;
+		const bool bIsSectionToKey = Track->GetAllSections().Num() > 1 && IsSectionToKey(Track, UnderlappingSectionObj);
 
 		// If this is the section to key, we draw the easing handle in the same bright green color as the border
 		// outline. We also make the handle a bit bigger, otherwise that border being drawn on top makes it look
diff --git a/Engine/Source/Runtime/MovieScene/Private/Evaluation/Blending/MovieSceneBlendType.cpp b/Engine/Source/Runtime/MovieScene/Private/Evaluation/Blending/MovieSceneBlendType.cpp
--- a/Engine/Source/Runtime/MovieScene/Private/Evaluation/Blending/MovieSceneBlendType.cpp
+++ b/Engine/Source/Runtime/MovieScene/Private/Evaluation/Blending/MovieSceneBlendType.cpp
@@ -22,7 +22,7 @@ FMovieSceneBlendTypeField FMovieSceneBlendTypeField::All()
 {
 	FMovieSceneBlendTypeField New;
 	New.Add(EMovieSceneBlendType::Absolute, EMovieSceneBlendType::Additive, EMovieSceneBlendType::Relative,
-			EMovieSceneBlendType::AdditiveFromBase);
+			EMovieSceneBlendType::AdditiveFromBase, EMovieSceneBlendType::Override);
 	return New;
 }
 
@@ -73,9 +73,11 @@ int32 FMovieSceneBlendTypeField::Num() const
 		(Contains(EMovieSceneBlendType::Absolute) ? 1 : 0) +
 		(Contains(EMovieSceneBlendType::Relative) ? 1 : 0) +
 		(Contains(EMovieSceneBlendType::Additive) ? 1 : 0) +
-		(Contains(EMovieSceneBlendType::AdditiveFromBase) ? 1 : 0);
+		(Contains(EMovieSceneBlendType::AdditiveFromBase) ? 1 : 0) +
+		(Contains(EMovieSceneBlendType::Override) ? 1 : 0);
 }
 
+
 void FMovieSceneBlendTypeFieldIterator::IterateToNext()
 {
 	do
diff --git a/Engine/Source/Runtime/MovieScene/Public/Evaluation/Blending/MovieSceneBlendType.h b/Engine/Source/Runtime/MovieScene/Public/Evaluation/Blending/MovieSceneBlendType.h
--- a/Engine/Source/Runtime/MovieScene/Public/Evaluation/Blending/MovieSceneBlendType.h
+++ b/Engine/Source/Runtime/MovieScene/Public/Evaluation/Blending/MovieSceneBlendType.h
@@ -24,6 +24,8 @@ enum class EMovieSceneBlendType : uint8
 	Relative			= 0x4,
 	/** Applies this value as an additive equal to the difference between the current value and the first value */
 	AdditiveFromBase    = 0x8,
+	/** The value will override the current value */
+	Override			= 0x10,
 
 	// If you add new blend types, don't forget to change:
 	// - All()
diff --git a/Engine/Source/Runtime/MovieSceneTracks/Private/Tracks/IMovieSceneSectionsToKey.cpp b/Engine/Source/Runtime/MovieSceneTracks/Private/Tracks/IMovieSceneSectionsToKey.cpp
new file mode 100644
--- /dev/null
+++ b/Engine/Source/Runtime/MovieSceneTracks/Private/Tracks/IMovieSceneSectionsToKey.cpp
@@ -0,0 +1,15 @@
+// Copyright Epic Games, Inc. All Rights Reserved.
+
+#include "Tracks/IMovieSceneSectionsToKey.h"
+
+#include UE_INLINE_GENERATED_CPP_BY_NAME(IMovieSceneSectionsToKey)
+
+
+IMovieSceneSectionsToKey::IMovieSceneSectionsToKey() 
+{
+}
+
+
+
+
+
diff --git a/Engine/Source/Runtime/MovieSceneTracks/Public/Sections/MovieScene3DTransformSection.h b/Engine/Source/Runtime/MovieSceneTracks/Public/Sections/MovieScene3DTransformSection.h
--- a/Engine/Source/Runtime/MovieSceneTracks/Public/Sections/MovieScene3DTransformSection.h
+++ b/Engine/Source/Runtime/MovieSceneTracks/Public/Sections/MovieScene3DTransformSection.h
@@ -231,8 +231,8 @@ public:
 #endif
 
 	/* From UMovieSection*/
-	
-	MOVIESCENETRACKS_API virtual bool ShowCurveForChannel(const void *Channel) const override;
+
+	MOVIESCENETRACKS_API virtual bool ShowCurveForChannel(const void* Channel) const override;
 	MOVIESCENETRACKS_API virtual void SetBlendType(EMovieSceneBlendType InBlendType) override;
 	MOVIESCENETRACKS_API virtual void OnBindingIDsUpdated(const TMap<UE::MovieScene::FFixedObjectBindingID, UE::MovieScene::FFixedObjectBindingID>& OldFixedToNewFixedMap, FMovieSceneSequenceID LocalSequenceID, TSharedRef<UE::MovieScene::FSharedPlaybackState> SharedPlaybackState) override;
 	MOVIESCENETRACKS_API virtual void GetReferencedBindings(TArray<FGuid>& OutBindings) override;
@@ -269,6 +269,10 @@ public:
 	*/
 	MOVIESCENETRACKS_API void SetUseQuaternionInterpolation(bool bInUseQuaternionInterpolation);
 
+	/**
+	* Get Weight Channel
+	*/
+	MOVIESCENETRACKS_API FMovieSceneFloatChannel* GetWeightChannel() { return &ManualWeight; }
 protected:
 
 	MOVIESCENETRACKS_API virtual TSharedPtr<FStructOnScope> GetKeyStruct(TArrayView<const FKeyHandle> KeyHandles) override;
diff --git a/Engine/Source/Runtime/MovieSceneTracks/Public/Tracks/IMovieSceneSectionsToKey.h b/Engine/Source/Runtime/MovieSceneTracks/Public/Tracks/IMovieSceneSectionsToKey.h
new file mode 100644
--- /dev/null
+++ b/Engine/Source/Runtime/MovieSceneTracks/Public/Tracks/IMovieSceneSectionsToKey.h
@@ -0,0 +1,37 @@
+// Copyright Epic Games, Inc. All Rights Reserved.
+
+#pragma once
+
+#include "CoreMinimal.h"
+#include "UObject/ObjectMacros.h"
+#include "UObject/NameTypes.h"
+#include "UObject/Interface.h"
+#include "IMovieSceneSectionsToKey.generated.h"
+
+class UMovieSceneSection;
+
+/**
+ * Functionality for having multiple sections per key
+ */
+
+UINTERFACE(MinimalAPI)
+class UMovieSceneSectionsToKey : public UInterface
+{
+public:
+	GENERATED_BODY()
+};
+
+/**
+ * Interface to be added to UMovieSceneTrack when they have multiple sections to key
+ */
+class IMovieSceneSectionsToKey
+{
+public:
+
+	GENERATED_BODY()
+	MOVIESCENETRACKS_API IMovieSceneSectionsToKey();
+
+	/* Get multiple sections to key*/
+	virtual TArray<TWeakObjectPtr<UMovieSceneSection>> GetSectionsToKey() const = 0;
+
+};
