NonUnity CIS FIx

diff --git a/Engine/Plugins/Animation/ControlRig/Source/ControlRigEditor/Private/Sequencer/AnimLayers/SAnimLayers.h b/Engine/Plugins/Animation/ControlRig/Source/ControlRigEditor/Private/Sequencer/AnimLayers/SAnimLayers.h
--- a/Engine/Plugins/Animation/ControlRig/Source/ControlRigEditor/Private/Sequencer/AnimLayers/SAnimLayers.h
+++ b/Engine/Plugins/Animation/ControlRig/Source/ControlRigEditor/Private/Sequencer/AnimLayers/SAnimLayers.h
@@ -7,6 +7,7 @@
 #include "Widgets/SCompoundWidget.h"
 #include "ISequencer.h"
 #include "Misc/Guid.h"
+#include "Widgets/DeclarativeSyntaxSupport.h"
 
 class UControlRig;
 class FControlRigEditMode;
