[Sequencer ControlRig] minor tweak to the filter menu of control rig picker, removed irrelevant filters

diff --git a/Engine/Plugins/Animation/ControlRig/Source/ControlRigEditor/Private/Sequencer/ControlRigParameterTrackEditor.cpp b/Engine/Plugins/Animation/ControlRig/Source/ControlRigEditor/Private/Sequencer/ControlRigParameterTrackEditor.cpp
--- a/Engine/Plugins/Animation/ControlRig/Source/ControlRigEditor/Private/Sequencer/ControlRigParameterTrackEditor.cpp
+++ b/Engine/Plugins/Animation/ControlRig/Source/ControlRigEditor/Private/Sequencer/ControlRigParameterTrackEditor.cpp
@@ -1602,6 +1602,7 @@ void FControlRigParameterTrackEditor::HandleAddTrackSubMenu(FMenuBuilder& MenuBu
 			AssetPickerConfig.SelectionMode = ESelectionMode::Single;
 			AssetPickerConfig.bAddFilterUI = true;
 			AssetPickerConfig.bFocusSearchBoxWhenOpened = true;
+			AssetPickerConfig.bForceShowPluginContent = true;
 			AssetPickerConfig.OnAssetSelected = FOnAssetSelected::CreateRaw(this, &FControlRigParameterTrackEditor::AddControlRig, BoundObject, ObjectBindings[0]);
 			AssetPickerConfig.OnAssetEnterPressed = FOnAssetEnterPressed::CreateRaw(this, &FControlRigParameterTrackEditor::AddControlRig, BoundObject, ObjectBindings[0]);
 			AssetPickerConfig.RefreshAssetViewDelegates.Add(&RefreshControlRigPickerDelegate);
@@ -1620,7 +1621,7 @@ void FControlRigParameterTrackEditor::HandleAddTrackSubMenu(FMenuBuilder& MenuBu
 			AssetPickerConfig.Filter.bRecursiveClasses = true;
 			AssetPickerConfig.Filter.ClassPaths.Add((UControlRigBlueprint::StaticClass())->GetClassPathName());
 			AssetPickerConfig.SaveSettingsName = TEXT("SequencerControlRigTrackAssetPicker");
-			TSharedRef<FFrontendFilterCategory>	ControlRigFilterCategory = MakeShared<FFrontendFilterCategory>(LOCTEXT("ControlRigFilterCategoryName", "ControlRig Tags"), LOCTEXT("ControlRigFilterCategoryToolTip", "Filter ControlRigs by tags specified in ControlRig Blueprint class settings"));
+			TSharedRef<FFrontendFilterCategory>	ControlRigFilterCategory = MakeShared<FFrontendFilterCategory>(LOCTEXT("ControlRigFilterCategoryName", "Control Rig Tags"), LOCTEXT("ControlRigFilterCategoryToolTip", "Filter ControlRigs by variant tags specified in ControlRig Blueprint class settings"));
 			const URigVMProjectSettings* Settings = GetDefault<URigVMProjectSettings>(URigVMProjectSettings::StaticClass());
 			TArray<FRigVMTag> AvailableTags = Settings->VariantTags;
 
@@ -1628,6 +1629,16 @@ void FControlRigParameterTrackEditor::HandleAddTrackSubMenu(FMenuBuilder& MenuBu
 			{
 				AssetPickerConfig.ExtraFrontendFilters.Add(MakeShared<FFrontendFilter_ControlRigFilterByAssetTag>(ControlRigFilterCategory, Tag));
 			}
+
+			// This is so that we can remove the "Other Filters" section easily
+			AssetPickerConfig.bUseSectionsForCustomFilterCategories = true;
+			// Make sure we only show ControlRig filters to avoid confusion 
+			AssetPickerConfig.OnExtendAddFilterMenu = FOnExtendAddFilterMenu::CreateLambda([ControlRigFilters = AssetPickerConfig.ExtraFrontendFilters](UToolMenu* InToolMenu)
+			{
+				// "AssetFilterBarFilterAdvancedAsset" taken from SAssetFilterBar.h PopulateAddFilterMenu()
+				InToolMenu->RemoveSection("AssetFilterBarFilterAdvancedAsset");
+				InToolMenu->RemoveSection("Other Filters");
+			});
 		}
 
 		FContentBrowserModule& ContentBrowserModule = FModuleManager::Get().LoadModuleChecked<FContentBrowserModule>(TEXT("ContentBrowser"));
