Anim Layers: Localization fix

diff --git a/Engine/Plugins/Animation/ControlRig/Source/ControlRigEditor/Private/Sequencer/AnimLayers/SAnimLayers.cpp b/Engine/Plugins/Animation/ControlRig/Source/ControlRigEditor/Private/Sequencer/AnimLayers/SAnimLayers.cpp
--- a/Engine/Plugins/Animation/ControlRig/Source/ControlRigEditor/Private/Sequencer/AnimLayers/SAnimLayers.cpp
+++ b/Engine/Plugins/Animation/ControlRig/Source/ControlRigEditor/Private/Sequencer/AnimLayers/SAnimLayers.cpp
@@ -1346,7 +1346,7 @@ void SAnimLayers::Construct(const FArguments& InArgs, FControlRigEditMode& InEdi
 							.OnClicked(this, &SAnimLayers::OnAddClicked)
 							.Icon(FAppStyle::Get().GetBrush("Icons.Plus"))
 							.Text(LOCTEXT("AnimLayer", "Layer"))
-							.ToolTipText(LOCTEXT("AnimLayer", "Add a new Animation Layer"))
+							.ToolTipText(LOCTEXT("AnimLayerTooltip", "Add a new Animation Layer"))
 						]
 						+ SHorizontalBox::Slot()
 						.FillWidth(10.f)
