[Seqeuencer ControlRig] for tags that marks things as invalid, the filter now inverted and mark with Exclude "tag name"

diff --git a/Engine/Plugins/Animation/ControlRig/Source/ControlRigEditor/Private/Sequencer/ControlRigParameterTrackEditor.cpp b/Engine/Plugins/Animation/ControlRig/Source/ControlRigEditor/Private/Sequencer/ControlRigParameterTrackEditor.cpp
--- a/Engine/Plugins/Animation/ControlRig/Source/ControlRigEditor/Private/Sequencer/ControlRigParameterTrackEditor.cpp
+++ b/Engine/Plugins/Animation/ControlRig/Source/ControlRigEditor/Private/Sequencer/ControlRigParameterTrackEditor.cpp
@@ -256,9 +256,9 @@ public:
 		FFrontendFilter(InCategory),
 		Tag(InTag)
 	{}
-	
-	FString GetName() const override { return Tag.Name.ToString(); }
-	FText GetDisplayName() const override { return FText::FromString(Tag.GetLabel()); }
+
+	FString GetName() const override { return Tag.bMarksSubjectAsInvalid ? TEXT("Exlude ") + Tag.Name.ToString() : Tag.Name.ToString(); }
+	FText GetDisplayName() const override { return FText::FromString(Tag.bMarksSubjectAsInvalid ? TEXT("Exlude ") + Tag.GetLabel() : Tag.GetLabel()); }
 	FText GetToolTipText() const override { return Tag.ToolTip; }
 	FLinearColor GetColor() const override { return Tag.Color; };
 	
@@ -274,7 +274,25 @@ public:
 			{
 				FRigVMVariant AssetVariant;
 				AssetVariantProperty->ImportText_Direct(*VariantStr, &AssetVariant, nullptr, EPropertyPortFlags::PPF_None);
-				if (AssetVariant.Tags.Contains(Tag))
+
+				if (Tag.bMarksSubjectAsInvalid)
+				{
+					if (!AssetVariant.Tags.Contains(Tag))
+					{
+						return true;
+					}
+				}
+				else
+				{
+					if (AssetVariant.Tags.Contains(Tag))
+					{
+						return true;
+					}	
+				}
+			}
+			else
+			{
+				if (Tag.bMarksSubjectAsInvalid)
 				{
 					return true;
 				}
