[Sequencer Control Rig] filter out rig modules by default as they are not supported for sequencer usage

diff --git a/Engine/Plugins/Animation/ControlRig/Source/ControlRigEditor/Private/Sequencer/ControlRigParameterTrackEditor.cpp b/Engine/Plugins/Animation/ControlRig/Source/ControlRigEditor/Private/Sequencer/ControlRigParameterTrackEditor.cpp
--- a/Engine/Plugins/Animation/ControlRig/Source/ControlRigEditor/Private/Sequencer/ControlRigParameterTrackEditor.cpp
+++ b/Engine/Plugins/Animation/ControlRig/Source/ControlRigEditor/Private/Sequencer/ControlRigParameterTrackEditor.cpp
@@ -1675,6 +1675,22 @@ void FControlRigParameterTrackEditor::HandleAddTrackSubMenu(FMenuBuilder& MenuBu
 
 bool FControlRigParameterTrackEditor::IsControlRigAllowed(const FAssetData& AssetData, TArray<UClass*> ExistingRigs, USkeleton* Skeleton)
 {
+	static const FName RigModuleSettingsPropertyName = GET_MEMBER_NAME_CHECKED(UControlRigBlueprint, RigModuleSettings);
+	const FProperty* RigModuleSettingsProperty = CastField<FProperty>(UControlRigBlueprint::StaticClass()->FindPropertyByName(RigModuleSettingsPropertyName));
+	const FString RigModuleSettingsStr = AssetData.GetTagValueRef<FString>(RigModuleSettingsPropertyName);
+	if(!RigModuleSettingsStr.IsEmpty())
+	{
+		FRigModuleSettings RigModuleSettings;
+		RigModuleSettingsProperty->ImportText_Direct(*RigModuleSettingsStr, &RigModuleSettings, nullptr, EPropertyPortFlags::PPF_None);
+		
+		// Currently rig module can only be used in a modular rig, not in sequencer
+		// see UControlRigBlueprint::IsControlRigModule()
+		if (RigModuleSettings.Identifier.IsValid())
+		{
+			return false;
+		}
+	}
+	
 	if (UControlRigBlueprint* LoadedControlRig = Cast<UControlRigBlueprint>(AssetData.FastGetAsset()))
 	{
 		if (ExistingRigs.Contains(LoadedControlRig->GetRigVMBlueprintGeneratedClass()))
