Control Rig: Follow up to 36223655, fix crash during modular rig edit

#rnx

diff --git a/Engine/Plugins/Animation/ControlRig/Source/ControlRig/Private/AnimNode_ControlRigBase.cpp b/Engine/Plugins/Animation/ControlRig/Source/ControlRig/Private/AnimNode_ControlRigBase.cpp
--- a/Engine/Plugins/Animation/ControlRig/Source/ControlRig/Private/AnimNode_ControlRigBase.cpp
+++ b/Engine/Plugins/Animation/ControlRig/Source/ControlRig/Private/AnimNode_ControlRigBase.cpp
@@ -181,9 +181,13 @@ void FAnimNode_ControlRig_PoseAdapter::ConvertToLocalPose()
 {
 	check(ParentPoseIndices.Num() == GlobalPose.Num());
 	LocalPose.SetNum(GlobalPose.Num());
-	for(int32 Index = 0; Index < ParentPoseIndices.Num(); Index++)
+
+	if(Hierarchy)
 	{
-		(void)GetLocalTransform(Index);
+		for(int32 Index = 0; Index < ParentPoseIndices.Num(); Index++)
+		{
+			(void)GetLocalTransform(Index);
+		}
 	}
 }
 void FAnimNode_ControlRig_PoseAdapter::ConvertToGlobalPose()
@@ -191,9 +195,12 @@ void FAnimNode_ControlRig_PoseAdapter::ConvertToGlobalPose()
 	check(ParentPoseIndices.Num() == LocalPose.Num());
 	GlobalPose.SetNum(LocalPose.Num());
 
-	for(int32 Index = 0; Index < ParentPoseIndices.Num(); Index++)
+	if(Hierarchy)
 	{
-		(void)GetGlobalTransform(Index);
+		for(int32 Index = 0; Index < ParentPoseIndices.Num(); Index++)
+		{
+			(void)GetGlobalTransform(Index);
+		}
 	}
 }
 
