Control Rig: Fix modular rig tree / details panel sync

#rnx

diff --git a/Engine/Plugins/Animation/ControlRig/Source/ControlRigEditor/Private/Editor/SModularRigTreeView.cpp b/Engine/Plugins/Animation/ControlRig/Source/ControlRigEditor/Private/Editor/SModularRigTreeView.cpp
--- a/Engine/Plugins/Animation/ControlRig/Source/ControlRigEditor/Private/Editor/SModularRigTreeView.cpp
+++ b/Engine/Plugins/Animation/ControlRig/Source/ControlRigEditor/Private/Editor/SModularRigTreeView.cpp
@@ -1033,6 +1033,7 @@ bool SModularRigTreeView::ReparentElement(const FString InKey, const FString InP
 void SModularRigTreeView::RefreshTreeView(bool bRebuildContent)
 {
 	TMap<FString, bool> ExpansionState;
+	TArray<FString> Selection;
 
 	if(bRebuildContent)
 	{
@@ -1047,6 +1048,8 @@ void SModularRigTreeView::RefreshTreeView(bool bRebuildContent)
 		RootElements.Reset();
 		ElementMap.Reset();
 		ParentMap.Reset();
+
+		Selection = GetSelectedKeys();
 	}
 
 	if(bRebuildContent)
@@ -1099,7 +1102,24 @@ void SModularRigTreeView::RefreshTreeView(bool bRebuildContent)
 
 	RequestTreeRefresh();
 	{
+		TGuardValue<bool> Guard(Delegates.bSuspendSelectionDelegate, true);
 		ClearSelection();
+
+		if(!Selection.IsEmpty())
+		{
+			TArray<TSharedPtr<FModularRigTreeElement>> SelectedElements;
+			for(const FString& SelectedPath : Selection)
+			{
+				if(const TSharedPtr<FModularRigTreeElement> ElementToSelect = FindElement(SelectedPath))
+				{
+					SelectedElements.Add(ElementToSelect);
+				}
+			}
+			if(!SelectedElements.IsEmpty())
+			{
+				SetSelection(SelectedElements);
+			}
+		}
 	}
 }
 
diff --git a/Engine/Plugins/Animation/ControlRig/Source/ControlRigEditor/Private/Editor/SModularRigTreeView.h b/Engine/Plugins/Animation/ControlRig/Source/ControlRigEditor/Private/Editor/SModularRigTreeView.h
--- a/Engine/Plugins/Animation/ControlRig/Source/ControlRigEditor/Private/Editor/SModularRigTreeView.h
+++ b/Engine/Plugins/Animation/ControlRig/Source/ControlRigEditor/Private/Editor/SModularRigTreeView.h
@@ -97,17 +97,19 @@ public:
 
 	void HandleSelectionChanged(TSharedPtr<FModularRigTreeElement> Selection, ESelectInfo::Type SelectInfo)
 	{
-		if(bSelectionIsChanged)
+		if(bSuspendSelectionDelegate)
 		{
 			return;
 		}
-		TGuardValue<bool> Guard(bSelectionIsChanged, true);
+		TGuardValue<bool> Guard(bSuspendSelectionDelegate, true);
 		(void)OnSelectionChanged.ExecuteIfBound(Selection, SelectInfo);
 	}
 
 private:
 
-	bool bSelectionIsChanged = false;
+	bool bSuspendSelectionDelegate = false;
+
+	friend class SModularRigTreeView;
 };
 
 
