Anim Details: avoid useless constraints evaluation

#rnx

diff --git a/Engine/Plugins/Animation/ControlRig/Source/ControlRigEditor/Private/EditMode/AnimDetailsProxy.cpp b/Engine/Plugins/Animation/ControlRig/Source/ControlRigEditor/Private/EditMode/AnimDetailsProxy.cpp
--- a/Engine/Plugins/Animation/ControlRig/Source/ControlRigEditor/Private/EditMode/AnimDetailsProxy.cpp
+++ b/Engine/Plugins/Animation/ControlRig/Source/ControlRigEditor/Private/EditMode/AnimDetailsProxy.cpp
@@ -758,39 +758,44 @@ EPropertyKeyedStatus UAnimDetailControlsKeyedProxy::GetPropertyKeyedStatus(TShar
 #if WITH_EDITOR
 void UAnimDetailControlsKeyedProxy::PostEditUndo()
 {
-	FRigControlModifiedContext Context;
-	Context.SetKey = EControlRigSetKey::Never;
-	UWorld* World = GCurrentLevelEditingViewportClient ? GCurrentLevelEditingViewportClient->GetWorld() : nullptr;
-	const FConstraintsManagerController& Controller = FConstraintsManagerController::Get(World);
-	Controller.EvaluateAllConstraints();
+	if (UWorld* World = GCurrentLevelEditingViewportClient ? GCurrentLevelEditingViewportClient->GetWorld() : nullptr)
+	{
+		const FConstraintsManagerController& Controller = FConstraintsManagerController::Get(World);
+		Controller.EvaluateAllConstraints();
+	}
+
+	static const FRigControlModifiedContext ContextNoKey(EControlRigSetKey::Never);
+	
 	for (const TPair<TWeakObjectPtr<UControlRig>, FControlRigProxyItem>& Items : ControlRigItems)
 	{
-		if (UControlRig* ControlRig = Items.Value.ControlRig.Get())
+		const FControlRigProxyItem& ProxyItem = Items.Value;
+		if (UControlRig* ControlRig = ProxyItem.ControlRig.Get())
 		{
-			for (const FName& CName : Items.Value.ControlElements)
+			for (const FName& CName : ProxyItem.ControlElements)
 			{
-				if (FRigControlElement* ControlElement = Items.Value.GetControlElement(CName))
+				if (FRigControlElement* ControlElement = ProxyItem.GetControlElement(CName))
 				{
-					if (ControlRig->GetHierarchy()->Contains(FRigElementKey(ControlElement->GetKey().Name, ERigElementType::Control)))
-					{
-						ControlRig->SelectControl(ControlElement->GetKey().Name, bSelected);
-						SetControlRigElementValueFromCurrent(ControlRig, ControlElement, Context);
-					}
+					ControlRig->SelectControl(ControlElement->GetKey().Name, bSelected);
+					SetControlRigElementValueFromCurrent(ControlRig, ControlElement, ContextNoKey);
 				}
 			}
 		}
 	}
+	
 	for (TPair <TWeakObjectPtr<UObject>, FSequencerProxyItem>& SItems : SequencerItems)
 	{
 		//we do this backwards so ValueChanged later is set up correctly since that iterates in the other direction
-		if (SItems.Key.IsValid())
+		TWeakObjectPtr<UObject> WeakObject = SItems.Key;
+		if (WeakObject.IsValid())
 		{
+			UObject* Object = WeakObject.Get();
 			for (FBindingAndTrack& Binding : SItems.Value.Bindings)
 			{
-				SetBindingValueFromCurrent(SItems.Key.Get(), Binding.Binding, Context);
+				SetBindingValueFromCurrent(Object, Binding.Binding, ContextNoKey);
 			}
 		}
 	}
+	
 	ValueChanged();
 }
 #endif
@@ -925,7 +930,7 @@ static FEulerTransform GetCurrentValue(UObject* InObject, TSharedPtr<FTrackInsta
 	return EulerTransform;
 }
 
-void UAnimDetailControlsProxyTransform::SetBindingValueFromCurrent(UObject* InObject, TSharedPtr<FTrackInstancePropertyBindings>& Binding, FRigControlModifiedContext& Context, bool bInteractive)
+void UAnimDetailControlsProxyTransform::SetBindingValueFromCurrent(UObject* InObject, TSharedPtr<FTrackInstancePropertyBindings>& Binding, const FRigControlModifiedContext& Context, bool bInteractive)
 {
 	if (InObject && Binding.IsValid())
 	{
@@ -2410,7 +2415,7 @@ void UAnimDetailControlsProxyFloat::GetChannelSelectionState(TWeakPtr<FCurveEdit
 	LocationSelectionCache.XSelected = CachePropertySelection(CurveEditor, this, GET_MEMBER_NAME_CHECKED(FAnimDetailProxyFloat, Float));
 }
 
-void UAnimDetailControlsProxyFloat::SetBindingValueFromCurrent(UObject* InObject, TSharedPtr<FTrackInstancePropertyBindings>& Binding, FRigControlModifiedContext& Context, bool bInteractive)
+void UAnimDetailControlsProxyFloat::SetBindingValueFromCurrent(UObject* InObject, TSharedPtr<FTrackInstancePropertyBindings>& Binding, const FRigControlModifiedContext& Context, bool bInteractive)
 {
 	if (InObject && Binding.IsValid())
 	{
@@ -2612,7 +2617,7 @@ void UAnimDetailControlsProxyBool::GetChannelSelectionState(TWeakPtr<FCurveEdito
 	LocationSelectionCache.XSelected = CachePropertySelection(CurveEditor, this, GET_MEMBER_NAME_CHECKED(FAnimDetailProxyBool, Bool));
 }
 
-void UAnimDetailControlsProxyBool::SetBindingValueFromCurrent(UObject* InObject, TSharedPtr<FTrackInstancePropertyBindings>& Binding, FRigControlModifiedContext& Context, bool bInteractive)
+void UAnimDetailControlsProxyBool::SetBindingValueFromCurrent(UObject* InObject, TSharedPtr<FTrackInstancePropertyBindings>& Binding, const FRigControlModifiedContext& Context, bool bInteractive)
 {
 	if (InObject && Binding.IsValid())
 	{
@@ -2801,7 +2806,7 @@ void UAnimDetailControlsProxyInteger::GetChannelSelectionState(TWeakPtr<FCurveEd
 	LocationSelectionCache.XSelected = CachePropertySelection(CurveEditor, this, GET_MEMBER_NAME_CHECKED(FAnimDetailProxyInteger, Integer));
 }
 
-void UAnimDetailControlsProxyInteger::SetBindingValueFromCurrent(UObject* InObject, TSharedPtr<FTrackInstancePropertyBindings>& Binding, FRigControlModifiedContext& Context, bool bInteractive)
+void UAnimDetailControlsProxyInteger::SetBindingValueFromCurrent(UObject* InObject, TSharedPtr<FTrackInstancePropertyBindings>& Binding, const FRigControlModifiedContext& Context, bool bInteractive)
 {
 	if (InObject && Binding.IsValid())
 	{
@@ -3060,10 +3065,16 @@ const TArray<UControlRigControlsProxy*> UControlRigDetailPanelControlProxies::Ge
 
 void UControlRigDetailPanelControlProxies::ValuesChanged()
 {
-	//need to do all proxies
-	UWorld* World = GCurrentLevelEditingViewportClient ? GCurrentLevelEditingViewportClient->GetWorld() : nullptr;
-	const FConstraintsManagerController& Controller = FConstraintsManagerController::Get(World);
-	Controller.EvaluateAllConstraints();
+	if (SelectedControlRigProxies.IsEmpty() && SelectedSequencerProxies.IsEmpty())
+	{
+		return;
+	}
+	
+	if (UWorld* World = GCurrentLevelEditingViewportClient ? GCurrentLevelEditingViewportClient->GetWorld() : nullptr)
+	{
+		const FConstraintsManagerController& Controller = FConstraintsManagerController::Get(World);
+		Controller.EvaluateAllConstraints();
+	}
 
 	for (UControlRigControlsProxy* P1 : SelectedControlRigProxies)
 	{
diff --git a/Engine/Plugins/Animation/ControlRig/Source/ControlRigEditor/Private/EditMode/AnimDetailsProxy.h b/Engine/Plugins/Animation/ControlRig/Source/ControlRigEditor/Private/EditMode/AnimDetailsProxy.h
--- a/Engine/Plugins/Animation/ControlRig/Source/ControlRigEditor/Private/EditMode/AnimDetailsProxy.h
+++ b/Engine/Plugins/Animation/ControlRig/Source/ControlRigEditor/Private/EditMode/AnimDetailsProxy.h
@@ -258,7 +258,7 @@ public:
 	virtual void GetChannelSelectionState(TWeakPtr<FCurveEditor>& CurveEditor, FAnimDetailVectorSelection& OutLocationSelection, FAnimDetailVectorSelection& OutRotationSelection,
 		FAnimDetailVectorSelection& OutScaleSelection) override;
 	virtual void UpdatePropertyNames(IDetailLayoutBuilder& DetailBuilder);
-	virtual void SetBindingValueFromCurrent(UObject* InObject, TSharedPtr<FTrackInstancePropertyBindings>& Binding, FRigControlModifiedContext& Context, bool bInteractive = false) override;
+	virtual void SetBindingValueFromCurrent(UObject* InObject, TSharedPtr<FTrackInstancePropertyBindings>& Binding, const FRigControlModifiedContext& Context, bool bInteractive = false) override;
 
 	FAnimDetailPropertyState State;
 
@@ -287,7 +287,7 @@ public:
 	virtual void GetChannelSelectionState(TWeakPtr<FCurveEditor>& CurveEditor, FAnimDetailVectorSelection& OutLocationSelection, FAnimDetailVectorSelection& OutRotationSelection,
 		FAnimDetailVectorSelection& OutScaleSelection) override;
 	virtual void UpdatePropertyNames(IDetailLayoutBuilder& DetailBuilder);
-	virtual void SetBindingValueFromCurrent(UObject* InObject, TSharedPtr<FTrackInstancePropertyBindings>& Binding, FRigControlModifiedContext& Context, bool bInteractive = false) override;
+	virtual void SetBindingValueFromCurrent(UObject* InObject, TSharedPtr<FTrackInstancePropertyBindings>& Binding, const FRigControlModifiedContext& Context, bool bInteractive = false) override;
 
 	
 	FAnimDetailPropertyState State;
@@ -317,7 +317,7 @@ public:
 	virtual void GetChannelSelectionState(TWeakPtr<FCurveEditor>& CurveEditor, FAnimDetailVectorSelection& OutLocationSelection, FAnimDetailVectorSelection& OutRotationSelection,
 		FAnimDetailVectorSelection& OutScaleSelection) override;
 	virtual void UpdatePropertyNames(IDetailLayoutBuilder& DetailBuilder);
-	virtual void SetBindingValueFromCurrent(UObject* InObject, TSharedPtr<FTrackInstancePropertyBindings>& Binding, FRigControlModifiedContext& Context, bool bInteractive = false) override;
+	virtual void SetBindingValueFromCurrent(UObject* InObject, TSharedPtr<FTrackInstancePropertyBindings>& Binding, const FRigControlModifiedContext& Context, bool bInteractive = false) override;
 
 
 	FAnimDetailPropertyState State;
@@ -374,7 +374,7 @@ public:
 	virtual TMap<FName, int32> GetPropertyNames() const override;
 	virtual bool IsMultiple(const FName& InPropertyName) const override;
 	virtual void SetControlRigElementValueFromCurrent(UControlRig* ControlRig, FRigControlElement* ControlElement, const FRigControlModifiedContext& Context) override;
-	virtual void SetBindingValueFromCurrent(UObject* InObject, TSharedPtr<FTrackInstancePropertyBindings>& Binding, FRigControlModifiedContext& Context, bool bInteractive = false) override;
+	virtual void SetBindingValueFromCurrent(UObject* InObject, TSharedPtr<FTrackInstancePropertyBindings>& Binding, const FRigControlModifiedContext& Context, bool bInteractive = false) override;
 	virtual void GetChannelSelectionState(TWeakPtr<FCurveEditor>& CurveEditor, FAnimDetailVectorSelection& OutLocationSelection, FAnimDetailVectorSelection& OutRotationSelection,
 		FAnimDetailVectorSelection& OutScaleSelection) override;
 
diff --git a/Engine/Plugins/Animation/ControlRig/Source/ControlRigEditor/Private/EditMode/ControlRigControlsProxy.h b/Engine/Plugins/Animation/ControlRig/Source/ControlRigEditor/Private/EditMode/ControlRigControlsProxy.h
--- a/Engine/Plugins/Animation/ControlRig/Source/ControlRigEditor/Private/EditMode/ControlRigControlsProxy.h
+++ b/Engine/Plugins/Animation/ControlRig/Source/ControlRigEditor/Private/EditMode/ControlRigControlsProxy.h
@@ -88,7 +88,7 @@ public:
 	virtual TMap<FName, int32> GetPropertyNames() const { TMap<FName, int32> Empty; return Empty; }
 	virtual bool IsMultiple(const FName& InPropertyName) const { return false; }
 	virtual void SetControlRigElementValueFromCurrent(UControlRig* ControlRig, FRigControlElement* ControlElement, const FRigControlModifiedContext& Context) {};
-	virtual void SetBindingValueFromCurrent(UObject* InObject, TSharedPtr<FTrackInstancePropertyBindings>& Binding, FRigControlModifiedContext& Context, bool bInteractive = false) {};
+	virtual void SetBindingValueFromCurrent(UObject* InObject, TSharedPtr<FTrackInstancePropertyBindings>& Binding, const FRigControlModifiedContext& Context, bool bInteractive = false) {};
 	virtual void GetChannelSelectionState(TWeakPtr<FCurveEditor>& CurveEditor, FAnimDetailVectorSelection& OutLocationSelection, FAnimDetailVectorSelection& OutRotationSelection,
 		FAnimDetailVectorSelection& OutScaleSelection) {};
 	virtual bool PropertyIsOnProxy(FProperty* Property, FProperty* MemberProperty){return false;}
