FModeToolkit : Add a new ShutdownUI() virtual method to provide a clear place for UI cleanup ( e.g. clearing status bar messages ).

Existing bindings to the ModeUILayer.Pin()->ToolkitHostShutdownUI() should be removed in favor of this override.  However, they will continue to work unless action is taken since we make our internal binding before the ModeUILayer becomes valid to the toolkit.

diff --git a/Engine/Plugins/Animation/ControlRig/Source/ControlRigEditor/Private/EditMode/ControlRigEditModeToolkit.cpp b/Engine/Plugins/Animation/ControlRig/Source/ControlRigEditor/Private/EditMode/ControlRigEditModeToolkit.cpp
--- a/Engine/Plugins/Animation/ControlRig/Source/ControlRigEditor/Private/EditMode/ControlRigEditModeToolkit.cpp
+++ b/Engine/Plugins/Animation/ControlRig/Source/ControlRigEditor/Private/EditMode/ControlRigEditModeToolkit.cpp
@@ -433,12 +433,6 @@ void FControlRigEditModeToolkit::RequestModeUITabs()
 			.SetGroup(MenuGroup)
 			.SetIcon(FSlateIcon(TEXT("ControlRigEditorStyle"), TEXT("ControlRig.AnimLayers")));
 		ModeUILayerPtr->GetTabManager()->RegisterDefaultTabWindowSize(AnimLayerTabName, FVector2D(425, 200));
-
-
-
-		ModeUILayer.Pin()->ToolkitHostShutdownUI().BindSP(this, &FControlRigEditModeToolkit::UnregisterAndRemoveFloatingTabs);
-
-
 	}
 };
 
@@ -475,6 +469,13 @@ void FControlRigEditModeToolkit::InvokeUI()
 	}	
 }
 
+void FControlRigEditModeToolkit::ShutdownUI()
+{
+	FModeToolkit::ShutdownUI();
+
+	UnregisterAndRemoveFloatingTabs();
+}
+
 void FControlRigEditModeToolkit::UnregisterAndRemoveFloatingTabs()
 {
 	if (FSlateApplication::IsInitialized())
diff --git a/Engine/Plugins/Animation/ControlRig/Source/ControlRigEditor/Private/EditMode/ControlRigEditModeToolkit.h b/Engine/Plugins/Animation/ControlRig/Source/ControlRigEditor/Private/EditMode/ControlRigEditModeToolkit.h
--- a/Engine/Plugins/Animation/ControlRig/Source/ControlRigEditor/Private/EditMode/ControlRigEditModeToolkit.h
+++ b/Engine/Plugins/Animation/ControlRig/Source/ControlRigEditor/Private/EditMode/ControlRigEditModeToolkit.h
@@ -89,6 +89,7 @@ protected:
 	/* FModeToolkit Interface */
 	virtual void RequestModeUITabs() override;
 	virtual void InvokeUI() override;
+	virtual void ShutdownUI() override;
 
 	//this also saves the layout
 	void UnregisterAndRemoveFloatingTabs();
diff --git a/Engine/Plugins/Editor/ModelingToolsEditorMode/Source/ModelingToolsEditorMode/Private/ModelingToolsEditorModeToolkit.cpp b/Engine/Plugins/Editor/ModelingToolsEditorMode/Source/ModelingToolsEditorMode/Private/ModelingToolsEditorModeToolkit.cpp
--- a/Engine/Plugins/Editor/ModelingToolsEditorMode/Source/ModelingToolsEditorMode/Private/ModelingToolsEditorModeToolkit.cpp
+++ b/Engine/Plugins/Editor/ModelingToolsEditorMode/Source/ModelingToolsEditorMode/Private/ModelingToolsEditorModeToolkit.cpp
@@ -2218,7 +2218,12 @@ void FModelingToolsEditorModeToolkit::InvokeUI()
 	}
 }
 
+void FModelingToolsEditorModeToolkit::ShutdownUI()
+{
+	FModeToolkit::ShutdownUI();
 
+	ClearNotification();
+}
 
 void FModelingToolsEditorModeToolkit::ForceToolPaletteRebuild()
 {
diff --git a/Engine/Plugins/Editor/ModelingToolsEditorMode/Source/ModelingToolsEditorMode/Public/ModelingToolsEditorModeToolkit.h b/Engine/Plugins/Editor/ModelingToolsEditorMode/Source/ModelingToolsEditorMode/Public/ModelingToolsEditorModeToolkit.h
--- a/Engine/Plugins/Editor/ModelingToolsEditorMode/Source/ModelingToolsEditorMode/Public/ModelingToolsEditorModeToolkit.h
+++ b/Engine/Plugins/Editor/ModelingToolsEditorMode/Source/ModelingToolsEditorMode/Public/ModelingToolsEditorModeToolkit.h
@@ -84,6 +84,7 @@ public:
 	void OnActiveViewportChanged(TSharedPtr<IAssetViewport>, TSharedPtr<IAssetViewport> );
 
 	virtual void InvokeUI() override;
+	virtual void ShutdownUI() override;
 	virtual void ExtendSecondaryModeToolbar(UToolMenu *InModeToolbarMenu) override;
 
 	virtual bool RequestAcceptCancelButtonOverride(IToolHostCustomizationAPI::FAcceptCancelButtonOverrideParams& Params);
diff --git a/Engine/Source/Editor/UnrealEd/Private/Toolkits/BaseToolkit.cpp b/Engine/Source/Editor/UnrealEd/Private/Toolkits/BaseToolkit.cpp
--- a/Engine/Source/Editor/UnrealEd/Private/Toolkits/BaseToolkit.cpp
+++ b/Engine/Source/Editor/UnrealEd/Private/Toolkits/BaseToolkit.cpp
@@ -234,9 +234,9 @@ void FModeToolkit::SetModeUILayer(const TSharedPtr<FAssetEditorModeUILayer> InLa
 	// responsible for a generic set of "Mode Commands"
 	ModeUILayer.Pin()->GetModeCommands()->Append(GetToolkitCommands());
 	ModeUILayer.Pin()->ToolkitHostReadyForUI().BindSP(this, &FModeToolkit::InvokeUI);
+	ModeUILayer.Pin()->ToolkitHostShutdownUI().BindSP(this, &FModeToolkit::ShutdownUI);
 
 	ModeUILayer.Pin()->RegisterSecondaryModeToolbarExtension().BindSP(this, &FModeToolkit::ExtendSecondaryModeToolbar);
-
 }
 
 void FModeToolkit::RegisterTabSpawners(const TSharedRef<FTabManager>& TabManager)
diff --git a/Engine/Source/Editor/UnrealEd/Public/Toolkits/BaseToolkit.h b/Engine/Source/Editor/UnrealEd/Public/Toolkits/BaseToolkit.h
--- a/Engine/Source/Editor/UnrealEd/Public/Toolkits/BaseToolkit.h
+++ b/Engine/Source/Editor/UnrealEd/Public/Toolkits/BaseToolkit.h
@@ -163,6 +163,14 @@ public:
 	UNREALED_API void SetModeSettingsObject(UObject* InSettingsObject);
 	UNREALED_API virtual void InvokeUI();
 
+	/**
+	* Override this function to perform any UI cleanup before the ModeUILayer is no longer valid.
+	* 
+	* Note: If you had an existing delegate binding to ModeUILayer.Pin()->ToolkitHostShutdownUI() you
+	* will need to remove that in order for this method to be invoked.  Overriding this method is preferred.
+	*/
+	virtual void ShutdownUI() {}
+
 	/**
 	 * Override this function to extend the secondary mode toolbar (that appears below the main toolbar) for your mode by
 	 * filling in the UToolMenu that is passed in
