[Backout] - CL36225939
Original CL Desc
-----------------------------------------------------------------
Control Rig: Follow up to 36223655, fix crash during modular rig edit

#rnx

diff --git a/Engine/Plugins/Animation/ControlRig/Source/ControlRig/Private/AnimNode_ControlRigBase.cpp b/Engine/Plugins/Animation/ControlRig/Source/ControlRig/Private/AnimNode_ControlRigBase.cpp
--- a/Engine/Plugins/Animation/ControlRig/Source/ControlRig/Private/AnimNode_ControlRigBase.cpp
+++ b/Engine/Plugins/Animation/ControlRig/Source/ControlRig/Private/AnimNode_ControlRigBase.cpp
@@ -181,13 +181,9 @@ void FAnimNode_ControlRig_PoseAdapter::ConvertToLocalPose()
 {
 	check(ParentPoseIndices.Num() == GlobalPose.Num());
 	LocalPose.SetNum(GlobalPose.Num());
-
-	if(Hierarchy)
+	for(int32 Index = 0; Index < ParentPoseIndices.Num(); Index++)
 	{
-		for(int32 Index = 0; Index < ParentPoseIndices.Num(); Index++)
-		{
-			(void)GetLocalTransform(Index);
-		}
+		(void)GetLocalTransform(Index);
 	}
 }
 void FAnimNode_ControlRig_PoseAdapter::ConvertToGlobalPose()
@@ -195,12 +191,9 @@ void FAnimNode_ControlRig_PoseAdapter::ConvertToGlobalPose()
 	check(ParentPoseIndices.Num() == LocalPose.Num());
 	GlobalPose.SetNum(LocalPose.Num());
 
-	if(Hierarchy)
+	for(int32 Index = 0; Index < ParentPoseIndices.Num(); Index++)
 	{
-		for(int32 Index = 0; Index < ParentPoseIndices.Num(); Index++)
-		{
-			(void)GetGlobalTransform(Index);
-		}
+		(void)GetGlobalTransform(Index);
 	}
 }
 
