Anim: Fix CIS

#rnx

diff --git a/Engine/Plugins/Animation/ControlRig/Source/ControlRigEditor/Private/Sequencer/AnimLayers/AnimLayers.h b/Engine/Plugins/Animation/ControlRig/Source/ControlRigEditor/Private/Sequencer/AnimLayers/AnimLayers.h
--- a/Engine/Plugins/Animation/ControlRig/Source/ControlRigEditor/Private/Sequencer/AnimLayers/AnimLayers.h
+++ b/Engine/Plugins/Animation/ControlRig/Source/ControlRigEditor/Private/Sequencer/AnimLayers/AnimLayers.h
@@ -104,10 +104,10 @@ struct FAnimLayerState
 	FText AnimLayerTypeToText()const;
 
 	UPROPERTY(BlueprintReadWrite, EditAnywhere, Category = "Data")
-	mutable ECheckBoxState bKeyed;
+	mutable ECheckBoxState bKeyed = ECheckBoxState::Unchecked;
 
 	UPROPERTY(BlueprintReadWrite, EditAnywhere, Category = "Data")
-	mutable ECheckBoxState bSelected;
+	mutable ECheckBoxState bSelected = ECheckBoxState::Unchecked;
 
 	UPROPERTY(BlueprintReadWrite, EditAnywhere, Category = "Data")
 	mutable bool bActive;
