Control Rig Edit Mode: always use control rig transforms instead of relying on the control shapes
    - this basically represents a rollback of 34128760, 34201378, 34375090 and 34809364 which were all related to numerical errors (UE-216574)

#rnx

diff --git a/Engine/Plugins/Animation/ControlRig/Source/ControlRigEditor/Private/EditMode/ControlRigEditMode.cpp b/Engine/Plugins/Animation/ControlRig/Source/ControlRigEditor/Private/EditMode/ControlRigEditMode.cpp
--- a/Engine/Plugins/Animation/ControlRig/Source/ControlRigEditor/Private/EditMode/ControlRigEditMode.cpp
+++ b/Engine/Plugins/Animation/ControlRig/Source/ControlRigEditor/Private/EditMode/ControlRigEditMode.cpp
@@ -4862,215 +4862,147 @@ void FControlRigEditMode::MoveControlShape(AControlRigShapeActor* ShapeActor, co
 	const bool bRotation, FRotator& InRot, const bool bScale, FVector& InScale, const FTransform& ToWorldTransform,
 	bool bUseLocal, bool bCalcLocal, FTransform& InOutLocal)
 {
+	UControlRig* ControlRig = ShapeActor ? ShapeActor->ControlRig.Get() : nullptr;
+	if (!ensure(ControlRig))
+	{
+		return;
+	}
+
 	bool bTransformChanged = false;
 
 	// In case for some reason the shape actor was detached, make sure to attach it again to the scene component
 	if (!ShapeActor->GetAttachParentActor())
 	{
-		if (UControlRig* ControlRig = ShapeActor->ControlRig.Get())
+		if (USceneComponent* SceneComponent = GetHostingSceneComponent(ControlRig))
 		{
-			if (USceneComponent* SceneComponent = GetHostingSceneComponent(ControlRig))
+			if (AActor* OwnerActor = SceneComponent->GetOwner())
 			{
-				if (AActor* OwnerActor = SceneComponent->GetOwner())
-				{
-					ShapeActor->AttachToActor(OwnerActor, FAttachmentTransformRules::KeepWorldTransform);
-				}
+				ShapeActor->AttachToActor(OwnerActor, FAttachmentTransformRules::KeepWorldTransform);
 			}
 		}
 	}
-
-	auto RotatorToStr = [](const FRotator& InRotator)
-	{
-		FRigPreferredEulerAngles EulerAngles;
-		EulerAngles.SetRotator(InRotator);
-		return EulerAngles.Current.ToString();
-	};
-
-	auto UpdatePreferredEulerAngles = [ShapeActor, bRotation, InRot, RotatorToStr](UControlRig* InControlRig)
-	{
-		if(bRotation)
-		{
-			if(FRigControlElement* ControlElement = InControlRig->GetHierarchy()->Find<FRigControlElement>(ShapeActor->GetElementKey()))
-			{
-				//if(ControlElement->Settings.bUsePreferredRotationOrder) always set rotation order since
-				//sequencer depends upon it
-				{
-					FRotator Rot = ControlElement->PreferredEulerAngles.GetRotator();
-
-					// Split the current rotation between winding and remainder
-					FRotator CurrentRotWind, CurrentRotRem;
-					Rot.GetWindingAndRemainder(CurrentRotWind, CurrentRotRem);
-
-					// Apply the delta to the current remainder, and normalize
-					FRotator NewRotRem = CurrentRotRem + InRot;
-					NewRotRem.Normalize();
-
-					// Add the current winding to the new remainder
-					const FRotator Final = CurrentRotWind + NewRotRem;
-					
-					InControlRig->GetHierarchy()->SetControlPreferredRotator(ControlElement, Final, false, false);
-				}
-			}
-		}
-	};
 	
 	//first case is where we do all controls by the local diff.
 	if (bUseLocal)
 	{
-		if (UControlRig* ControlRig = ShapeActor->ControlRig.Get())
-		{
-			FRigControlModifiedContext Context;
-			Context.EventName = FRigUnit_BeginExecution::EventName;
-			FTransform CurrentLocalTransform = ControlRig->GetControlLocalTransform(ShapeActor->ControlName);
-			if (bRotation)
-			{
-
-				FQuat CurrentRotation = CurrentLocalTransform.GetRotation();
-				CurrentRotation = (CurrentRotation * InOutLocal.GetRotation());
-				CurrentLocalTransform.SetRotation(CurrentRotation);
-				bTransformChanged = true;
-			}
-
-			if (bTranslation)
-			{
-				FVector CurrentLocation = CurrentLocalTransform.GetLocation();
-				CurrentLocation = CurrentLocation + InOutLocal.GetLocation();
-				CurrentLocalTransform.SetLocation(CurrentLocation);
-				bTransformChanged = true;
-			}
-
-			if (bTransformChanged)
-			{
-				ControlRig->InteractionType = InteractionType;
-				ControlRig->ElementsBeingInteracted.AddUnique(ShapeActor->GetElementKey());
-				
-				ControlRig->SetControlLocalTransform(ShapeActor->ControlName, CurrentLocalTransform,true, FRigControlModifiedContext(), true, /*fix eulers*/ true);
-				//UpdatePreferredEulerAngles(ControlRig);
-
-				FTransform CurrentTransform  = ControlRig->GetControlGlobalTransform(ShapeActor->ControlName);			// assumes it's attached to actor
-				CurrentTransform = ToWorldTransform * CurrentTransform;
-
-				// make the transform relative to the offset transform again.
-				// first we'll make it relative to the offset used at the time of starting the drag
-				// and then we'll make it absolute again based on the current offset. these two can be
-				// different if we are interacting on a control on an animated character
-				CurrentTransform = CurrentTransform.GetRelativeTransform(ShapeActor->OffsetTransform);
-				CurrentTransform = CurrentTransform * ControlRig->GetHierarchy()->GetGlobalControlOffsetTransform(ShapeActor->GetElementKey(), false);
-				
-				ShapeActor->SetGlobalTransform(CurrentTransform);
-
-				ControlRig->Evaluate_AnyThread();
-			}
-		}
-	}
-	if(!bTransformChanged) //not local or doing scale.
-	{
-		// Get the global transform from shape actor to avoid drifting
-		FTransform CurrentTransform;
-		if (ShapeActor->StaticMeshComponent && ShapeActor->StaticMeshComponent->GetStaticMesh())
-		{
-			if (ShapeActor->GetAttachParentActor())
-			{
-				const FTransform& ParentTransform = ShapeActor->GetAttachParentActor()->GetTransform();
-				CurrentTransform = ShapeActor->GetGlobalTransform() * ParentTransform;
-			}
-			else
-			{
-				CurrentTransform = ShapeActor->GetGlobalTransform();
-			}
-		}
-		else
-		{
-			// If the static mesh is not valid, we cannot rely on the shape's transform.
-			// This happens for FKControlRigs (and other control types)
-			// We will need to rely on the information we have in the rig hierarchy
-			CurrentTransform = GetControlShapeTransform(ShapeActor) * ToWorldTransform;
-		}
+		FTransform CurrentLocalTransform = ControlRig->GetControlLocalTransform(ShapeActor->ControlName);
 
 		if (bRotation)
 		{
-			FQuat CurrentRotation = CurrentTransform.GetRotation();
-			CurrentRotation = (InRot.Quaternion() * CurrentRotation);
-			CurrentTransform.SetRotation(CurrentRotation);
+			FQuat CurrentRotation = CurrentLocalTransform.GetRotation();
+			CurrentRotation = (CurrentRotation * InOutLocal.GetRotation());
+			CurrentLocalTransform.SetRotation(CurrentRotation);
 			bTransformChanged = true;
 		}
 
 		if (bTranslation)
 		{
-			FVector CurrentLocation = CurrentTransform.GetLocation();
-			CurrentLocation = CurrentLocation + InDrag;
-			CurrentTransform.SetLocation(CurrentLocation);
+			FVector CurrentLocation = CurrentLocalTransform.GetLocation();
+			CurrentLocation = CurrentLocation + InOutLocal.GetLocation();
+			CurrentLocalTransform.SetLocation(CurrentLocation);
 			bTransformChanged = true;
 		}
 
-		if (bScale)
+		if (bTransformChanged)
 		{
-			FVector CurrentScale = CurrentTransform.GetScale3D();
-			CurrentScale = CurrentScale + InScale;
-			CurrentTransform.SetScale3D(CurrentScale);
-			bTransformChanged = true;
+			ControlRig->InteractionType = InteractionType;
+			ControlRig->ElementsBeingInteracted.AddUnique(ShapeActor->GetElementKey());
+			
+			ControlRig->SetControlLocalTransform(ShapeActor->ControlName, CurrentLocalTransform,true, FRigControlModifiedContext(), true, /*fix eulers*/ true);
+
+			FTransform CurrentTransform  = ControlRig->GetControlGlobalTransform(ShapeActor->ControlName);			// assumes it's attached to actor
+			CurrentTransform = ToWorldTransform * CurrentTransform;
+
+			// make the transform relative to the offset transform again.
+			// first we'll make it relative to the offset used at the time of starting the drag
+			// and then we'll make it absolute again based on the current offset. these two can be
+			// different if we are interacting on a control on an animated character
+			CurrentTransform = CurrentTransform.GetRelativeTransform(ShapeActor->OffsetTransform);
+			CurrentTransform = CurrentTransform * ControlRig->GetHierarchy()->GetGlobalControlOffsetTransform(ShapeActor->GetElementKey(), false);
+			
+			ShapeActor->SetGlobalTransform(CurrentTransform);
+
+			ControlRig->Evaluate_AnyThread();
+
+			return;
 		}
+	}
 
-		if (bTransformChanged)
-		{
-			if (UControlRig* ControlRig = ShapeActor->ControlRig.Get())
-			{
-				ControlRig->InteractionType = InteractionType;
-				ControlRig->ElementsBeingInteracted.AddUnique(ShapeActor->GetElementKey());
 
-				FTransform NewTransform = CurrentTransform.GetRelativeTransform(ToWorldTransform);
-				FRigControlModifiedContext Context;
-				Context.EventName = FRigUnit_BeginExecution::EventName;
-				Context.bConstraintUpdate = true;
-				if (bCalcLocal)
-				{
-					InOutLocal = ControlRig->GetControlLocalTransform(ShapeActor->ControlName);
-				}
+	//not local or doing scale.
+	FTransform CurrentTransform = GetControlShapeTransform(ShapeActor) * ToWorldTransform;
 
-				bool bPrintPythonCommands = false;
-				if (UWorld* World = ControlRig->GetWorld())
-				{
-					bPrintPythonCommands = World->IsPreviewWorld();
-				}
+	if (bRotation)
+	{
+		FQuat CurrentRotation = CurrentTransform.GetRotation();
+		CurrentRotation = (InRot.Quaternion() * CurrentRotation);
+		CurrentTransform.SetRotation(CurrentRotation);
+		bTransformChanged = true;
+	}
 
-				bool bIsTransientControl = false;
-				if(const FRigControlElement* ControlElement = ControlRig->FindControl(ShapeActor->ControlName))
-				{
-					bIsTransientControl = ControlElement->Settings.bIsTransientControl;
-				}
+	if (bTranslation)
+	{
+		FVector CurrentLocation = CurrentTransform.GetLocation();
+		CurrentLocation = CurrentLocation + InDrag;
+		CurrentTransform.SetLocation(CurrentLocation);
+		bTransformChanged = true;
+	}
 
-				// if we are operating on a PIE instance which is playing we need to reapply the input pose
-				// since the hierarchy will also have been brought into the solved pose. by reapplying the
-				// input pose we avoid double transformation / double forward solve results.
-				if(bIsTransientControl && ControlRig->GetWorld()->IsPlayInEditor() && !ControlRig->GetWorld()->IsPaused())
-				{
-					ControlRig->GetHierarchy()->SetPose(ControlRig->InputPoseOnDebuggedRig);
-				}
-				
-				ControlRig->Evaluate_AnyThread();
-				//fix flips and do rotation orders only if not additive
-				const bool bFixEulerFlips = ControlRig->IsAdditive() == false ? true : false;
-				SetControlShapeTransform(ShapeActor, NewTransform, ToWorldTransform, Context, bPrintPythonCommands, bFixEulerFlips);
-				//UpdatePreferredEulerAngles(ControlRig);
-				NotifyDrivenControls(ControlRig, ShapeActor->GetElementKey(),Context);
-				if(const FRigControlElement* ControlElement = ControlRig->FindControl(ShapeActor->ControlName))
-				{
-					if(!bIsTransientControl)
-					{
-						ControlRig->Evaluate_AnyThread();
-					}
-				}
-				
-				// Don't set the global transform to the shape actor to avoid drifting
-				//ShapeActor->SetGlobalTransform(CurrentTransform);
+	if (bScale)
+	{
+		FVector CurrentScale = CurrentTransform.GetScale3D();
+		CurrentScale = CurrentScale + InScale;
+		CurrentTransform.SetScale3D(CurrentScale);
+		bTransformChanged = true;
+	}
 
-				if (bCalcLocal)
-				{
-					FTransform NewLocal = ControlRig->GetControlLocalTransform(ShapeActor->ControlName);
-					InOutLocal = NewLocal.GetRelativeTransform(InOutLocal);
-				}
+	if (bTransformChanged)
+	{
+		ControlRig->InteractionType = InteractionType;
+		ControlRig->ElementsBeingInteracted.AddUnique(ShapeActor->GetElementKey());
 
-			}
+		FTransform NewTransform = CurrentTransform.GetRelativeTransform(ToWorldTransform);
+		
+		FRigControlModifiedContext Context;
+		Context.EventName = FRigUnit_BeginExecution::EventName;
+		Context.bConstraintUpdate = true;
+
+		if (bCalcLocal)
+		{
+			InOutLocal = ControlRig->GetControlLocalTransform(ShapeActor->ControlName);
+		}
+
+		const UWorld* World = ControlRig->GetWorld();
+		const bool bPrintPythonCommands = World ? World->IsPreviewWorld() : false;
+
+		const FRigControlElement* ControlElement = ControlRig->FindControl(ShapeActor->ControlName);
+		const bool bIsTransientControl = ControlElement ? ControlElement->Settings.bIsTransientControl : false;
+
+		// if we are operating on a PIE instance which is playing we need to reapply the input pose
+		// since the hierarchy will also have been brought into the solved pose. by reapplying the
+		// input pose we avoid double transformation / double forward solve results.
+		if(bIsTransientControl && World && World->IsPlayInEditor() && !World->IsPaused())
+		{
+			ControlRig->GetHierarchy()->SetPose(ControlRig->InputPoseOnDebuggedRig);
+		}
+		
+		ControlRig->Evaluate_AnyThread();
+		//fix flips and do rotation orders only if not additive
+		const bool bFixEulerFlips = !ControlRig->IsAdditive();
+		SetControlShapeTransform(ShapeActor, NewTransform, ToWorldTransform, Context, bPrintPythonCommands, bFixEulerFlips);
+		NotifyDrivenControls(ControlRig, ShapeActor->GetElementKey(),Context);
+
+		if (ControlElement && !bIsTransientControl)
+		{
+			ControlRig->Evaluate_AnyThread();
+		}
+		
+		ShapeActor->SetGlobalTransform(CurrentTransform);
+
+		if (bCalcLocal)
+		{
+			FTransform NewLocal = ControlRig->GetControlLocalTransform(ShapeActor->ControlName);
+			InOutLocal = NewLocal.GetRelativeTransform(InOutLocal);
 		}
 	}
 }
