Anim Node Control Rig: improve PropagateInputProperties using cached update functions

#rnx

diff --git a/Engine/Plugins/Animation/ControlRig/Source/ControlRig/Private/AnimNode_ControlRig.cpp b/Engine/Plugins/Animation/ControlRig/Source/ControlRig/Private/AnimNode_ControlRig.cpp
--- a/Engine/Plugins/Animation/ControlRig/Source/ControlRig/Private/AnimNode_ControlRig.cpp
+++ b/Engine/Plugins/Animation/ControlRig/Source/ControlRig/Private/AnimNode_ControlRig.cpp
@@ -9,6 +9,7 @@
 #include "Animation/NodeMappingContainer.h"
 #include "AnimationRuntime.h"
 #include "ControlRigObjectBinding.h"
+#include "Algo/ForEach.h"
 #include "Animation/AnimCurveUtils.h"
 #include "Animation/AnimStats.h"
 #include "Units/Execution/RigUnit_PrepareForExecution.h"
@@ -78,6 +79,7 @@ void FAnimNode_ControlRig::GatherDebugData(FNodeDebugData& DebugData)
 void FAnimNode_ControlRig::Update_AnyThread(const FAnimationUpdateContext& Context)
 {
 	DECLARE_SCOPE_HIERARCHICAL_COUNTER_FUNC()
+	QUICK_SCOPE_CYCLE_COUNTER(STAT_AnimNode_ControlRig_Update_AnyThread);
 
 	if (IsLODEnabled(Context.AnimInstanceProxy))
 	{
@@ -589,267 +591,602 @@ FName FAnimNode_ControlRig::GetIOMapping(bool bInput, const FName& SourcePropert
 	return NAME_None;
 }
 
+namespace AnimNodeLocals
+{
+
+const TCHAR* UseFunctionsName = TEXT("ControlRig.AnimNode.UseFunctions");
+
+static bool bUseFunctions = true;
+static FAutoConsoleVariableRef CVarUseCachedFunctions(
+   UseFunctionsName,
+   bUseFunctions,
+   TEXT("Cache and use propagate functions.")
+   );
+
+template<typename ValueType, typename HierarchyValueType = ValueType>
+void SetControlValue(FRigControlElement* InControlElement, FProperty* InSourceProperty, URigHierarchy* InTargetHierarchy, const UObject* InSourceInstance)
+{
+	const ValueType* ValuePtr = InSourceProperty->ContainerPtrToValuePtr<ValueType>(InSourceInstance);
+	const FRigControlValue Value = FRigControlValue::Make<HierarchyValueType>(*ValuePtr);
+	InTargetHierarchy->SetControlValue(InControlElement, Value, ERigControlValueType::Current);
+}
+
+template<typename ValueType, typename VariableValueType = ValueType>
+void SetVariableValue(FRigVMExternalVariable& InVariable, FProperty* InSourceProperty, const UObject* InSourceInstance)
+{
+	const ValueType* ValuePtr = InSourceProperty->ContainerPtrToValuePtr<ValueType>(InSourceInstance);
+	InVariable.SetValue<VariableValueType>(*ValuePtr);
+}
+
+}
+
+void FAnimNode_ControlRig::AddUpdateFunction(PropertyUpdateFunction&& InFunction)
+{
+	UpdateFunctions.Add( MoveTemp(InFunction) );
+}
+
+void FAnimNode_ControlRig::AddControlFunction(FRigControlElement* InControlElement, FProperty* InSourceProperty, URigHierarchy* InTargetHierarchy)
+{
+	using namespace AnimNodeLocals;
+	
+	switch(InControlElement->Settings.ControlType)
+	{
+		case ERigControlType::Bool:
+		{
+			if (ensure(CastField<FBoolProperty>(InSourceProperty)))
+			{
+				return AddUpdateFunction([InSourceProperty, InTargetHierarchy, InControlElement](const UObject* InSourceInstance)
+				{
+					SetControlValue<bool>(InControlElement, InSourceProperty, InTargetHierarchy, InSourceInstance);
+				});
+			}
+			break;
+		}
+		case ERigControlType::Float:
+		case ERigControlType::ScaleFloat:
+		{
+			if(ensure(CastField<FFloatProperty>(InSourceProperty)))
+			{
+				return AddUpdateFunction([InSourceProperty, InTargetHierarchy, InControlElement](const UObject* InSourceInstance)
+				{
+					SetControlValue<float>(InControlElement, InSourceProperty, InTargetHierarchy, InSourceInstance);
+				});
+			}
+			break;
+		}
+		case ERigControlType::Integer:
+		{
+			if(ensure(CastField<FIntProperty>(InSourceProperty)))
+			{
+				return AddUpdateFunction([InSourceProperty, InTargetHierarchy, InControlElement](const UObject* InSourceInstance)
+				{
+					SetControlValue<int32>(InControlElement, InSourceProperty, InTargetHierarchy, InSourceInstance);
+				});
+			}
+			break;
+		}
+		case ERigControlType::Vector2D:
+		{
+			FStructProperty* StructProperty = CastField<FStructProperty>(InSourceProperty);
+			if(ensure(StructProperty))
+			{
+				if(ensure(StructProperty->Struct == TBaseStructure<FVector2D>::Get()))
+				{
+					return AddUpdateFunction([InSourceProperty, InTargetHierarchy, InControlElement](const UObject* InSourceInstance)
+					{
+						SetControlValue<FVector2D>(InControlElement, InSourceProperty, InTargetHierarchy, InSourceInstance);
+					});
+				}
+			}
+			break;
+		}
+		case ERigControlType::Position:
+		case ERigControlType::Scale:
+		{
+			FStructProperty* StructProperty = CastField<FStructProperty>(InSourceProperty);
+			if(ensure(StructProperty))
+			{
+				if(ensure(StructProperty->Struct == TBaseStructure<FVector>::Get()))
+				{
+					return AddUpdateFunction([InSourceProperty, InTargetHierarchy, InControlElement](const UObject* InSourceInstance)
+					{
+						SetControlValue<FVector>(InControlElement, InSourceProperty, InTargetHierarchy, InSourceInstance);
+					});
+				}
+			}
+			break;
+		}
+		case ERigControlType::Rotator:
+		{
+			FStructProperty* StructProperty = CastField<FStructProperty>(InSourceProperty);
+			if(ensure(StructProperty))
+			{
+				if(ensure(StructProperty->Struct == TBaseStructure<FRotator>::Get()))
+				{
+					return AddUpdateFunction([InSourceProperty, InTargetHierarchy, InControlElement](const UObject* InSourceInstance)
+					{
+						SetControlValue<FRotator>(InControlElement, InSourceProperty, InTargetHierarchy, InSourceInstance);
+					});
+				}
+			}
+			break;
+		}
+		case ERigControlType::Transform:
+		{
+			FStructProperty* StructProperty = CastField<FStructProperty>(InSourceProperty);
+			if(ensure(StructProperty))
+			{
+				if(ensure(StructProperty->Struct == TBaseStructure<FTransform>::Get()))
+				{
+					return AddUpdateFunction([InSourceProperty, InTargetHierarchy, InControlElement](const UObject* InSourceInstance)
+	                {
+						SetControlValue<FTransform>(InControlElement, InSourceProperty, InTargetHierarchy, InSourceInstance);
+	                });
+				}
+			}
+			break;
+		}
+		case ERigControlType::TransformNoScale:
+		{
+			FStructProperty* StructProperty = CastField<FStructProperty>(InSourceProperty);
+			if(ensure(StructProperty))
+			{
+				if(ensure(StructProperty->Struct == TBaseStructure<FTransform>::Get()))
+				{
+					return AddUpdateFunction([InSourceProperty, InTargetHierarchy, InControlElement](const UObject* InSourceInstance)
+					{
+						SetControlValue<FTransform, FTransformNoScale>(InControlElement, InSourceProperty, InTargetHierarchy, InSourceInstance);
+					});
+				}
+			}
+			break;
+		}
+		case ERigControlType::EulerTransform:
+		{
+			FStructProperty* StructProperty = CastField<FStructProperty>(InSourceProperty);
+			if(ensure(StructProperty))
+			{
+				if(ensure(StructProperty->Struct == TBaseStructure<FTransform>::Get()))
+				{
+					return AddUpdateFunction([InSourceProperty, InTargetHierarchy, InControlElement](const UObject* InSourceInstance)
+					{
+						const FTransform* ValuePtr = InSourceProperty->ContainerPtrToValuePtr<FTransform>(InSourceInstance);
+						const FRigControlValue Value = FRigControlValue::Make<FEulerTransform>(FEulerTransform(*ValuePtr));
+						InTargetHierarchy->SetControlValue(InControlElement, Value, ERigControlValueType::Current);
+					});
+				}
+			}
+			break;
+		}
+		default:
+		{
+			checkNoEntry();
+		}
+	}
+}
+
+void FAnimNode_ControlRig::AddVariableFunction(const FRigVMExternalVariable& InVariable, FProperty* InSourceProperty)
+{
+	using namespace AnimNodeLocals;
+	
+	const int32 VariableIndex = Variables.Add(InVariable);
+	
+	if (CastField<FBoolProperty>(InSourceProperty) && InVariable.TypeName == RigVMTypeUtils::BoolTypeName)
+	{
+		return AddUpdateFunction([InSourceProperty, this, VariableIndex](const UObject* InSourceInstance)
+		{
+			SetVariableValue<bool>(Variables[VariableIndex], InSourceProperty, InSourceInstance);
+		});
+	}
+	else if (CastField<FFloatProperty>(InSourceProperty) && (InVariable.TypeName == RigVMTypeUtils::FloatTypeName || InVariable.TypeName == RigVMTypeUtils::DoubleTypeName))
+	{
+		if(InVariable.TypeName == RigVMTypeUtils::FloatTypeName)
+		{
+			return AddUpdateFunction([InSourceProperty, this, VariableIndex](const UObject* InSourceInstance)
+			{
+				SetVariableValue<float>(Variables[VariableIndex], InSourceProperty, InSourceInstance);
+			});
+		}
+		else
+		{
+			return AddUpdateFunction([InSourceProperty, this, VariableIndex](const UObject* InSourceInstance)
+			{
+				SetVariableValue<float, double>(Variables[VariableIndex], InSourceProperty, InSourceInstance);
+			});
+		}
+	}
+	else if (CastField<FDoubleProperty>(InSourceProperty) && (InVariable.TypeName == RigVMTypeUtils::FloatTypeName || InVariable.TypeName == RigVMTypeUtils::DoubleTypeName))
+	{
+		if(InVariable.TypeName == RigVMTypeUtils::FloatTypeName)
+		{
+			return AddUpdateFunction([InSourceProperty, this, VariableIndex](const UObject* InSourceInstance)
+			{
+				SetVariableValue<double, float>(Variables[VariableIndex], InSourceProperty, InSourceInstance);
+			});
+		}
+		else
+		{
+			return AddUpdateFunction([InSourceProperty, this, VariableIndex](const UObject* InSourceInstance)
+			{
+				SetVariableValue<double>(Variables[VariableIndex], InSourceProperty, InSourceInstance);
+			});
+		}
+	}
+	else if (CastField<FIntProperty>(InSourceProperty) && InVariable.TypeName == RigVMTypeUtils::Int32TypeName)
+	{
+		return AddUpdateFunction([InSourceProperty, this, VariableIndex](const UObject* InSourceInstance)
+		{
+			SetVariableValue<int32>(Variables[VariableIndex], InSourceProperty, InSourceInstance);
+		});
+	}
+	else if (CastField<FNameProperty>(InSourceProperty) && InVariable.TypeName == RigVMTypeUtils::FNameTypeName)
+	{
+		return AddUpdateFunction([InSourceProperty, this, VariableIndex](const UObject* InSourceInstance)
+		{
+			SetVariableValue<FName>(Variables[VariableIndex], InSourceProperty, InSourceInstance);
+		});
+	}
+	else if (CastField<FNameProperty>(InSourceProperty) && InVariable.TypeName == RigVMTypeUtils::FStringTypeName)
+	{
+		return AddUpdateFunction([InSourceProperty, this, VariableIndex](const UObject* InSourceInstance)
+		{
+			SetVariableValue<FString>(Variables[VariableIndex], InSourceProperty, InSourceInstance);
+		});
+	}
+	else if (FStructProperty* StructProperty = CastField<FStructProperty>(InSourceProperty))
+	{
+		return AddUpdateFunction([StructProperty, this, VariableIndex](const UObject* InSourceInstance)
+		{
+			const uint8* SrcPtr = StructProperty->ContainerPtrToValuePtr<uint8>(InSourceInstance);
+			StructProperty->Struct->CopyScriptStruct(Variables[VariableIndex].Memory, SrcPtr, 1);
+		});
+	}
+	else if(FArrayProperty* ArrayProperty = CastField<FArrayProperty>(InSourceProperty))
+	{
+		if(ensure(ArrayProperty->SameType(InVariable.Property)))
+		{
+			return AddUpdateFunction([ArrayProperty, this, VariableIndex](const UObject* InSourceInstance)
+			{
+				const uint8* SrcPtr = ArrayProperty->ContainerPtrToValuePtr<uint8>(InSourceInstance);
+				ArrayProperty->CopyCompleteValue(Variables[VariableIndex].Memory, SrcPtr);
+			});
+		}
+	}
+	else if(FObjectProperty* ObjectProperty = CastField<FObjectProperty>(InSourceProperty))
+	{
+		if(ensure(ObjectProperty->SameType(InVariable.Property)))
+		{
+			return AddUpdateFunction([ObjectProperty, this, VariableIndex](const UObject* InSourceInstance)
+			{
+				const uint8* SrcPtr = ObjectProperty->ContainerPtrToValuePtr<uint8>(InSourceInstance);
+				ObjectProperty->CopyCompleteValue(Variables[VariableIndex].Memory, SrcPtr);
+			});
+		}
+	}
+	else if(FEnumProperty* EnumProperty = CastField<FEnumProperty>(InSourceProperty))
+	{
+		if(ensure(EnumProperty->SameType(InVariable.Property)))
+		{
+			return AddUpdateFunction([EnumProperty, this, VariableIndex](const UObject* InSourceInstance)
+			{
+				const uint8* SrcPtr = EnumProperty->ContainerPtrToValuePtr<uint8>(InSourceInstance);
+				EnumProperty->CopyCompleteValue(Variables[VariableIndex].Memory, SrcPtr);
+			});
+		}
+	}
+	else
+	{
+		ensureMsgf(false, TEXT("Property %s type %s not recognized"), *InSourceProperty->GetName(), *InSourceProperty->GetCPPType());
+	}
+}
+
 void FAnimNode_ControlRig::InitializeProperties(const UObject* InSourceInstance, UClass* InTargetClass)
 {
 	// Build property lists
 	SourceProperties.Reset(SourcePropertyNames.Num());
 	DestProperties.Reset(SourcePropertyNames.Num());
+	UpdateFunctions.Reset(SourcePropertyNames.Num());
+	Variables.Reset();
 
 	check(SourcePropertyNames.Num() == DestPropertyNames.Num());
 
+	UControlRig* TargetControlRig = Cast<UControlRig>(TargetInstance);
+	URigHierarchy* TargetHierarchy = TargetControlRig ? TargetControlRig->GetHierarchy() : nullptr;
+
+	const UClass* SourceClass = InSourceInstance->GetClass();
 	for (int32 Idx = 0; Idx < SourcePropertyNames.Num(); ++Idx)
 	{
-		FName& SourceName = SourcePropertyNames[Idx];
-		UClass* SourceClass = InSourceInstance->GetClass();
+		const FName& SourceName = SourcePropertyNames[Idx];
 
 		FProperty* SourceProperty = FindFProperty<FProperty>(SourceClass, SourceName);
 		SourceProperties.Add(SourceProperty);
 		DestProperties.Add(nullptr);
+
+        if (TargetHierarchy && SourceProperty)
+        {
+	        const FName& DestPropertyName = DestPropertyNames[Idx];
+        	if(FRigControlElement* ControlElement = TargetControlRig->FindControl(DestPropertyName))
+        	{
+        		AddControlFunction(ControlElement, SourceProperty, TargetHierarchy);
+        	}
+        	else
+        	{
+        		const FRigVMExternalVariable Variable = TargetControlRig->GetPublicVariableByName(DestPropertyName);
+        		if (Variable.IsValid() && !Variable.bIsReadOnly)
+        		{
+        			AddVariableFunction(Variable, SourceProperty);
+        		}
+        	}
+        }
 	}
 }
 
 void FAnimNode_ControlRig::PropagateInputProperties(const UObject* InSourceInstance)
 {
-	if (TargetInstance)
+	DECLARE_SCOPE_HIERARCHICAL_COUNTER_FUNC()
+	QUICK_SCOPE_CYCLE_COUNTER(STAT_AnimNode_ControlRig_PropagateInputProperties);
+
+	if (!InSourceInstance)
 	{
-		UControlRig* TargetControlRig = Cast<UControlRig>((UObject*)TargetInstance);
-		if(TargetControlRig == nullptr)
-		{
-			return;
-		}
-		URigHierarchy* TargetHierarchy = TargetControlRig->GetHierarchy();
-		if(TargetHierarchy == nullptr)
+		return;
+	}
+	
+	UControlRig* TargetControlRig = Cast<UControlRig>(TargetInstance);
+	if (!TargetControlRig)
+	{
+		return;
+	}
+	
+	URigHierarchy* TargetHierarchy = TargetControlRig->GetHierarchy();
+	if (!TargetHierarchy)
+	{
+		return;
+	}
+
+	if (AnimNodeLocals::bUseFunctions)
+	{
+		Algo::ForEach(UpdateFunctions, [InSourceInstance](const PropertyUpdateFunction& InFunc)
 		{
-			return;
-		}
-		
-		// First copy properties
-		check(SourceProperties.Num() == DestProperties.Num());
-		for (int32 PropIdx = 0; PropIdx < SourceProperties.Num(); ++PropIdx)
+			InFunc(InSourceInstance);
+		});
+	}
+	else
+	{
+		PropagateInputPropertiesNoCache(InSourceInstance, TargetControlRig, TargetHierarchy);
+	}
+}
+
+void FAnimNode_ControlRig::PropagateInputPropertiesNoCache(
+	const UObject* InSourceInstance,
+	const UControlRig* InTargetControlRig,
+	URigHierarchy* InTargetHierarchy)
+{
+	check(SourceProperties.Num() == DestProperties.Num());
+	
+	for (int32 PropIdx = 0; PropIdx < SourceProperties.Num(); ++PropIdx)
+	{
+		FProperty* CallerProperty = SourceProperties[PropIdx];
+
+		if(FRigControlElement* ControlElement = InTargetControlRig->FindControl(DestPropertyNames[PropIdx]))
 		{
-			FProperty* CallerProperty = SourceProperties[PropIdx];
+			const uint8* SrcPtr = CallerProperty->ContainerPtrToValuePtr<uint8>(InSourceInstance);
 
-			if(FRigControlElement* ControlElement = TargetControlRig->FindControl(DestPropertyNames[PropIdx]))
+			bool bIsValid = false;
+			FRigControlValue Value;
+			switch(ControlElement->Settings.ControlType)
 			{
-				const uint8* SrcPtr = CallerProperty->ContainerPtrToValuePtr<uint8>(InSourceInstance);
-
-				bool bIsValid = false;
-				FRigControlValue Value;
-				switch(ControlElement->Settings.ControlType)
+				case ERigControlType::Bool:
 				{
-					case ERigControlType::Bool:
+					if(ensure(CastField<FBoolProperty>(CallerProperty)))
 					{
-						if(ensure(CastField<FBoolProperty>(CallerProperty)))
-						{
-							Value = FRigControlValue::Make<bool>(*(bool*)SrcPtr);
-							bIsValid = true;
-						}
-						break;
+						Value = FRigControlValue::Make<bool>(*(bool*)SrcPtr);
+						bIsValid = true;
 					}
-					case ERigControlType::Float:
-					case ERigControlType::ScaleFloat:
+					break;
+				}
+				case ERigControlType::Float:
+				case ERigControlType::ScaleFloat:
+				{
+					if(ensure(CastField<FFloatProperty>(CallerProperty)))
 					{
-						if(ensure(CastField<FFloatProperty>(CallerProperty)))
-						{
-							Value = FRigControlValue::Make<float>(*(float*)SrcPtr);
-							bIsValid = true;
-						}
-						break;
+						Value = FRigControlValue::Make<float>(*(float*)SrcPtr);
+						bIsValid = true;
 					}
-					case ERigControlType::Integer:
+					break;
+				}
+				case ERigControlType::Integer:
+				{
+					if(ensure(CastField<FIntProperty>(CallerProperty)))
 					{
-						if(ensure(CastField<FIntProperty>(CallerProperty)))
-						{
-							Value = FRigControlValue::Make<int32>(*(int32*)SrcPtr);
-							bIsValid = true;
-						}
-						break;
+						Value = FRigControlValue::Make<int32>(*(int32*)SrcPtr);
+						bIsValid = true;
 					}
-					case ERigControlType::Vector2D:
+					break;
+				}
+				case ERigControlType::Vector2D:
+				{
+					FStructProperty* StructProperty = CastField<FStructProperty>(CallerProperty);
+					if(ensure(StructProperty))
 					{
-						FStructProperty* StructProperty = CastField<FStructProperty>(CallerProperty);
-						if(ensure(StructProperty))
+						if(ensure(StructProperty->Struct == TBaseStructure<FVector2D>::Get()))
 						{
-							if(ensure(StructProperty->Struct == TBaseStructure<FVector2D>::Get()))
-							{
-								const FVector2D& SrcVector = *(FVector2D*)SrcPtr;
-								Value = FRigControlValue::Make<FVector2D>(SrcVector);
-								bIsValid = true;
-							}
+							const FVector2D& SrcVector = *(FVector2D*)SrcPtr;
+							Value = FRigControlValue::Make<FVector2D>(SrcVector);
+							bIsValid = true;
 						}
-						break;
 					}
-					case ERigControlType::Position:
-					case ERigControlType::Scale:
+					break;
+				}
+				case ERigControlType::Position:
+				case ERigControlType::Scale:
+				{
+					FStructProperty* StructProperty = CastField<FStructProperty>(CallerProperty);
+					if(ensure(StructProperty))
 					{
-						FStructProperty* StructProperty = CastField<FStructProperty>(CallerProperty);
-						if(ensure(StructProperty))
+						if(ensure(StructProperty->Struct == TBaseStructure<FVector>::Get()))
 						{
-							if(ensure(StructProperty->Struct == TBaseStructure<FVector>::Get()))
-							{
-								const FVector& SrcVector = *(FVector*)SrcPtr;  
-								Value = FRigControlValue::Make<FVector>(SrcVector);
-								bIsValid = true;
-							}
+							const FVector& SrcVector = *(FVector*)SrcPtr;  
+							Value = FRigControlValue::Make<FVector>(SrcVector);
+							bIsValid = true;
 						}
-						break;
 					}
-					case ERigControlType::Rotator:
+					break;
+				}
+				case ERigControlType::Rotator:
+				{
+					FStructProperty* StructProperty = CastField<FStructProperty>(CallerProperty);
+					if(ensure(StructProperty))
 					{
-						FStructProperty* StructProperty = CastField<FStructProperty>(CallerProperty);
-						if(ensure(StructProperty))
+						if(ensure(StructProperty->Struct == TBaseStructure<FRotator>::Get()))
 						{
-							if(ensure(StructProperty->Struct == TBaseStructure<FRotator>::Get()))
-							{
-								const FRotator& SrcRotator = *(FRotator*)SrcPtr;
-								Value = FRigControlValue::Make<FRotator>(SrcRotator);
-								bIsValid = true;
-							}
+							const FRotator& SrcRotator = *(FRotator*)SrcPtr;
+							Value = FRigControlValue::Make<FRotator>(SrcRotator);
+							bIsValid = true;
 						}
-						break;
 					}
-					case ERigControlType::Transform:
+					break;
+				}
+				case ERigControlType::Transform:
+				{
+					FStructProperty* StructProperty = CastField<FStructProperty>(CallerProperty);
+					if(ensure(StructProperty))
 					{
-						FStructProperty* StructProperty = CastField<FStructProperty>(CallerProperty);
-						if(ensure(StructProperty))
+						if(ensure(StructProperty->Struct == TBaseStructure<FTransform>::Get()))
 						{
-							if(ensure(StructProperty->Struct == TBaseStructure<FTransform>::Get()))
-							{
-								const FTransform& SrcTransform = *(FTransform*)SrcPtr;  
-								Value = FRigControlValue::Make<FTransform>(SrcTransform);
-								bIsValid = true;
-							}
+							const FTransform& SrcTransform = *(FTransform*)SrcPtr;  
+							Value = FRigControlValue::Make<FTransform>(SrcTransform);
+							bIsValid = true;
 						}
-						break;
 					}
-					case ERigControlType::TransformNoScale:
+					break;
+				}
+				case ERigControlType::TransformNoScale:
+				{
+					FStructProperty* StructProperty = CastField<FStructProperty>(CallerProperty);
+					if(ensure(StructProperty))
 					{
-						FStructProperty* StructProperty = CastField<FStructProperty>(CallerProperty);
-						if(ensure(StructProperty))
+						if(ensure(StructProperty->Struct == TBaseStructure<FTransform>::Get()))
 						{
-							if(ensure(StructProperty->Struct == TBaseStructure<FTransform>::Get()))
-							{
-								const FTransform& SrcTransform = *(FTransform*)SrcPtr;  
-								Value = FRigControlValue::Make<FTransformNoScale>(SrcTransform);
-								bIsValid = true;
-							}
+							const FTransform& SrcTransform = *(FTransform*)SrcPtr;  
+							Value = FRigControlValue::Make<FTransformNoScale>(SrcTransform);
+							bIsValid = true;
 						}
-						break;
 					}
-					case ERigControlType::EulerTransform:
+					break;
+				}
+				case ERigControlType::EulerTransform:
+				{
+					FStructProperty* StructProperty = CastField<FStructProperty>(CallerProperty);
+					if(ensure(StructProperty))
 					{
-						FStructProperty* StructProperty = CastField<FStructProperty>(CallerProperty);
-						if(ensure(StructProperty))
+						if(ensure(StructProperty->Struct == TBaseStructure<FTransform>::Get()))
 						{
-							if(ensure(StructProperty->Struct == TBaseStructure<FTransform>::Get()))
-							{
-								const FTransform& SrcTransform = *(FTransform*)SrcPtr;  
-								Value = FRigControlValue::Make<FEulerTransform>(FEulerTransform(SrcTransform));
-								bIsValid = true;
-							}
+							const FTransform& SrcTransform = *(FTransform*)SrcPtr;  
+							Value = FRigControlValue::Make<FEulerTransform>(FEulerTransform(SrcTransform));
+							bIsValid = true;
 						}
-						break;
-					}
-					default:
-					{
-						checkNoEntry();
 					}
+					break;
 				}
-
-				if(bIsValid)
+				default:
 				{
-					TargetHierarchy->SetControlValue(ControlElement, Value, ERigControlValueType::Current);
+					checkNoEntry();
 				}
-				continue;
 			}
 
-			FRigVMExternalVariable Variable = TargetControlRig->GetPublicVariableByName(DestPropertyNames[PropIdx]);
-			if(Variable.IsValid())
+			if(bIsValid)
 			{
-				if (Variable.bIsReadOnly)
-				{
-					continue;
-				}
+				InTargetHierarchy->SetControlValue(ControlElement, Value, ERigControlValueType::Current);
+			}
+			continue;
+		}
 
-				const uint8* SrcPtr = CallerProperty->ContainerPtrToValuePtr<uint8>(InSourceInstance);
+		FRigVMExternalVariable Variable = InTargetControlRig->GetPublicVariableByName(DestPropertyNames[PropIdx]);
+		if(Variable.IsValid())
+		{
+			if (Variable.bIsReadOnly)
+			{
+				continue;
+			}
 
-				if (CastField<FBoolProperty>(CallerProperty) != nullptr && Variable.TypeName == RigVMTypeUtils::BoolTypeName)
-				{
-					const bool Value = *(const bool*)SrcPtr;
-					Variable.SetValue<bool>(Value);
-				}
-				else if (CastField<FFloatProperty>(CallerProperty) != nullptr && (Variable.TypeName == RigVMTypeUtils::FloatTypeName ||  Variable.TypeName == RigVMTypeUtils::DoubleTypeName))
-				{
-					const float Value = *(const float*)SrcPtr;
-					if(Variable.TypeName == RigVMTypeUtils::FloatTypeName)
-					{
-						Variable.SetValue<float>(Value);
-					}
-					else
-					{
-						Variable.SetValue<double>(Value);
-					}
-				}
-				else if (CastField<FDoubleProperty>(CallerProperty) != nullptr && (Variable.TypeName == RigVMTypeUtils::FloatTypeName ||  Variable.TypeName == RigVMTypeUtils::DoubleTypeName))
-				{
-					const double Value = *(const double*)SrcPtr;
-					if(Variable.TypeName == RigVMTypeUtils::FloatTypeName)
-					{
-						Variable.SetValue<float>((float)Value);
-					}
-					else
-					{
-						Variable.SetValue<double>(Value);
-					}
-				}
-				else if (CastField<FIntProperty>(CallerProperty) != nullptr && Variable.TypeName == RigVMTypeUtils::Int32TypeName)
+			const uint8* SrcPtr = CallerProperty->ContainerPtrToValuePtr<uint8>(InSourceInstance);
+
+			if (CastField<FBoolProperty>(CallerProperty) != nullptr && Variable.TypeName == RigVMTypeUtils::BoolTypeName)
+			{
+				const bool Value = *(const bool*)SrcPtr;
+				Variable.SetValue<bool>(Value);
+			}
+			else if (CastField<FFloatProperty>(CallerProperty) != nullptr && (Variable.TypeName == RigVMTypeUtils::FloatTypeName ||  Variable.TypeName == RigVMTypeUtils::DoubleTypeName))
+			{
+				const float Value = *(const float*)SrcPtr;
+				if(Variable.TypeName == RigVMTypeUtils::FloatTypeName)
 				{
-					const int32 Value = *(const int32*)SrcPtr;
-					Variable.SetValue<int32>(Value);
+					Variable.SetValue<float>(Value);
 				}
-				else if (CastField<FNameProperty>(CallerProperty) != nullptr && Variable.TypeName == RigVMTypeUtils::FNameTypeName)
+				else
 				{
-					const FName Value = *(const FName*)SrcPtr;
-					Variable.SetValue<FName>(Value);
+					Variable.SetValue<double>(Value);
 				}
-				else if (CastField<FNameProperty>(CallerProperty) != nullptr && Variable.TypeName == RigVMTypeUtils::FStringTypeName)
+			}
+			else if (CastField<FDoubleProperty>(CallerProperty) != nullptr && (Variable.TypeName == RigVMTypeUtils::FloatTypeName ||  Variable.TypeName == RigVMTypeUtils::DoubleTypeName))
+			{
+				const double Value = *(const double*)SrcPtr;
+				if(Variable.TypeName == RigVMTypeUtils::FloatTypeName)
 				{
-					const FString Value = *(const FString*)SrcPtr;
-					Variable.SetValue<FString>(Value);
+					Variable.SetValue<float>((float)Value);
 				}
-				else if (FStructProperty* StructProperty = CastField<FStructProperty>(CallerProperty))
+				else
 				{
-					if (StructProperty->Struct == Variable.TypeObject)
-					{
-						StructProperty->Struct->CopyScriptStruct(Variable.Memory, SrcPtr, 1);
-					}
+					Variable.SetValue<double>(Value);
 				}
-				else if(FArrayProperty* ArrayProperty = CastField<FArrayProperty>(CallerProperty))
+			}
+			else if (CastField<FIntProperty>(CallerProperty) != nullptr && Variable.TypeName == RigVMTypeUtils::Int32TypeName)
+			{
+				const int32 Value = *(const int32*)SrcPtr;
+				Variable.SetValue<int32>(Value);
+			}
+			else if (CastField<FNameProperty>(CallerProperty) != nullptr && Variable.TypeName == RigVMTypeUtils::FNameTypeName)
+			{
+				const FName Value = *(const FName*)SrcPtr;
+				Variable.SetValue<FName>(Value);
+			}
+			else if (CastField<FNameProperty>(CallerProperty) != nullptr && Variable.TypeName == RigVMTypeUtils::FStringTypeName)
+			{
+				const FString Value = *(const FString*)SrcPtr;
+				Variable.SetValue<FString>(Value);
+			}
+			else if (FStructProperty* StructProperty = CastField<FStructProperty>(CallerProperty))
+			{
+				if (StructProperty->Struct == Variable.TypeObject)
 				{
-					if(ensure(ArrayProperty->SameType(Variable.Property)))
-					{
-						ArrayProperty->CopyCompleteValue(Variable.Memory, SrcPtr);
-					}
+					StructProperty->Struct->CopyScriptStruct(Variable.Memory, SrcPtr, 1);
 				}
-				else if(FObjectProperty* ObjectProperty = CastField<FObjectProperty>(CallerProperty))
+			}
+			else if(FArrayProperty* ArrayProperty = CastField<FArrayProperty>(CallerProperty))
+			{
+				if(ensure(ArrayProperty->SameType(Variable.Property)))
 				{
-					if(ensure(ObjectProperty->SameType(Variable.Property)))
-					{
-						ObjectProperty->CopyCompleteValue(Variable.Memory, SrcPtr);
-					}
+					ArrayProperty->CopyCompleteValue(Variable.Memory, SrcPtr);
 				}
-				else if(FEnumProperty* EnumProperty = CastField<FEnumProperty>(CallerProperty))
+			}
+			else if(FObjectProperty* ObjectProperty = CastField<FObjectProperty>(CallerProperty))
+			{
+				if(ensure(ObjectProperty->SameType(Variable.Property)))
 				{
-					if(ensure(EnumProperty->SameType(Variable.Property)))
-					{
-						EnumProperty->CopyCompleteValue(Variable.Memory, SrcPtr);
-					}
+					ObjectProperty->CopyCompleteValue(Variable.Memory, SrcPtr);
 				}
-				else
+			}
+			else if(FEnumProperty* EnumProperty = CastField<FEnumProperty>(CallerProperty))
+			{
+				if(ensure(EnumProperty->SameType(Variable.Property)))
 				{
-					ensureMsgf(false, TEXT("Property %s type %s not recognized"), *CallerProperty->GetName(), *CallerProperty->GetCPPType());
+					EnumProperty->CopyCompleteValue(Variable.Memory, SrcPtr);
 				}
 			}
+			else
+			{
+				ensureMsgf(false, TEXT("Property %s type %s not recognized"), *CallerProperty->GetName(), *CallerProperty->GetCPPType());
+			}
 		}
 	}
 }
diff --git a/Engine/Plugins/Animation/ControlRig/Source/ControlRig/Public/AnimNode_ControlRig.h b/Engine/Plugins/Animation/ControlRig/Source/ControlRig/Public/AnimNode_ControlRig.h
--- a/Engine/Plugins/Animation/ControlRig/Source/ControlRig/Public/AnimNode_ControlRig.h
+++ b/Engine/Plugins/Animation/ControlRig/Source/ControlRig/Public/AnimNode_ControlRig.h
@@ -132,6 +132,23 @@ private:
 	FCurveMappings InputCurveMappings;
 	FCurveMappings OutputCurveMappings;
 
+	/* This function is used to copy the property values going thru the hierarchy without any caching mechanism. */
+	void PropagateInputPropertiesNoCache(const UObject* InSourceInstance, const UControlRig* InControlRig, URigHierarchy* InHierarchy);
+	
+	/*
+	 * Cached functions used to copy property values to avoid looking for them for every PropagateInputProperties calls.
+	 * They are stored in InitializeProperties and basically represent 'set value' functions on controls and variables.
+	 */
+	void AddControlFunction(FRigControlElement* InControlElement, FProperty* InSourceProperty, URigHierarchy* InTargetHierarchy);
+	void AddVariableFunction(const FRigVMExternalVariable& InVariable, FProperty* InSourceProperty);
+	
+	using PropertyUpdateFunction = TFunction<void(const UObject*)>;
+
+	void AddUpdateFunction(PropertyUpdateFunction&& InFunction);
+
+	TArray<PropertyUpdateFunction> UpdateFunctions;
+	TArray<FRigVMExternalVariable> Variables;
+
 protected:
 	virtual UClass* GetTargetClass() const override;
 	virtual void UpdateInput(UControlRig* InControlRig, FPoseContext& InOutput) override;
