Control Rig: Sustain connections as much as possible

#rnx

diff --git a/Engine/Plugins/Animation/ControlRig/Source/ControlRig/Private/ModularRigController.cpp b/Engine/Plugins/Animation/ControlRig/Source/ControlRig/Private/ModularRigController.cpp
--- a/Engine/Plugins/Animation/ControlRig/Source/ControlRig/Private/ModularRigController.cpp
+++ b/Engine/Plugins/Animation/ControlRig/Source/ControlRig/Private/ModularRigController.cpp
@@ -243,7 +243,8 @@ bool UModularRigController::ConnectConnectorToElement(const FRigElementKey& InCo
 
 #if WITH_EDITOR
 	FName TargetModulePathName = NAME_None;
-	if (UModularRig* ModularRig = Cast<UModularRig>(Blueprint->GetObjectBeingDebugged()))
+	UModularRig* ModularRig = Cast<UModularRig>(Blueprint->GetObjectBeingDebugged());
+	if(ModularRig)
 	{
 		if (URigHierarchy* Hierarchy = ModularRig->GetHierarchy())
 		{
@@ -260,19 +261,34 @@ bool UModularRigController::ConnectConnectorToElement(const FRigElementKey& InCo
 #endif 
 
 	// First disconnect before connecting to anything else. This might disconnect other secondary/optional connectors.
+	TMap<FRigElementKey, FRigElementKey> PreviousConnections;
 	if (CurrentTarget.IsValid())
 	{
 		const TGuardValue<bool> DisableAutomaticReparenting(bAutomaticReparenting, false);
-		DisconnectConnector(InConnectorKey, false, bSetupUndo);
+		DisconnectConnector_Internal(InConnectorKey, false, &PreviousConnections, bSetupUndo);
 	}
 	
 	Model->Connections.AddConnection(InConnectorKey, InTargetKey);
+
+	// restore previous connections if possible
+	for(const TPair<FRigElementKey, FRigElementKey>& PreviousConnection : PreviousConnections)
+	{
+		if(!Model->Connections.HasConnection(PreviousConnection.Key))
+		{
+			FText ErrorMessageForPreviousConnection;
+			if(CanConnectConnectorToElement(PreviousConnection.Key, PreviousConnection.Value, ErrorMessageForPreviousConnection))
+			{
+				(void)ConnectConnectorToElement(PreviousConnection.Key, PreviousConnection.Value, bSetupUndo, false, false);
+			}
+		}
+	}
+	
 	Notify(EModularRigNotification::ConnectionChanged, Module);
 
 #if WITH_EDITOR
 	if (UControlRig* RigCDO = Module->Class->GetDefaultObject<UControlRig>())
 	{
-		if (UModularRig* ModularRig = Cast<UModularRig>(Blueprint->GetObjectBeingDebugged()))
+		if (ModularRig)
 		{
 			if (URigHierarchy* Hierarchy = ModularRig->GetHierarchy())
 			{
@@ -320,6 +336,12 @@ bool UModularRigController::ConnectConnectorToElement(const FRigElementKey& InCo
 }
 
 bool UModularRigController::DisconnectConnector(const FRigElementKey& InConnectorKey, bool bDisconnectSubModules, bool bSetupUndo)
+{
+	return DisconnectConnector_Internal(InConnectorKey, bDisconnectSubModules, nullptr, bSetupUndo);
+};
+
+bool UModularRigController::DisconnectConnector_Internal(const FRigElementKey& InConnectorKey, bool bDisconnectSubModules,
+	TMap<FRigElementKey, FRigElementKey>* OutRemovedConnections, bool bSetupUndo)
 {
 	FString ConnectorModulePath, ConnectorName;
 	if (!URigHierarchy::SplitNameSpace(InConnectorKey.Name.ToString(), &ConnectorModulePath, &ConnectorName))
@@ -369,6 +391,10 @@ bool UModularRigController::DisconnectConnector(const FRigElementKey& InConnecto
 	}
 #endif 
 
+	if(OutRemovedConnections)
+	{
+		OutRemovedConnections->Add(InConnectorKey, Model->Connections.FindTargetFromConnector(InConnectorKey));
+	}
 	Model->Connections.RemoveConnection(InConnectorKey);
 
 	if (ModuleConnector->IsPrimary())
@@ -384,6 +410,10 @@ bool UModularRigController::DisconnectConnector(const FRigElementKey& InConnecto
 		}
 		for (const FRigElementKey& ToRemove : ConnectionsToRemove)
 		{
+			if(OutRemovedConnections)
+			{
+				OutRemovedConnections->Add(ToRemove, Model->Connections.FindTargetFromConnector(ToRemove));
+			}
 			Model->Connections.RemoveConnection(ToRemove);
 		}
 	}
@@ -402,6 +432,10 @@ bool UModularRigController::DisconnectConnector(const FRigElementKey& InConnecto
 		}
 		for (const FRigElementKey& ToRemove : ConnectionsToRemove)
 		{
+			if(OutRemovedConnections)
+			{
+				OutRemovedConnections->Add(ToRemove, Model->Connections.FindTargetFromConnector(ToRemove));
+			}
 			Model->Connections.RemoveConnection(ToRemove);
 		}
 	}
diff --git a/Engine/Plugins/Animation/ControlRig/Source/ControlRig/Private/ModularRigRuleManager.cpp b/Engine/Plugins/Animation/ControlRig/Source/ControlRig/Private/ModularRigRuleManager.cpp
--- a/Engine/Plugins/Animation/ControlRig/Source/ControlRig/Private/ModularRigRuleManager.cpp
+++ b/Engine/Plugins/Animation/ControlRig/Source/ControlRig/Private/ModularRigRuleManager.cpp
@@ -25,7 +25,16 @@ FModularRigResolveResult UModularRigRuleManager::FindMatches(FWorkData& InWorkDa
 	{
 		if (ControlRig->IsConstructionRequired())
 		{
+			const FRigElementKey ConnectorKey = InWorkData.Connector ? InWorkData.Connector->GetKey() : FRigElementKey();
+			
 			ControlRig->Execute(FRigUnit_PrepareForExecution::EventName);
+
+			// restore the connector. executing the control rig may have destroyed the previous connector element
+			InWorkData.Connector = Hierarchy->Find<FRigConnectorElement>(ConnectorKey);
+			if(ConnectorKey.IsValid() && InWorkData.Connector == nullptr)
+			{
+				return Result;
+			}
 		}
 	}
 
diff --git a/Engine/Plugins/Animation/ControlRig/Source/ControlRig/Private/Rigs/RigHierarchyCache.cpp b/Engine/Plugins/Animation/ControlRig/Source/ControlRig/Private/Rigs/RigHierarchyCache.cpp
--- a/Engine/Plugins/Animation/ControlRig/Source/ControlRig/Private/Rigs/RigHierarchyCache.cpp
+++ b/Engine/Plugins/Animation/ControlRig/Source/ControlRig/Private/Rigs/RigHierarchyCache.cpp
@@ -2,6 +2,7 @@
 
 #include "Rigs/RigHierarchyCache.h"
 #include "Rigs/RigHierarchy.h"
+#include "ModularRigModel.h"
 
 #include UE_INLINE_GENERATED_CPP_BY_NAME(RigHierarchyCache)
 
@@ -108,6 +109,21 @@ FRigElementKeyRedirector::FRigElementKeyRedirector(const FRigElementKeyRedirecto
 	}
 }
 
+FRigElementKeyRedirector::FRigElementKeyRedirector(const FModularRigConnections& InOther, const URigHierarchy* InHierarchy)
+{
+	check(InHierarchy);
+	InternalKeyToExternalKey.Reserve(InOther.Num());
+	ExternalKeys.Reserve(InOther.Num());
+
+	Hash = 0;
+	for(const FModularRigSingleConnection& Connection : InOther)
+	{
+		check(Connection.Connector.IsValid());
+		check(Connection.Target.IsValid());
+		Add(Connection.Connector, Connection.Target, InHierarchy);
+	}
+}
+
 const FRigElementKey* FRigElementKeyRedirector::FindReverse(const FRigElementKey& InKey) const
 {
 	for(const TPair<FRigElementKey, FCachedRigElement>& Pair : InternalKeyToExternalKey)
diff --git a/Engine/Plugins/Animation/ControlRig/Source/ControlRig/Public/ControlRig.h b/Engine/Plugins/Animation/ControlRig/Source/ControlRig/Public/ControlRig.h
--- a/Engine/Plugins/Animation/ControlRig/Source/ControlRig/Public/ControlRig.h
+++ b/Engine/Plugins/Animation/ControlRig/Source/ControlRig/Public/ControlRig.h
@@ -936,6 +936,7 @@ private:
 	friend class FRigTransformElementDetails;
 	friend class FControlRigEditorModule;
 	friend class UModularRig;
+	friend class UModularRigController;
 };
 
 class CONTROLRIG_API FControlRigBracketScope
diff --git a/Engine/Plugins/Animation/ControlRig/Source/ControlRig/Public/ModularRigController.h b/Engine/Plugins/Animation/ControlRig/Source/ControlRig/Public/ModularRigController.h
--- a/Engine/Plugins/Animation/ControlRig/Source/ControlRig/Public/ModularRigController.h
+++ b/Engine/Plugins/Animation/ControlRig/Source/ControlRig/Public/ModularRigController.h
@@ -35,7 +35,8 @@ class CONTROLRIG_API UModularRigController : public UObject
 
 	UFUNCTION(BlueprintCallable, Category = "Control Rig | Modules")
 	bool DisconnectConnector(const FRigElementKey& InConnectorKey, bool bDisconnectSubModules = false, bool bSetupUndo = true);
-	
+	bool DisconnectConnector_Internal(const FRigElementKey& InConnectorKey, bool bDisconnectSubModules = false, TMap<FRigElementKey, FRigElementKey>*OutRemovedConnections = nullptr, bool bSetupUndo = true);
+
 	UFUNCTION(BlueprintCallable, Category = "Control Rig | Modules")
 	TArray<FRigElementKey> DisconnectCyclicConnectors(bool bSetupUndo = true);
 
diff --git a/Engine/Plugins/Animation/ControlRig/Source/ControlRig/Public/Rigs/RigHierarchyCache.h b/Engine/Plugins/Animation/ControlRig/Source/ControlRig/Public/Rigs/RigHierarchyCache.h
--- a/Engine/Plugins/Animation/ControlRig/Source/ControlRig/Public/Rigs/RigHierarchyCache.h
+++ b/Engine/Plugins/Animation/ControlRig/Source/ControlRig/Public/Rigs/RigHierarchyCache.h
@@ -9,6 +9,7 @@
 #include "RigHierarchyCache.generated.h"
 
 class URigHierarchy;
+struct FModularRigConnections;
 
 USTRUCT(BlueprintType)
 struct CONTROLRIG_API FCachedRigElement
@@ -168,6 +169,7 @@ public:
 
 	FRigElementKeyRedirector(const TMap<FRigElementKey, FRigElementKey>& InMap, const URigHierarchy* InHierarchy);
 	FRigElementKeyRedirector(const FRigElementKeyRedirector& InOther, const URigHierarchy* InHierarchy);
+	FRigElementKeyRedirector(const FModularRigConnections& InOther, const URigHierarchy* InHierarchy);
 
 	bool Contains(const FRigElementKey& InKey) const { return InternalKeyToExternalKey.Contains(InKey); }
 	const FCachedRigElement* Find(const FRigElementKey& InKey) const { return InternalKeyToExternalKey.Find(InKey); }
diff --git a/Engine/Plugins/Animation/ControlRig/Source/ControlRigEditor/Private/ControlRigModuleDetails.cpp b/Engine/Plugins/Animation/ControlRig/Source/ControlRigEditor/Private/ControlRigModuleDetails.cpp
--- a/Engine/Plugins/Animation/ControlRig/Source/ControlRigEditor/Private/ControlRigModuleDetails.cpp
+++ b/Engine/Plugins/Animation/ControlRig/Source/ControlRigEditor/Private/ControlRigModuleDetails.cpp
@@ -365,12 +365,15 @@ void FRigModuleInstanceDetails::CustomizeDetails(IDetailLayoutBuilder& DetailBui
 				{
 					return PerModuleInfos[0].GetModularRig()->GetHierarchy();
 				});
-				TreeDelegates.OnRigTreeIsItemVisible = FOnRigTreeIsItemVisible::CreateLambda([Matches](const FRigElementKey& InTarget)
+
+				TArray<FRigElementKey> MatchingKeys;
+				for(const FRigElementResolveResult& SingleMatch : Matches)
 				{
-					return Matches.ContainsByPredicate([InTarget](const FRigElementResolveResult& Match)
-					{
-						return Match.GetKey() == InTarget;
-					});
+					MatchingKeys.Add(SingleMatch.GetKey());
+				}
+				TreeDelegates.OnRigTreeIsItemVisible = FOnRigTreeIsItemVisible::CreateLambda([MatchingKeys](const FRigElementKey& InTarget)
+				{
+					return MatchingKeys.Contains(InTarget);
 				});
 				TreeDelegates.OnGetSelection.BindLambda([this, Connector]() -> TArray<FRigElementKey>
 				{
@@ -407,6 +410,7 @@ void FRigModuleInstanceDetails::CustomizeDetails(IDetailLayoutBuilder& DetailBui
 				static const FSlateBrush* OptionalBrush = FControlRigEditorStyle::Get().GetBrush("ControlRig.ConnectorOptional");
 
 				const FSlateBrush* IconBrush = Connector.IsPrimary() ? PrimaryBrush : (Connector.IsOptional() ? OptionalBrush : SecondaryBrush);
+				TSharedPtr<SSearchableRigHierarchyTreeView>& SearchableTreeView = ConnectionListBox.FindOrAdd(Connector.Name);
 
 				ConnectionsCategory.AddCustomRow(Label)
 					.NameContent()
@@ -465,7 +469,7 @@ void FRigModuleInstanceDetails::CustomizeDetails(IDetailLayoutBuilder& DetailBui
 								.Visibility(EVisibility::Visible)
 								.BorderImage(FAppStyle::GetBrush("Menu.Background"))
 								[
-									SAssignNew(ConnectionListBox.FindOrAdd(Connector.Name), SSearchableRigHierarchyTreeView)
+									SAssignNew(SearchableTreeView, SSearchableRigHierarchyTreeView)
 										.RigTreeDelegates(TreeDelegates)
 								]
 							]
diff --git a/Engine/Plugins/Animation/ControlRig/Source/ControlRigEditor/Private/Editor/SModularRigModel.cpp b/Engine/Plugins/Animation/ControlRig/Source/ControlRigEditor/Private/Editor/SModularRigModel.cpp
--- a/Engine/Plugins/Animation/ControlRig/Source/ControlRigEditor/Private/Editor/SModularRigModel.cpp
+++ b/Engine/Plugins/Animation/ControlRig/Source/ControlRigEditor/Private/Editor/SModularRigModel.cpp
@@ -1194,9 +1194,9 @@ void SModularRigModel::HandleConnectorResolved(const FRigElementKey& InConnector
 		UModularRigController* Controller = ControlRigBlueprint->GetModularRigController();
 		check(Controller);
 
-		if (ControlRigBeingDebuggedPtr.IsValid())
+		if (const UModularRig* ModularRig = GetModularRig())
 		{
-			Controller->ConnectConnectorToElement(InConnector, InTarget, true, ControlRigBeingDebuggedPtr->GetModularRigSettings().bAutoResolve);
+			Controller->ConnectConnectorToElement(InConnector, InTarget, true, ModularRig->GetModularRigSettings().bAutoResolve);
 		}
 	}
 }
@@ -1350,6 +1350,10 @@ UModularRig* SModularRigModel::GetModularRig() const
 		{
 			return Cast<UModularRig>(DebuggedRig);
 		}
+		if(UControlRig* DebuggedRig = ControlRigBlueprint->GetDebuggedControlRig())
+		{
+			return Cast<UModularRig>(DebuggedRig);
+		}
 	}
 	if (ControlRigEditor.IsValid())
 	{
diff --git a/Engine/Plugins/Animation/ControlRig/Source/ControlRigEditor/Private/Editor/SRigHierarchyTreeView.cpp b/Engine/Plugins/Animation/ControlRig/Source/ControlRigEditor/Private/Editor/SRigHierarchyTreeView.cpp
--- a/Engine/Plugins/Animation/ControlRig/Source/ControlRigEditor/Private/Editor/SRigHierarchyTreeView.cpp
+++ b/Engine/Plugins/Animation/ControlRig/Source/ControlRigEditor/Private/Editor/SRigHierarchyTreeView.cpp
@@ -1396,7 +1396,7 @@ void SSearchableRigHierarchyTreeView::Construct(const FArguments& InArgs)
 		]
 	];
 
-	if (MaxHeight > 0)
+	if (MaxHeight > SMALL_NUMBER)
 	{
 		VerticalBox->GetSlot(1).SetMaxHeight(MaxHeight);
 	}
diff --git a/Engine/Plugins/Animation/ControlRig/Source/ControlRigEditor/Private/Editor/SRigHierarchyTreeView.h b/Engine/Plugins/Animation/ControlRig/Source/ControlRigEditor/Private/Editor/SRigHierarchyTreeView.h
--- a/Engine/Plugins/Animation/ControlRig/Source/ControlRigEditor/Private/Editor/SRigHierarchyTreeView.h
+++ b/Engine/Plugins/Animation/ControlRig/Source/ControlRigEditor/Private/Editor/SRigHierarchyTreeView.h
@@ -387,7 +387,9 @@ class SSearchableRigHierarchyTreeView : public SCompoundWidget
 {
 public:
 
-	SLATE_BEGIN_ARGS(SSearchableRigHierarchyTreeView) {}
+	SLATE_BEGIN_ARGS(SSearchableRigHierarchyTreeView)
+		:_MaxHeight(0.f)
+	{}
 		SLATE_ARGUMENT(FRigTreeDelegates, RigTreeDelegates)
 		SLATE_ARGUMENT(FText, InitialFilterText)
 		SLATE_ARGUMENT(float, MaxHeight)
