Sequencer: NULL check before getting outer crash fix.
https://crashreporter.ol.epicgames.net/crashreporter/crash-groups/5d65469e5f74e841311aa401534456e1

diff --git a/Engine/Plugins/Animation/ControlRig/Source/ControlRigEditor/Private/EditMode/ControlRigEditMode.cpp b/Engine/Plugins/Animation/ControlRig/Source/ControlRigEditor/Private/EditMode/ControlRigEditMode.cpp
--- a/Engine/Plugins/Animation/ControlRig/Source/ControlRigEditor/Private/EditMode/ControlRigEditMode.cpp
+++ b/Engine/Plugins/Animation/ControlRig/Source/ControlRigEditor/Private/EditMode/ControlRigEditMode.cpp
@@ -5842,13 +5842,16 @@ void FDetailKeyFrameCacheAndHandler::OnKeyPropertyClicked(const IPropertyHandle&
 	KeyedPropertyHandle.GetOuterObjects(Objects);
 	for (UObject* Object : Objects)
 	{
-		if (UControlRigControlsProxy* Proxy = Cast< UControlRigControlsProxy>(Object))
+		if (Object)
 		{
-			Proxy->SetKey(SequencerPtr, KeyedPropertyHandle);
-		}
-		else if (UAnimLayer* AnimLayer = (Object->GetTypedOuter<UAnimLayer>()))
-		{
-			AnimLayer->SetKey(SequencerPtr, KeyedPropertyHandle);
+			if (UControlRigControlsProxy* Proxy = Cast< UControlRigControlsProxy>(Object))
+			{
+				Proxy->SetKey(SequencerPtr, KeyedPropertyHandle);
+			}
+			else if (UAnimLayer* AnimLayer = (Object->GetTypedOuter<UAnimLayer>()))
+			{
+				AnimLayer->SetKey(SequencerPtr, KeyedPropertyHandle);
+			}
 		}
 	}
 }
@@ -5889,13 +5892,16 @@ EPropertyKeyedStatus FDetailKeyFrameCacheAndHandler::GetPropertyKeyedStatus(cons
 	
 	for (UObject* Object : OuterObjects)
 	{
-		if (UControlRigControlsProxy* Proxy = Cast< UControlRigControlsProxy>(Object))
+		if (Object)
 		{
-			KeyedStatus = Proxy->GetPropertyKeyedStatus(SequencerPtr,PropertyHandle);
-		}
-		else if (UAnimLayer* AnimLayer = (Object->GetTypedOuter<UAnimLayer>()))
-		{
-			KeyedStatus = AnimLayer->GetPropertyKeyedStatus(SequencerPtr, PropertyHandle);
+			if (UControlRigControlsProxy* Proxy = Cast< UControlRigControlsProxy>(Object))
+			{
+				KeyedStatus = Proxy->GetPropertyKeyedStatus(SequencerPtr, PropertyHandle);
+			}
+			else if (UAnimLayer* AnimLayer = (Object->GetTypedOuter<UAnimLayer>()))
+			{
+				KeyedStatus = AnimLayer->GetPropertyKeyedStatus(SequencerPtr, PropertyHandle);
+			}
 		}
 		//else check to see if it's in sequencer
 	}
