Control Rig: AnimNode Curve exchange loosing values for flipbook faces

#rnx

diff --git a/Engine/Plugins/Animation/ControlRig/Source/ControlRig/Private/AnimNode_ControlRigBase.cpp b/Engine/Plugins/Animation/ControlRig/Source/ControlRig/Private/AnimNode_ControlRigBase.cpp
--- a/Engine/Plugins/Animation/ControlRig/Source/ControlRig/Private/AnimNode_ControlRigBase.cpp
+++ b/Engine/Plugins/Animation/ControlRig/Source/ControlRig/Private/AnimNode_ControlRigBase.cpp
@@ -514,6 +514,7 @@ void FAnimNode_ControlRigBase::UpdateInput(UControlRig* ControlRig, FPoseContext
 		
 		if(InputSettings.bUpdateCurves || OutputSettings.bUpdateCurves)
 		{
+			Hierarchy->UnsetCurveValues();
 			Hierarchy->ResetChangedCurveIndices();
 			
 			const TArray<FRigBaseElement*> HierarchyCurves = Hierarchy->GetCurvesFast();
@@ -527,8 +528,7 @@ void FAnimNode_ControlRigBase::UpdateInput(UControlRig* ControlRig, FPoseContext
 				}
 			}
 
-			PoseAdapter->PoseCurveToHierarchyCurve.Reset();
-			PoseAdapter->PoseCurveToHierarchyCurve.SetNum(InOutput.Curve.Num());
+			PoseAdapter->PoseCurveToHierarchyCurve.SetNumUninitialized(InOutput.Curve.Num());
 
 			int32 CurveIndex = 0;
 			InOutput.Curve.ForEachElement([this, &CurveIndex, &HierarchyCurves](const UE::Anim::FCurveElement& InCurveElement)
@@ -540,7 +540,10 @@ void FAnimNode_ControlRigBase::UpdateInput(UControlRig* ControlRig, FPoseContext
 				{
 					const int32& Index = *IndexPtr;
 					FRigCurveElement* HierarchyCurve = CastChecked<FRigCurveElement>(HierarchyCurves[Index]);
-					HierarchyCurve->Set(InCurveElement.Value, false);
+
+					// when setting the curve we need to mark it as "value set", otherwise the copy
+					// pose may reset it to unset - thus we'll loose the value that was just copied in.
+					HierarchyCurve->Set(InCurveElement.Value, true);
 					PoseAdapter->PoseCurveToHierarchyCurve[CurveIndex] = Index;
 				}
 				CurveIndex++;
