Control Rig: Remap connections after rename / duplicate

#rnx

diff --git a/Engine/Plugins/Animation/ControlRig/Source/ControlRigDeveloper/Private/ControlRigBlueprint.cpp b/Engine/Plugins/Animation/ControlRig/Source/ControlRigDeveloper/Private/ControlRigBlueprint.cpp
--- a/Engine/Plugins/Animation/ControlRig/Source/ControlRigDeveloper/Private/ControlRigBlueprint.cpp
+++ b/Engine/Plugins/Animation/ControlRig/Source/ControlRigDeveloper/Private/ControlRigBlueprint.cpp
@@ -1430,6 +1430,32 @@ void UControlRigBlueprint::HandleConfigureRigVMController(const FRigVMClient* In
 
 #endif
 
+void UControlRigBlueprint::UpdateConnectionMapAfterRename(const FString& InOldNameSpace)
+{
+	const FString OldNameSpace = InOldNameSpace + UModularRig::NamespaceSeparator;
+	const FString NewNameSpace = RigModuleSettings.Identifier.Name + UModularRig::NamespaceSeparator;
+	
+	TMap<FRigElementKey, FRigElementKey> FixedConnectionMap;
+	for(const TPair<FRigElementKey, FRigElementKey>& Pair : ConnectionMap)
+	{
+		auto FixUpConnectionMap = [OldNameSpace, NewNameSpace](const FRigElementKey& InKey) -> FRigElementKey
+		{
+			const FString NameString = InKey.Name.ToString();
+			if(NameString.StartsWith(OldNameSpace, ESearchCase::CaseSensitive))
+			{
+				return FRigElementKey(*(NewNameSpace + NameString.Mid(OldNameSpace.Len())), InKey.Type);
+			}
+			return InKey;
+		};
+
+		const FRigElementKey Key = FixUpConnectionMap(Pair.Key);
+		const FRigElementKey Value = FixUpConnectionMap(Pair.Value);
+		FixedConnectionMap.FindOrAdd(Key) = Value;
+	}
+
+	Swap(ConnectionMap, FixedConnectionMap);
+}
+
 UClass* UControlRigBlueprint::GetRigVMEdGraphNodeClass() const
 {
 	return UControlRigGraphNode::StaticClass();
@@ -1585,7 +1611,9 @@ void UControlRigBlueprint::PostDuplicate(bool bDuplicateForPIE)
 	// update the rig module identifier after save-as or duplicate asset
 	if(IsControlRigModule())
 	{
+		const FString OldNameSpace = RigModuleSettings.Identifier.Name;
 		RigModuleSettings.Identifier.Name = URigHierarchy::GetSanitizedName(FRigName(GetName())).ToString();
+		UpdateConnectionMapAfterRename(OldNameSpace);
 	}
 
 	ModularRigModel.UpdateCachedChildren();
@@ -1599,7 +1627,9 @@ void UControlRigBlueprint::PostRename(UObject* OldOuter, const FName OldName)
 	// update the rig module identifier after renaming the asset
 	if(IsControlRigModule())
 	{
+		const FString OldNameSpace = RigModuleSettings.Identifier.Name; 
 		RigModuleSettings.Identifier.Name = URigHierarchy::GetSanitizedName(FRigName(GetName())).ToString();
+		UpdateConnectionMapAfterRename(OldNameSpace);
 	}
 }
 
diff --git a/Engine/Plugins/Animation/ControlRig/Source/ControlRigDeveloper/Public/ControlRigBlueprint.h b/Engine/Plugins/Animation/ControlRig/Source/ControlRigDeveloper/Public/ControlRigBlueprint.h
--- a/Engine/Plugins/Animation/ControlRig/Source/ControlRigDeveloper/Public/ControlRigBlueprint.h
+++ b/Engine/Plugins/Animation/ControlRig/Source/ControlRigDeveloper/Public/ControlRigBlueprint.h
@@ -397,6 +397,8 @@ private:
 	virtual void HandleConfigureRigVMController(const FRigVMClient* InClient, URigVMController* InControllerToConfigure) override;
 #endif
 
+	void UpdateConnectionMapAfterRename(const FString& InOldNameSpace);
+
 	// Class used to temporarily cache all 
 	// current control values and reapply them
 	// on destruction
