Control Rig: hide animnode optimization work behind a console variable and default to off

#rnx

diff --git a/Engine/Plugins/Animation/ControlRig/Source/ControlRig/Private/AnimNode_ControlRigBase.cpp b/Engine/Plugins/Animation/ControlRig/Source/ControlRig/Private/AnimNode_ControlRigBase.cpp
--- a/Engine/Plugins/Animation/ControlRig/Source/ControlRig/Private/AnimNode_ControlRigBase.cpp
+++ b/Engine/Plugins/Animation/ControlRig/Source/ControlRig/Private/AnimNode_ControlRigBase.cpp
@@ -27,6 +27,9 @@ TAutoConsoleVariable<int32> CVarAnimNodeControlRigDebug(TEXT("a.AnimNode.Control
 // CVar to disable control rig execution within an anim node
 static TAutoConsoleVariable<int32> CVarControlRigDisableExecutionAnimNode(TEXT("ControlRig.DisableExecutionInAnimNode"), 0, TEXT("if nonzero we disable the execution of Control Rigs inside an anim node."));
 
+// CVar to enable performance optimizations within the anim node pose exchange
+static TAutoConsoleVariable<int32> CVarControlRigEnableAnimNodePerformanceOptimizations(TEXT("ControlRig.EnableAnimNodePerformanceOptimizations"), 0, TEXT("if nonzero we enable the (experimental) execution performance optimizations of Control Rig AnimNodes."));
+
 void FAnimNode_ControlRig_PoseAdapter::PostLinked(URigHierarchy* InHierarchy)
 {
 	FRigHierarchyPoseAdapter::PostLinked(InHierarchy);
@@ -360,12 +363,13 @@ FAnimNode_ControlRigBase::FAnimNode_ControlRigBase()
 	, bResetInputPoseToInitial(true) 
 	, bTransferInputPose(true)
 	, bTransferInputCurves(true)
-	, bTransferPoseInGlobalSpace(false)
+	, bTransferPoseInGlobalSpace(CVarControlRigEnableAnimNodePerformanceOptimizations->GetInt() == 0) // default to local in optimized mode
 	, InputSettings(FControlRigIOSettings())
 	, OutputSettings(FControlRigIOSettings())
 	, bExecute(true)
 	, InternalBlendAlpha (1.f)
 	, bControlRigRequiresInitialization(true)
+	, bEnablePoseAdapter(CVarControlRigEnableAnimNodePerformanceOptimizations->GetInt() != 0)
 	, LastBonesSerialNumberForCacheBones(0)
 {
 }
@@ -377,7 +381,10 @@ void FAnimNode_ControlRigBase::OnInitializeAnimInstance(const FAnimInstanceProxy
 	FAnimNode_CustomProperty::OnInitializeAnimInstance(InProxy, InAnimInstance);
 
 	WeakAnimInstanceObject = TWeakObjectPtr<const UAnimInstance>(InAnimInstance);
-	PoseAdapter = MakeShareable(new FAnimNode_ControlRig_PoseAdapter());
+	if(bEnablePoseAdapter)
+	{
+		PoseAdapter = MakeShareable(new FAnimNode_ControlRig_PoseAdapter());
+	}
 
 	USkeletalMeshComponent* Component = InAnimInstance->GetOwningComponent();
 	UControlRig* ControlRig = GetControlRig();
@@ -438,7 +445,7 @@ bool FAnimNode_ControlRigBase::CanExecute()
 		return false;
 	}
 
-	if(!PoseAdapter.IsValid())
+	if(bEnablePoseAdapter && !PoseAdapter.IsValid())
 	{
 		return false;
 	}
@@ -476,79 +483,188 @@ void FAnimNode_ControlRigBase::UpdateInput(UControlRig* ControlRig, FPoseContext
 
 	DECLARE_SCOPE_HIERARCHICAL_COUNTER_FUNC()
 
-	if (PoseAdapter && InputSettings.bUpdatePose)
+	if(bEnablePoseAdapter)
 	{
-		if(InputSettings.bUpdatePose)
+		if (PoseAdapter && InputSettings.bUpdatePose)
 		{
-			// reset transforms here to prevent additive transforms from accumulating to INF
-			// we only update transforms from the mesh pose for bones in the current LOD, 
-			// so the reset here ensures excluded bones are also reset
-			if(!PoseAdapter->BonesToResetToInitial.IsEmpty())
+			if(InputSettings.bUpdatePose)
 			{
-				FRigHierarchyValidityBracket ValidityBracket(Hierarchy);
+				// reset transforms here to prevent additive transforms from accumulating to INF
+				// we only update transforms from the mesh pose for bones in the current LOD, 
+				// so the reset here ensures excluded bones are also reset
+				if(!PoseAdapter->BonesToResetToInitial.IsEmpty())
 				{
+					FRigHierarchyValidityBracket ValidityBracket(Hierarchy);
+					{
 #if WITH_EDITOR
-					// make sure transient controls don't get reset
-					UControlRig::FTransientControlPoseScope PoseScope(ControlRig);
+						// make sure transient controls don't get reset
+						UControlRig::FTransientControlPoseScope PoseScope(ControlRig);
 #endif
-					for(const int32& BoneElementIndex : PoseAdapter->BonesToResetToInitial)
-					{
-						if(FRigTransformElement* BoneElement = Hierarchy->Get<FRigBoneElement>(BoneElementIndex))
+						for(const int32& BoneElementIndex : PoseAdapter->BonesToResetToInitial)
 						{
-							const FTransform InitialLocalTransform = Hierarchy->GetTransform(BoneElement, ERigTransformType::InitialLocal);
-							BoneElement->GetTransform().Set(ERigTransformType::CurrentLocal, InitialLocalTransform);
-							BoneElement->GetDirtyState().MarkClean(ERigTransformType::CurrentLocal);
-							BoneElement->GetDirtyState().MarkDirty(ERigTransformType::CurrentGlobal);
+							if(FRigTransformElement* BoneElement = Hierarchy->Get<FRigBoneElement>(BoneElementIndex))
+							{
+								const FTransform InitialLocalTransform = Hierarchy->GetTransform(BoneElement, ERigTransformType::InitialLocal);
+								BoneElement->GetTransform().Set(ERigTransformType::CurrentLocal, InitialLocalTransform);
+								BoneElement->GetDirtyState().MarkClean(ERigTransformType::CurrentLocal);
+								BoneElement->GetDirtyState().MarkDirty(ERigTransformType::CurrentGlobal);
+							}
 						}
 					}
 				}
+				
+				if(bTransferInputPose && InOutput.Pose.GetNumBones() == PoseAdapter->LocalPose.Num())
+				{
+					PoseAdapter->MarkDependentsDirty();
+					InOutput.Pose.CopyBonesTo(PoseAdapter->LocalPose);
+					PoseAdapter->UpdateDirtyStates(true);
+				}
 			}
 			
-			if(bTransferInputPose && InOutput.Pose.GetNumBones() == PoseAdapter->LocalPose.Num())
+			if(InputSettings.bUpdateCurves || OutputSettings.bUpdateCurves)
 			{
-				PoseAdapter->MarkDependentsDirty();
-				InOutput.Pose.CopyBonesTo(PoseAdapter->LocalPose);
-				PoseAdapter->UpdateDirtyStates(true);
+				Hierarchy->UnsetCurveValues();
+				
+				const TArray<FRigBaseElement*> HierarchyCurves = Hierarchy->GetCurvesFast();
+				if(PoseAdapter->HierarchyCurveLookup.Num() != HierarchyCurves.Num())
+				{
+					PoseAdapter->HierarchyCurveLookup.Reset();
+					PoseAdapter->HierarchyCurveLookup.Reserve(HierarchyCurves.Num());
+					for(int32 Index = 0; Index < HierarchyCurves.Num(); Index++)
+					{
+						PoseAdapter->HierarchyCurveLookup.Add(HierarchyCurves[Index]->GetFName(), Index);
+					}
+				}
+
+				PoseAdapter->PoseCurveToHierarchyCurve.SetNumUninitialized(InOutput.Curve.Num());
+
+				int32 CurveIndex = 0;
+				InOutput.Curve.ForEachElement([this, &CurveIndex, &HierarchyCurves](const UE::Anim::FCurveElement& InCurveElement)
+				{
+					PoseAdapter->PoseCurveToHierarchyCurve[CurveIndex] = INDEX_NONE;
+
+					// the index stored here is the sub index of the curve (the index of the curve within the list of curves)
+					if(const int32* IndexPtr = PoseAdapter->HierarchyCurveLookup.Find(InCurveElement.Name))
+					{
+						const int32& Index = *IndexPtr;
+						FRigCurveElement* HierarchyCurve = CastChecked<FRigCurveElement>(HierarchyCurves[Index]);
+
+						// when setting the curve we need to mark it as "value set", otherwise the copy
+						// pose may reset it to unset - thus we'll loose the value that was just copied in.
+						HierarchyCurve->Set(InCurveElement.Value, true);
+						PoseAdapter->PoseCurveToHierarchyCurve[CurveIndex] = Index;
+					}
+					CurveIndex++;
+				});
 			}
 		}
-		
-		if(InputSettings.bUpdateCurves || OutputSettings.bUpdateCurves)
+	}
+	else
+	{
+		if (InputSettings.bUpdatePose && bTransferInputPose)
 		{
-			Hierarchy->UnsetCurveValues();
-			Hierarchy->ResetChangedCurveIndices();
-			
-			const TArray<FRigBaseElement*> HierarchyCurves = Hierarchy->GetCurvesFast();
-			if(PoseAdapter->HierarchyCurveLookup.Num() != HierarchyCurves.Num())
+			const FBoneContainer& RequiredBones = InOutput.Pose.GetBoneContainer();
+
+			// reset transforms here to prevent additive transforms from accumulating to INF
+			// we only update transforms from the mesh pose for bones in the current LOD, 
+			// so the reset here ensures excluded bones are also reset
+			if(!ControlRigBoneInputMappingByName.IsEmpty() || bResetInputPoseToInitial)
 			{
-				PoseAdapter->HierarchyCurveLookup.Reset();
-				PoseAdapter->HierarchyCurveLookup.Reserve(HierarchyCurves.Num());
-				for(int32 Index = 0; Index < HierarchyCurves.Num(); Index++)
+				FRigHierarchyValidityBracket ValidityBracket(Hierarchy);
+
 				{
-					PoseAdapter->HierarchyCurveLookup.Add(HierarchyCurves[Index]->GetFName(), Index);
+	#if WITH_EDITOR
+					// make sure transient controls don't get reset
+					UControlRig::FTransientControlPoseScope PoseScope(ControlRig);
+	#endif 
+					Hierarchy->ResetPoseToInitial(ERigElementType::Bone);
 				}
 			}
 
-			PoseAdapter->PoseCurveToHierarchyCurve.SetNumUninitialized(InOutput.Curve.Num());
-
-			int32 CurveIndex = 0;
-			InOutput.Curve.ForEachElement([this, &CurveIndex, &HierarchyCurves](const UE::Anim::FCurveElement& InCurveElement)
+			if(bTransferPoseInGlobalSpace || NodeMappingContainer.IsValid())
 			{
-				PoseAdapter->PoseCurveToHierarchyCurve[CurveIndex] = INDEX_NONE;
+				// get component pose from control rig
+				FCSPose<FCompactPose> MeshPoses;
+				// first I need to convert to local pose
+				MeshPoses.InitPose(InOutput.Pose);
 
-				// the index stored here is the sub index of the curve (the index of the curve within the list of curves)
-				if(const int32* IndexPtr = PoseAdapter->HierarchyCurveLookup.Find(InCurveElement.Name))
+				if(!ControlRigBoneInputMappingByIndex.IsEmpty())
+				{
+					for (const TPair<uint16, uint16>& Pair : ControlRigBoneInputMappingByIndex)
+					{
+						const uint16 ControlRigIndex = Pair.Key;
+						const uint16 SkeletonIndex = Pair.Value;
+						
+						FCompactPoseBoneIndex CompactPoseIndex(SkeletonIndex);
+						const FTransform& ComponentTransform = MeshPoses.GetComponentSpaceTransform(CompactPoseIndex);
+						Hierarchy->SetGlobalTransformByIndex(ControlRigIndex, ComponentTransform, false);
+					}
+				}
+				else
 				{
-					const int32& Index = *IndexPtr;
-					FRigCurveElement* HierarchyCurve = CastChecked<FRigCurveElement>(HierarchyCurves[Index]);
+					for (auto Iter = ControlRigBoneInputMappingByName.CreateConstIterator(); Iter; ++Iter)
+					{
+						const FName& Name = Iter.Key();
+						const uint16 Index = Iter.Value();
+						const FRigElementKey Key(Name, ERigElementType::Bone);
+
+						FCompactPoseBoneIndex CompactPoseIndex(Index);
 
-					// when setting the curve we need to mark it as "value set", otherwise the copy
-					// pose may reset it to unset - thus we'll loose the value that was just copied in.
-					HierarchyCurve->Set(InCurveElement.Value, true);
-					PoseAdapter->PoseCurveToHierarchyCurve[CurveIndex] = Index;
+						const FTransform& ComponentTransform = MeshPoses.GetComponentSpaceTransform(CompactPoseIndex);
+						if (NodeMappingContainer.IsValid())
+						{
+							const FTransform& RelativeTransformReverse = NodeMappingContainer->GetSourceToTargetTransform(Name).GetRelativeTransformReverse(ComponentTransform);
+							Hierarchy->SetGlobalTransform(Key, RelativeTransformReverse, false);
+						}
+						else
+						{
+							Hierarchy->SetGlobalTransform(Key, ComponentTransform, false);
+						}
+						
+					}
 				}
-				CurveIndex++;
-			});
+			}
+			else
+			{
+				if(!ControlRigBoneInputMappingByIndex.IsEmpty())
+				{
+					for (const TPair<uint16, uint16>& Pair : ControlRigBoneInputMappingByIndex)
+					{
+						const uint16 ControlRigIndex = Pair.Key;
+						const uint16 SkeletonIndex = Pair.Value;
+						
+						FCompactPoseBoneIndex CompactPoseIndex(SkeletonIndex);
+						const FTransform& LocalTransform = InOutput.Pose[CompactPoseIndex];
+						Hierarchy->SetLocalTransformByIndex(ControlRigIndex, LocalTransform, false);
+					}
+				}
+				else
+				{
+					for (auto Iter = ControlRigBoneInputMappingByName.CreateConstIterator(); Iter; ++Iter)
+					{
+						const FName& Name = Iter.Key();
+						const uint16 SkeletonIndex = Iter.Value();
+						const FRigElementKey Key(Name, ERigElementType::Bone);
+
+						FCompactPoseBoneIndex CompactPoseIndex(SkeletonIndex);
+						const FTransform& LocalTransform = InOutput.Pose[CompactPoseIndex];
+						Hierarchy->SetLocalTransform(Key, LocalTransform, false);
+					}
+				}
+			}
+			
 		}
+
+		if (InputSettings.bUpdateCurves && bTransferInputCurves)
+		{
+			Hierarchy->UnsetCurveValues();
+			
+			InOutput.Curve.ForEachElement([Hierarchy](const UE::Anim::FCurveElement& InCurveElement)
+			{
+				const FRigElementKey Key(InCurveElement.Name, ERigElementType::Curve);
+				Hierarchy->SetCurveValue(Key, InCurveElement.Value);
+			});
+		}		
 	}
 
 #if WITH_EDITOR
@@ -576,83 +692,189 @@ void FAnimNode_ControlRigBase::UpdateOutput(UControlRig* ControlRig, FPoseContex
 
 	DECLARE_SCOPE_HIERARCHICAL_COUNTER_FUNC()
 
-	if (OutputSettings.bUpdatePose)
+	if(bEnablePoseAdapter)
 	{
-		if(PoseAdapter)
+		if (OutputSettings.bUpdatePose)
 		{
-			// make sure the local respectively global transforms are all ready
-			if(PoseAdapter->bTransferInLocalSpace)
-			{
-				PoseAdapter->ConvertToLocalPose();
-			}
-			else
+			if(PoseAdapter)
 			{
-				PoseAdapter->ConvertToGlobalPose();
+				// make sure the local respectively global transforms are all ready
+				if(PoseAdapter->bTransferInLocalSpace)
+				{
+					PoseAdapter->ConvertToLocalPose();
+				}
+				else
+				{
+					PoseAdapter->ConvertToGlobalPose();
+				}
+
+				// if we transfered in global - let's convert back to local
+				if(!PoseAdapter->bTransferInLocalSpace)
+				{
+					if(PoseAdapter->GlobalPose.Num() == InOutput.Pose.GetNumBones())
+					{
+						PoseAdapter->ConvertToLocalPose();
+					}
+				}
+				
+				InOutput.Pose.CopyBonesFrom(PoseAdapter->LocalPose);
 			}
+		}
 
-			// if we transfered in global - let's convert back to local
-			if(!PoseAdapter->bTransferInLocalSpace)
+		if(OutputSettings.bUpdateCurves)
+		{
+			const TArray<int32>& ChangedCurveIndices = Hierarchy->GetChangedCurveIndices();
+			if(ChangedCurveIndices.Num() > 0)
 			{
-				if(PoseAdapter->GlobalPose.Num() == InOutput.Pose.GetNumBones())
+				const TArray<FRigBaseElement*> HierarchyCurves = Hierarchy->GetCurvesFast();
+
+				int32 CurveIndex = 0;
+				int32 CurvesCopied = 0;
+				InOutput.Curve.ForEachElement([this, &CurveIndex, &CurvesCopied, &HierarchyCurves](const UE::Anim::FCurveElement& InCurveElement)
 				{
-					PoseAdapter->ConvertToLocalPose();
+					// the index stored here is the sub index of the curve (the index of the curve within the list of curves)
+					const int32& HierarchyIndex = PoseAdapter->PoseCurveToHierarchyCurve[CurveIndex];
+					if(HierarchyIndex != INDEX_NONE)
+					{
+						const FRigCurveElement* HierarchyCurve = CastChecked<FRigCurveElement>(HierarchyCurves[HierarchyIndex]);
+						if(HierarchyCurve->IsValueSet())
+						{
+							const_cast<UE::Anim::FCurveElement*>(&InCurveElement)->Value = HierarchyCurve->Get();
+						}
+						CurvesCopied++;
+					}
+					CurveIndex++;
+				});
+
+				if(CurvesCopied < HierarchyCurves.Num())
+				{
+					PoseAdapter->HierarchyCurveCopied.SetNumUninitialized(Hierarchy->Num());
+					FMemory::Memzero(PoseAdapter->HierarchyCurveCopied.GetData(), PoseAdapter->HierarchyCurveCopied.GetAllocatedSize());
+					
+					FBlendedCurve ControlRigCurves;
+					ControlRigCurves.Reserve(ChangedCurveIndices.Num());
+					for(const int32& ChangedCurveIndex : ChangedCurveIndices)
+					{
+						if(!PoseAdapter->HierarchyCurveCopied[ChangedCurveIndex])
+						{
+							if(const FRigCurveElement* HierarchyCurve = Hierarchy->Get<FRigCurveElement>(ChangedCurveIndex))
+							{
+								if(HierarchyCurve->IsValueSet())
+								{
+									ControlRigCurves.Add(HierarchyCurve->GetFName(), HierarchyCurve->Get());
+								}
+							}
+							PoseAdapter->HierarchyCurveCopied[ChangedCurveIndex] = true;
+						}
+					}
+
+					InOutput.Curve.Combine(ControlRigCurves);
 				}
 			}
-			
-			InOutput.Pose.CopyBonesFrom(PoseAdapter->LocalPose);
 		}
 	}
-
-	if(OutputSettings.bUpdateCurves)
+	else
 	{
-		const TArray<int32>& ChangedCurveIndices = Hierarchy->GetChangedCurveIndices();
-		if(ChangedCurveIndices.Num() > 0)
+		if (OutputSettings.bUpdatePose)
 		{
-			const TArray<FRigBaseElement*> HierarchyCurves = Hierarchy->GetCurvesFast();
+			// copy output of the rig
+			const FBoneContainer& RequiredBones = InOutput.Pose.GetBoneContainer();
+
+			TMap<FName, uint16>& NameBasedMapping = ControlRigBoneOutputMappingByName;
+			TArray<TPair<uint16, uint16>>& IndexBasedMapping = ControlRigBoneOutputMappingByIndex;
+
+			// if we don't have a different mapping for outputs, use the input mapping
+			if(NameBasedMapping.IsEmpty() && IndexBasedMapping.IsEmpty())
+			{
+				NameBasedMapping = ControlRigBoneInputMappingByName;
+				IndexBasedMapping = ControlRigBoneInputMappingByIndex;
+			}
 
-			int32 CurveIndex = 0;
-			int32 CurvesCopied = 0;
-			InOutput.Curve.ForEachElement([this, &CurveIndex, &CurvesCopied, &HierarchyCurves](const UE::Anim::FCurveElement& InCurveElement)
+			if(bTransferPoseInGlobalSpace || NodeMappingContainer.IsValid())
 			{
-				// the index stored here is the sub index of the curve (the index of the curve within the list of curves)
-				const int32& HierarchyIndex = PoseAdapter->PoseCurveToHierarchyCurve[CurveIndex];
-				if(HierarchyIndex != INDEX_NONE)
+				// get component pose from control rig
+				FCSPose<FCompactPose> MeshPoses;
+				MeshPoses.InitPose(InOutput.Pose);
+
+				if(!IndexBasedMapping.IsEmpty())
 				{
-					const FRigCurveElement* HierarchyCurve = CastChecked<FRigCurveElement>(HierarchyCurves[HierarchyIndex]);
-					if(HierarchyCurve->IsValueSet())
+					for (const TPair<uint16, uint16>& Pair : IndexBasedMapping)
 					{
-						const_cast<UE::Anim::FCurveElement*>(&InCurveElement)->Value = HierarchyCurve->Get();
+						const uint16 ControlRigIndex = Pair.Key;
+						const uint16 SkeletonIndex = Pair.Value;
+
+						FCompactPoseBoneIndex CompactPoseIndex(SkeletonIndex);
+						FTransform ComponentTransform = Hierarchy->GetGlobalTransformByIndex(ControlRigIndex);
+						MeshPoses.SetComponentSpaceTransform(CompactPoseIndex, ComponentTransform);
 					}
-					CurvesCopied++;
 				}
-				CurveIndex++;
-			});
-
-			if(CurvesCopied < HierarchyCurves.Num())
-			{
-				PoseAdapter->HierarchyCurveCopied.SetNumUninitialized(Hierarchy->Num());
-				FMemory::Memzero(PoseAdapter->HierarchyCurveCopied.GetData(), PoseAdapter->HierarchyCurveCopied.GetAllocatedSize());
-				
-				FBlendedCurve ControlRigCurves;
-				ControlRigCurves.Reserve(ChangedCurveIndices.Num());
-				for(const int32& ChangedCurveIndex : ChangedCurveIndices)
+				else
 				{
-					if(!PoseAdapter->HierarchyCurveCopied[ChangedCurveIndex])
+					for (auto Iter = NameBasedMapping.CreateConstIterator(); Iter; ++Iter)
 					{
-						if(const FRigCurveElement* HierarchyCurve = Hierarchy->Get<FRigCurveElement>(ChangedCurveIndex))
+						const FName& Name = Iter.Key();
+						const uint16 SkeletonIndex = Iter.Value();
+						const FRigElementKey Key(Name, ERigElementType::Bone);
+
+						FCompactPoseBoneIndex CompactPoseIndex(SkeletonIndex);
+						FTransform ComponentTransform = Hierarchy->GetGlobalTransform(Key);
+						if (NodeMappingContainer.IsValid())
 						{
-							if(HierarchyCurve->IsValueSet())
-							{
-								ControlRigCurves.Add(HierarchyCurve->GetFName(), HierarchyCurve->Get());
-							}
+							ComponentTransform = NodeMappingContainer->GetSourceToTargetTransform(Name) * ComponentTransform;
 						}
-						PoseAdapter->HierarchyCurveCopied[ChangedCurveIndex] = true;
+
+						MeshPoses.SetComponentSpaceTransform(CompactPoseIndex, ComponentTransform);
 					}
 				}
 
-				InOutput.Curve.Combine(ControlRigCurves);
+				FCSPose<FCompactPose>::ConvertComponentPosesToLocalPosesSafe(MeshPoses, InOutput.Pose);
+				InOutput.Pose.NormalizeRotations();
+			}
+			else
+			{
+				if(!IndexBasedMapping.IsEmpty())
+				{
+					for (const TPair<uint16, uint16>& Pair : IndexBasedMapping)
+					{
+						const uint16 ControlRigIndex = Pair.Key;
+						const uint16 SkeletonIndex = Pair.Value;
+
+						FCompactPoseBoneIndex CompactPoseIndex(SkeletonIndex);
+						FTransform LocalTransform = Hierarchy->GetLocalTransformByIndex(ControlRigIndex);
+						InOutput.Pose[CompactPoseIndex] = LocalTransform;
+					}
+				}
+				else
+				{
+					for (auto Iter = NameBasedMapping.CreateConstIterator(); Iter; ++Iter)
+					{
+						const FName& Name = Iter.Key();
+						const uint16 Index = Iter.Value();
+						const FRigElementKey Key(Name, ERigElementType::Bone);
+
+						FCompactPoseBoneIndex CompactPoseIndex(Index);
+						FTransform LocalTransform = Hierarchy->GetLocalTransform(Key);
+						InOutput.Pose[CompactPoseIndex] = LocalTransform;
+					}
+				}
 			}
 		}
+
+		if (OutputSettings.bUpdateCurves)
+		{
+			FBlendedCurve ControlRigCurves;
+			ControlRigCurves.Reserve(Hierarchy->Num(ERigElementType::Curve));
+			Hierarchy->ForEach<FRigCurveElement>([&ControlRigCurves](const FRigCurveElement* InElement)
+			{
+				if(InElement->IsValueSet())
+				{
+					ControlRigCurves.Add(InElement->GetFName(), InElement->Get());
+				}
+				return true;
+			});
+
+			InOutput.Curve.Combine(ControlRigCurves);
+		}
 	}
 
 #if WITH_EDITOR
@@ -668,7 +890,7 @@ void FAnimNode_ControlRigBase::Evaluate_AnyThread(FPoseContext& Output)
 {
 	DECLARE_SCOPE_HIERARCHICAL_COUNTER_FUNC()
 
-	if(!PoseAdapter.IsValid())
+	if(bEnablePoseAdapter && !PoseAdapter.IsValid())
 	{
 		return;
 	}
@@ -739,12 +961,15 @@ void FAnimNode_ControlRigBase::ExecuteControlRig(FPoseContext& InOutput)
 			return;
 		}
 
-		const uint32 TopologyVersion = Hierarchy->GetTopologyVersion();
-		if(PoseAdapter && PoseAdapter->LastTopologyVersion != TopologyVersion)
+		if(bEnablePoseAdapter)
 		{
-			Hierarchy->UnlinkPoseAdapter();
-			UpdateInputOutputMappingIfRequired(ControlRig, InOutput.Pose.GetBoneContainer());
-			Hierarchy->LinkPoseAdapter(PoseAdapter);
+			const uint32 TopologyVersion = Hierarchy->GetTopologyVersion();
+			if(PoseAdapter && PoseAdapter->LastTopologyVersion != TopologyVersion)
+			{
+				Hierarchy->UnlinkPoseAdapter();
+				UpdateInputOutputMappingIfRequired(ControlRig, InOutput.Pose.GetBoneContainer());
+				Hierarchy->LinkPoseAdapter(PoseAdapter);
+			}
 		}
 
 		// first update input to the system
@@ -813,6 +1038,38 @@ void FAnimNode_ControlRigBase::ExecuteControlRig(FPoseContext& InOutput)
 	}
 }
 
+struct FControlRigControlScope
+{
+	FControlRigControlScope(UControlRig* InControlRig)
+		: ControlRig(InControlRig)
+	{
+		if (ControlRig.IsValid())
+		{
+			URigHierarchy* Hierarchy = ControlRig->GetHierarchy();
+			Hierarchy->ForEach<FRigControlElement>([this, Hierarchy](FRigControlElement* ControlElement) -> bool
+			{
+				ControlValues.Add(ControlElement->GetKey(), Hierarchy->GetControlValueByIndex(ControlElement->GetIndex()));
+				return true; // continue
+			});
+		}
+	}
+
+	~FControlRigControlScope()
+	{
+		if (ControlRig.IsValid())
+		{
+			URigHierarchy* Hierarchy = ControlRig->GetHierarchy();
+			for (const TPair<FRigElementKey, FRigControlValue>& Pair: ControlValues)
+			{
+				Hierarchy->SetControlValue(Pair.Key, Pair.Value);
+			}
+		}
+	}
+
+	TMap<FRigElementKey, FRigControlValue> ControlValues;
+	TWeakObjectPtr<UControlRig> ControlRig;
+};
+
 void FAnimNode_ControlRigBase::CacheBones_AnyThread(const FAnimationCacheBonesContext& Context)
 {
 	DECLARE_SCOPE_HIERARCHICAL_COUNTER_FUNC()
@@ -856,9 +1113,12 @@ void FAnimNode_ControlRigBase::CacheBones_AnyThread(const FAnimationCacheBonesCo
 		
 		LastBonesSerialNumberForCacheBones = BonesSerialNumber;
 
-		if(URigHierarchy* Hierarchy = ControlRig->GetHierarchy())
+		if(bEnablePoseAdapter)
 		{
-			Hierarchy->LinkPoseAdapter(PoseAdapter);
+			if(URigHierarchy* Hierarchy = ControlRig->GetHierarchy())
+			{
+				Hierarchy->LinkPoseAdapter(PoseAdapter);
+			}
 		}
 	}
 }
@@ -870,181 +1130,358 @@ void FAnimNode_ControlRigBase::UpdateInputOutputMappingIfRequired(UControlRig* I
 	{
 		return;
 	}
-	Hierarchy->UnlinkPoseAdapter();
 
-	if(!PoseAdapter.IsValid())
+	if(bEnablePoseAdapter)
 	{
-		return;
-	}
+		Hierarchy->UnlinkPoseAdapter();
 
-	PoseAdapter->ParentPoseIndices.Reset();
-	PoseAdapter->RequiresHierarchyForSpaceConversion.Reset();
-	PoseAdapter->ElementIndexToPoseIndex.Reset();
+		if(!PoseAdapter.IsValid())
+		{
+			return;
+		}
 
-	const int32 NumBonesInPose = InRequiredBones.GetCompactPoseNumBones();
+		PoseAdapter->ParentPoseIndices.Reset();
+		PoseAdapter->RequiresHierarchyForSpaceConversion.Reset();
+		PoseAdapter->ElementIndexToPoseIndex.Reset();
 
-	PoseAdapter->GlobalPose.Reset();
-	PoseAdapter->LocalPose.Reset();
-	PoseAdapter->ParentPoseIndices.Reserve(NumBonesInPose);
-	PoseAdapter->RequiresHierarchyForSpaceConversion.Reserve(NumBonesInPose);
-	PoseAdapter->GlobalPose.AddDefaulted(NumBonesInPose);
-	PoseAdapter->LocalPose.AddDefaulted(NumBonesInPose);
-	PoseAdapter->HierarchyCurveLookup.Reset();
+		const int32 NumBonesInPose = InRequiredBones.GetCompactPoseNumBones();
 
-	PoseAdapter->bTransferInLocalSpace = !(bTransferPoseInGlobalSpace || NodeMappingContainer.IsValid());
+		PoseAdapter->GlobalPose.Reset();
+		PoseAdapter->LocalPose.Reset();
+		PoseAdapter->ParentPoseIndices.Reserve(NumBonesInPose);
+		PoseAdapter->RequiresHierarchyForSpaceConversion.Reserve(NumBonesInPose);
+		PoseAdapter->GlobalPose.AddDefaulted(NumBonesInPose);
+		PoseAdapter->LocalPose.AddDefaulted(NumBonesInPose);
+		PoseAdapter->HierarchyCurveLookup.Reset();
 
-	for(int32 Index = 0; Index < NumBonesInPose; Index++)
-	{
-		PoseAdapter->ParentPoseIndices.Add(INDEX_NONE);
+		PoseAdapter->bTransferInLocalSpace = !(bTransferPoseInGlobalSpace || NodeMappingContainer.IsValid());
 
-		const FCompactPoseBoneIndex ParentBoneIndex = InRequiredBones.GetParentBoneIndex(FCompactPoseBoneIndex(Index));
-		if(ParentBoneIndex.IsValid())
+		for(int32 Index = 0; Index < NumBonesInPose; Index++)
 		{
-			PoseAdapter->ParentPoseIndices[Index] = ParentBoneIndex.GetInt();
-		}
-		PoseAdapter->RequiresHierarchyForSpaceConversion.Add(false);
-	}
-	PoseAdapter->UpdateDirtyStates();
+			PoseAdapter->ParentPoseIndices.Add(INDEX_NONE);
 
-	TArray<int32> MappedBoneElementIndices;
-	if(InRequiredBones.IsValid())
-	{
-		PoseAdapter->ElementIndexToPoseIndex.Reserve(NumBonesInPose);
-		PoseAdapter->PoseIndexToElementIndex.Reserve(NumBonesInPose);
+			const FCompactPoseBoneIndex ParentBoneIndex = InRequiredBones.GetParentBoneIndex(FCompactPoseBoneIndex(Index));
+			if(ParentBoneIndex.IsValid())
+			{
+				PoseAdapter->ParentPoseIndices[Index] = ParentBoneIndex.GetInt();
+			}
+			PoseAdapter->RequiresHierarchyForSpaceConversion.Add(false);
+		}
+		PoseAdapter->UpdateDirtyStates();
 
-		const FReferenceSkeleton* RefSkeleton = &InRequiredBones.GetReferenceSkeleton();
-		if(const USkeleton* Skeleton = InRequiredBones.GetSkeletonAsset())
+		TArray<int32> MappedBoneElementIndices;
+		if(InRequiredBones.IsValid())
 		{
-			RefSkeleton = &Skeleton->GetReferenceSkeleton();
-		}			
+			PoseAdapter->ElementIndexToPoseIndex.Reserve(NumBonesInPose);
+			PoseAdapter->PoseIndexToElementIndex.Reserve(NumBonesInPose);
 
-		// @todo: thread-safe? probably not in editor, but it may not be a big issue in editor
-		if (NodeMappingContainer.IsValid())
-		{
-			// get target to source mapping table - this is reversed mapping table
-			TMap<FName, FName> TargetToSourceMappingTable;
-			NodeMappingContainer->GetTargetToSourceMappingTable(TargetToSourceMappingTable);
+			const FReferenceSkeleton* RefSkeleton = &InRequiredBones.GetReferenceSkeleton();
+			if(const USkeleton* Skeleton = InRequiredBones.GetSkeletonAsset())
+			{
+				RefSkeleton = &Skeleton->GetReferenceSkeleton();
+			}			
 
-			// now fill up node name
-			for (uint16 Index = 0; Index < NumBonesInPose; ++Index)
+			// @todo: thread-safe? probably not in editor, but it may not be a big issue in editor
+			if (NodeMappingContainer.IsValid())
 			{
-				// get bone name, and find reverse mapping
-				const FSkeletonPoseBoneIndex BoneIndex = InRequiredBones.GetSkeletonPoseIndexFromCompactPoseIndex(FCompactPoseBoneIndex(Index));
-				if(BoneIndex.IsValid())
+				// get target to source mapping table - this is reversed mapping table
+				TMap<FName, FName> TargetToSourceMappingTable;
+				NodeMappingContainer->GetTargetToSourceMappingTable(TargetToSourceMappingTable);
+
+				// now fill up node name
+				for (uint16 Index = 0; Index < NumBonesInPose; ++Index)
 				{
-					FName TargetNodeName = RefSkeleton->GetBoneName(BoneIndex.GetInt());
-					FName* SourceName = TargetToSourceMappingTable.Find(TargetNodeName);
-					if (SourceName)
+					// get bone name, and find reverse mapping
+					const FSkeletonPoseBoneIndex BoneIndex = InRequiredBones.GetSkeletonPoseIndexFromCompactPoseIndex(FCompactPoseBoneIndex(Index));
+					if(BoneIndex.IsValid())
 					{
-						const int32 ElementIndex = Hierarchy->GetIndex({*SourceName, ERigElementType::Bone});
-						if(ElementIndex != INDEX_NONE)
+						FName TargetNodeName = RefSkeleton->GetBoneName(BoneIndex.GetInt());
+						FName* SourceName = TargetToSourceMappingTable.Find(TargetNodeName);
+						if (SourceName)
 						{
-							MappedBoneElementIndices.Add(ElementIndex);
-							PoseAdapter->ElementIndexToPoseIndex.Add(static_cast<uint16>(ElementIndex), Index);
-							PoseAdapter->PoseIndexToElementIndex.Add(ElementIndex);
-							PoseAdapter->LocalPose[Index] = Hierarchy->GetLocalTransform(ElementIndex);
-							PoseAdapter->GlobalPose[Index] = Hierarchy->GetGlobalTransform(ElementIndex);
-							continue;
+							const int32 ElementIndex = Hierarchy->GetIndex({*SourceName, ERigElementType::Bone});
+							if(ElementIndex != INDEX_NONE)
+							{
+								MappedBoneElementIndices.Add(ElementIndex);
+								PoseAdapter->ElementIndexToPoseIndex.Add(static_cast<uint16>(ElementIndex), Index);
+								PoseAdapter->PoseIndexToElementIndex.Add(ElementIndex);
+								PoseAdapter->LocalPose[Index] = Hierarchy->GetLocalTransform(ElementIndex);
+								PoseAdapter->GlobalPose[Index] = Hierarchy->GetGlobalTransform(ElementIndex);
+								continue;
+							}
 						}
 					}
+					PoseAdapter->PoseIndexToElementIndex.Add(INDEX_NONE);
 				}
-				PoseAdapter->PoseIndexToElementIndex.Add(INDEX_NONE);
 			}
-		}
-		else
-		{
-			TArray<FName> NodeNames;
-			TArray<FNodeItem> NodeItems;
-			InControlRig->GetMappableNodeData(NodeNames, NodeItems);
-
-			// even if not mapped, we map only node that exists in the controlrig
-			for (uint16 Index = 0; Index < NumBonesInPose; ++Index)
+			else
 			{
-				const FSkeletonPoseBoneIndex BoneIndex = InRequiredBones.GetSkeletonPoseIndexFromCompactPoseIndex(FCompactPoseBoneIndex(Index));
-				if(BoneIndex.IsValid())
+				TArray<FName> NodeNames;
+				TArray<FNodeItem> NodeItems;
+				InControlRig->GetMappableNodeData(NodeNames, NodeItems);
+
+				// even if not mapped, we map only node that exists in the controlrig
+				for (uint16 Index = 0; Index < NumBonesInPose; ++Index)
 				{
-					const FName& BoneName = RefSkeleton->GetBoneName(BoneIndex.GetInt());
-					if (NodeNames.Contains(BoneName))
+					const FSkeletonPoseBoneIndex BoneIndex = InRequiredBones.GetSkeletonPoseIndexFromCompactPoseIndex(FCompactPoseBoneIndex(Index));
+					if(BoneIndex.IsValid())
 					{
-						const int32 ElementIndex = Hierarchy->GetIndex({BoneName, ERigElementType::Bone});
-						if(ElementIndex != INDEX_NONE)
+						const FName& BoneName = RefSkeleton->GetBoneName(BoneIndex.GetInt());
+						if (NodeNames.Contains(BoneName))
 						{
-							MappedBoneElementIndices.Add(ElementIndex);
-							PoseAdapter->ElementIndexToPoseIndex.Add(static_cast<uint16>(ElementIndex), Index);
-							PoseAdapter->PoseIndexToElementIndex.Add(ElementIndex);
-							PoseAdapter->LocalPose[Index] = Hierarchy->GetLocalTransform(ElementIndex);
-							PoseAdapter->GlobalPose[Index] = Hierarchy->GetGlobalTransform(ElementIndex);
-							continue;
+							const int32 ElementIndex = Hierarchy->GetIndex({BoneName, ERigElementType::Bone});
+							if(ElementIndex != INDEX_NONE)
+							{
+								MappedBoneElementIndices.Add(ElementIndex);
+								PoseAdapter->ElementIndexToPoseIndex.Add(static_cast<uint16>(ElementIndex), Index);
+								PoseAdapter->PoseIndexToElementIndex.Add(ElementIndex);
+								PoseAdapter->LocalPose[Index] = Hierarchy->GetLocalTransform(ElementIndex);
+								PoseAdapter->GlobalPose[Index] = Hierarchy->GetGlobalTransform(ElementIndex);
+								continue;
+							}
 						}
 					}
+					PoseAdapter->PoseIndexToElementIndex.Add(INDEX_NONE);
 				}
-				PoseAdapter->PoseIndexToElementIndex.Add(INDEX_NONE);
 			}
-		}
 
-		// once we know all of the bones we are going to transfer - we can check if any of these bones has a different parenting
-		// relationship in the skeleton used in the anim graph vs the hierarchy in the rig. in that case we have to transfer in global
-		if(PoseAdapter->bTransferInLocalSpace)
-		{
-			for(const int32& BoneElementIndex : MappedBoneElementIndices)
+			// once we know all of the bones we are going to transfer - we can check if any of these bones has a different parenting
+			// relationship in the skeleton used in the anim graph vs the hierarchy in the rig. in that case we have to transfer in global
+			if(PoseAdapter->bTransferInLocalSpace)
 			{
-				const int32 HierarchyParentIndex = Hierarchy->GetFirstParent(BoneElementIndex);
-				const int16 PoseIndex = PoseAdapter->ElementIndexToPoseIndex.FindChecked(BoneElementIndex); 
-				const FCompactPoseBoneIndex CompactPoseParentIndex(PoseAdapter->ParentPoseIndices[PoseIndex]);
+				for(const int32& BoneElementIndex : MappedBoneElementIndices)
+				{
+					const int32 HierarchyParentIndex = Hierarchy->GetFirstParent(BoneElementIndex);
+					const int16 PoseIndex = PoseAdapter->ElementIndexToPoseIndex.FindChecked(BoneElementIndex); 
+					const FCompactPoseBoneIndex CompactPoseParentIndex(PoseAdapter->ParentPoseIndices[PoseIndex]);
 
-				FName HierarchyParentName(NAME_None);
-				FName PoseParentName(NAME_None);
+					FName HierarchyParentName(NAME_None);
+					FName PoseParentName(NAME_None);
 				
-				if(HierarchyParentIndex != INDEX_NONE)
-				{
-					HierarchyParentName = Hierarchy->Get(HierarchyParentIndex)->GetFName();
+					if(HierarchyParentIndex != INDEX_NONE)
+					{
+						HierarchyParentName = Hierarchy->Get(HierarchyParentIndex)->GetFName();
+					}
+					if(CompactPoseParentIndex.IsValid())
+					{
+						const FSkeletonPoseBoneIndex SkeletonIndex = InRequiredBones.GetSkeletonPoseIndexFromCompactPoseIndex(CompactPoseParentIndex);
+						PoseParentName = RefSkeleton->GetBoneName(SkeletonIndex.GetInt());
+					}
+
+					if(HierarchyParentName.IsEqual(PoseParentName, ENameCase::CaseSensitive))
+					{
+						continue;
+					}
+				
+					PoseAdapter->RequiresHierarchyForSpaceConversion[PoseIndex] = true;
+					check(PoseAdapter->PoseIndexToElementIndex[PoseIndex] != INDEX_NONE);
+					PoseAdapter->bTransferInLocalSpace = false;
 				}
-				if(CompactPoseParentIndex.IsValid())
+			}
+
+			// only reset the full pose if we are not mapping all bones
+			const TArray<FRigBaseElement*>& HierarchyBones = Hierarchy->GetBonesFast();
+			const bool bMapsAllBones = MappedBoneElementIndices.Num() == HierarchyBones.Num();
+			PoseAdapter->BonesToResetToInitial.Reset();
+			PoseAdapter->bRequiresResetPoseToInitial = bResetInputPoseToInitial && !bMapsAllBones;
+
+			if(PoseAdapter->bRequiresResetPoseToInitial)
+			{
+				PoseAdapter->BonesToResetToInitial.Reserve(HierarchyBones.Num() - MappedBoneElementIndices.Num());
+
+				// bone is mapped stores sub indices (bone index within the list of bones)
+				TArray<bool> BoneIsMapped;
+				BoneIsMapped.AddZeroed(HierarchyBones.Num());
+				for(const int32& MappedTransformIndex : MappedBoneElementIndices)
 				{
-					const FSkeletonPoseBoneIndex SkeletonIndex = InRequiredBones.GetSkeletonPoseIndexFromCompactPoseIndex(CompactPoseParentIndex);
-					PoseParentName = RefSkeleton->GetBoneName(SkeletonIndex.GetInt());
+					const FRigBaseElement* MappedElement = Hierarchy->Get(MappedTransformIndex);
+					check(MappedElement);
+					BoneIsMapped[MappedElement->GetSubIndex()] = true;
 				}
 
-				if(HierarchyParentName.IsEqual(PoseParentName, ENameCase::CaseSensitive))
+				// when we want to know which bones to reset we want to convert back to a global index
+				for(int32 UnmappedBoneIndex = 0; UnmappedBoneIndex < BoneIsMapped.Num(); UnmappedBoneIndex++)
 				{
-					continue;
+					if(!BoneIsMapped[UnmappedBoneIndex])
+					{
+						PoseAdapter->BonesToResetToInitial.Add(HierarchyBones[UnmappedBoneIndex]->GetIndex());
+					}
 				}
-				
-				PoseAdapter->RequiresHierarchyForSpaceConversion[PoseIndex] = true;
-				check(PoseAdapter->PoseIndexToElementIndex[PoseIndex] != INDEX_NONE);
-				PoseAdapter->bTransferInLocalSpace = false;
 			}
 		}
+	}
+	else
+	{
+		ControlRigBoneInputMappingByIndex.Reset();
+		ControlRigBoneOutputMappingByIndex.Reset();
+		ControlRigCurveMappingByIndex.Reset();
+		ControlRigBoneInputMappingByName.Reset();
+		ControlRigBoneOutputMappingByName.Reset();
+		ControlRigCurveMappingByName.Reset();
+
+		if(InRequiredBones.IsValid())
+		{
+			const TArray<FBoneIndexType>& RequiredBonesArray = InRequiredBones.GetBoneIndicesArray();
+			const int32 NumBones = RequiredBonesArray.Num();
 
-		// only reset the full pose if we are not mapping all bones
-		const TArray<FRigBaseElement*>& HierarchyBones = Hierarchy->GetBonesFast();
-		const bool bMapsAllBones = MappedBoneElementIndices.Num() == HierarchyBones.Num();
-		PoseAdapter->BonesToResetToInitial.Reset();
-		PoseAdapter->bRequiresResetPoseToInitial = bResetInputPoseToInitial && !bMapsAllBones;
+			const FReferenceSkeleton& RefSkeleton = InRequiredBones.GetReferenceSkeleton();
 
-		if(PoseAdapter->bRequiresResetPoseToInitial)
-		{
-			PoseAdapter->BonesToResetToInitial.Reserve(HierarchyBones.Num() - MappedBoneElementIndices.Num());
+			// @todo: thread-safe? probably not in editor, but it may not be a big issue in editor
+			if (NodeMappingContainer.IsValid())
+			{
+				// get target to source mapping table - this is reversed mapping table
+				TMap<FName, FName> TargetToSourceMappingTable;
+				NodeMappingContainer->GetTargetToSourceMappingTable(TargetToSourceMappingTable);
+
+				// now fill up node name
+				for (uint16 Index = 0; Index < NumBones; ++Index)
+				{
+					// get bone name, and find reverse mapping
+					FName TargetNodeName = RefSkeleton.GetBoneName(RequiredBonesArray[Index]);
+					FName* SourceName = TargetToSourceMappingTable.Find(TargetNodeName);
+					if (SourceName)
+					{
+						ControlRigBoneInputMappingByName.Add(*SourceName, Index);
+					}
+				}
+			}
+			else
+			{
+				TArray<FName> NodeNames;
+				TArray<FNodeItem> NodeItems;
+				InControlRig->GetMappableNodeData(NodeNames, NodeItems);
+
+				// even if not mapped, we map only node that exists in the controlrig
+				for (uint16 Index = 0; Index < NumBones; ++Index)
+				{
+					const FName& BoneName = RefSkeleton.GetBoneName(RequiredBonesArray[Index]);
+					if (NodeNames.Contains(BoneName))
+					{
+						ControlRigBoneInputMappingByName.Add(BoneName, Index);
+					}
+				}
+			}
+
+			auto UpdatingMappingFromSpecificTransferList = [] (
+				TArray<FBoneReference>& InTransferList,
+				const TWeakObjectPtr<UNodeMappingContainer>& InMappingContainer,
+				const FBoneContainer& InRequiredBones,
+				const FReferenceSkeleton& InRefSkeleton,
+				const TArray<FBoneIndexType>& InRequiredBonesArray,
+				const UControlRig* InControlRig,
+				TMap<FName, uint16>& OutMapping
+			) {
+				OutMapping.Reset();
+				
+				if (InMappingContainer.IsValid())
+				{
+					// get target to source mapping table - this is reversed mapping table
+					TMap<FName, FName> TargetToSourceMappingTable;
+					InMappingContainer->GetTargetToSourceMappingTable(TargetToSourceMappingTable);
+
+					for(FBoneReference& InputBoneToTransfer : InTransferList)
+					{
+						if(!InputBoneToTransfer.Initialize(InRequiredBones))
+						{
+							continue;
+						}
+						const FName TargetNodeName = InRefSkeleton.GetBoneName(InputBoneToTransfer.BoneIndex);
+						if (const FName* SourceName = TargetToSourceMappingTable.Find(TargetNodeName))
+						{
+							OutMapping.Add(*SourceName, InputBoneToTransfer.BoneIndex);
+						}
+					}
+				}
+				else
+				{
+					TArray<FName> NodeNames;
+					TArray<FNodeItem> NodeItems;
+					InControlRig->GetMappableNodeData(NodeNames, NodeItems);
 
-			// bone is mapped stores sub indices (bone index within the list of bones)
-			TArray<bool> BoneIsMapped;
-			BoneIsMapped.AddZeroed(HierarchyBones.Num());
-			for(const int32& MappedTransformIndex : MappedBoneElementIndices)
+					for(FBoneReference& InputBoneToTransfer : InTransferList)
+					{
+						if(!InputBoneToTransfer.Initialize(InRequiredBones))
+						{
+							continue;
+						}
+						if (InRequiredBonesArray.IsValidIndex(InputBoneToTransfer.BoneIndex))
+						{
+							const FName& BoneName = InRefSkeleton.GetBoneName(InRequiredBonesArray[InputBoneToTransfer.BoneIndex]);
+							if (NodeNames.Contains(BoneName))
+							{
+								OutMapping.Add(BoneName, InputBoneToTransfer.BoneIndex);
+							}
+						}
+					}
+				}
+			};
+			
+			if(!InputBonesToTransfer.IsEmpty())
 			{
-				const FRigBaseElement* MappedElement = Hierarchy->Get(MappedTransformIndex);
-				check(MappedElement);
-				BoneIsMapped[MappedElement->GetSubIndex()] = true;
+				ControlRigBoneOutputMappingByName = ControlRigBoneInputMappingByName;
+
+				UpdatingMappingFromSpecificTransferList(
+					InputBonesToTransfer,
+					NodeMappingContainer,
+					InRequiredBones,
+					RefSkeleton,
+					RequiredBonesArray,
+					InControlRig,
+					ControlRigBoneInputMappingByName);
 			}
 
-			// when we want to know which bones to reset we want to convert back to a global index
-			for(int32 UnmappedBoneIndex = 0; UnmappedBoneIndex < BoneIsMapped.Num(); UnmappedBoneIndex++)
+			if(!OutputBonesToTransfer.IsEmpty())
 			{
-				if(!BoneIsMapped[UnmappedBoneIndex])
+				UpdatingMappingFromSpecificTransferList(
+					OutputBonesToTransfer,
+					NodeMappingContainer,
+					InRequiredBones,
+					RefSkeleton,
+					RequiredBonesArray,
+					InControlRig,
+					ControlRigBoneOutputMappingByName);
+			}
+
+			// check if we can switch the bones to an index based mapping.
+			// we can only do that if there is no node mapping container set.
+			if(!NodeMappingContainer.IsValid())
+			{
+				for(int32 InputOutput = 0; InputOutput < 2; InputOutput++)
 				{
-					PoseAdapter->BonesToResetToInitial.Add(HierarchyBones[UnmappedBoneIndex]->GetIndex());
+					bool bIsMappingByIndex = true;
+					TMap<FName, uint16>& NameBasedMapping = InputOutput == 0 ? ControlRigBoneInputMappingByName : ControlRigBoneOutputMappingByName;
+					if(NameBasedMapping.IsEmpty())
+					{
+						continue;
+					}
+					
+					TArray<TPair<uint16, uint16>>& IndexBasedMapping = InputOutput == 0 ? ControlRigBoneInputMappingByIndex : ControlRigBoneOutputMappingByIndex;
+					
+					for (auto Iter = NameBasedMapping.CreateConstIterator(); Iter; ++Iter)
+					{
+						const uint16 SkeletonIndex = Iter.Value();
+						const int32 ControlRigIndex = Hierarchy->GetIndex(FRigElementKey(Iter.Key(), ERigElementType::Bone));
+						if(ControlRigIndex != INDEX_NONE)
+						{
+							IndexBasedMapping.Add(TPair<uint16, uint16>((uint16)ControlRigIndex, SkeletonIndex));
+						}
+						else
+						{
+							bIsMappingByIndex = false;
+						}
+					}
+
+					if(bIsMappingByIndex)
+					{
+						NameBasedMapping.Reset();
+					}
+					else
+					{
+						IndexBasedMapping.Reset();
+					}
 				}
 			}
-		}
+		}		
 	}
 }
 
diff --git a/Engine/Plugins/Animation/ControlRig/Source/ControlRig/Public/AnimNode_ControlRigBase.h b/Engine/Plugins/Animation/ControlRig/Source/ControlRig/Public/AnimNode_ControlRigBase.h
--- a/Engine/Plugins/Animation/ControlRig/Source/ControlRig/Public/AnimNode_ControlRigBase.h
+++ b/Engine/Plugins/Animation/ControlRig/Source/ControlRig/Public/AnimNode_ControlRigBase.h
@@ -197,7 +197,7 @@ protected:
 	 * If this list is empty all bones will be transferred.
 	 */
 	UPROPERTY()
-	TArray<FBoneReference> InputBonesToTransfer_DEPRECATED;
+	TArray<FBoneReference> InputBonesToTransfer;
 
 	/**
 	 * An inclusive list of bones to transfer as part
@@ -205,7 +205,7 @@ protected:
 	 * If this list is empty all bones will be transferred.
 	 */
 	UPROPERTY()
-	TArray<FBoneReference> OutputBonesToTransfer_DEPRECATED;
+	TArray<FBoneReference> OutputBonesToTransfer;
 
 	UPROPERTY(EditAnywhere, AdvancedDisplay, Category = Settings)
 	TArray<TObjectPtr<UAssetUserData>> AssetUserData;
@@ -222,6 +222,20 @@ protected:
 	UPROPERTY(transient)
 	FControlRigIOSettings OutputSettings;
 
+	/** Complete mapping from skeleton to control rig bone index */
+	TArray<TPair<uint16, uint16>> ControlRigBoneInputMappingByIndex;
+	TArray<TPair<uint16, uint16>> ControlRigBoneOutputMappingByIndex;
+
+	/** Complete mapping from skeleton to curve name */
+	TArray<TPair<uint16, FName>> ControlRigCurveMappingByIndex;
+
+	/** Rig Hierarchy bone name to required array index mapping */
+	TMap<FName, uint16> ControlRigBoneInputMappingByName;
+	TMap<FName, uint16> ControlRigBoneOutputMappingByName;
+
+	/** Rig Curve name to Curve mapping */
+	TMap<FName, FName> ControlRigCurveMappingByName;
+	
 	TSharedPtr<FAnimNode_ControlRig_PoseAdapter> PoseAdapter;
 
 	UPROPERTY(transient)
@@ -251,6 +265,7 @@ protected:
 	void UpdateGetAssetUserDataDelegate(UControlRig* InControlRig) const;
 
 	bool bControlRigRequiresInitialization;
+	bool bEnablePoseAdapter;
 	uint16 LastBonesSerialNumberForCacheBones;
 	TWeakObjectPtr<const UAnimInstance> WeakAnimInstanceObject;
 
