Anim Node Control Rig: disable input properties update functions caching by default

#rnx

diff --git a/Engine/Plugins/Animation/ControlRig/Source/ControlRig/Private/AnimNode_ControlRig.cpp b/Engine/Plugins/Animation/ControlRig/Source/ControlRig/Private/AnimNode_ControlRig.cpp
--- a/Engine/Plugins/Animation/ControlRig/Source/ControlRig/Private/AnimNode_ControlRig.cpp
+++ b/Engine/Plugins/Animation/ControlRig/Source/ControlRig/Private/AnimNode_ControlRig.cpp
@@ -596,7 +596,7 @@ namespace AnimNodeLocals
 
 const TCHAR* UseFunctionsName = TEXT("ControlRig.AnimNode.UseFunctions");
 
-static bool bUseFunctions = true;
+static bool bUseFunctions = false;
 static FAutoConsoleVariableRef CVarUseCachedFunctions(
    UseFunctionsName,
    bUseFunctions,
@@ -902,21 +902,24 @@ void FAnimNode_ControlRig::InitializeProperties(const UObject* InSourceInstance,
 		SourceProperties.Add(SourceProperty);
 		DestProperties.Add(nullptr);
 
-        if (TargetHierarchy && SourceProperty)
+		if (AnimNodeLocals::bUseFunctions)
         {
-	        const FName& DestPropertyName = DestPropertyNames[Idx];
-        	if(FRigControlElement* ControlElement = TargetControlRig->FindControl(DestPropertyName))
-        	{
-        		AddControlFunction(ControlElement, SourceProperty, TargetHierarchy);
-        	}
-        	else
-        	{
-        		const FRigVMExternalVariable Variable = TargetControlRig->GetPublicVariableByName(DestPropertyName);
-        		if (Variable.IsValid() && !Variable.bIsReadOnly)
-        		{
-        			AddVariableFunction(Variable, SourceProperty);
-        		}
-        	}
+			if (TargetHierarchy && SourceProperty)
+			{
+				const FName& DestPropertyName = DestPropertyNames[Idx];
+				if(FRigControlElement* ControlElement = TargetControlRig->FindControl(DestPropertyName))
+				{
+					AddControlFunction(ControlElement, SourceProperty, TargetHierarchy);
+				}
+				else
+				{
+					const FRigVMExternalVariable Variable = TargetControlRig->GetPublicVariableByName(DestPropertyName);
+					if (Variable.IsValid() && !Variable.bIsReadOnly)
+					{
+						AddVariableFunction(Variable, SourceProperty);
+					}
+				}
+			}
         }
 	}
 }
@@ -943,7 +946,7 @@ void FAnimNode_ControlRig::PropagateInputProperties(const UObject* InSourceInsta
 		return;
 	}
 
-	if (AnimNodeLocals::bUseFunctions)
+	if (AnimNodeLocals::bUseFunctions && !UpdateFunctions.IsEmpty())
 	{
 		Algo::ForEach(UpdateFunctions, [InSourceInstance](const PropertyUpdateFunction& InFunc)
 		{
