Control Rig: Modular Name Spaces to be considered case insensitive

#rnx

diff --git a/Engine/Plugins/Animation/ControlRig/Source/ControlRig/Private/ModularRig.cpp b/Engine/Plugins/Animation/ControlRig/Source/ControlRig/Private/ModularRig.cpp
--- a/Engine/Plugins/Animation/ControlRig/Source/ControlRig/Private/ModularRig.cpp
+++ b/Engine/Plugins/Animation/ControlRig/Source/ControlRig/Private/ModularRig.cpp
@@ -187,7 +187,7 @@ const FRigConnectorElement* FRigModuleInstance::FindPrimaryConnector() const
 					const FString ModulePath = Hierarchy->GetModulePath(Connector->GetKey());
 					if(!ModulePath.IsEmpty())
 					{
-						if(ModulePath.Equals(MyModulePath, ESearchCase::CaseSensitive))
+						if(ModulePath.Equals(MyModulePath, ESearchCase::IgnoreCase))
 						{
 							PrimaryConnector = Connector;
 							return PrimaryConnector;
@@ -214,7 +214,7 @@ TArray<const FRigConnectorElement*> FRigModuleInstance::FindConnectors() const
 				const FString ModulePath = Hierarchy->GetModulePath(Connector->GetKey());
 				if(!ModulePath.IsEmpty())
 				{
-					if(ModulePath.Equals(MyModulePath, ESearchCase::CaseSensitive))
+					if(ModulePath.Equals(MyModulePath, ESearchCase::IgnoreCase))
 					{
 						Connectors.Add(Connector);
 					}
diff --git a/Engine/Plugins/Animation/ControlRig/Source/ControlRig/Private/ModularRigController.cpp b/Engine/Plugins/Animation/ControlRig/Source/ControlRig/Private/ModularRigController.cpp
--- a/Engine/Plugins/Animation/ControlRig/Source/ControlRig/Private/ModularRigController.cpp
+++ b/Engine/Plugins/Animation/ControlRig/Source/ControlRig/Private/ModularRigController.cpp
@@ -403,7 +403,7 @@ bool UModularRigController::DisconnectConnector_Internal(const FRigElementKey& I
 		TArray<FRigElementKey> ConnectionsToRemove;
 		for (const FModularRigSingleConnection& Connection : Model->Connections)
 		{
-			if (Connection.Connector.Name.ToString().StartsWith(ConnectorModulePath, ESearchCase::CaseSensitive))
+			if (Connection.Connector.Name.ToString().StartsWith(ConnectorModulePath, ESearchCase::IgnoreCase))
 			{
 				ConnectionsToRemove.Add(Connection.Connector);
 			}
@@ -425,7 +425,7 @@ bool UModularRigController::DisconnectConnector_Internal(const FRigElementKey& I
 		{
 			FString OtherConnectorModulePath, OtherConnectorName;
 			(void)URigHierarchy::SplitNameSpace(Connection.Connector.Name.ToString(), &OtherConnectorModulePath, &OtherConnectorName);
-			if (OtherConnectorModulePath.StartsWith(ConnectorModulePath, ESearchCase::CaseSensitive) && OtherConnectorModulePath.Len() > ConnectorModulePath.Len())
+			if (OtherConnectorModulePath.StartsWith(ConnectorModulePath, ESearchCase::IgnoreCase) && OtherConnectorModulePath.Len() > ConnectorModulePath.Len())
 			{
 				ConnectionsToRemove.Add(Connection.Connector);
 			}
@@ -828,7 +828,7 @@ TArray<FString> UModularRigController::GetPossibleBindings(const FString& InModu
 	Model->ForEachModule([this, &PossibleBindings, InModulePath, InVariableName, InvalidModulePrefix](const FRigModuleReference* InModule) -> bool
 	{
 		const FString CurModulePath = InModule->GetPath();
-		if (InModulePath != CurModulePath && !CurModulePath.StartsWith(InvalidModulePrefix, ESearchCase::CaseSensitive))
+		if (InModulePath != CurModulePath && !CurModulePath.StartsWith(InvalidModulePrefix, ESearchCase::IgnoreCase))
 		{
 			if (!InModule->Class.IsValid())
 			{
@@ -895,7 +895,7 @@ bool UModularRigController::CanBindModuleVariable(const FString& InModulePath, c
 			return false;
 		}
 
-		if (SourceModulePath.StartsWith(InModulePath, ESearchCase::CaseSensitive))
+		if (SourceModulePath.StartsWith(InModulePath, ESearchCase::IgnoreCase))
 		{
 			OutErrorMessage = FText::FromString(FString::Printf(TEXT("Cannot bind variable of module %s to a variable of module %s because the source module is a child of the target module"), *InModulePath, *SourceModulePath));
 			return false;
@@ -1179,11 +1179,11 @@ FString UModularRigController::RenameModule(const FString& InModulePath, const F
 		const FString NewNamespace = NewPath + TEXT(":");
 		for (FModularRigSingleConnection& Connection : Model->Connections)
 		{
-			if (Connection.Connector.Name.ToString().StartsWith(OldNamespace, ESearchCase::CaseSensitive))
+			if (Connection.Connector.Name.ToString().StartsWith(OldNamespace, ESearchCase::IgnoreCase))
 			{
 				Connection.Connector.Name = *FString::Printf(TEXT("%s%s"), *NewNamespace, *Connection.Connector.Name.ToString().RightChop(OldNamespace.Len()));
 			}
-			if (Connection.Target.Name.ToString().StartsWith(OldNamespace, ESearchCase::CaseSensitive))
+			if (Connection.Target.Name.ToString().StartsWith(OldNamespace, ESearchCase::IgnoreCase))
 			{
 				Connection.Target.Name = *FString::Printf(TEXT("%s%s"), *NewNamespace, *Connection.Target.Name.ToString().RightChop(OldNamespace.Len()));
 			}
@@ -1263,7 +1263,7 @@ FString UModularRigController::ReparentModule(const FString& InModulePath, const
 	const FRigModuleReference* NewParentModule = FindModule(InNewParentModulePath);
 	const FString PreviousParentPath = Module->ParentPath;
 	const FString ParentPath = (NewParentModule) ? NewParentModule->GetPath() : FString();
-	if(PreviousParentPath.Equals(ParentPath, ESearchCase::CaseSensitive))
+	if(PreviousParentPath.Equals(ParentPath, ESearchCase::IgnoreCase))
 	{
 		return Module->GetPath();
 	}
@@ -1310,11 +1310,11 @@ FString UModularRigController::ReparentModule(const FString& InModulePath, const
 	{
 		for (FModularRigSingleConnection& Connection : Model->Connections)
 		{
-			if (Connection.Connector.Name.ToString().StartsWith(OldPath, ESearchCase::CaseSensitive))
+			if (Connection.Connector.Name.ToString().StartsWith(OldPath, ESearchCase::IgnoreCase))
 			{
 				Connection.Connector.Name = *FString::Printf(TEXT("%s%s"), *NewPath, *Connection.Connector.Name.ToString().RightChop(OldPath.Len()));
 			}
-			if (Connection.Target.Name.ToString().StartsWith(OldPath, ESearchCase::CaseSensitive))
+			if (Connection.Target.Name.ToString().StartsWith(OldPath, ESearchCase::IgnoreCase))
 			{
 				Connection.Target.Name = *FString::Printf(TEXT("%s%s"), *NewPath, *Connection.Target.Name.ToString().RightChop(OldPath.Len()));
 			}
@@ -2089,7 +2089,7 @@ void UModularRigController::UpdateShortNames()
 				}
 			}
 
-			if(!Module.ShortName.Equals(ShortPath, ESearchCase::CaseSensitive))
+			if(!Module.ShortName.Equals(ShortPath, ESearchCase::IgnoreCase))
 			{
 				Module.ShortName = ShortPath;
 				Notify(EModularRigNotification::ModuleShortNameChanged, &Module);
diff --git a/Engine/Plugins/Animation/ControlRig/Source/ControlRig/Private/ModularRigModel.cpp b/Engine/Plugins/Animation/ControlRig/Source/ControlRig/Private/ModularRigModel.cpp
--- a/Engine/Plugins/Animation/ControlRig/Source/ControlRig/Private/ModularRigModel.cpp
+++ b/Engine/Plugins/Animation/ControlRig/Source/ControlRig/Private/ModularRigModel.cpp
@@ -43,7 +43,7 @@ const FRigConnectorElement* FRigModuleReference::FindPrimaryConnector(const URig
 				const FString ModulePath = InHierarchy->GetModulePath(Connector->GetKey());
 				if(!ModulePath.IsEmpty())
 				{
-					if(ModulePath.Equals(MyModulePath, ESearchCase::CaseSensitive))
+					if(ModulePath.Equals(MyModulePath, ESearchCase::IgnoreCase))
 					{
 						return Connector;
 					}
@@ -66,7 +66,7 @@ TArray<const FRigConnectorElement*> FRigModuleReference::FindConnectors(const UR
 			const FString ModulePath = InHierarchy->GetModulePath(Connector->GetKey());
 			if(!ModulePath.IsEmpty())
 			{
-				if(ModulePath.Equals(MyModulePath, ESearchCase::CaseSensitive))
+				if(ModulePath.Equals(MyModulePath, ESearchCase::IgnoreCase))
 				{
 					Connectors.Add(Connector);
 				}
diff --git a/Engine/Plugins/Animation/ControlRig/Source/ControlRig/Private/ModularRigRuleManager.cpp b/Engine/Plugins/Animation/ControlRig/Source/ControlRig/Private/ModularRigRuleManager.cpp
--- a/Engine/Plugins/Animation/ControlRig/Source/ControlRig/Private/ModularRigRuleManager.cpp
+++ b/Engine/Plugins/Animation/ControlRig/Source/ControlRig/Private/ModularRigRuleManager.cpp
@@ -225,12 +225,12 @@ void UModularRigRuleManager::FilterInvalidNameSpaces(FWorkData& InOutWorkData)
 		if(!MatchNameSpace.IsNone())
 		{
 			const FString MatchNameSpaceString = MatchNameSpace.ToString();
-			if(MatchNameSpaceString.Equals(NameSpaceString, ESearchCase::CaseSensitive))
+			if(MatchNameSpaceString.Equals(NameSpaceString, ESearchCase::IgnoreCase))
 			{
 				static const FText CannotConnectWithinNameSpaceMessage = LOCTEXT("CannotConnectWithinNameSpace", "Cannot connect within the same namespace.");
 				Result.SetInvalidTarget(CannotConnectWithinNameSpaceMessage);
 			}
-			else if(MatchNameSpaceString.StartsWith(NameSpaceString, ESearchCase::CaseSensitive))
+			else if(MatchNameSpaceString.StartsWith(NameSpaceString, ESearchCase::IgnoreCase))
 			{
 				static const FText CannotConnectBelowNameSpaceMessage = LOCTEXT("CannotConnectBelowNameSpace", "Cannot connect to element below the connector's namespace.");
 				Result.SetInvalidTarget(CannotConnectBelowNameSpaceMessage);
diff --git a/Engine/Plugins/Animation/ControlRig/Source/ControlRig/Private/Rigs/RigConnectionRules.cpp b/Engine/Plugins/Animation/ControlRig/Source/ControlRig/Private/Rigs/RigConnectionRules.cpp
--- a/Engine/Plugins/Animation/ControlRig/Source/ControlRig/Private/Rigs/RigConnectionRules.cpp
+++ b/Engine/Plugins/Animation/ControlRig/Source/ControlRig/Private/Rigs/RigConnectionRules.cpp
@@ -116,7 +116,7 @@ const FRigConnectorElement* FRigConnectionRuleInput::FindPrimaryConnector(FText*
 			if(Connector->IsPrimary())
 			{
 				const FName ConnectorNameSpace = Hierarchy->GetNameSpaceFName(Connector->GetKey());
-				if(!ConnectorNameSpace.IsNone() && ConnectorNameSpace.IsEqual(ModuleNameSpace, ENameCase::CaseSensitive))
+				if(!ConnectorNameSpace.IsNone() && ConnectorNameSpace.IsEqual(ModuleNameSpace, ENameCase::IgnoreCase))
 				{
 					PrimaryConnector = Connector;
 					return false; // stop the search
@@ -152,7 +152,7 @@ TArray<const FRigConnectorElement*> FRigConnectionRuleInput::FindSecondaryConnec
 			if(Connector->IsSecondary() && Connector->IsOptional() == bOptional)
 			{
 				const FName ConnectorNameSpace = Hierarchy->GetNameSpaceFName(Connector->GetKey());
-				if(!ConnectorNameSpace.IsNone() && ConnectorNameSpace.IsEqual(ModuleNameSpace, ENameCase::CaseSensitive))
+				if(!ConnectorNameSpace.IsNone() && ConnectorNameSpace.IsEqual(ModuleNameSpace, ENameCase::IgnoreCase))
 				{
 					SecondaryConnectors.Add(Connector);
 				}
diff --git a/Engine/Plugins/Animation/ControlRig/Source/ControlRig/Private/Rigs/RigHierarchyController.cpp b/Engine/Plugins/Animation/ControlRig/Source/ControlRig/Private/Rigs/RigHierarchyController.cpp
--- a/Engine/Plugins/Animation/ControlRig/Source/ControlRig/Private/Rigs/RigHierarchyController.cpp
+++ b/Engine/Plugins/Animation/ControlRig/Source/ControlRig/Private/Rigs/RigHierarchyController.cpp
@@ -633,7 +633,7 @@ FRigElementKey URigHierarchyController::AddConnector(FName InName, FRigConnector
 					if(!NameSpace.IsEmpty())
 					{
 						const FString ConnectorNameSpace = Hierarchy->GetNameSpace(Connector->GetKey());
-						if(!ConnectorNameSpace.IsEmpty() && ConnectorNameSpace.Equals(NameSpace, ESearchCase::CaseSensitive))
+						if(!ConnectorNameSpace.IsEmpty() && ConnectorNameSpace.Equals(NameSpace, ESearchCase::IgnoreCase))
 						{
 							static constexpr TCHAR Format[] = TEXT("Cannot add connector '%s' - there already is a primary connector.");
 							ReportAndNotifyErrorf(Format, *InName.ToString());
@@ -2104,7 +2104,7 @@ int32 URigHierarchyController::AddElement(FRigBaseElement* InElementToAdd, FRigB
 
 		if(!CRContext.GetRigModuleNameSpace().IsEmpty())
 		{
-			if(InElementToAdd->GetName().StartsWith(CRContext.GetRigModuleNameSpace(), ESearchCase::CaseSensitive))
+			if(InElementToAdd->GetName().StartsWith(CRContext.GetRigModuleNameSpace(), ESearchCase::IgnoreCase))
 			{
 				Hierarchy->SetNameMetadata(InElementToAdd->Key, URigHierarchy::NameSpaceMetadataName, *CRContext.GetRigModuleNameSpace());
 				Hierarchy->SetNameMetadata(InElementToAdd->Key, URigHierarchy::ModuleMetadataName, *CRContext.GetRigModuleNameSpace().LeftChop(1));
diff --git a/Engine/Plugins/Animation/ControlRig/Source/ControlRig/Private/Units/Execution/RigUnit_RigModules.cpp b/Engine/Plugins/Animation/ControlRig/Source/ControlRig/Private/Units/Execution/RigUnit_RigModules.cpp
--- a/Engine/Plugins/Animation/ControlRig/Source/ControlRig/Private/Units/Execution/RigUnit_RigModules.cpp
+++ b/Engine/Plugins/Animation/ControlRig/Source/ControlRig/Private/Units/Execution/RigUnit_RigModules.cpp
@@ -85,7 +85,7 @@ FRigUnit_IsItemInCurrentNameSpace_Execute()
 	Result = false;
 	if(!CurrentNameSpace.IsEmpty() && !ItemNameSpace.IsEmpty())
 	{
-		Result = ItemNameSpace.Equals(CurrentNameSpace, ESearchCase::CaseSensitive); 
+		Result = ItemNameSpace.Equals(CurrentNameSpace, ESearchCase::IgnoreCase); 
 	}
 }
 
@@ -129,7 +129,7 @@ FRigUnit_GetItemsInNameSpace_Execute()
 					const FName ItemNameSpace = Hierarchy->GetNameSpaceFName(Key);
 					if(!ItemNameSpace.IsNone())
 					{
-						if(ItemNameSpace.IsEqual(NameSpaceName, ENameCase::CaseSensitive))
+						if(ItemNameSpace.IsEqual(NameSpaceName, ENameCase::IgnoreCase))
 						{
 							Collection.AddUnique(Key);
 						}
diff --git a/Engine/Plugins/Animation/ControlRig/Source/ControlRigEditor/Private/ControlRigModuleDetails.cpp b/Engine/Plugins/Animation/ControlRig/Source/ControlRigEditor/Private/ControlRigModuleDetails.cpp
--- a/Engine/Plugins/Animation/ControlRig/Source/ControlRigEditor/Private/ControlRigModuleDetails.cpp
+++ b/Engine/Plugins/Animation/ControlRig/Source/ControlRigEditor/Private/ControlRigModuleDetails.cpp
@@ -764,7 +764,7 @@ FText FRigModuleInstanceDetails::GetName() const
 		{
 			if(const FRigModuleInstance* Module = PerModuleInfos[i].GetModule())
 			{
-				if (!Module->Name.IsEqual(FirstValue, ENameCase::CaseSensitive))
+				if (!Module->Name.IsEqual(FirstValue, ENameCase::IgnoreCase))
 				{
 					bSame = false;
 					break;
@@ -843,7 +843,7 @@ FText FRigModuleInstanceDetails::GetShortName() const
 		{
 			if(const FRigModuleInstance* Module = PerModuleInfos[i].GetModule())
 			{
-				if (!Module->GetShortName().Equals(FirstValue, ESearchCase::CaseSensitive))
+				if (!Module->GetShortName().Equals(FirstValue, ESearchCase::IgnoreCase))
 				{
 					bSame = false;
 					break;
@@ -915,7 +915,7 @@ FText FRigModuleInstanceDetails::GetLongName() const
 		{
 			if(const FRigModuleInstance* Module = PerModuleInfos[i].GetModule())
 			{
-				if (!Module->GetLongName().Equals(FirstValue, ESearchCase::CaseSensitive))
+				if (!Module->GetLongName().Equals(FirstValue, ESearchCase::IgnoreCase))
 				{
 					bSame = false;
 					break;
diff --git a/Engine/Plugins/Animation/ControlRig/Source/ControlRigEditor/Private/Editor/ControlRigSchematicModel.cpp b/Engine/Plugins/Animation/ControlRig/Source/ControlRigEditor/Private/Editor/ControlRigSchematicModel.cpp
--- a/Engine/Plugins/Animation/ControlRig/Source/ControlRigEditor/Private/Editor/ControlRigSchematicModel.cpp
+++ b/Engine/Plugins/Animation/ControlRig/Source/ControlRigEditor/Private/Editor/ControlRigSchematicModel.cpp
@@ -617,7 +617,7 @@ void FControlRigSchematicModel::HandleModularRigModified(EModularRigNotification
 						if(const FSchematicGraphNode* ConnectorNode = FindElementKeyNode(ConnectorKey))
 						{
 							const FString NameString = ConnectorKey.Name.ToString();
-							if(NameString.StartsWith(InModule->GetNamespace(), ESearchCase::CaseSensitive))
+							if(NameString.StartsWith(InModule->GetNamespace(), ESearchCase::IgnoreCase))
 							{
 								const FString LocalName = NameString.Mid(InModule->GetNamespace().Len());
 								if(!LocalName.Contains(UModularRig::NamespaceSeparator))
diff --git a/Engine/Plugins/Animation/ControlRig/Source/ControlRigEditor/Private/Editor/SModularRigModel.cpp b/Engine/Plugins/Animation/ControlRig/Source/ControlRigEditor/Private/Editor/SModularRigModel.cpp
--- a/Engine/Plugins/Animation/ControlRig/Source/ControlRigEditor/Private/Editor/SModularRigModel.cpp
+++ b/Engine/Plugins/Animation/ControlRig/Source/ControlRigEditor/Private/Editor/SModularRigModel.cpp
@@ -1006,7 +1006,7 @@ void SModularRigModel::HandleReresolveModules(const TArray<FString>& InPaths)
 		{
 			FString ModulePath = PathAndConnector;
 			FString ConnectorName;
-			(void)PathAndConnector.Split(TEXT("|"), &ModulePath, &ConnectorName, ESearchCase::CaseSensitive, ESearchDir::FromEnd);
+			(void)PathAndConnector.Split(TEXT("|"), &ModulePath, &ConnectorName, ESearchCase::IgnoreCase, ESearchDir::FromEnd);
 
 			const FRigModuleReference* Module = Controller->Model->FindModule(ModulePath);
 			if (!Module)
@@ -1020,7 +1020,7 @@ void SModularRigModel::HandleReresolveModules(const TArray<FString>& InPaths)
 				// if we are executing this on a primary connector we want to re-resolve all secondaries
 				const FRigConnectorElement* PrimaryConnector = Module->FindPrimaryConnector(Hierarchy);
 				const FName DesiredName = Hierarchy->GetNameMetadata(PrimaryConnector->GetKey(), URigHierarchy::DesiredNameMetadataName, NAME_None);
-				if(!DesiredName.IsNone() && DesiredName.ToString().Equals(ConnectorName, ESearchCase::CaseSensitive))
+				if(!DesiredName.IsNone() && DesiredName.ToString().Equals(ConnectorName, ESearchCase::IgnoreCase))
 				{
 					ConnectorName.Reset();
 				}
@@ -1038,7 +1038,7 @@ void SModularRigModel::HandleReresolveModules(const TArray<FString>& InPaths)
 					else
 					{
 						const FName DesiredName = Hierarchy->GetNameMetadata(Connector->GetKey(), URigHierarchy::DesiredNameMetadataName, NAME_None);
-						if(!DesiredName.IsNone() && DesiredName.ToString().Equals(ConnectorName, ESearchCase::CaseSensitive))
+						if(!DesiredName.IsNone() && DesiredName.ToString().Equals(ConnectorName, ESearchCase::IgnoreCase))
 						{
 							ConnectorKeys.AddUnique(Connector->GetKey());
 							break;
