Sequencer: Animation Layers:  Merging layers work with Overrides, Passthrough action to key the value under the Overide layer, various misc fixes.

diff --git a/Engine/Plugins/Animation/ControlRig/Source/ControlRigEditor/Private/EditMode/SControlRigEditModeTools.cpp b/Engine/Plugins/Animation/ControlRig/Source/ControlRigEditor/Private/EditMode/SControlRigEditModeTools.cpp
--- a/Engine/Plugins/Animation/ControlRig/Source/ControlRigEditor/Private/EditMode/SControlRigEditModeTools.cpp
+++ b/Engine/Plugins/Animation/ControlRig/Source/ControlRigEditor/Private/EditMode/SControlRigEditModeTools.cpp
@@ -1306,8 +1306,8 @@ void SControlRigEditModeTools::CustomizeToolBarPalette(FToolBarBuilder& ToolBarB
 	ToolBarBuilder.AddToolBarButton(
 		FExecuteAction::CreateRaw(OwningToolkit.Pin().Get(), &FControlRigEditModeToolkit::TryInvokeToolkitUI, FControlRigEditModeToolkit::AnimLayerTabName),
 		NAME_None,
-		LOCTEXT("AnimLayers", "Anim Layers"),
-		LOCTEXT("AnimLayersTooltip", "Display anim layers"),
+		LOCTEXT("Layers", "Layers"),
+		LOCTEXT("AnimLayersTooltip", "Display animation layers"),
 		FSlateIcon(TEXT("ControlRigEditorStyle"), TEXT("ControlRig.AnimLayers")),
 		EUserInterfaceActionType::Button
 	);
diff --git a/Engine/Plugins/Animation/ControlRig/Source/ControlRigEditor/Private/Sequencer/AnimLayers/AnimLayers.cpp b/Engine/Plugins/Animation/ControlRig/Source/ControlRigEditor/Private/Sequencer/AnimLayers/AnimLayers.cpp
--- a/Engine/Plugins/Animation/ControlRig/Source/ControlRigEditor/Private/Sequencer/AnimLayers/AnimLayers.cpp
+++ b/Engine/Plugins/Animation/ControlRig/Source/ControlRigEditor/Private/Sequencer/AnimLayers/AnimLayers.cpp
@@ -36,6 +36,7 @@
 #include "ISequencerPropertyKeyedStatus.h"
 #include "BakingAnimationKeySettings.h"
 #include "Algo/Accumulate.h"
+#include "MovieSceneToolHelpers.h"
 
 #define LOCTEXT_NAMESPACE "AnimLayers"
 
@@ -1035,7 +1036,7 @@ double UAnimLayer::GetWeight() const
 							const FFrameNumber CurrentTime = SequencerPtr->GetLocalTime().Time.FloorToFrame();
 							float Value = 0.f;
 							FloatChannel->Evaluate(CurrentTime, Value);
-							if (State.Weight != Value)
+							if (State.Weight != Value || WeightProxy->Weight != Value)
 							{
 								DifferentWeightValue = Value;
 							}
@@ -1191,7 +1192,6 @@ void UAnimLayerWeightProxy::PostEditUndo()
 	}
 }
 #endif
-
 EAnimLayerType UAnimLayer::GetType() const
 {
 	TOptional<EMovieSceneBlendType> CurrentVal;
@@ -1235,6 +1235,64 @@ EAnimLayerType UAnimLayer::GetType() const
 	return (EAnimLayerType)(State.Type);
 }
 
+static void SetDefaultsForOverride(UMovieSceneSection* InSection)
+{
+	if (InSection->IsA<UMovieSceneControlRigParameterSection>())
+	{
+		return; //control rig sections already handle this
+	}
+	TSharedPtr<ISequencer> SequencerPtr = UAnimLayers::GetSequencerFromAsset();
+	if (SequencerPtr.IsValid() == false)
+	{
+		return;
+	}
+	FFrameNumber FrameNumber = SequencerPtr->GetLocalTime().Time.GetFrame();
+	if (UMovieSceneTrack* OwnerTrack = InSection->GetTypedOuter<UMovieSceneTrack>())
+	{
+		TArray<UMovieSceneSection*> TrackSections = OwnerTrack->GetAllSections();
+		int32 SectionIndex = TrackSections.Find((InSection));
+		if (SectionIndex != INDEX_NONE)
+		{
+			InSection->Modify();
+			TrackSections.SetNum(SectionIndex); //this will gives us up to the section 
+			TArray<UMovieSceneSection*> Sections;
+			TArray<UMovieSceneSection*> AbsoluteSections;
+			MovieSceneToolHelpers::SplitSectionsByBlendType(EMovieSceneBlendType::Absolute, TrackSections, Sections, AbsoluteSections);
+			TArrayView<FMovieSceneFloatChannel*> BaseFloatChannels = InSection->GetChannelProxy().GetChannels<FMovieSceneFloatChannel>();
+			TArrayView<FMovieSceneDoubleChannel*> BaseDoubleChannels = InSection->GetChannelProxy().GetChannels<FMovieSceneDoubleChannel>();
+			if (BaseDoubleChannels.Num() > 0)
+			{
+				int32 NumChannels = BaseDoubleChannels.Num();
+				const int32 StartIndex = 0;
+				const int32 EndIndex = NumChannels - 1;
+				TArray<double> ChannelValues = MovieSceneToolHelpers::GetChannelValues<FMovieSceneDoubleChannel,
+					double>(StartIndex,EndIndex, Sections, AbsoluteSections, FrameNumber);
+				for (int32 Index = 0; Index < NumChannels; ++Index)
+				{
+					FMovieSceneDoubleChannel* DoubleChannel = BaseDoubleChannels[Index];
+					const double Value = ChannelValues[Index];
+					DoubleChannel->SetDefault(Value);
+				}
+			}
+			else if (BaseFloatChannels.Num() > 0)
+			{
+				int32 NumChannels = BaseFloatChannels.Num();
+				const int32 StartIndex = 0;
+				const int32 EndIndex = NumChannels - 1;
+				TArray<float> ChannelValues = MovieSceneToolHelpers::GetChannelValues<FMovieSceneFloatChannel,
+					float>(StartIndex, EndIndex, Sections, AbsoluteSections, FrameNumber);
+				for (int32 Index = 0; Index < NumChannels; ++Index)
+				{
+					FMovieSceneFloatChannel* FloatChannel = BaseFloatChannels[Index];
+					const float Value = ChannelValues[Index];
+					FloatChannel->SetDefault(Value);
+				}
+			}
+			SequencerPtr->NotifyMovieSceneDataChanged(EMovieSceneDataChangeType::TrackValueChanged);
+		}
+	}
+}
+
 void UAnimLayer::SetType(EAnimLayerType LayerType)
 {
 	const FScopedTransaction Transaction(LOCTEXT("SetType_Transaction", "Set Type"), !GIsTransacting);
@@ -1256,6 +1314,7 @@ void UAnimLayer::SetType(EAnimLayerType LayerType)
 						break;
 					case EAnimLayerType::Override:
 						SectionItem.Section.Get()->SetBlendType(EMovieSceneBlendType::Override);
+						SetDefaultsForOverride(SectionItem.Section.Get());
 						break;
 					case EAnimLayerType::Base:
 						SectionItem.Section.Get()->SetBlendType(EMovieSceneBlendType::Absolute);
@@ -1263,7 +1322,6 @@ void UAnimLayer::SetType(EAnimLayerType LayerType)
 					}
 				}
 			}
-
 		}
 	}
 }
@@ -1387,46 +1445,30 @@ static void CopySectionIntoAnother(UMovieSceneSection* ToSection, UMovieSceneSec
 	TArray<UMovieSceneSection*> AbsoluteSections;
 	TArray<UMovieSceneSection*> AdditiveSections;
 
-	if (ToSection && ToSection->GetBlendType().Get() == EMovieSceneBlendType::Additive)
-	{
-		AdditiveSections.Add(ToSection);
-	}
-	else if (ToSection && ToSection->GetBlendType().Get() == EMovieSceneBlendType::Absolute)
-	{
-		AbsoluteSections.Add(ToSection);
-	}
-	if (FromSection && FromSection->GetBlendType().Get() == EMovieSceneBlendType::Additive)
-	{
-		AdditiveSections.Add(FromSection);
-	}
-	else if (FromSection && FromSection->GetBlendType().Get() == EMovieSceneBlendType::Absolute)
-	{
-		AbsoluteSections.Add(FromSection);
-	}
-	//if CR or transform we just do certain channels so we avoid baking in the weight channel
-	//if not we do it by property type
-	if (ToSection && ToSection->IsA<UMovieSceneControlRigParameterSection>())
-	{
-		MovieSceneToolHelpers::MergeSections<FMovieSceneFloatChannel>(ToSection, AbsoluteSections, AdditiveSections, Range, true /*bSkipLastChannel*/);
-	}
-	else if (ToSection && ToSection->IsA<UMovieScene3DTransformSection>())
-	{
-		MovieSceneToolHelpers::MergeSections<FMovieSceneDoubleChannel>(ToSection, AbsoluteSections, AdditiveSections, Range, false /*bSkipLastChannel*/);
-	}
-	else if(ToSection && FromSection)
+	AdditiveSections.Add(ToSection);
+	AdditiveSections.Add(FromSection);
+	
+
+	FMovieSceneChannelProxy& ChannelProxy = ToSection->GetChannelProxy();
+	for (const FMovieSceneChannelEntry& Entry : ToSection->GetChannelProxy().GetAllEntries())
 	{
-		FMovieSceneChannelProxy& ChannelProxy = ToSection->GetChannelProxy();
-		for (const FMovieSceneChannelEntry& Entry : ToSection->GetChannelProxy().GetAllEntries())
+		const FName ChannelTypeName = Entry.GetChannelTypeName();
+
+		if (ChannelTypeName == FMovieSceneFloatChannel::StaticStruct()->GetFName())
 		{
-			const FName ChannelTypeName = Entry.GetChannelTypeName();
-			if (ChannelTypeName == FMovieSceneFloatChannel::StaticStruct()->GetFName())
-			{
-				MovieSceneToolHelpers::MergeSections<FMovieSceneFloatChannel>(ToSection, AbsoluteSections, AdditiveSections, Range, false /*bSkipLastChannel*/);
-			}
-			else if (ChannelTypeName == FMovieSceneDoubleChannel::StaticStruct()->GetFName())
-			{
-				MovieSceneToolHelpers::MergeSections<FMovieSceneDoubleChannel>(ToSection, AbsoluteSections, AdditiveSections, Range, false /*bSkipLastChannel*/);
-			}
+			TArrayView<FMovieSceneFloatChannel*> BaseFloatChannels = ToSection->GetChannelProxy().GetChannels<FMovieSceneFloatChannel>();
+			const int StartIndex = 0;
+			const int EndIndex = BaseFloatChannels.Num() - 1;
+			MovieSceneToolHelpers::MergeSections<FMovieSceneFloatChannel>(ToSection, AbsoluteSections, AdditiveSections, 
+				StartIndex, EndIndex, Range);
+		}
+		else if (ChannelTypeName == FMovieSceneDoubleChannel::StaticStruct()->GetFName())
+		{
+			TArrayView<FMovieSceneDoubleChannel*> BaseDoubleChannels = ToSection->GetChannelProxy().GetChannels<FMovieSceneDoubleChannel>();
+			const int StartIndex = 0;
+			const int EndIndex = BaseDoubleChannels.Num() - 1;
+			MovieSceneToolHelpers::MergeSections<FMovieSceneDoubleChannel>(ToSection, AbsoluteSections, AdditiveSections, 
+				StartIndex, EndIndex, Range);
 		}
 	}
 }
@@ -1446,6 +1488,7 @@ int32 UAnimLayers::DuplicateAnimLayer(ISequencer* SequencerPtr, int32 Index)
 			const FScopedTransaction Transaction(LOCTEXT("DuplicateAnimLayer_Transaction", "Duplicate Anim Layer"), !GIsTransacting);
 			Modify();
 			UAnimLayer* NewAnimLayer = NewObject<UAnimLayer>(this, NAME_None, RF_Transactional);
+			NewAnimLayer->SetType(ExistingAnimLayer->GetType());
 			bool bItemAdded = false;
 			for (const TPair<TWeakObjectPtr<UObject>,FAnimLayerItem>& Pair : ExistingAnimLayer->AnimLayerItems)
 			{
@@ -1486,6 +1529,7 @@ int32 UAnimLayers::DuplicateAnimLayer(ISequencer* SequencerPtr, int32 Index)
 												AnimLayerItem.SectionItems.Add(NewSectionItem);
 												NewAnimLayer->AnimLayerItems.Add(ControlRig, AnimLayerItem);
 												SetUpSectionDefaults(SequencerPtr, NewAnimLayer, Track, NewSection, FloatChannel);
+												NewSection->SetBlendType(Section->GetBlendType().Get());
 												TArray<FName> ControlNames;
 												NewSectionItem.AnimLayerSet.Names.GetKeys(ControlNames);
 												SetUpControlRigSection(NewSection, ControlNames);
@@ -1521,8 +1565,9 @@ int32 UAnimLayers::DuplicateAnimLayer(ISequencer* SequencerPtr, int32 Index)
 										AnimLayerItem.SectionItems.Add(NewSectionItem);
 										NewAnimLayer->AnimLayerItems.Add(Pair.Key, AnimLayerItem);
 										SetUpSectionDefaults(SequencerPtr, NewAnimLayer, Track, NewSection, FloatChannel);
+										NewSection->SetBlendType(SectionItem.Section->GetBlendType().Get());
 										//current copy keys
-										CopySectionIntoAnother(NewSection, Section);
+										CopySectionIntoAnother(NewSection, SectionItem.Section.Get());
 									}
 								}
 							}
@@ -1532,9 +1577,11 @@ int32 UAnimLayers::DuplicateAnimLayer(ISequencer* SequencerPtr, int32 Index)
 			}
 			if (bItemAdded)
 			{
-				FString LayerName = FString::Printf(TEXT("Anim Layer %d"), AnimLayers.Num());
+				FString ExistingName = ExistingAnimLayer->GetName().ToString();
+				FString NewLayerName = FString::Printf(TEXT("%s_Duplicate"), *ExistingName);
+
 				FText LayerText;
-				LayerText = LayerText.FromString(LayerName);
+				LayerText = LayerText.FromString(NewLayerName);
 				NewAnimLayer->SetName(LayerText); //need items/sections to be added so we can change their track row names
 				NewIndex = AnimLayers.Add(NewAnimLayer);
 				NewAnimLayer->SetKeyed();
@@ -1586,22 +1633,150 @@ FAnimLayerSectionItem* FAnimLayerItem::FindMatchingSectionItem(UMovieSceneSectio
 	return nullptr;
 }
 
+template<typename ChannelType>
+void AddKeyToChannel(ChannelType* Channel, EMovieSceneKeyInterpolation DefaultInterpolation, const FFrameNumber& FrameNumber, double Value)
+{
+	switch (DefaultInterpolation)
+	{
+	case EMovieSceneKeyInterpolation::Linear:
+		Channel->AddLinearKey(FrameNumber, Value);
+		break;
+	case EMovieSceneKeyInterpolation::Constant:
+		Channel->AddConstantKey(FrameNumber, Value);
+		break;
+	case  EMovieSceneKeyInterpolation::Auto:
+		Channel->AddCubicKey(FrameNumber, Value, ERichCurveTangentMode::RCTM_Auto);
+		break;
+	case  EMovieSceneKeyInterpolation::SmartAuto:
+	default:
+		Channel->AddCubicKey(FrameNumber, Value, ERichCurveTangentMode::RCTM_SmartAuto);
+		break;
+	}
+}
+
+bool UAnimLayers::SetPassthroughKey(ISequencer* InSequencer, int32 InIndex)
+{
+	if (InSequencer == nullptr || InIndex <= 0 || InIndex >= AnimLayers.Num())
+	{
+		return false;
+	}
+	FFrameNumber FrameNumber = InSequencer->GetLocalTime().Time.GetFrame();
+	EMovieSceneKeyInterpolation DefaultInterpolation = InSequencer->GetKeyInterpolation();
+	if (UAnimLayer* AnimLayer = AnimLayers[InIndex].Get())
+	{
+		for (TPair<TWeakObjectPtr<UObject>, FAnimLayerItem>& Pair : AnimLayer->AnimLayerItems)
+		{
+			if (Pair.Key != nullptr)
+			{
+				for (FAnimLayerSectionItem& SectionItem : Pair.Value.SectionItems)
+				{
+					if (SectionItem.Section.IsValid())
+					{
+						if (UMovieSceneTrack* OwnerTrack = SectionItem.Section->GetTypedOuter<UMovieSceneTrack>())
+						{
+							TArray<UMovieSceneSection*> TrackSections = OwnerTrack->GetAllSections();
+							int32 SectionIndex = TrackSections.Find((SectionItem.Section.Get()));
+							if (SectionIndex != INDEX_NONE)
+							{
+								SectionItem.Section->Modify();
+								TrackSections.SetNum(SectionIndex); //this will gives us up to the section 
+								TArray<UMovieSceneSection*> Sections;
+								TArray<UMovieSceneSection*> AbsoluteSections;
+								MovieSceneToolHelpers::SplitSectionsByBlendType(EMovieSceneBlendType::Absolute, TrackSections, Sections, AbsoluteSections);
+								TArrayView<FMovieSceneFloatChannel*> BaseFloatChannels = SectionItem.Section->GetChannelProxy().GetChannels<FMovieSceneFloatChannel>();
+								TArrayView<FMovieSceneDoubleChannel*> BaseDoubleChannels = SectionItem.Section->GetChannelProxy().GetChannels<FMovieSceneDoubleChannel>();
+								if (BaseDoubleChannels.Num() > 0)
+								{
+									if (SectionItem.Section->GetBlendType().IsValid() &&
+										SectionItem.Section->GetBlendType() == EMovieSceneBlendType::Additive)
+									{
+										const double Value = 0.0;
+										for (FMovieSceneDoubleChannel* DoubleChannel : BaseDoubleChannels)
+										{
+											AddKeyToChannel(DoubleChannel, FrameNumber, Value, DefaultInterpolation);
+										}
+									}
+									else if (SectionItem.Section->GetBlendType().IsValid() &&
+										SectionItem.Section->GetBlendType() == EMovieSceneBlendType::Override)
+									{
+										int32 NumChannels = BaseDoubleChannels.Num();
+										const int32 StartIndex = 0;
+										const int32 EndIndex = NumChannels - 1;
+										TArray<double> ChannelValues = MovieSceneToolHelpers::GetChannelValues<FMovieSceneDoubleChannel,
+											double>(StartIndex, EndIndex, Sections,AbsoluteSections,FrameNumber);
+										for (int32 Index = 0; Index < NumChannels; ++Index)
+										{
+											FMovieSceneDoubleChannel* DoubleChannel = BaseDoubleChannels[Index];
+											const double Value = ChannelValues[Index];
+											AddKeyToChannel(DoubleChannel, FrameNumber, Value, DefaultInterpolation);
+										}
+									}
+								}
+								else if (BaseFloatChannels.Num() > 0)
+								{
+									int32 NumChannels = SectionItem.Section->IsA<UMovieSceneControlRigParameterSection>() ?
+										BaseFloatChannels.Num() - 1 : BaseFloatChannels.Num();// skip weight if CR section
+									const int32 StartIndex = 0;
+									const int32 EndIndex = NumChannels - 1;
+									if (SectionItem.Section->GetBlendType().IsValid() &&
+										SectionItem.Section->GetBlendType() == EMovieSceneBlendType::Additive)
+									{
+										const float Value = 0.0;
+										for (int32 Index = 0; Index < NumChannels; ++Index)
+										{
+											FMovieSceneFloatChannel* FloatChannel = BaseFloatChannels[Index];
+											AddKeyToChannel(FloatChannel, FrameNumber, Value, DefaultInterpolation);
+										}
+									}
+									else if (SectionItem.Section->GetBlendType().IsValid() &&
+										SectionItem.Section->GetBlendType() == EMovieSceneBlendType::Override)
+									{
+										TArray<float> ChannelValues = MovieSceneToolHelpers::GetChannelValues<FMovieSceneFloatChannel,
+											float>(StartIndex,EndIndex, Sections,AbsoluteSections, FrameNumber);
+										for (int32 Index = 0; Index < NumChannels; ++Index)
+										{
+											FMovieSceneFloatChannel* FloatChannel = BaseFloatChannels[Index];
+											const float Value = ChannelValues[Index];
+											AddKeyToChannel(FloatChannel, FrameNumber, Value, DefaultInterpolation);
+										}
+									}
+								}
+								InSequencer->NotifyMovieSceneDataChanged(EMovieSceneDataChangeType::TrackValueChanged);
+
+							}
+						}
+						else
+						{
+							return false;
+						}
+					}
+				}
+			}
+		}
+	}
+	else
+	{
+		return false;
+	}
+	return true;
+}
+
 bool UAnimLayers::MergeAnimLayers(ISequencer* InSequencer, const TArray<int32>& Indices, const FBakingAnimationKeySettings* InSettings)
 {
 	if (InSequencer == nullptr)
 	{
 		return false;
 	}
-	UAnimLayer* BaseLayer = nullptr;
 	TArray<UAnimLayer*> LayersToMerge;
 	const FFrameNumber Min = InSettings ? InSettings->StartFrame : TNumericLimits<FFrameNumber>::Lowest();
-	const FFrameNumber Max = InSettings ? InSettings->EndFrame :  TNumericLimits<FFrameNumber>::Max();
-	TRange<FFrameNumber> Range(Min,Max);
+	const FFrameNumber Max = InSettings ? InSettings->EndFrame : TNumericLimits<FFrameNumber>::Max();
+	TRange<FFrameNumber> Range(Min, Max);
 	FScopedTransaction Transaction(LOCTEXT("Merge Anim Layers", "Merge Anim Layers"), !GIsTransacting);
 
 	TArray<int32> SortedIndices = Indices;
+	//we go backwards to the first one
 	SortedIndices.Sort([](const int32& Index1, const int32& Index2) {
-		return  Index1 < Index2; 
+		return  Index1 > Index2;
 		});
 
 	for (int32 Index : SortedIndices)
@@ -1610,27 +1785,22 @@ bool UAnimLayers::MergeAnimLayers(ISequencer* InSequencer, const TArray<int32>&
 		{
 			if (UAnimLayer* AnimLayer = AnimLayers[Index].Get())
 			{
-				if (BaseLayer == nullptr)
-				{
-					BaseLayer = AnimLayer;
-				}
-				else
-				{
-					LayersToMerge.Add(AnimLayer);
-				}
+				LayersToMerge.Add(AnimLayer);
 			}
 		}
 	}
-	if (BaseLayer == nullptr || LayersToMerge.Num() < 1)
+	if (LayersToMerge.Num() < 1)
 	{
-		return false; 
+		return false;
 	}
 	Modify();
-	BaseLayer->Modify();
 
-	for (UAnimLayer* AnimLayer : LayersToMerge)
+	for (int32 Index = 0; Index < LayersToMerge.Num() -1; ++Index)
 	{
-		for (TPair<TWeakObjectPtr<UObject>,FAnimLayerItem>& Pair : AnimLayer->AnimLayerItems)
+		UAnimLayer* BaseLayer = LayersToMerge[Index + 1];
+		UAnimLayer* AnimLayer = LayersToMerge[Index];
+		BaseLayer->Modify();
+		for (TPair<TWeakObjectPtr<UObject>, FAnimLayerItem>& Pair : AnimLayer->AnimLayerItems)
 		{
 			if (Pair.Key != nullptr)
 			{
@@ -1644,36 +1814,30 @@ bool UAnimLayers::MergeAnimLayers(ISequencer* InSequencer, const TArray<int32>&
 							BaseSectionItem = Owner->FindMatchingSectionItem(SectionItem.Section.Get());
 
 						}
-						if(BaseSectionItem && BaseSectionItem->Section.IsValid())
+						if (BaseSectionItem && BaseSectionItem->Section.IsValid())
 						{
 							if (SectionItem.Section->IsActive())//active sections merge them
 							{
-								TArray<UMovieSceneSection*> AbsoluteSections;
-								TArray<UMovieSceneSection*> AdditiveSections;
-								//base section get's added to both
-								AdditiveSections.Add(BaseSectionItem->Section.Get());
-								AbsoluteSections.Add(BaseSectionItem->Section.Get());
-
-								if (SectionItem.Section->GetBlendType().Get() == EMovieSceneBlendType::Additive)
-								{
-									AdditiveSections.Add(SectionItem.Section.Get());
-								}
-								else if (SectionItem.Section->GetBlendType().Get() == EMovieSceneBlendType::Absolute)
-								{
-									AbsoluteSections.Add(SectionItem.Section.Get());
-								}
+								
 								TArrayView<FMovieSceneFloatChannel*> BaseFloatChannels = BaseSectionItem->Section->GetChannelProxy().GetChannels<FMovieSceneFloatChannel>();
 								TArrayView<FMovieSceneDoubleChannel*> BaseDoubleChannels = BaseSectionItem->Section->GetChannelProxy().GetChannels<FMovieSceneDoubleChannel>();
 								if (BaseDoubleChannels.Num() > 0)
 								{
-									MovieSceneToolHelpers::MergeSections<FMovieSceneDoubleChannel>(BaseSectionItem->Section.Get(), AbsoluteSections, AdditiveSections, Range, false /*bSkipLastChannel*/);
+									int32 StartIndex = 0;
+									int32 EndIndex = BaseDoubleChannels.Num() - 1;
+									MovieSceneToolHelpers::MergeSections<FMovieSceneDoubleChannel>(BaseSectionItem->Section.Get(),
+										SectionItem.Section.Get(), StartIndex, EndIndex, Range);
 
 								}
 								else if (BaseFloatChannels.Num() > 0)
 								{
-									MovieSceneToolHelpers::MergeSections<FMovieSceneFloatChannel>(BaseSectionItem->Section.Get(), AbsoluteSections, AdditiveSections, Range, true /*bSkipLastChannel*/);
+									int32 StartIndex = 0;
+									int32 EndIndex = BaseFloatChannels.Num() - 1;
+									MovieSceneToolHelpers::MergeSections<FMovieSceneFloatChannel>(BaseSectionItem->Section.Get(),
+										SectionItem.Section.Get(), StartIndex, EndIndex, Range);
 								}
 							}
+			
 							if (UMovieSceneControlRigParameterSection* CRSection = Cast<UMovieSceneControlRigParameterSection>(BaseSectionItem->Section))
 							{
 								if (SectionItem.AnimLayerSet.Names.Num() > 0)
@@ -1686,7 +1850,6 @@ bool UAnimLayers::MergeAnimLayers(ISequencer* InSequencer, const TArray<int32>&
 										UAnimLayers::SetUpControlRigSection(CRSection, AllControls);
 									}
 								}
-								
 							}
 						}
 						else //okay this object doesn't exist in the first layer we are merging into so we need to move it to the other one 
@@ -1711,14 +1874,28 @@ bool UAnimLayers::MergeAnimLayers(ISequencer* InSequencer, const TArray<int32>&
 						}
 					}
 				}
-				int32 Index = GetAnimLayerIndex(AnimLayer);
-				if (Index != INDEX_NONE)
-				{
-					DeleteAnimLayer(InSequencer, Index);
-				}
 			}
 		}
-		//we now delete the layer that was merged, this will delete sections, etc..
+		if (AnimLayer->GetType() == EAnimLayerType::Override &&
+			(BaseLayer->GetType() == EAnimLayerType::Additive))
+		{
+			BaseLayer->SetType(EAnimLayerType::Override);
+		}
+		int32 LayerIndex = GetAnimLayerIndex(AnimLayer);
+		if (LayerIndex != INDEX_NONE)
+		{
+			DeleteAnimLayer(InSequencer, LayerIndex);
+		}
+	}
+	UAnimLayer* BaseLayer = LayersToMerge[LayersToMerge.Num() - 1];
+	FString Merged(TEXT("Merged"));
+	FString ExistingName = BaseLayer->GetName().ToString();
+	if (ExistingName.Contains(Merged) == false)
+	{
+		FString NewLayerName = FString::Printf(TEXT("%s_Merged"), *ExistingName);
+		FText LayerText;
+		LayerText = LayerText.FromString(NewLayerName);
+		BaseLayer->SetName(LayerText); //need items/sections to be added so we can change their track row names
 	}
 	InSequencer->NotifyMovieSceneDataChanged(EMovieSceneDataChangeType::MovieSceneStructureItemAdded);
 	return true;
@@ -1729,7 +1906,6 @@ void UAnimLayers::AnimLayerListChangedBroadcast()
 	OnAnimLayerListChanged.Broadcast(this);
 }
 
-
 TArray<UMovieSceneSection*> UAnimLayers::GetSelectedLayerSections() const
 {
 	TArray<UMovieSceneSection*> Sections;
@@ -2080,6 +2256,7 @@ void UAnimLayers::SetUpSectionDefaults(ISequencer* SequencerPtr, UAnimLayer* Lay
 		break;
 	case EAnimLayerType::Override:
 		NewSection->SetBlendType(EMovieSceneBlendType::Override);
+		SetDefaultsForOverride(NewSection);
 		break;
 	case EAnimLayerType::Base:
 		NewSection->SetBlendType(EMovieSceneBlendType::Absolute);
@@ -2188,26 +2365,7 @@ TArray<TPair<FFrameNumber, double>>& PercentageDifferences)
 		PercentageDifferences.SetNum(0);
 	}
 }
-template<typename ChannelType>
-void AddKeyToChannel(ChannelType* Channel, EMovieSceneKeyInterpolation DefaultInterpolation, const FFrameNumber& FrameNumber, double Value)
-{
-	switch (DefaultInterpolation)
-	{
-	case EMovieSceneKeyInterpolation::Linear:
-		Channel->AddLinearKey(FrameNumber, Value);
-		break;
-	case EMovieSceneKeyInterpolation::Constant:
-		Channel->AddConstantKey(FrameNumber, Value);
-		break;
-	case  EMovieSceneKeyInterpolation::Auto:
-		Channel->AddCubicKey(FrameNumber, Value, ERichCurveTangentMode::RCTM_Auto);
-		break;
-	case  EMovieSceneKeyInterpolation::SmartAuto:
-	default:
-		Channel->AddCubicKey(FrameNumber, Value, ERichCurveTangentMode::RCTM_SmartAuto);
-		break;
-	}
-}
+
 template<typename ChannelType>
 void AdjustmentBlend(UMovieSceneSection* Section, TArrayView<ChannelType*>& BaseChannels, 
 	TArrayView<ChannelType*>& LayerChannels,ISequencer* InSequencer)
diff --git a/Engine/Plugins/Animation/ControlRig/Source/ControlRigEditor/Private/Sequencer/AnimLayers/AnimLayers.h b/Engine/Plugins/Animation/ControlRig/Source/ControlRigEditor/Private/Sequencer/AnimLayers/AnimLayers.h
--- a/Engine/Plugins/Animation/ControlRig/Source/ControlRigEditor/Private/Sequencer/AnimLayers/AnimLayers.h
+++ b/Engine/Plugins/Animation/ControlRig/Source/ControlRigEditor/Private/Sequencer/AnimLayers/AnimLayers.h
@@ -278,7 +278,8 @@ public:
 	int32 AddAnimLayerFromSelection(ISequencer* InSequencer);
 	void GetAnimLayerStates(TArray<FAnimLayerState>& OutStates);
 	bool MergeAnimLayers(ISequencer* InSequencer, const TArray<int32>& Indices, const FBakingAnimationKeySettings* InSettings);
-	
+	bool SetPassthroughKey(ISequencer* InSequencer, int32 Index);
+
 	//will always blend to base fo rnow
 	bool AdjustmentBlendLayers(ISequencer* InSequencer, int32 LayerIndex);
 
diff --git a/Engine/Plugins/Animation/ControlRig/Source/ControlRigEditor/Private/Sequencer/AnimLayers/SAnimLayers.cpp b/Engine/Plugins/Animation/ControlRig/Source/ControlRigEditor/Private/Sequencer/AnimLayers/SAnimLayers.cpp
--- a/Engine/Plugins/Animation/ControlRig/Source/ControlRigEditor/Private/Sequencer/AnimLayers/SAnimLayers.cpp
+++ b/Engine/Plugins/Animation/ControlRig/Source/ControlRigEditor/Private/Sequencer/AnimLayers/SAnimLayers.cpp
@@ -32,6 +32,7 @@
 #include "LevelEditor.h"
 #include "Modules/ModuleManager.h"
 #include "Sequencer/ControlRigParameterTrackEditor.h"
+#include "Algo/IndexOf.h"
 
 #define LOCTEXT_NAMESPACE "AnimLayers"
 
@@ -62,6 +63,7 @@ public:
 	void RemoveSelected() const;
 	void Duplicate() const;
 	void DeleteAnimLayer() const;
+	void SetPassthroughKey() const;
 
 	ECheckBoxState GetKeyed() const;
 	void SetKeyed();
@@ -184,6 +186,26 @@ void FAnimLayerSourceUIEntry::Duplicate() const
 	}
 }
 
+void FAnimLayerSourceUIEntry::SetPassthroughKey() const
+{
+	if (AnimLayer)
+	{
+		if (TSharedPtr<ISequencer> SequencerPtr = UAnimLayers::GetSequencerFromAsset())
+		{
+			if (UAnimLayers* AnimLayers = UAnimLayers::GetAnimLayers(SequencerPtr.Get()))
+			{
+				int32 Index = AnimLayers->GetAnimLayerIndex(AnimLayer);
+				if (Index != INDEX_NONE)
+				{
+					AnimLayers->SetPassthroughKey(SequencerPtr.Get(), Index);
+				}
+			}
+		}
+		ClearCaches();
+	}
+}
+
+
 ECheckBoxState FAnimLayerSourceUIEntry::GetKeyed() const
 {
 	if (bKeyedStateIsValid == false)
@@ -433,13 +455,13 @@ private:
 
 	mutable int32  FocusOnIndex = INDEX_NONE;
 
-
 	void AddSelected();
 	void RemoveSelected();
 	void SelectObjects();
 	void Duplicate();
 	void MergeLayers();
 	void AdjustmentBlend();
+	void SetPassthroughKey();
 	void DeleteAnimLayer();
 	void Rename();
 
@@ -897,10 +919,10 @@ void FAnimLayerSourcesView::CreateSourcesListView()
 			.FillWidth(120.f)
 			.DefaultLabel(LOCTEXT("LayerColumnName", "Name"))
 			+ SHeaderRow::Column(AnimLayerSourceListUI::StatusColumnName)
-			.FillWidth(20.f)
+			.FillWidth(40.f)
 			.DefaultLabel(LOCTEXT("StatusColumnName", "Status"))
 			+ SHeaderRow::Column(AnimLayerSourceListUI::WeightColumnName)
-			.FillWidth(60.f)
+			.FillWidth(80.f)
 			.DefaultLabel(LOCTEXT("WeightColumnName", "Weight"))
 			+ SHeaderRow::Column(AnimLayerSourceListUI::TypeColumnName)
 			.FillWidth(80.f)
@@ -1021,6 +1043,16 @@ void FAnimLayerSourcesView::MergeLayers()
 	}
 }
 
+void FAnimLayerSourcesView::SetPassthroughKey()
+{
+	TArray<FAnimLayerSourceUIEntryPtr> Selected;
+	SourcesListView->GetSelectedItems(Selected);
+	for (const FAnimLayerSourceUIEntryPtr& Ptr : Selected)
+	{
+		Ptr->SetPassthroughKey();
+	}
+}
+
 void FAnimLayerSourcesView::Rename()
 {
 	TArray<FAnimLayerSourceUIEntryPtr> Selected;
@@ -1049,7 +1081,27 @@ TSharedPtr<SWidget> FAnimLayerSourcesView::OnSourceConstructContextMenu()
 
 			TArray<FAnimLayerSourceUIEntryPtr> Selected;
 			SourcesListView->GetSelectedItems(Selected);
-			if (Selected.Num() > 0)
+			int32 BaseLayerIndex = Algo::IndexOfByPredicate(Selected, [AnimLayers](const FAnimLayerSourceUIEntryPtr& Key)
+				{
+					return (Key.IsValid() && Key->GetAnimLayerIndex(AnimLayers) == 0);
+				});
+			//if we have a base layer selected only show Merge 
+			if (BaseLayerIndex != INDEX_NONE)
+			{
+				if (Selected.Num() > 1)
+				{
+					MenuBuilder.BeginSection("AnimLayerContextMenuLayer", LOCTEXT("AnimLayerContextMenuLayer", "Layer"));
+
+					FUIAction Action = FUIAction(FExecuteAction::CreateRaw((this), &FAnimLayerSourcesView::MergeLayers));
+					const FText Label = LOCTEXT("MergeLayers", "Merge Layers");
+					const FText ToolTipText = LOCTEXT("MergeLayerstooltip", "Merge selected layers");
+					MenuBuilder.AddMenuEntry(Label, ToolTipText, FSlateIcon(), Action);
+				}
+				MenuBuilder.AddMenuSeparator();
+				return MenuBuilder.MakeWidget();
+
+			}
+			else if (Selected.Num() > 0)
 			{
 				MenuBuilder.BeginSection("AnimLayerContextMenuLayer", LOCTEXT("AnimLayerContextMenuLayer", "Layer"));
 				{
@@ -1086,6 +1138,12 @@ TSharedPtr<SWidget> FAnimLayerSourcesView::OnSourceConstructContextMenu()
 					const FText ToolTipText = LOCTEXT("MergeLayerstooltip", "Merge selected layers");
 					MenuBuilder.AddMenuEntry(Label, ToolTipText, FSlateIcon(), Action);
 				}
+				{
+					FUIAction Action = FUIAction(FExecuteAction::CreateRaw((this), &FAnimLayerSourcesView::SetPassthroughKey));
+					const FText Label = LOCTEXT("SetPassthroughKey", "Passthrough Key");
+					const FText ToolTipText = LOCTEXT("SetPassthroughKeytooltip", "Set zero key(Additive) or previous value(Override)");
+					MenuBuilder.AddMenuEntry(Label, ToolTipText, FSlateIcon(), Action);
+				}
 				/* adjustment blending not in yet
 				if (Selected.Num() == 1 && Selected[0]->GetAnimLayerIndex(AnimLayers) != 0)
 				{
diff --git a/Engine/Source/Editor/EditorStyle/Private/StarshipStyle.cpp b/Engine/Source/Editor/EditorStyle/Private/StarshipStyle.cpp
--- a/Engine/Source/Editor/EditorStyle/Private/StarshipStyle.cpp
+++ b/Engine/Source/Editor/EditorStyle/Private/StarshipStyle.cpp
@@ -2371,6 +2371,7 @@ void FStarshipEditorStyle::FStyle::SetupSequencerStyles()
 		Set("EMovieSceneBlendType::Absolute", new IMAGE_BRUSH("Sequencer/EMovieSceneBlendType_Absolute", FVector2D(32, 16)));
 		Set("EMovieSceneBlendType::Relative", new IMAGE_BRUSH("Sequencer/EMovieSceneBlendType_Relative", FVector2D(32, 16)));
 		Set("EMovieSceneBlendType::Additive", new IMAGE_BRUSH("Sequencer/EMovieSceneBlendType_Additive", FVector2D(32, 16)));
+		Set("EMovieSceneBlendType::Override", new IMAGE_BRUSH_SVG("Sequencer/EMovieSceneBlendType_Override", FVector2D(32, 16)));
 		Set("EMovieSceneBlendType::AdditiveFromBase", new IMAGE_BRUSH("Sequencer/EMovieSceneBlendType_AdditiveFromBase", FVector2D(32, 16)));
 
 		Set("Sequencer.TrackIsolate", new IMAGE_BRUSH_SVG("Sequencer/TrackIsolate", Icon16x16));
diff --git a/Engine/Source/Editor/MovieSceneTools/Private/MovieSceneToolHelpers.cpp b/Engine/Source/Editor/MovieSceneTools/Private/MovieSceneToolHelpers.cpp
--- a/Engine/Source/Editor/MovieSceneTools/Private/MovieSceneToolHelpers.cpp
+++ b/Engine/Source/Editor/MovieSceneTools/Private/MovieSceneToolHelpers.cpp
@@ -5055,106 +5055,75 @@ bool MovieSceneToolHelpers::OptimizeSection(const FKeyDataOptimizationParams& In
 	}
 	return true;
 }
+void MovieSceneToolHelpers::SplitSectionsByBlendType(EMovieSceneBlendType BlendType, const TArray<UMovieSceneSection*>& InSections, TArray<UMovieSceneSection*>& Sections, TArray<UMovieSceneSection*>& OutBlendSections)
+{
+	for (UMovieSceneSection* InSection : InSections)
+	{
+		if (InSection &&
+			InSection->IsActive() &&
+			InSection->GetBlendType().IsValid())
+		{
+			if (InSection->GetBlendType().Get() == BlendType)
+			{
+				OutBlendSections.Add(InSection);
+			}
+			else 
+			{
+				Sections.Add(InSection);
+			}
+		}
+	}
+}
 
 bool MovieSceneToolHelpers::CollapseSection(TSharedPtr<ISequencer>& SequencerPtr, UMovieSceneTrack* OwnerTrack, TArray<UMovieSceneSection*> Sections,
 	const FBakingAnimationKeySettings& InSettings)
 {
-	if (SequencerPtr.IsValid() && Sections.Num() > 1)
+	if (SequencerPtr.IsValid() && Sections.Num() > 1 && OwnerTrack)
 	{
 		TRange<FFrameNumber> Range(InSettings.StartFrame, InSettings.EndFrame);
 		//we get the  first absolute section or if no absolute sections, first additive,
 		// that's the one we collapse onto
 		UMovieSceneSection* BaseSection = nullptr;
-		//get first absolute
-		for (UMovieSceneSection* Section: Sections)
-		{
-			if (Section->IsActive() == false)
-			{
-				continue;
-			}
-			if (Section->GetBlendType().Get() == EMovieSceneBlendType::Absolute)
-			{
-				BaseSection = Section;
-				break;
-			}
-			else if (Section->GetBlendType().Get() == EMovieSceneBlendType::Additive)
-			{
-				if (BaseSection == nullptr)
-				{
-					BaseSection = Section; //don't stop though may have an absolute later
-				}
-			}
-		}
-		if (BaseSection == nullptr)
-		{
-			UE_LOG(LogMovieScene, Warning, TEXT("CollapseSection:: Invalid section(s) to collapse"));
-			return false;
-		}
-		//now find other sections.
-		TArray<UMovieSceneSection*> AbsoluteSections;
-		TArray<UMovieSceneSection*> AdditiveSections;
-		for (UMovieSceneSection* Section : Sections)
+		
+		FScopedTransaction Transaction(NSLOCTEXT("MovieSceneTools", "CollapseAllSections", "Collapse All Sections"));
+		for (int32 SectionIndex = Sections.Num() - 1; SectionIndex > 0; --SectionIndex)
 		{
-			if (Section != BaseSection && Section->IsActive())
+			BaseSection = Sections[SectionIndex - 1];
+			UMovieSceneSection* Section = Sections[SectionIndex];
+			if (BaseSection && Section && BaseSection->GetBlendType().IsValid() && Section->GetBlendType().IsValid())
 			{
-				if (Section->GetBlendType().Get() == EMovieSceneBlendType::Absolute)
+				OwnerTrack->Modify();
+
+				if (Section->IsActive())//active sections merge them
 				{
-					AbsoluteSections.Add(Section);
+					TArrayView<FMovieSceneFloatChannel*> BaseFloatChannels = BaseSection->GetChannelProxy().GetChannels<FMovieSceneFloatChannel>();
+					TArrayView<FMovieSceneDoubleChannel*> BaseDoubleChannels = BaseSection->GetChannelProxy().GetChannels<FMovieSceneDoubleChannel>();
+					if (BaseDoubleChannels.Num() > 0)
+					{
+						const int StartIndex = 0;
+						const int EndIndex = BaseDoubleChannels.Num() - 1;
+						MovieSceneToolHelpers::MergeSections<FMovieSceneDoubleChannel>(BaseSection,
+							Section, StartIndex, EndIndex, Range);
+
+					}
+					else if (BaseFloatChannels.Num() > 0)
+					{
+						const int StartIndex = 0;
+						const int EndIndex = BaseFloatChannels.Num() - 1;
+						MovieSceneToolHelpers::MergeSections<FMovieSceneFloatChannel>(BaseSection,
+							Section,StartIndex, EndIndex, Range);
+					}
 				}
-				else if (Section->GetBlendType().Get() == EMovieSceneBlendType::Additive)
+
+				//if merging override onto an additive we change it to an additive
+				if (Section->GetBlendType().Get() == EMovieSceneBlendType::Override &&
+					(BaseSection->GetBlendType().Get() == EMovieSceneBlendType::Additive))
 				{
-					AdditiveSections.Add(Section);
+					BaseSection->SetBlendType(EMovieSceneBlendType::Override);
 				}
-			}
-		}
-		//now make sure we have sections to blend with and if so add base section to them
-		if (AbsoluteSections.Num() > 0 || AdditiveSections.Num() > 0)
-		{
-			if (AbsoluteSections.Num() > 0)
-			{
-				AbsoluteSections.Insert(BaseSection, 0);
-			}
-			if (AdditiveSections.Num() > 0)
-			{
-				AdditiveSections.Insert(BaseSection, 0);
-			}
-		}
-		else
-		{
-			UE_LOG(LogMovieScene, Warning, TEXT("CollapseSection:: Invalid section(s) to collapse"));
-			return false;
-		}
-
-		FScopedTransaction Transaction(NSLOCTEXT("MovieSceneTools", "CollapseAllSections", "Collapse All Sections"));
-
-		TArrayView<FMovieSceneFloatChannel*> BaseFloatChannels = BaseSection->GetChannelProxy().GetChannels<FMovieSceneFloatChannel>();
-		TArrayView<FMovieSceneDoubleChannel*> BaseDoubleChannels = BaseSection->GetChannelProxy().GetChannels<FMovieSceneDoubleChannel>();
-		if (BaseDoubleChannels.Num() > 0) //transforms
-		{
-			if (MovieSceneToolHelpers::MergeSections<FMovieSceneDoubleChannel>(BaseSection, AbsoluteSections, AdditiveSections, Range, false /*bSkipLastChannel*/) == false)
-			{
-				Transaction.Cancel();
-				return false;
-			}
-		}
-		else if(BaseFloatChannels.Num() > 0) //control rig
-		{
-			//skip weight channel for control rig
-			if (MovieSceneToolHelpers::MergeSections<FMovieSceneFloatChannel>(BaseSection,  AbsoluteSections, AdditiveSections, Range, true /*bSkipLastChannel*/) == false)
-			{
-				Transaction.Cancel();
-				return false;
-			}
-		}
-		//delete other sections
-		OwnerTrack->Modify();
-		for (int32 SectionIndex = Sections.Num() - 1; SectionIndex >= 0; --SectionIndex)
-		{
-			if (Sections[SectionIndex] != BaseSection)
-			{
 				TArray<UMovieSceneSection*> AllSections = OwnerTrack->GetAllSections();
 				int32 TrackSectionIndex = INDEX_NONE;
-				if (AllSections.Find(Sections[SectionIndex], TrackSectionIndex))
+				if (AllSections.Find(Section, TrackSectionIndex))
 				{
 					if (TrackSectionIndex != INDEX_NONE)
 					{
@@ -5163,7 +5132,7 @@ bool MovieSceneToolHelpers::CollapseSection(TSharedPtr<ISequencer>& SequencerPtr
 				}
 			}
 		}
-		if (InSettings.bReduceKeys)
+		if (InSettings.bReduceKeys && BaseSection)
 		{
 			FKeyDataOptimizationParams Params;
 			Params.bAutoSetInterpolation = true;
diff --git a/Engine/Source/Editor/MovieSceneTools/Public/MovieSceneToolHelpers.h b/Engine/Source/Editor/MovieSceneTools/Public/MovieSceneToolHelpers.h
--- a/Engine/Source/Editor/MovieSceneTools/Public/MovieSceneToolHelpers.h
+++ b/Engine/Source/Editor/MovieSceneTools/Public/MovieSceneToolHelpers.h
@@ -22,7 +22,6 @@
 #include "Channels/MovieSceneChannelProxy.h"
 #include "MovieSceneSection.h"
 
-class ISequencer;
 class UMovieScene;
 class UMovieSceneSequence;
 struct FMovieSceneObjectBindingID;
@@ -111,7 +110,9 @@ enum class FChannelMergeAlgorithm : uint8
 	/**Average values together*/
 	Average,
 	/**Add values together*/
-	Add
+	Add,
+	/** Override values together*/
+	Override,
 };
 
 class MOVIESCENETOOLS_API MovieSceneToolHelpers
@@ -668,110 +669,151 @@ public:
 	 */
 	static bool IsValidAsset(UMovieSceneSequence* Sequence, const FAssetData& InAssetData);
 
-	static bool CollapseSection(TSharedPtr<ISequencer>& SequencerPtr, UMovieSceneTrack* OwnerTrack, TArray<UMovieSceneSection*> Sections,
+	/*
+	* Collapse all of the sections specified onto the first one
+	* @param InSequencer  Sequencer we are collpasing at
+	* @param InOwnerTrack  The track that should be owning the sections we are collapsing
+	* @param InSections The sections we are collapsing. Sections[0] will remain, te otheer ones will be deleted
+	* @param InSettings Baking settings ussed
+	* @return Return true if succeeds false if it doesn't
+	*/
+	static bool CollapseSection(TSharedPtr<ISequencer>& InSequencer, UMovieSceneTrack* InOwnerTrack, TArray<UMovieSceneSection*> InSections,
 		const FBakingAnimationKeySettings& InSettings);
 
-	//merge the channels onto the first one using the passed in algorithm, we may skip the last channel also (since it may be weight).
-	template<typename ChannelType>
-	static bool MergeChannels(TArray<ChannelType*>& Channels, const  TArray<UMovieSceneSection*>& Sections, const TRange<FFrameNumber>& Range,
-		FChannelMergeAlgorithm MergeAlgorithm)
+	/*
+	* Split set of sections to one containing the BlendType the other not
+	* @param InSections The sections we are searching
+	* @param OutSections All non-BlendType sections
+	* @param OutAbsoluteSections All BlendType sections
+	*/
+	static void SplitSectionsByBlendType(EMovieSceneBlendType BlendType, const TArray<UMovieSceneSection*>& InSections, TArray<UMovieSceneSection*>& OutSections, TArray<UMovieSceneSection*>& OutBlendTypeSections);
+
+	/*
+	* Get the channel values at the specified time with the specified sections, will get
+	* the channels from within the specified start and end indices
+	* @param StartIndex Start Index to get values from
+	* @param EndIndex End Index to get values from, should be no greater than the number of channels -1
+	* @param Sections List of non-absolute(additive and override) sections to evaluate. Note this may not be all of the sections in the track,
+	* but just a subset, which allows us to evaluate up to a certain point in the track.
+	* @param AbsoluteSections List of absolute channels to evaluate
+	* @param FrameTime Frame to evaluate
+	* @return Returns the total channel value of those sections for each channel. Remember if OverrideChannelIndex is set this will just be one channel.
+	*/
+	template<typename ChannelType, typename CurveValueType>
+	static TArray<CurveValueType> GetChannelValues(const int32 StartIndex, const int32 EndIndex, const TArray<UMovieSceneSection*>& Sections, const TArray<UMovieSceneSection*>& AbsoluteSections, const FFrameNumber& FrameTime)
 	{
-		if (Channels.Num() < 2 && Channels.Num() != Sections.Num())
+		TArray<CurveValueType> Values;
+		int32 NumChannels = 0;
+		if (Sections.Num() > 0)
 		{
-			return false;
+			TArrayView<ChannelType*>Channels = Sections[0]->GetChannelProxy().GetChannels<ChannelType>();
+			NumChannels = Channels.Num();
 		}
-		using ChannelValueType = typename ChannelType::ChannelValueType;
-		using CurveValueType = typename ChannelType::CurveValueType;
-
-		//base channel we set values on
-		ChannelType* BaseChannel = Channels[0];
-		TMovieSceneChannelData<ChannelValueType> BaseChannelData = BaseChannel->GetData();
-		//iterate over each key
-		TArray<FFrameNumber> KeyTimes;
-		TArray<FKeyHandle> Handles;
-		//cached set that we set at the end
-		TArray<TPair< FFrameNumber, ChannelValueType>> KeysToSet;
-		for (int32 ChannelIndex = 0; ChannelIndex < Channels.Num(); ++ChannelIndex)
+		else if(AbsoluteSections.Num() > 0)
 		{
-			KeyTimes.Reset();
-			Handles.Reset();
-			ChannelType* Channel = Channels[ChannelIndex];
-			Channel->GetKeys(Range, &KeyTimes, &Handles);
-			for (int32 FrameIndex = 0; FrameIndex < KeyTimes.Num(); ++FrameIndex)
+			TArrayView<ChannelType*>Channels = AbsoluteSections[0]->GetChannelProxy().GetChannels<ChannelType>();
+			NumChannels = Channels.Num();
+		}
+		else
+		{
+			UE_LOG(LogMovieScene, Warning, TEXT("GetChannelValues:: Invalid number of channels"));
+			return Values;
+		}
+		if (StartIndex < 0 || EndIndex >= NumChannels || EndIndex < StartIndex)
+		{
+			UE_LOG(LogMovieScene, Warning, TEXT("GetChannelValues:: Invalid Start/End indices"));
+			return Values;
+		}
+		
+		for (int32 ChannelIndex = StartIndex; ChannelIndex <= EndIndex; ++ChannelIndex)
+		{
+			CurveValueType Value = 0.0;
+			if (AbsoluteSections.Num() > 0)
 			{
-				const FFrameNumber& Frame = KeyTimes[FrameIndex];
-				if (Sections[0]->GetRange().Contains(Frame) == false)  //frame is outside base range so skip
+				for (UMovieSceneSection* AbsoluteSection : AbsoluteSections)
 				{
-					continue;
+					TArrayView<ChannelType*>Channels = AbsoluteSection->GetChannelProxy().GetChannels<ChannelType>();
+					float Weight = AbsoluteSection->GetTotalWeightValue(FrameTime);
+					CurveValueType WeightedValue = 0.0;
+					ChannelType* Channel = Channels[ChannelIndex];
+					Channel->Evaluate(FrameTime, WeightedValue);
+					WeightedValue *= ((double)Weight);
+					Value += WeightedValue;
 				}
-				const FFrameTime FrameTime(Frame);
-				int32 KeyIndex = Channel->GetData().GetIndex(Handles[FrameIndex]);
-				ChannelValueType Value = Channel->GetData().GetValues()[KeyIndex];
-				//got value with tangents and times, now we perform the operation
-				Value.Value = 0.0; //zero out the value we calculate it 
-				if (MergeAlgorithm == FChannelMergeAlgorithm::Average)
+
+				Value /= (double)AbsoluteSections.Num();		
+			}
+
+			for (UMovieSceneSection* Section : Sections)
+			{
+				TArrayView<ChannelType*> Channels = Section->GetChannelProxy().GetChannels<ChannelType>();
+				float Weight = Section->GetTotalWeightValue(FrameTime);
+				CurveValueType WeightedValue = 0.0;
+				ChannelType* Channel = Channels[ChannelIndex];
+				Channel->Evaluate(FrameTime, WeightedValue);
+				if (Section->GetBlendType().Get() == EMovieSceneBlendType::Additive)
 				{
-					double DNumChannels = 0.0;
-					for (int32 WeightIndex = 0; WeightIndex < Sections.Num(); ++WeightIndex)
-					{
-						if (Sections[WeightIndex]->GetRange().Contains(Frame))
-						{
-							float Weight = Sections[WeightIndex]->GetTotalWeightValue(FrameTime);
-							CurveValueType WeightedValue = 0.0;
-							ChannelType* EachChannel = Channels[WeightIndex];
-							EachChannel->Evaluate(FrameTime, WeightedValue);
-							WeightedValue *= ((double)Weight);
-							Value.Value += WeightedValue;
-							DNumChannels += 1.0;
-						}
-					}
-					if (DNumChannels > 0.0) //should always happen since base(0) at least be here
-					{
-						Value.Value /= DNumChannels;
-					}
+					WeightedValue *= Weight;
+					Value += WeightedValue;
 				}
-				else if (MergeAlgorithm == FChannelMergeAlgorithm::Add)
+				else if (Section->GetBlendType().Get() == EMovieSceneBlendType::Override)
 				{
-					for (int32 WeightIndex = 0; WeightIndex < Sections.Num(); ++WeightIndex)
-					{
-						if (Sections[WeightIndex]->GetRange().Contains(Frame))
-						{
-							float Weight = Sections[WeightIndex]->GetTotalWeightValue(FrameTime);
-							if (Sections[WeightIndex]->GetBlendType().IsValid() == false ||
-								Sections[WeightIndex]->GetBlendType().Get() != EMovieSceneBlendType::Additive)
-							{
-								Weight = 1.0;
-							}
-							CurveValueType WeightedValue = 0.0;
-							ChannelType* EachChannel = Channels[WeightIndex];
-							EachChannel->Evaluate(FrameTime, WeightedValue);
-							WeightedValue *= Weight;
-							Value.Value += WeightedValue;
-						}
-					}
+					Value = (Value * (1.0 - Weight)) +
+						(Value * Weight);
 				}
-				KeysToSet.Add(TPair<FFrameNumber, ChannelValueType>(Frame, Value));
 			}
+			Values.Add(Value);
 		}
-		for (TPair<FFrameNumber, ChannelValueType>& KeyToSet : KeysToSet)
-		{
-			MovieSceneToolHelpers::SetOrAddKey(BaseChannelData, KeyToSet.Key, KeyToSet.Value);
-		}
-
-		return true;
+		return Values;
 	}
 
-	template<typename ChannelType>
-	static bool MergeNonWeightedChannels(TArray<ChannelType*>& Channels, const  TArray<UMovieSceneSection*>& Sections, const TRange<FFrameNumber>& Range,
+	/*
+	* Merge the set of passed in channels from each section at the specified section channel index.
+	* For the algorithm to be relaible it should be done per section blending type.
+	* @param SectionChannelIndex  The index of the channel of this type in the sections.
+	* @param Channels  Set of channels of this type that we will merge together onto the first specified one.
+	* @param Sections Set of corresponding sections for each channel specified.
+	* @param Range Range of time over which we will merge
+	* @param MergeAlgorithm The algorithm to use to blend the channels together
+	* @return Returns true if successful
+	*/	template<typename ChannelType>
+	static bool MergeChannels(int32 SectionChannelIndex, TArray<ChannelType*>& Channels, const  TArray<UMovieSceneSection*>& Sections, const TRange<FFrameNumber>& Range,
 		FChannelMergeAlgorithm MergeAlgorithm)
 	{
-		if (Channels.Num() < 2 && Channels.Num() != Sections.Num())
+		if (Channels.Num() < 2 || Channels.Num() != Sections.Num())
+		{
+			UE_LOG(LogMovieScene, Warning, TEXT("MergeChannels:: Invalid number of channels"));
+			return false;
+		}
+		//for overrides since the need to evaluate the full value we only support 2 sections
+		if (MergeAlgorithm == FChannelMergeAlgorithm::Override && Sections.Num() != 2)
 		{
+			UE_LOG(LogMovieScene, Warning, TEXT("MergeChannels:: Can only do Override Blend with two sections"));
 			return false;
 		}
 		using ChannelValueType = typename ChannelType::ChannelValueType;
 		using CurveValueType = typename ChannelType::CurveValueType;
 
+		TArray<UMovieSceneSection*> OtherSections;
+		TArray<UMovieSceneSection*> AbsoluteSections;
+		int32 OverrideChannelIndex = SectionChannelIndex;
+		int32 NumChannelsInSection = INDEX_NONE;
+		if (MergeAlgorithm == FChannelMergeAlgorithm::Override) //if we are overriding we need to get the absolute value at the frame!
+		{
+			if (UMovieSceneTrack* OwnerTrack = Sections[1]->GetTypedOuter<UMovieSceneTrack>())
+			{
+				TArray<UMovieSceneSection*> TrackSections = OwnerTrack->GetAllSections();
+				int32 SectionIndex = TrackSections.Find(Sections[1]);
+				if (SectionIndex != INDEX_NONE)
+				{
+					TrackSections.SetNum(SectionIndex + 1); //this wil make sure we include the current sections
+					MovieSceneToolHelpers::SplitSectionsByBlendType(EMovieSceneBlendType::Absolute,TrackSections, OtherSections, AbsoluteSections);
+					TArrayView<ChannelType*> OverrideChannels = Sections[1]->GetChannelProxy().GetChannels<ChannelType>();
+					NumChannelsInSection = OverrideChannels.Num();
+				}
+			}
+		}
+
 		//base channel we set values on
 		ChannelType* BaseChannel = Channels[0];
 		TMovieSceneChannelData<ChannelValueType> BaseChannelData = BaseChannel->GetData();
@@ -839,6 +881,23 @@ public:
 						}
 					}
 				}
+				else if (MergeAlgorithm == FChannelMergeAlgorithm::Override)
+				{
+					if (Sections[1]->GetRange().Contains(Frame))
+					{
+						float Weight = Sections[1]->GetTotalWeightValue(FrameTime);
+						TArray<CurveValueType> ChannelValues = MovieSceneToolHelpers::GetChannelValues<ChannelType,
+							CurveValueType>(OverrideChannelIndex, OverrideChannelIndex, OtherSections, AbsoluteSections, Frame);
+						if (ChannelValues.Num() == 1)
+						{
+							CurveValueType WeightedValue = 0.0;
+							ChannelType* EachChannel = Channels[1];
+							EachChannel->Evaluate(FrameTime, WeightedValue);
+							WeightedValue *= Weight;
+							Value.Value = WeightedValue + (1.0 - Weight) * ChannelValues[0];
+						}
+					}
+				}
 				KeysToSet.Add(TPair<FFrameNumber, ChannelValueType>(Frame, Value));
 			}
 		}
@@ -846,15 +905,98 @@ public:
 		{
 			MovieSceneToolHelpers::SetOrAddKey(BaseChannelData, KeyToSet.Key, KeyToSet.Value);
 		}
+		return true;
+	}
+
+	/*
+	* Merge the TopSection onto the BaseSection. This function should be used when merging
+	* a mix of Override/Additive and Absolute section Blend Types.
+	* @param BaseSection  The section to merge onto
+	* @param TopSection  The section we will merge from
+	* @param StartIndex Start Index of the channels we are merging
+	* @param EndIndex End Index of the channels we are merging, should be no greater than the number of channels -1
+	* @param Range Range of time over which we will merge
+	* Control Rig sections and don't want to merge the last weight float channel
+	* @return Returns true if successful
+	*/	
+	template<typename ChannelType>
+	static bool MergeSections(UMovieSceneSection* BaseSection, UMovieSceneSection* TopSection,
+		int32 StartIndex, int32 EndIndex, const TRange<FFrameNumber>& Range)
+	{
+		TArrayView<ChannelType*> BaseChannels = BaseSection->GetChannelProxy().GetChannels<ChannelType>();
+		TArrayView<ChannelType*> TopChannels = TopSection->GetChannelProxy().GetChannels<ChannelType>();
+		if (TopChannels.Num() != BaseChannels.Num())
+		{
+			UE_LOG(LogMovieScene, Warning, TEXT("MergeSections:: Invalid number of channels"));
+			return false;
+		}
+		if (StartIndex < 0 || EndIndex >= BaseChannels.Num() || EndIndex < StartIndex)
+		{
+			UE_LOG(LogMovieScene, Warning, TEXT("MergeSections:: Invalid Start/End indices"));
+			return false;
+		}
+
+		BaseSection->Modify();
 
+		TArray<ChannelType*> Channels;
+		int32 ChannelIndex = 0;
+		TArray<UMovieSceneSection*> Sections;
+		Sections.Add(BaseSection);
+		Sections.Add(TopSection);
+		FChannelMergeAlgorithm MergeAlgorithm = FChannelMergeAlgorithm::Add;
+		if (TopSection->GetBlendType().IsValid() && TopSection->GetBlendType().Get() == EMovieSceneBlendType::Absolute)
+		{
+			MergeAlgorithm = FChannelMergeAlgorithm::Average;
+		}
+		else if (TopSection->GetBlendType().IsValid() && TopSection->GetBlendType().Get() == EMovieSceneBlendType::Override)
+		{
+			MergeAlgorithm = FChannelMergeAlgorithm::Override;
+		}
+
+		for (ChannelIndex = StartIndex; ChannelIndex <= EndIndex; ++ChannelIndex)
+		{
+			Channels.Reset();
+			Channels.Add(BaseChannels[ChannelIndex]);
+			Channels.Add(TopChannels[ChannelIndex]);
+
+			if (MergeChannels(ChannelIndex, Channels, Sections, Range,
+				MergeAlgorithm) == false)
+			{
+				UE_LOG(LogMovieScene, Warning, TEXT("MergeSections:: Could not merge channels"));
+				return false;
+			}
+		}
+		
+		for (ChannelType* Channel : BaseChannels)
+		{
+			Channel->AutoSetTangents();
+		}
+		
 		return true;
 	}
 
+	/*
+	* Merge the following set of sections. Should be used when blending just Absolute with additive sections. If Merging with
+	* overrides, use the above MergeSections function
+	* @param BaseSection  The section to merge onto
+	* @param AbsoluteSections  Set of absolute functions to merge
+	* @param AddditveSections  Set of Additive functions to merge
+	* @param StartIndex Start Index of the channels we are merging
+	* @param EndIndex End Index of the channels we are merging, should be no greater than the number of channels -1
+	* @param Range Range of time over which we will merge
+	* Control Rig sections and don't want to merge the last weight float channel
+	* @return Returns true if successful
+	*/	
 	template<typename ChannelType>
 	static bool MergeSections(UMovieSceneSection* BaseSection, TArray<UMovieSceneSection*>& AbsoluteSections, TArray<UMovieSceneSection*>& AdditiveSections,
-		const TRange<FFrameNumber>& Range, bool bSkipLastChannel)
+		int32 StartIndex, int32 EndIndex, const TRange<FFrameNumber>& Range)
 	{
 		TArrayView<ChannelType*> BaseChannels = BaseSection->GetChannelProxy().GetChannels<ChannelType>();
+		if (StartIndex < 0 || EndIndex >= BaseChannels.Num() || EndIndex < StartIndex)
+		{
+			UE_LOG(LogMovieScene, Warning, TEXT("MergeSections:: Invalid Start/End indices"));
+			return false;
+		}
 		if (BaseChannels.Num() > 0)
 		{
 			int32 SectionIndex = 0;
@@ -880,10 +1022,7 @@ public:
 			BaseSection->Modify();
 
 			TArray<ChannelType*> Channels;
-			int32 ChannelIndex = 0;
-			bSkipLastChannel = false;
-			int32 BaseChannelsNum = bSkipLastChannel ? BaseChannels.Num() - 1 : BaseChannels.Num();
-			for (ChannelIndex = 0; ChannelIndex < BaseChannels.Num(); ++ChannelIndex)
+			for (int32 ChannelIndex = StartIndex; ChannelIndex <= EndIndex; ++ChannelIndex)
 			{
 				if (AbsoluteSections.Num() > 0)
 				{
@@ -894,7 +1033,7 @@ public:
 						Channels.Add(OurChannels[ChannelIndex]);
 					}
 					//now blend them
-					if (MergeChannels(Channels, AbsoluteSections, Range,
+					if (MergeChannels(ChannelIndex, Channels, AbsoluteSections, Range,
 						FChannelMergeAlgorithm::Average) == false)
 					{
 						UE_LOG(LogMovieScene, Warning, TEXT("MergeSections:: Could not merge channels"));
@@ -911,7 +1050,7 @@ public:
 						Channels.Add(OurChannels[ChannelIndex]);
 					}
 					//now blend them
-					if (MergeChannels(Channels, AdditiveSections, Range,
+					if (MergeChannels(ChannelIndex,Channels, AdditiveSections, Range,
 						FChannelMergeAlgorithm::Add) == false)
 					{
 						UE_LOG(LogMovieScene, Warning, TEXT("MergeSections:: Could not merge channels"));
@@ -926,6 +1065,7 @@ public:
 		}
 		else
 		{
+			UE_LOG(LogMovieScene, Warning, TEXT("MergeSections:: Invalid number of channels"));
 			return false;
 		}
 		return true;
