Sequencer: Hopeful fix for Sequencer delegate not getting cleared, make sure it's cleared if we get a new Sequencer, also changed from a Lambda to a Raw which are supposed to be safer.

diff --git a/Engine/Plugins/Animation/ControlRig/Source/ControlRigEditor/Private/EditMode/SControlRigDetails.cpp b/Engine/Plugins/Animation/ControlRig/Source/ControlRigEditor/Private/EditMode/SControlRigDetails.cpp
--- a/Engine/Plugins/Animation/ControlRig/Source/ControlRigEditor/Private/EditMode/SControlRigDetails.cpp
+++ b/Engine/Plugins/Animation/ControlRig/Source/ControlRigEditor/Private/EditMode/SControlRigDetails.cpp
@@ -504,15 +504,20 @@ bool SControlRigDetails::IsReadOnlyPropertyOnDetailCustomization(const FProperty
 
 FSequencerTracker::~FSequencerTracker()
 {
-	TSharedPtr<ISequencer> Sequencer = WeakSequencer.Pin();
-	if (Sequencer)
+	RemoveDelegates();
+}
+
+void FSequencerTracker::RemoveDelegates()
+{
+	if (TSharedPtr<ISequencer> Sequencer = WeakSequencer.Pin())
 	{
-		Sequencer->GetSelectionChangedObjectGuids().Remove(OnSelectionChangedHandle);
+		Sequencer->GetSelectionChangedObjectGuids().RemoveAll(this);
 	}
 }
 
 void FSequencerTracker::SetSequencerAndDetails(TWeakPtr<ISequencer> InWeakSequencer, SControlRigDetails* InControlRigDetails)
 {
+	RemoveDelegates();
 	WeakSequencer = InWeakSequencer;
 	ControlRigDetails = InControlRigDetails;
 	if (WeakSequencer.IsValid() == false || InControlRigDetails == nullptr)
@@ -525,15 +530,11 @@ void FSequencerTracker::SetSequencerAndDetails(TWeakPtr<ISequencer> InWeakSequen
 	Sequencer->GetSelectedObjects(SequencerSelectedObjects);
 	UpdateSequencerBindings(SequencerSelectedObjects);
 
-	OnSelectionChangedHandle = Sequencer->GetSelectionChangedObjectGuids().AddLambda([this](TArray<FGuid> NewSelection)
-	{
-		UpdateSequencerBindings(NewSelection);
-
-	});
-
+	Sequencer->GetSelectionChangedObjectGuids().AddRaw(this, &FSequencerTracker::UpdateSequencerBindings);
+	
 }
 
-void FSequencerTracker::UpdateSequencerBindings(const TArray<FGuid>& SequencerBindings)
+void FSequencerTracker::UpdateSequencerBindings(TArray<FGuid> SequencerBindings)
 {
 	const FDateTime StartTime = FDateTime::Now();
 
diff --git a/Engine/Plugins/Animation/ControlRig/Source/ControlRigEditor/Private/EditMode/SControlRigDetails.h b/Engine/Plugins/Animation/ControlRig/Source/ControlRigEditor/Private/EditMode/SControlRigDetails.h
--- a/Engine/Plugins/Animation/ControlRig/Source/ControlRigEditor/Private/EditMode/SControlRigDetails.h
+++ b/Engine/Plugins/Animation/ControlRig/Source/ControlRigEditor/Private/EditMode/SControlRigDetails.h
@@ -38,8 +38,9 @@ public:
 	void SetSequencerAndDetails(TWeakPtr<ISequencer> InWeakSequencer, SControlRigDetails* InControlRigDetails);
 	TMap<UObject*, FArrayOfPropertyTracks>& GetObjectsTracked() { return ObjectsTracked; }
 private:
-	void UpdateSequencerBindings(const TArray<FGuid>& SequencerBindings);
-	FDelegateHandle OnSelectionChangedHandle;
+
+	void RemoveDelegates();
+	void UpdateSequencerBindings(TArray<FGuid> SequencerBindings);
 	TWeakPtr<ISequencer> WeakSequencer;
 	TMap<UObject*, FArrayOfPropertyTracks> ObjectsTracked;
 	SControlRigDetails* ControlRigDetails = nullptr;
