Anim Layers: Deleting Layer could leave it unkeyed, so force the base to be the keyed section.
#trivial

diff --git a/Engine/Plugins/Animation/ControlRig/Source/ControlRigEditor/Private/Sequencer/AnimLayers/AnimLayers.cpp b/Engine/Plugins/Animation/ControlRig/Source/ControlRigEditor/Private/Sequencer/AnimLayers/AnimLayers.cpp
--- a/Engine/Plugins/Animation/ControlRig/Source/ControlRigEditor/Private/Sequencer/AnimLayers/AnimLayers.cpp
+++ b/Engine/Plugins/Animation/ControlRig/Source/ControlRigEditor/Private/Sequencer/AnimLayers/AnimLayers.cpp
@@ -1428,6 +1428,10 @@ bool UAnimLayers::DeleteAnimLayer(ISequencer* SequencerPtr, int32 Index)
 			AnimLayers.RemoveAt(Index);
 			SequencerPtr->NotifyMovieSceneDataChanged(EMovieSceneDataChangeType::MovieSceneStructureItemRemoved);
 		}
+		if (UAnimLayer* BaseAnimLayer = AnimLayers[0]) //set base as keyed
+		{
+			BaseAnimLayer->SetKeyed();
+		}
 		AnimLayerListChangedBroadcast();
 	}
 	else
