Control Rig Edit Mode improvements
    - avoid updating shape actors when the viewport is in game view (as they're hidden)
    - reuse the UpdateShape function to use a single path to update shapes
    - start moving shape-related functions into separate files

#rnx

diff --git a/Engine/Plugins/Animation/ControlRig/Source/ControlRigEditor/Private/EditMode/ControlRigEditMode.cpp b/Engine/Plugins/Animation/ControlRig/Source/ControlRigEditor/Private/EditMode/ControlRigEditMode.cpp
--- a/Engine/Plugins/Animation/ControlRig/Source/ControlRigEditor/Private/EditMode/ControlRigEditMode.cpp
+++ b/Engine/Plugins/Animation/ControlRig/Source/ControlRigEditor/Private/EditMode/ControlRigEditMode.cpp
@@ -76,6 +76,7 @@
 #include "DragTool_BoxSelect.h"
 #include "DragTool_FrustumSelect.h"
 #include "AnimationEditorViewportClient.h"
+#include "ControlShapeActorHelper.h"
 #include "EditorInteractiveGizmoManager.h"
 #include "Editor/ControlRigViewportToolbarExtensions.h"
 #include "Editor/Sequencer/Private/SSequencer.h"
@@ -679,7 +680,7 @@ void FControlRigEditMode::Tick(FEditorViewportClient* ViewportClient, float Delt
 	{
 		// We need to tick here since changing a bone for example
 		// might have changed the transform of the Control
-		PostPoseUpdate();
+		PostPoseUpdate(ViewportClient);
 
 		if (!AreEditingControlRigDirectly() == false) //only do this check if not in level editor
 		{
@@ -3932,33 +3933,6 @@ void FControlRigEditMode::RequestToRecreateControlShapeActors(UControlRig* Contr
 	}
 }
 
-// temporarily we just support following types of gizmo
-static bool IsSupportedControlType(const ERigControlType ControlType)
-{
-	switch (ControlType)
-	{
-	case ERigControlType::Float:
-	case ERigControlType::ScaleFloat:
-	case ERigControlType::Integer:
-	case ERigControlType::Vector2D:
-	case ERigControlType::Position:
-	case ERigControlType::Scale:
-	case ERigControlType::Rotator:
-	case ERigControlType::Transform:
-	case ERigControlType::TransformNoScale:
-	case ERigControlType::EulerTransform:
-	{
-		return true;
-	}
-	default:
-	{
-		break;
-	}
-	}
-
-	return false;
-}
-
 bool FControlRigEditMode::TryUpdatingControlsShapes(UControlRig* InControlRig)
 {
 	using namespace ControlRigEditMode::Shapes;
@@ -4061,7 +4035,7 @@ bool FControlRigEditMode::TryUpdatingControlsShapes(UControlRig* InControlRig)
 		ComponentTransform = GetHostingSceneComponentTransform(InControlRig);
 	}
 		
-	const FShapeUpdateParams Params(InControlRig, ComponentTransform, IsControlRigSkelMeshVisible(InControlRig));
+	const FShapeUpdateParams Params(InControlRig, ComponentTransform, IsControlRigSkelMeshVisible(InControlRig), IsInLevelEditor());
 	for (int32 ShapeActorIndex = 0; ShapeActorIndex < NumShapes; ShapeActorIndex++)
 	{
 		AControlRigShapeActor* ShapeActor = Shapes[ShapeActorIndex].Get();
@@ -5051,9 +5025,7 @@ void FControlRigEditMode::ChangeControlShapeTransform(AControlRigShapeActor* Sha
 	{
 		if (UControlRig* ControlRig = ShapeActor->ControlRig.Get())
 		{
-
 			FTransform NewTransform = CurrentTransform.GetRelativeTransform(ToWorldTransform);
-			FRigControlModifiedContext Context;
 
 			if (FRigControlElement* ControlElement = ControlRig->GetHierarchy()->Find<FRigControlElement>(ShapeActor->GetElementKey()))
 			{
@@ -5080,96 +5052,14 @@ bool FControlRigEditMode::ModeSupportedByShapeActor(const AControlRigShapeActor*
 {
 	if (UControlRig* ControlRig = ShapeActor->ControlRig.Get())
 	{
-		const FRigControlElement* ControlElement = ControlRig->FindControl(ShapeActor->ControlName);
-		if (ControlElement)
+		if (const FRigControlElement* ControlElement = ControlRig->FindControl(ShapeActor->ControlName))
 		{
 			if (bIsChangingControlShapeTransform)
 			{
 				return true;
 			}
 
-			if (IsSupportedControlType(ControlElement->Settings.ControlType))
-			{
-				switch (InMode)
-				{
-					case UE::Widget::WM_None:
-						return true;
-					case UE::Widget::WM_Rotate:
-					{
-						switch (ControlElement->Settings.ControlType)
-						{
-							case ERigControlType::Rotator:
-							case ERigControlType::Transform:
-							case ERigControlType::TransformNoScale:
-							case ERigControlType::EulerTransform:
-							{
-								return true;
-							}
-							default:
-							{
-								break;
-							}
-						}
-						break;
-					}
-					case UE::Widget::WM_Translate:
-					{
-						switch (ControlElement->Settings.ControlType)
-						{
-							case ERigControlType::Float:
-							case ERigControlType::Integer:
-							case ERigControlType::Vector2D:
-							case ERigControlType::Position:
-							case ERigControlType::Transform:
-							case ERigControlType::TransformNoScale:
-							case ERigControlType::EulerTransform:
-							{
-								return true;
-							}
-							default:
-							{
-								break;
-							}
-						}
-						break;
-					}
-					case UE::Widget::WM_Scale:
-					{
-						switch (ControlElement->Settings.ControlType)
-						{
-							case ERigControlType::Scale:
-							case ERigControlType::ScaleFloat:
-							case ERigControlType::Transform:
-							case ERigControlType::EulerTransform:
-							{
-								return true;
-							}
-							default:
-							{
-								break;
-							}
-						}
-						break;
-					}
-					case UE::Widget::WM_TranslateRotateZ:
-					{
-						switch (ControlElement->Settings.ControlType)
-						{
-							case ERigControlType::Transform:
-							case ERigControlType::TransformNoScale:
-							case ERigControlType::EulerTransform:
-							{
-								return true;
-							}
-							default:
-							{
-								break;
-							}
-						}
-						break;
-					}
-				}
-			}
+			return ControlRigEditMode::Shapes::IsModeSupported(ControlElement->Settings.ControlType, InMode);
 		}
 	}
 	return false;
@@ -5192,49 +5082,6 @@ bool FControlRigEditMode::IsControlRigSkelMeshVisible(const UControlRig* InContr
 	return true;
 }
 
-void FControlRigEditMode::TickControlShape(AControlRigShapeActor* ShapeActor, const FTransform& ComponentTransform) const
-{
-	const UControlRigEditModeSettings* Settings = GetDefault<UControlRigEditModeSettings>();
-	if (ShapeActor)
-	{
-		if (UControlRig* ControlRig = ShapeActor->ControlRig.Get())
-		{
-			const FTransform Transform = ControlRig->GetControlGlobalTransform(ShapeActor->ControlName);
-			ShapeActor->SetActorTransform(Transform * ComponentTransform);
-
-			if (FRigControlElement* ControlElement = ControlRig->FindControl(ShapeActor->ControlName))
-			{
-				const bool bControlsHiddenInViewport = Settings->bHideControlShapes || !ControlRig->GetControlsVisible()
-					|| (IsControlRigSkelMeshVisible(ControlRig) == false);
-
-				bool bIsVisible = ControlElement->Settings.IsVisible();
-				bool bRespectVisibilityForSelection = IsInLevelEditor() ? false: true;
-
-				if(!bControlsHiddenInViewport)
-				{
-					if(ControlElement->Settings.AnimationType == ERigControlAnimationType::ProxyControl)
-					{					
-						bRespectVisibilityForSelection = false;
-						if(Settings->bShowAllProxyControls)
-						{
-							bIsVisible = true;
-						}
-					}
-				}
-				
-				ShapeActor->SetShapeColor(ShapeActor->OverrideColor.A < SMALL_NUMBER ?
-					ControlElement->Settings.ShapeColor : ShapeActor->OverrideColor);
-				ShapeActor->SetIsTemporarilyHiddenInEditor(
-					!bIsVisible ||
-					bControlsHiddenInViewport);
-				
-				ShapeActor->SetSelectable(
-					ControlElement->Settings.IsSelectable(bRespectVisibilityForSelection));
-			}
-		}
-	}
-}
-
 AControlRigShapeActor* FControlRigEditMode::GetControlShapeFromControlName(UControlRig* InControlRig,const FName& ControlName) const
 {
 	const auto* ShapeActors = ControlRigShapeActors.Find(InControlRig);
@@ -5613,21 +5460,37 @@ void FControlRigEditMode::OnPoseInitialized()
 	OnAnimSystemInitializedDelegate.Broadcast();
 }
 
-void FControlRigEditMode::PostPoseUpdate() const
+void FControlRigEditMode::PostPoseUpdate(const FEditorViewportClient* InViewportClient) const
 {
-	for (auto& ShapeActors : ControlRigShapeActors)
+	using namespace ControlRigEditMode::Shapes;
+
+	const bool bIsGameView = InViewportClient ? InViewportClient->IsInGameView() : false;
+	if (bIsGameView)
 	{
-		FTransform ComponentTransform = FTransform::Identity;
-		if (!AreEditingControlRigDirectly())
+		// no need to update the shape actors in game view, they just need to be hidden
+		static constexpr bool bHiddenInEditor = true;
+		for (const TPair<UControlRig*,TArray<TObjectPtr<AControlRigShapeActor>>>& ShapeActors : ControlRigShapeActors)
 		{
-			ComponentTransform = GetHostingSceneComponentTransform(ShapeActors.Key);
+			for (AControlRigShapeActor* ShapeActor : ShapeActors.Value)
+			{
+				ShapeActor->SetIsTemporarilyHiddenInEditor(bHiddenInEditor);
+			}
 		}
+		return;
+	}
+
+	const bool bAreEditingControlRigDirectly = AreEditingControlRigDirectly();
+	for (const TPair<UControlRig*,TArray<TObjectPtr<AControlRigShapeActor>>>& ShapeActors : ControlRigShapeActors)
+	{
+		const UControlRig* ControlRig = ShapeActors.Key;
+		const FTransform ComponentTransform = bAreEditingControlRigDirectly ? FTransform::Identity : GetHostingSceneComponentTransform(ControlRig); 
+
+		FShapeUpdateParams Params(ControlRig, ComponentTransform, IsControlRigSkelMeshVisible(ControlRig), IsInLevelEditor());
 		for (AControlRigShapeActor* ShapeActor : ShapeActors.Value)
 		{
-			TickControlShape(ShapeActor, ComponentTransform);	
+			UpdateControlShape(ShapeActor, ControlRig->FindControl(ShapeActor->ControlName), Params);
 		}
 	}
-
 }
 
 void FControlRigEditMode::NotifyDrivenControls(UControlRig* InControlRig, const FRigElementKey& InKey, const FRigControlModifiedContext& InContext)
@@ -5975,114 +5838,14 @@ void FDetailKeyFrameCacheAndHandler::UpdateIfDirty()
 	{
 		if (FMovieSceneConstraintChannelHelper::bDoNotCompensate == false) //if compensating don't reset this.
 		{
-			if (EditMode && EditMode->GetControlProxy())
+			UControlRigDetailPanelControlProxies* ControlProxies = EditMode ? EditMode->GetDetailProxies() : nullptr;
+			if (ControlProxies)
 			{
-				EditMode->GetControlProxy()->ValuesChanged();
+				ControlProxies->ValuesChanged();
 			}
 			bValuesDirty = false;
 		}
 	}
 }
 
-namespace ControlRigEditMode::Shapes
-{
-
-void GetControlsEligibleForShapes(UControlRig* InControlRig, TArray<FRigControlElement*>& OutControls)
-{
-	OutControls.Reset();
-
-	const URigHierarchy* Hierarchy = InControlRig ? InControlRig->GetHierarchy() : nullptr;
-	if (!Hierarchy)
-	{
-		return;
-	}
-
-	OutControls = Hierarchy->GetFilteredElements<FRigControlElement>([](const FRigControlElement* ControlElement)
-	{
-		const FRigControlSettings& ControlSettings = ControlElement->Settings;
-		return ControlSettings.SupportsShape() && IsSupportedControlType(ControlSettings.ControlType);
-	});
-}
-
-void DestroyShapesActorsFromWorld(const TArray<TObjectPtr<AControlRigShapeActor>>& InShapeActorsToDestroy)
-{
-	// NOTE: should UWorld::EditorDestroyActor really modify the level when removing the shapes?
-	// kept for legacy but I guess this should be set to false
-	static constexpr bool bShouldModifyLevel = true;
-
-	for (const TObjectPtr<AControlRigShapeActor>& ShapeActorPtr: InShapeActorsToDestroy)
-	{
-		if (AControlRigShapeActor* ShapeActor = ShapeActorPtr.Get())
-		{
-			if (UWorld* World = ShapeActor->GetWorld())
-			{
-				if (ShapeActor->GetAttachParentActor())
-				{
-					ShapeActor->DetachFromActor(FDetachmentTransformRules::KeepRelativeTransform);
-				}
-				World->EditorDestroyActor(ShapeActor,bShouldModifyLevel);
-			}			
-		}
-	}
-}
-
-FShapeUpdateParams::FShapeUpdateParams(const UControlRig* InControlRig, const FTransform& InComponentTransform, const bool InSkeletalMeshVisible)
-	: ControlRig(InControlRig)
-	, Hierarchy(InControlRig->GetHierarchy())
-	, Settings(GetDefault<UControlRigEditModeSettings>())
-	, ComponentTransform(InComponentTransform)
-	, bIsSkeletalMeshVisible(InSkeletalMeshVisible)
-{}
-
-bool FShapeUpdateParams::IsValid() const
-{
-	return ControlRig && Hierarchy && Settings;
-}
-
-void UpdateControlShape(AControlRigShapeActor* InShapeActor, FRigControlElement* InControlElement, const FShapeUpdateParams& InUpdateParams)
-{
-	if (!InShapeActor || !InControlElement || !InUpdateParams.IsValid())
-	{
-		return;
-	}
-
-	// update transform
-	const FTransform Transform = InUpdateParams.Hierarchy->GetTransform(InControlElement, ERigTransformType::CurrentGlobal);
-	InShapeActor->SetActorTransform(Transform * InUpdateParams.ComponentTransform);
-
-	const FRigControlSettings& ControlSettings = InControlElement->Settings;
-
-	// update visibility & color
-	bool bIsVisible = ControlSettings.IsVisible();
-	bool bRespectVisibilityForSelection = true; 
-
-	const bool bControlsHiddenInViewport =
-		InUpdateParams.Settings->bHideControlShapes ||
-		!InUpdateParams.ControlRig->GetControlsVisible() ||
-		!InUpdateParams.bIsSkeletalMeshVisible;
-	
-	if (!bControlsHiddenInViewport)
-	{
-		if (ControlSettings.AnimationType == ERigControlAnimationType::ProxyControl)
-		{					
-			bRespectVisibilityForSelection = false;
-			if (InUpdateParams.Settings->bShowAllProxyControls)
-			{
-				bIsVisible = true;
-			}
-		}
-	}
-
-	InShapeActor->SetIsTemporarilyHiddenInEditor(!bIsVisible || bControlsHiddenInViewport);
-
-	// update color
-	InShapeActor->SetShapeColor(InShapeActor->OverrideColor.A < SMALL_NUMBER ?
-		ControlSettings.ShapeColor : InShapeActor->OverrideColor);
-	
-	// update selectability
-	InShapeActor->SetSelectable( ControlSettings.IsSelectable(bRespectVisibilityForSelection) );
-}
-	
-}
-
 #undef LOCTEXT_NAMESPACE
diff --git a/Engine/Plugins/Animation/ControlRig/Source/ControlRigEditor/Private/EditMode/ControlShapeActorHelper.cpp b/Engine/Plugins/Animation/ControlRig/Source/ControlRigEditor/Private/EditMode/ControlShapeActorHelper.cpp
new file mode 100644
--- /dev/null
+++ b/Engine/Plugins/Animation/ControlRig/Source/ControlRigEditor/Private/EditMode/ControlShapeActorHelper.cpp
@@ -0,0 +1,235 @@
+// Copyright Epic Games, Inc. All Rights Reserved.
+
+#include "ControlShapeActorHelper.h"
+
+#include "ControlRig.h"
+#include "ControlRigEditModeSettings.h"
+#include "ControlRigGizmoActor.h"
+#include "Rigs/RigHierarchyElements.h"
+
+namespace ControlRigEditMode::Shapes
+{
+
+bool IsSupportedControlType(const ERigControlType& ControlType)
+{
+	switch (ControlType)
+	{
+	case ERigControlType::Float:
+	case ERigControlType::ScaleFloat:
+	case ERigControlType::Integer:
+	case ERigControlType::Vector2D:
+	case ERigControlType::Position:
+	case ERigControlType::Scale:
+	case ERigControlType::Rotator:
+	case ERigControlType::Transform:
+	case ERigControlType::TransformNoScale:
+	case ERigControlType::EulerTransform:
+		{
+			return true;
+		}
+	default:
+		{
+			break;
+		}
+	}
+
+	return false;
+}
+
+bool IsModeSupported(const ERigControlType& InControlType, const UE::Widget::EWidgetMode& InMode)
+{
+	if (IsSupportedControlType(InControlType))
+	{
+		switch (InMode)
+		{
+			case UE::Widget::WM_None:
+				return true;
+			case UE::Widget::WM_Rotate:
+			{
+				switch (InControlType)
+				{
+					case ERigControlType::Rotator:
+					case ERigControlType::Transform:
+					case ERigControlType::TransformNoScale:
+					case ERigControlType::EulerTransform:
+					{
+						return true;
+					}
+					default:
+					{
+						break;
+					}
+				}
+				break;
+			}
+			case UE::Widget::WM_Translate:
+			{
+				switch (InControlType)
+				{
+					case ERigControlType::Float:
+					case ERigControlType::Integer:
+					case ERigControlType::Vector2D:
+					case ERigControlType::Position:
+					case ERigControlType::Transform:
+					case ERigControlType::TransformNoScale:
+					case ERigControlType::EulerTransform:
+					{
+						return true;
+					}
+					default:
+					{
+						break;
+					}
+				}
+				break;
+			}
+			case UE::Widget::WM_Scale:
+			{
+				switch (InControlType)
+				{
+					case ERigControlType::Scale:
+					case ERigControlType::ScaleFloat:
+					case ERigControlType::Transform:
+					case ERigControlType::EulerTransform:
+					{
+						return true;
+					}
+					default:
+					{
+						break;
+					}
+				}
+				break;
+			}
+			case UE::Widget::WM_TranslateRotateZ:
+			{
+				switch (InControlType)
+				{
+					case ERigControlType::Transform:
+					case ERigControlType::TransformNoScale:
+					case ERigControlType::EulerTransform:
+					{
+						return true;
+					}
+					default:
+					{
+						break;
+					}
+				}
+				break;
+			}
+		}
+	}
+
+	return false;
+}
+	
+void GetControlsEligibleForShapes(const UControlRig* InControlRig, TArray<FRigControlElement*>& OutControls)
+{
+	OutControls.Reset();
+
+	const URigHierarchy* Hierarchy = InControlRig ? InControlRig->GetHierarchy() : nullptr;
+	if (!Hierarchy)
+	{
+		return;
+	}
+
+	OutControls = Hierarchy->GetFilteredElements<FRigControlElement>([](const FRigControlElement* ControlElement)
+	{
+		const FRigControlSettings& ControlSettings = ControlElement->Settings;
+		return ControlSettings.SupportsShape() && IsSupportedControlType(ControlSettings.ControlType);
+	});
+}
+	
+void DestroyShapesActorsFromWorld(const TArray<TObjectPtr<AControlRigShapeActor>>& InShapeActorsToDestroy)
+{
+	// NOTE: should UWorld::EditorDestroyActor really modify the level when removing the shapes?
+	// kept for legacy but I guess this should be set to false
+	static constexpr bool bShouldModifyLevel = true;
+
+	for (const TObjectPtr<AControlRigShapeActor>& ShapeActorPtr: InShapeActorsToDestroy)
+	{
+		if (AControlRigShapeActor* ShapeActor = ShapeActorPtr.Get())
+		{
+			if (UWorld* World = ShapeActor->GetWorld())
+			{
+				if (ShapeActor->GetAttachParentActor())
+				{
+					ShapeActor->DetachFromActor(FDetachmentTransformRules::KeepRelativeTransform);
+				}
+				World->EditorDestroyActor(ShapeActor,bShouldModifyLevel);
+			}			
+		}
+	}
+}
+
+FShapeUpdateParams::FShapeUpdateParams(const UControlRig* InControlRig, const FTransform& InComponentTransform, const bool InSkeletalMeshVisible, const bool IsInLevelEditor)
+	: ControlRig(InControlRig)
+	, Hierarchy(InControlRig->GetHierarchy())
+	, Settings(GetDefault<UControlRigEditModeSettings>())
+	, ComponentTransform(InComponentTransform)
+	, bIsSkeletalMeshVisible(InSkeletalMeshVisible)
+	, bIsInLevelEditor(IsInLevelEditor)
+{
+	if (IsValid())
+	{
+		bControlsHiddenInViewport = Settings->bHideControlShapes || !ControlRig->GetControlsVisible() || !bIsSkeletalMeshVisible;
+	}
+}
+
+bool FShapeUpdateParams::IsValid() const
+{
+	return ControlRig && Hierarchy && Settings;
+}
+
+void UpdateControlShape(AControlRigShapeActor* InShapeActor, FRigControlElement* InControlElement, const FShapeUpdateParams& InUpdateParams)
+{
+	if (!InShapeActor || !InControlElement || !InUpdateParams.IsValid())
+	{
+		return;
+	}
+
+	if (InUpdateParams.bIsGameView)
+	{
+		InShapeActor->SetIsTemporarilyHiddenInEditor(true);
+		return;
+	}
+	
+	// update transform
+	const FTransform Transform = InUpdateParams.Hierarchy->GetTransform(InControlElement, ERigTransformType::CurrentGlobal);
+	InShapeActor->SetActorTransform(Transform * InUpdateParams.ComponentTransform);
+
+	const FRigControlSettings& ControlSettings = InControlElement->Settings;
+
+	// update visibility & color
+	bool bIsVisible = ControlSettings.IsVisible();
+	bool bRespectVisibilityForSelection = !InUpdateParams.bIsInLevelEditor;
+
+	const bool bControlsHiddenInViewport =
+		InUpdateParams.Settings->bHideControlShapes ||
+		!InUpdateParams.ControlRig->GetControlsVisible() ||
+		!InUpdateParams.bIsSkeletalMeshVisible;
+	
+	if (!bControlsHiddenInViewport)
+	{
+		if (ControlSettings.AnimationType == ERigControlAnimationType::ProxyControl)
+		{					
+			bRespectVisibilityForSelection = false;
+			if (InUpdateParams.Settings->bShowAllProxyControls)
+			{
+				bIsVisible = true;
+			}
+		}
+	}
+
+	InShapeActor->SetIsTemporarilyHiddenInEditor(!bIsVisible || InUpdateParams.bControlsHiddenInViewport);
+
+	// update color
+	InShapeActor->SetShapeColor(InShapeActor->OverrideColor.A < SMALL_NUMBER ?
+		ControlSettings.ShapeColor : InShapeActor->OverrideColor);
+	
+	// update selectability
+	InShapeActor->SetSelectable( ControlSettings.IsSelectable(bRespectVisibilityForSelection) );
+}
+	
+}
diff --git a/Engine/Plugins/Animation/ControlRig/Source/ControlRigEditor/Private/EditMode/ControlShapeActorHelper.h b/Engine/Plugins/Animation/ControlRig/Source/ControlRigEditor/Private/EditMode/ControlShapeActorHelper.h
new file mode 100644
--- /dev/null
+++ b/Engine/Plugins/Animation/ControlRig/Source/ControlRigEditor/Private/EditMode/ControlShapeActorHelper.h
@@ -0,0 +1,47 @@
+// Copyright Epic Games, Inc. All Rights Reserved.
+
+#pragma once
+
+#include "Math/Transform.h"
+#include "UObject/ObjectPtr.h"
+#include "Containers/Array.h"
+#include "UnrealWidgetFwd.h"
+
+enum class ERigControlType : uint8;
+class UControlRigEditModeSettings;
+class URigHierarchy;
+class AControlRigShapeActor;
+struct FRigControlElement;
+class UControlRig;
+
+namespace ControlRigEditMode::Shapes
+{
+	bool IsSupportedControlType(const ERigControlType& ControlType);
+
+	// Returns true if the control type supports being modified using this widget mode
+	bool IsModeSupported(const ERigControlType& InControlType, const UE::Widget::EWidgetMode& InMode);
+	
+	// Returns the list of controls for which a shape is expected
+	void GetControlsEligibleForShapes(const UControlRig* InControlRig, TArray<FRigControlElement*>& OutControls);
+
+	// Destroys shape actors and removes them from their UWorld
+	void DestroyShapesActorsFromWorld(const TArray<TObjectPtr<AControlRigShapeActor>>& InShapeActorsToDestroy);
+	
+	// Parameters used to update shape actors (transform, visibility, etc.)
+	struct FShapeUpdateParams
+	{
+		FShapeUpdateParams(const UControlRig* InControlRig, const FTransform& InComponentTransform, const bool InSkeletalMeshVisible, const bool IsInLevelEditor);
+		const UControlRig* ControlRig = nullptr;
+		const URigHierarchy* Hierarchy = nullptr;
+		const UControlRigEditModeSettings* Settings = nullptr;
+		const FTransform& ComponentTransform = FTransform::Identity;
+		const bool bIsSkeletalMeshVisible = false;
+		const bool bIsInLevelEditor = false;
+		bool bControlsHiddenInViewport = false;
+		bool bIsGameView = false;
+		bool IsValid() const;
+	};
+
+	// Updates shape actors transform, visibility, etc.
+	void UpdateControlShape(AControlRigShapeActor* InShapeActor, FRigControlElement* InControlElement, const FShapeUpdateParams& InUpdateParams);
+}
diff --git a/Engine/Plugins/Animation/ControlRig/Source/ControlRigEditor/Public/EditMode/ControlRigEditMode.h b/Engine/Plugins/Animation/ControlRig/Source/ControlRigEditor/Public/EditMode/ControlRigEditMode.h
--- a/Engine/Plugins/Animation/ControlRig/Source/ControlRigEditor/Public/EditMode/ControlRigEditMode.h
+++ b/Engine/Plugins/Animation/ControlRig/Source/ControlRigEditor/Public/EditMode/ControlRigEditMode.h
@@ -462,7 +462,6 @@ public:
 	/** Request a certain transform widget for the next update */
 	void RequestTransformWidgetMode(UE::Widget::EWidgetMode InWidgetMode);
 	
-	UControlRigDetailPanelControlProxies* GetControlProxy() { return ControlProxy; }
 private:
 	/** Whether or not Pivot Transforms have changed, in which case we need to redraw viewport*/
 	bool HasPivotTransformsChanged() const;
@@ -516,7 +515,6 @@ private:
 	void ChangeControlShapeTransform(AControlRigShapeActor* ShapeActor, const bool bTranslation, FVector& InDrag,
 		const bool bRotation, FRotator& InRot, const bool bScale, FVector& InScale, const FTransform& ToWorldTransform);
 
-	void TickControlShape(AControlRigShapeActor* ShapeActor, const FTransform& ComponentTransform) const;
 	bool ModeSupportedByShapeActor(const AControlRigShapeActor* ShapeActor, UE::Widget::EWidgetMode InMode) const;
 
 public:
@@ -540,7 +538,7 @@ private:
 
 	// Post pose update handler
 	void OnPoseInitialized();
-	void PostPoseUpdate() const;
+	void PostPoseUpdate(const FEditorViewportClient* InViewportClient = nullptr) const;
 	void UpdateSelectabilityOnSkeletalMeshes(UControlRig* InControlRig, bool bEnabled);
 
 	bool IsMovingCamera(const FViewport* InViewport) const;
@@ -608,28 +606,4 @@ private:
 	friend class FControlRigEditModeGenericDetails;
 	friend class UControlRigEditModeDelegateHelper;
 	friend class SControlRigEditModeTools;
-};
-
-namespace ControlRigEditMode::Shapes
-{
-	// Returns the list of controls for which a shape is expected
-	void GetControlsEligibleForShapes(UControlRig* InControlRig, TArray<FRigControlElement*>& OutControls);
-
-	// Destroys shape actors and removes them from their UWorld
-	void DestroyShapesActorsFromWorld(const TArray<TObjectPtr<AControlRigShapeActor>>& InShapeActorsToDestroy);
-
-	// Parameters used to update shape actors (transform, visibility, etc.)
-	struct FShapeUpdateParams
-	{
-		FShapeUpdateParams(const UControlRig* InControlRig, const FTransform& InComponentTransform, const bool InSkeletalMeshVisible);
-		const UControlRig* ControlRig = nullptr;
-		const URigHierarchy* Hierarchy = nullptr;
-		const UControlRigEditModeSettings* Settings = nullptr;
-		const FTransform& ComponentTransform = FTransform::Identity;
-		const bool bIsSkeletalMeshVisible = false;
-		bool IsValid() const;
-	};
-
-	// Updates shape actors transform, visibility, etc.
-	void UpdateControlShape(AControlRigShapeActor* InShapeActor, FRigControlElement* InControlElement, const FShapeUpdateParams& InUpdateParams);
-}
\ No newline at end of file
+};
\ No newline at end of file
