Control Rig: Automatically hide the schematic if a control is selected

#rnx

diff --git a/Engine/Plugins/Animation/ControlRig/Source/ControlRigEditor/Private/Editor/ControlRigEditor.cpp b/Engine/Plugins/Animation/ControlRig/Source/ControlRigEditor/Private/Editor/ControlRigEditor.cpp
--- a/Engine/Plugins/Animation/ControlRig/Source/ControlRigEditor/Private/Editor/ControlRigEditor.cpp
+++ b/Engine/Plugins/Animation/ControlRig/Source/ControlRigEditor/Private/Editor/ControlRigEditor.cpp
@@ -2418,6 +2418,7 @@ void FControlRigEditor::HandleViewportCreated(const TSharedRef<class IPersonaVie
 																.PaddingTop(60)
 																.PaddingBottom(60)
 																.PaddingInterNode(5)
+																.Visibility(this, &FControlRigEditor::GetSchematicOverlayVisibility)
 				;
 				InViewport->AddOverlayWidget(SchematicViewport.ToSharedRef());
 
@@ -2548,6 +2549,22 @@ bool FControlRigEditor::IsSchematicViewportActive() const
 	return false;
 }
 
+EVisibility FControlRigEditor::GetSchematicOverlayVisibility() const
+{
+	if(const URigHierarchy* Hierarchy = GetHierarchyBeingDebugged())
+	{
+		TArray<const FRigBaseElement*> SelectedElements = Hierarchy->GetSelectedElements();
+		if(SelectedElements.ContainsByPredicate([](const FRigBaseElement* InSelectedElement) -> bool
+		{
+			return InSelectedElement->IsA<FRigControlElement>();
+		}))
+		{
+			return EVisibility::Hidden;
+		}
+	}
+	return EVisibility::HitTestInvisible;
+}
+
 bool FControlRigEditor::GetToolbarDrawAxesOnSelection() const
 {
 	if (const UControlRigEditModeSettings* Settings = GetDefault<UControlRigEditModeSettings>())
diff --git a/Engine/Plugins/Animation/ControlRig/Source/ControlRigEditor/Private/Editor/ControlRigEditor.h b/Engine/Plugins/Animation/ControlRig/Source/ControlRigEditor/Private/Editor/ControlRigEditor.h
--- a/Engine/Plugins/Animation/ControlRig/Source/ControlRigEditor/Private/Editor/ControlRigEditor.h
+++ b/Engine/Plugins/Animation/ControlRig/Source/ControlRigEditor/Private/Editor/ControlRigEditor.h
@@ -258,6 +258,7 @@ private:
 	void OnToolbarAxesScaleChanged(float InValue);
 	void HandleToggleSchematicViewport();
 	bool IsSchematicViewportActive() const;
+	EVisibility GetSchematicOverlayVisibility() const;
 
 		/** Handle switching skeletal meshes */
 	void HandlePreviewMeshChanged(USkeletalMesh* InOldSkeletalMesh, USkeletalMesh* InNewSkeletalMesh);
diff --git a/Engine/Source/Editor/AnimationEditorWidgets/Private/SchematicGraphPanel/SSchematicGraphPanel.cpp b/Engine/Source/Editor/AnimationEditorWidgets/Private/SchematicGraphPanel/SSchematicGraphPanel.cpp
--- a/Engine/Source/Editor/AnimationEditorWidgets/Private/SchematicGraphPanel/SSchematicGraphPanel.cpp
+++ b/Engine/Source/Editor/AnimationEditorWidgets/Private/SchematicGraphPanel/SSchematicGraphPanel.cpp
@@ -594,7 +594,15 @@ void SSchematicGraphPanel::Construct(const FArguments& InArgs)
 	
 	SNodePanel::Construct();
 
-	SetVisibility(bIsOverlay ? EVisibility::SelfHitTestInvisible : EVisibility::Visible);
+	if(InArgs._Visibility.IsBound() || InArgs._Visibility.IsSet())
+	{
+		SetVisibility(InArgs._Visibility);
+	}
+	else
+	{
+		SetVisibility(bIsOverlay ? EVisibility::SelfHitTestInvisible : EVisibility::Visible);
+	}
+	
 	if (GraphData)
 	{
 		GraphData->OnNodeAdded().AddSP(this, &SSchematicGraphPanel::AddNode);
