Control Rig: New connectors should be optional by default

#rnx

diff --git a/Engine/Plugins/Animation/ControlRig/Source/ControlRigEditor/Private/Editor/SRigHierarchy.cpp b/Engine/Plugins/Animation/ControlRig/Source/ControlRigEditor/Private/Editor/SRigHierarchy.cpp
--- a/Engine/Plugins/Animation/ControlRig/Source/ControlRigEditor/Private/Editor/SRigHierarchy.cpp
+++ b/Engine/Plugins/Animation/ControlRig/Source/ControlRigEditor/Private/Editor/SRigHierarchy.cpp
@@ -2227,11 +2227,12 @@ void SRigHierarchy::HandleNewItem(ERigElementType InElementType, bool bIsAnimati
 						
 						FRigConnectorSettings Settings;
 						Settings.Type = bIsPrimary ? EConnectorType::Primary : EConnectorType::Secondary;
-							if(!bIsPrimary)
-							{
-								Settings.Rules.Reset();
-								Settings.AddRule(FRigChildOfPrimaryConnectionRule());
-							}
+						if(!bIsPrimary)
+						{
+							Settings.Rules.Reset();
+							Settings.AddRule(FRigChildOfPrimaryConnectionRule());
+							Settings.bOptional = true;
+						}
 						NewItemKey = Controller->AddConnector(NewElementName, Settings, true);
 						(void)ResolveConnector(NewItemKey, ParentKey);
 						break;
