Prevent control rig debug draw displaying in game view mode when high resolution screenshots are taken

diff --git a/Engine/Plugins/Animation/ControlRig/Source/ControlRigEditor/Private/EditMode/ControlRigEditMode.cpp b/Engine/Plugins/Animation/ControlRig/Source/ControlRigEditor/Private/EditMode/ControlRigEditMode.cpp
--- a/Engine/Plugins/Animation/ControlRig/Source/ControlRigEditor/Private/EditMode/ControlRigEditMode.cpp
+++ b/Engine/Plugins/Animation/ControlRig/Source/ControlRigEditor/Private/EditMode/ControlRigEditMode.cpp
@@ -604,15 +604,6 @@ void FControlRigEditMode::Tick(FEditorViewportClient* ViewportClient, float Delt
 	}
 	CheckMovieSceneSig();
 
-	if (bool* GameView = ViewportToGameView.Find(ViewportClient->Viewport))
-	{
-		*GameView = ViewportClient->IsInGameView();
-	} 
-	else
-	{
-		ViewportToGameView.Add(ViewportClient->Viewport, ViewportClient->IsInGameView());
-	}
-
 	if(DeferredItemsToFrame.Num() > 0)
 	{
 		TGuardValue<FEditorViewportClient*> ViewportGuard(CurrentViewportClient, ViewportClient);
@@ -844,7 +835,8 @@ void FControlRigEditMode::Render(const FSceneView* View, FViewport* Viewport, FP
 	DragToolHandler.Render3DDragTool(View, PDI);
 
 	const UControlRigEditModeSettings* Settings = GetDefault<UControlRigEditModeSettings>();
-	const bool bIsInGameView = !AreEditingControlRigDirectly() ? (ViewportToGameView.Find(Viewport) && ViewportToGameView[Viewport]) : false;
+	FEditorViewportClient* EditorViewportClient = (FEditorViewportClient*)Viewport->GetClient();
+	const bool bIsInGameView = !AreEditingControlRigDirectly() ? EditorViewportClient && EditorViewportClient->IsInGameView() : false;
 	bool bRender = !Settings->bHideControlShapes;
 	for (TWeakObjectPtr<UControlRig>& ControlRigPtr : RuntimeControlRigs)
 	{
diff --git a/Engine/Plugins/Animation/ControlRig/Source/ControlRigEditor/Public/EditMode/ControlRigEditMode.h b/Engine/Plugins/Animation/ControlRig/Source/ControlRigEditor/Public/EditMode/ControlRigEditMode.h
--- a/Engine/Plugins/Animation/ControlRig/Source/ControlRigEditor/Public/EditMode/ControlRigEditMode.h
+++ b/Engine/Plugins/Animation/ControlRig/Source/ControlRigEditor/Public/EditMode/ControlRigEditMode.h
@@ -584,9 +584,6 @@ private:
 	TArray<TWeakObjectPtr<UControlRig>> RuntimeControlRigs;
 	TMap<UControlRig*,TStrongObjectPtr<UControlRigEditModeDelegateHelper>> DelegateHelpers;
 
-	//hack since we can't get the viewport client from the viewport, so in the tick we set the gameview bool and then in render/tickcontrolshapes we use it.
-	TMap<FViewport*, bool>  ViewportToGameView;
-
 	TArray<FRigElementKey> DeferredItemsToFrame;
 
 	/** Computes the current interaction types based on the widget mode */
