Control Rig: Rely on skeletal mesh for curves import where possible

#rnx

diff --git a/Engine/Plugins/Animation/ControlRig/Source/ControlRig/Private/Rigs/RigHierarchyController.cpp b/Engine/Plugins/Animation/ControlRig/Source/ControlRig/Private/Rigs/RigHierarchyController.cpp
--- a/Engine/Plugins/Animation/ControlRig/Source/ControlRig/Private/Rigs/RigHierarchyController.cpp
+++ b/Engine/Plugins/Animation/ControlRig/Source/ControlRig/Private/Rigs/RigHierarchyController.cpp
@@ -6,13 +6,13 @@
 #include "UObject/Package.h"
 #include "ModularRig.h"
 #include "HelperUtil.h"
+#include "Engine/SkeletalMesh.h"
 
 #if WITH_EDITOR
 #include "Framework/Notifications/NotificationManager.h"
 #include "Widgets/Notifications/SNotificationList.h"
 #include "Styling/AppStyle.h"
 #include "ScopedTransaction.h"
-#include "Engine/SkeletalMesh.h"
 #include "RigVMPythonUtils.h"
 #endif
 
@@ -1076,8 +1076,12 @@ TArray<FRigElementKey> URigHierarchyController::ImportBonesFromAsset(FString InA
 }
 
 TArray<FRigElementKey> URigHierarchyController::ImportCurvesFromAsset(FString InAssetPath, FName InNameSpace,
-	bool bSelectCurves, bool bSetupUndo)
+                                                                      bool bSelectCurves, bool bSetupUndo)
 {
+	if(USkeletalMesh* SkeletalMesh = GetSkeletalMeshFromAssetPath(InAssetPath))
+	{
+		return ImportCurvesFromSkeletalMesh(SkeletalMesh, InNameSpace, bSelectCurves, bSetupUndo);
+	}
 	if(USkeleton* Skeleton = GetSkeletonFromAssetPath(InAssetPath))
 	{
 		return ImportCurves(Skeleton, InNameSpace, bSelectCurves, bSetupUndo);
@@ -1085,6 +1089,22 @@ TArray<FRigElementKey> URigHierarchyController::ImportCurvesFromAsset(FString In
 	return TArray<FRigElementKey>();
 }
 
+USkeletalMesh* URigHierarchyController::GetSkeletalMeshFromAssetPath(const FString& InAssetPath)
+{
+	UObject* AssetObject = StaticLoadObject(UObject::StaticClass(), NULL, *InAssetPath, NULL, LOAD_None, NULL);
+	if(AssetObject == nullptr)
+	{
+		return nullptr;
+	}
+
+	if(USkeletalMesh* SkeletalMesh = Cast<USkeletalMesh>(AssetObject))
+	{
+		return SkeletalMesh;
+	}
+
+	return nullptr;
+}
+
 USkeleton* URigHierarchyController::GetSkeletonFromAssetPath(const FString& InAssetPath)
 {
 	UObject* AssetObject = StaticLoadObject(UObject::StaticClass(), NULL, *InAssetPath, NULL, LOAD_None, NULL);
@@ -1108,23 +1128,23 @@ USkeleton* URigHierarchyController::GetSkeletonFromAssetPath(const FString& InAs
 
 #endif
 
-TArray<FRigElementKey> URigHierarchyController::ImportCurves(USkeleton* InSkeleton, FName InNameSpace,
-                                                             bool bSelectCurves, bool bSetupUndo, bool bPrintPythonCommand)
+TArray<FRigElementKey> URigHierarchyController::ImportCurves(UAnimCurveMetaData* InAnimCurvesMetadata, FName InNameSpace, bool bSetupUndo)
 {
-	if (InSkeleton == nullptr)
+	TArray<FRigElementKey> Keys;
+	if (InAnimCurvesMetadata == nullptr)
 	{
-		return TArray<FRigElementKey>();
+		return Keys;
 	}
 
-	TArray<FRigElementKey> Keys;
 	if(!IsValid())
 	{
 		return Keys;
 	}
 
 	URigHierarchy* Hierarchy = GetHierarchy();
+	FRigHierarchyInteractionBracket InteractionBracket(Hierarchy);
 
-	InSkeleton->ForEachCurveMetaData([this, Hierarchy, InNameSpace, &Keys, bSetupUndo](const FName& InCurveName, const FCurveMetaData& InMetaData)
+	InAnimCurvesMetadata->ForEachCurveMetaData([this, Hierarchy, InNameSpace, &Keys, bSetupUndo](const FName& InCurveName, const FCurveMetaData& InMetaData)
 	{
 		FName Name = InCurveName;
 		if (!InNameSpace.IsNone())
@@ -1143,6 +1163,37 @@ TArray<FRigElementKey> URigHierarchyController::ImportCurves(USkeleton* InSkelet
 		Keys.Add(FRigElementKey(Name, ERigElementType::Curve));
 	});
 
+	return Keys;
+}
+
+TArray<FRigElementKey> URigHierarchyController::ImportCurves(USkeleton* InSkeleton, FName InNameSpace,
+                                                             bool bSelectCurves, bool bSetupUndo, bool bPrintPythonCommand)
+{
+	TArray<FRigElementKey> Keys;
+	if (InSkeleton == nullptr)
+	{
+		return Keys;
+	}
+
+	if(!IsValid())
+	{
+		return Keys;
+	}
+
+	URigHierarchy* Hierarchy = GetHierarchy();
+	FRigHierarchyInteractionBracket InteractionBracket(Hierarchy);
+
+#if WITH_EDITOR
+	TSharedPtr<FScopedTransaction> TransactionPtr;
+	if(bSetupUndo)
+	{
+		TransactionPtr = MakeShared<FScopedTransaction>(NSLOCTEXT("RigHierarchyController", "Import Curves", "Import Curves"));
+		Hierarchy->Modify();
+	}
+#endif
+
+	Keys.Append(ImportCurves(InSkeleton->GetAssetUserData<UAnimCurveMetaData>(), InNameSpace, bSetupUndo));
+
 	if(bSelectCurves)
 	{
 		SetSelection(Keys);
@@ -1168,6 +1219,60 @@ TArray<FRigElementKey> URigHierarchyController::ImportCurves(USkeleton* InSkelet
 	return Keys;
 }
 
+TArray<FRigElementKey> URigHierarchyController::ImportCurvesFromSkeletalMesh(USkeletalMesh* InSkeletalMesh, FName InNameSpace,
+	bool bSelectCurves, bool bSetupUndo, bool bPrintPythonCommand)
+{
+	TArray<FRigElementKey> Keys;
+	if(InSkeletalMesh == nullptr)
+	{
+		return Keys;
+	}
+
+	if(!IsValid())
+	{
+		return Keys;
+	}
+
+	URigHierarchy* Hierarchy = GetHierarchy();
+	FRigHierarchyInteractionBracket InteractionBracket(Hierarchy);
+
+#if WITH_EDITOR
+	TSharedPtr<FScopedTransaction> TransactionPtr;
+	if(bSetupUndo)
+	{
+		TransactionPtr = MakeShared<FScopedTransaction>(NSLOCTEXT("RigHierarchyController", "Import Curves", "Import Curves"));
+		Hierarchy->Modify();
+	}
+#endif
+
+	Keys.Append(ImportCurves(InSkeletalMesh->GetSkeleton(), InNameSpace, false, bSetupUndo, false));
+	Keys.Append(ImportCurves(InSkeletalMesh->GetAssetUserData<UAnimCurveMetaData>(), InNameSpace, bSetupUndo));
+	
+	if(bSelectCurves)
+	{
+		SetSelection(Keys);
+	}
+	
+#if WITH_EDITOR
+	if (!Keys.IsEmpty() && bPrintPythonCommand && !bSuspendPythonPrinting)
+	{
+		UBlueprint* Blueprint = GetTypedOuter<UBlueprint>();
+		if (Blueprint)
+		{
+			RigVMPythonUtils::Print(Blueprint->GetFName().ToString(),
+				FString::Printf(TEXT("hierarchy_controller.import_curves_from_asset('%s', '%s', %s)"),
+				*InSkeletalMesh->GetPathName(),
+				*InNameSpace.ToString(),
+				(bSelectCurves) ? TEXT("True") : TEXT("False")));
+		}
+	}
+#endif
+
+	Hierarchy->EnsureCacheValidity();
+
+	return Keys;
+}
+
 FString URigHierarchyController::ExportSelectionToText() const
 {
 	if(!IsValid())
diff --git a/Engine/Plugins/Animation/ControlRig/Source/ControlRig/Public/Rigs/RigHierarchyController.h b/Engine/Plugins/Animation/ControlRig/Source/ControlRig/Public/Rigs/RigHierarchyController.h
--- a/Engine/Plugins/Animation/ControlRig/Source/ControlRig/Public/Rigs/RigHierarchyController.h
+++ b/Engine/Plugins/Animation/ControlRig/Source/ControlRig/Public/Rigs/RigHierarchyController.h
@@ -343,6 +343,18 @@ public:
         bool bSetupUndo = false);
 #endif
 
+	/**
+	 * Imports all curves from an anim curve metadata object to the hierarchy
+	 * @param InAnimCurvesMetadata The anim curve metadata object to import the curves from
+	 * @param InNameSpace The namespace to prefix the bone names with
+	 * @param bSetupUndo If set to true the stack will record the change for undo / redo
+	 * @return The keys of the imported elements
+	 */
+	TArray<FRigElementKey> ImportCurves(
+		UAnimCurveMetaData* InAnimCurvesMetadata, 
+		FName InNameSpace = NAME_None,
+		bool bSetupUndo = false);
+
 	/**
 	 * Imports all curves from a skeleton to the hierarchy
 	 * @param InSkeleton The skeleton to import the curves from
@@ -359,6 +371,22 @@ public:
 		bool bSetupUndo = false,
 		bool bPrintPythonCommand = false);
 
+	/**
+	 * Imports all curves from a skeletalmesh to the hierarchy
+	 * @param InSkeletalMesh The skeletalmesh to import the curves from
+	 * @param InNameSpace The namespace to prefix the bone names with
+	 * @param bSelectCurves If true the curves will be selected upon import
+	 * @param bSetupUndo If set to true the stack will record the change for undo / redo
+	 * @return The keys of the imported elements
+	 */
+	UFUNCTION(BlueprintCallable, Category = URigHierarchyController)
+	TArray<FRigElementKey> ImportCurvesFromSkeletalMesh(
+		USkeletalMesh* InSkeletalMesh, 
+		FName InNameSpace = NAME_None,  
+		bool bSelectCurves = false,
+		bool bSetupUndo = false,
+		bool bPrintPythonCommand = false);
+
 #if WITH_EDITOR
 	/**
 	 * Imports all curves from a skeleton to the hierarchy
@@ -821,6 +849,7 @@ private:
 	void RemoveElementToDirty(FRigBaseElement* InParent, FRigBaseElement* InElementToRemove) const;
 
 #if WITH_EDITOR
+	static USkeletalMesh* GetSkeletalMeshFromAssetPath(const FString& InAssetPath);
 	static USkeleton* GetSkeletonFromAssetPath(const FString& InAssetPath);
 #endif
 
diff --git a/Engine/Plugins/Animation/ControlRig/Source/ControlRigEditor/Private/ControlRigBlueprintActions.cpp b/Engine/Plugins/Animation/ControlRig/Source/ControlRigEditor/Private/ControlRigBlueprintActions.cpp
--- a/Engine/Plugins/Animation/ControlRig/Source/ControlRigEditor/Private/ControlRigBlueprintActions.cpp
+++ b/Engine/Plugins/Animation/ControlRig/Source/ControlRigEditor/Private/ControlRigBlueprintActions.cpp
@@ -275,7 +275,14 @@ UControlRigBlueprint* FControlRigBlueprintActions::CreateControlRigFromSkeletalM
 	if(URigHierarchyController* Controller = NewControlRigBlueprint->GetHierarchyController())
 	{
 		Controller->ImportBones(*RefSkeleton, NAME_None, false, false, false, false);
-		Controller->ImportCurves(Skeleton, NAME_None, false, false);
+		if(SkeletalMesh)
+		{
+			Controller->ImportCurvesFromSkeletalMesh(SkeletalMesh, NAME_None, false, false);
+		}
+		else
+		{
+			Controller->ImportCurves(Skeleton, NAME_None, false, false);
+		}
 	}
 	NewControlRigBlueprint->SourceHierarchyImport = Skeleton;
 	NewControlRigBlueprint->SourceCurveImport = Skeleton;
diff --git a/Engine/Plugins/Animation/ControlRig/Source/ControlRigEditor/Private/Editor/ControlRigEditor.cpp b/Engine/Plugins/Animation/ControlRig/Source/ControlRigEditor/Private/Editor/ControlRigEditor.cpp
--- a/Engine/Plugins/Animation/ControlRig/Source/ControlRigEditor/Private/Editor/ControlRigEditor.cpp
+++ b/Engine/Plugins/Animation/ControlRig/Source/ControlRigEditor/Private/Editor/ControlRigEditor.cpp
@@ -2923,7 +2923,14 @@ void FControlRigEditor::HandlePreviewMeshChanged(USkeletalMesh* InOldSkeletalMes
 						
 						USkeleton* Skeleton = InNewSkeletalMesh ? InNewSkeletalMesh->GetSkeleton() : nullptr;
 						Controller->ImportBones(Skeleton, NAME_None, true, true, false, true, true);
-						Controller->ImportCurves(Skeleton, NAME_None, false, true, true);
+						if(InNewSkeletalMesh)
+						{
+							Controller->ImportCurvesFromSkeletalMesh(InNewSkeletalMesh, NAME_None, false, true, true);
+						}
+						else
+						{
+							Controller->ImportCurves(Skeleton, NAME_None, false, true, true);
+						}
 
 						// Recreate manually created elements
 						for (const TTuple<FRigElementKey, FRigElementKey, FTransform>& Tuple : ConnectorsAndSocketsToParents)
diff --git a/Engine/Plugins/Animation/ControlRig/Source/ControlRigEditor/Private/Editor/SRigCurveContainer.cpp b/Engine/Plugins/Animation/ControlRig/Source/ControlRigEditor/Private/Editor/SRigCurveContainer.cpp
--- a/Engine/Plugins/Animation/ControlRig/Source/ControlRigEditor/Private/Editor/SRigCurveContainer.cpp
+++ b/Engine/Plugins/Animation/ControlRig/Source/ControlRigEditor/Private/Editor/SRigCurveContainer.cpp
@@ -686,7 +686,8 @@ void SRigCurveContainer::ImportCurve(const FAssetData& InAssetData)
 		TGuardValue<bool> SuspendBlueprintNotifs(ControlRigBlueprint->bSuspendAllNotifications, true);
 
 		USkeleton* Skeleton = nullptr;
-		if (USkeletalMesh* Mesh = Cast<USkeletalMesh>(InAssetData.GetAsset()))
+		USkeletalMesh* Mesh = Cast<USkeletalMesh>(InAssetData.GetAsset());
+		if (Mesh)
 		{
 			Skeleton = Mesh->GetSkeleton();
 			ControlRigBlueprint->SourceCurveImport = Skeleton;
@@ -707,7 +708,14 @@ void SRigCurveContainer::ImportCurve(const FAssetData& InAssetData)
 			if(URigHierarchyController* Controller = Hierarchy->GetController())
 			{
 				Controller->ClearSelection();
-				Controller->ImportCurves(Skeleton, NAME_None, false, true, true);
+				if(Mesh)
+				{
+					Controller->ImportCurvesFromSkeletalMesh(Mesh, NAME_None, false, true, true);
+				}
+				else
+				{
+					Controller->ImportCurves(Skeleton, NAME_None, false, true, true);
+				}
 			}
 
 			FSlateApplication::Get().DismissAllMenus();
diff --git a/Engine/Plugins/Animation/ControlRig/Source/ControlRigEditor/Private/Editor/SRigHierarchy.cpp b/Engine/Plugins/Animation/ControlRig/Source/ControlRigEditor/Private/Editor/SRigHierarchy.cpp
--- a/Engine/Plugins/Animation/ControlRig/Source/ControlRigEditor/Private/Editor/SRigHierarchy.cpp
+++ b/Engine/Plugins/Animation/ControlRig/Source/ControlRigEditor/Private/Editor/SRigHierarchy.cpp
@@ -1656,12 +1656,12 @@ void SRigHierarchy::RefreshHierarchy(const FAssetData& InAssetData, bool bOnlyRe
 			{
 				Controller->RemoveElement(CurveKey, true, true);
 			}			
-			Controller->ImportCurves(Mesh->GetSkeleton(), NAME_None, false, true, true);
+			Controller->ImportCurvesFromSkeletalMesh(Mesh, NAME_None, false, true, true);
 		}
 		else
 		{
 			Controller->ImportBones(Mesh->GetSkeleton(), NAME_None, true, true, bSelectBones, true, true);
-			Controller->ImportCurves(Mesh->GetSkeleton(), NAME_None, false, true, true);
+			Controller->ImportCurvesFromSkeletalMesh(Mesh, NAME_None, false, true, true);
 		}
 	}
 
@@ -1786,7 +1786,7 @@ void SRigHierarchy::ImportHierarchy(const FAssetData& InAssetData)
 		check(Controller);
 
 		const TArray<FRigElementKey> ImportedBones = Controller->ImportBones(Mesh->GetSkeleton(), NAME_None, false, false, bSelectBones, true, true);
-		Controller->ImportCurves(Mesh->GetSkeleton(), NAME_None, false, true);
+		Controller->ImportCurvesFromSkeletalMesh(Mesh, NAME_None, false, true);
 
 		ControlRigBlueprint->SourceHierarchyImport = Mesh->GetSkeleton();
 		ControlRigBlueprint->SourceCurveImport = Mesh->GetSkeleton();
