Control Rig: Include nulls when resetting the pose during construction

#rnx

diff --git a/Engine/Plugins/Animation/ControlRig/Source/ControlRig/Public/Rigs/RigHierarchyDefines.h b/Engine/Plugins/Animation/ControlRig/Source/ControlRig/Public/Rigs/RigHierarchyDefines.h
--- a/Engine/Plugins/Animation/ControlRig/Source/ControlRig/Public/Rigs/RigHierarchyDefines.h
+++ b/Engine/Plugins/Animation/ControlRig/Source/ControlRig/Public/Rigs/RigHierarchyDefines.h
@@ -35,7 +35,7 @@ enum class ERigElementType : uint8
 	First = Bone UMETA(Hidden), 
 	Last = Socket UMETA(Hidden), 
 	All = Bone | Null | Control | Curve | Physics | Reference | Connector | Socket,
-	ToResetAfterConstructionEvent = Bone | Control | Curve | Socket UMETA(Hidden),
+	ToResetAfterConstructionEvent = Bone | Control | Null | Curve | Socket UMETA(Hidden),
 };
 
 UENUM(BlueprintType)
