Control Rig: Query optimizatoins cvar during initialize instance

#rnx

diff --git a/Engine/Plugins/Animation/ControlRig/Source/ControlRig/Private/AnimNode_ControlRigBase.cpp b/Engine/Plugins/Animation/ControlRig/Source/ControlRig/Private/AnimNode_ControlRigBase.cpp
--- a/Engine/Plugins/Animation/ControlRig/Source/ControlRig/Private/AnimNode_ControlRigBase.cpp
+++ b/Engine/Plugins/Animation/ControlRig/Source/ControlRig/Private/AnimNode_ControlRigBase.cpp
@@ -379,6 +379,8 @@ void FAnimNode_ControlRigBase::OnInitializeAnimInstance(const FAnimInstanceProxy
 	DECLARE_SCOPE_HIERARCHICAL_COUNTER_FUNC()
 
 	FAnimNode_CustomProperty::OnInitializeAnimInstance(InProxy, InAnimInstance);
+	
+	bEnablePoseAdapter = (CVarControlRigEnableAnimNodePerformanceOptimizations->GetInt() != 0);
 
 	WeakAnimInstanceObject = TWeakObjectPtr<const UAnimInstance>(InAnimInstance);
 	if(bEnablePoseAdapter)
