Sequencer: Anim Details:  Show anim details for sequencer selected objects/components not editor selected, per animator feedback

diff --git a/Engine/Plugins/Animation/ControlRig/Source/ControlRigEditor/Private/EditMode/SControlRigDetails.cpp b/Engine/Plugins/Animation/ControlRig/Source/ControlRigEditor/Private/EditMode/SControlRigDetails.cpp
--- a/Engine/Plugins/Animation/ControlRig/Source/ControlRigEditor/Private/EditMode/SControlRigDetails.cpp
+++ b/Engine/Plugins/Animation/ControlRig/Source/ControlRigEditor/Private/EditMode/SControlRigDetails.cpp
@@ -117,22 +117,11 @@ void SControlRigDetails::SelectedSequencerObjects(const TMap<UObject*, FArrayOfP
 	TMap<UObject*, FArrayOfPropertyTracks> SequencerObjects;
 	for (const TPair<UObject*, FArrayOfPropertyTracks>& Pair : InObjectsTracked)
 	{
-		if (AActor* Actor = Cast<AActor>(Pair.Key))
+		if(Pair.Key && (Pair.Key->IsA<AActor>() || Pair.Key->IsA<UActorComponent>()))
 		{
-			if (Actor->IsSelectedInEditor())
-			{
-				SequencerObjects.Add(Pair);
-			}
-		}
-		else if (UActorComponent* Component = Cast<UActorComponent>(Pair.Key))
-		{
-			if (Component->IsSelectedInEditor())
-			{
-				SequencerObjects.Add(Pair);
-			}
+			SequencerObjects.Add(Pair);
 		}
 	}
-	//make sure the objects that are selected are actually selected in the world
 
 	HandleSequencerObjects(SequencerObjects);
 	UpdateProxies();
