Sequencer: Anim Details: Fix them showing multiples with controls do to Sequencer object also being selected, regression from recent submission

diff --git a/Engine/Plugins/Animation/ControlRig/Source/ControlRigEditor/Private/EditMode/AnimDetailsProxy.cpp b/Engine/Plugins/Animation/ControlRig/Source/ControlRigEditor/Private/EditMode/AnimDetailsProxy.cpp
--- a/Engine/Plugins/Animation/ControlRig/Source/ControlRigEditor/Private/EditMode/AnimDetailsProxy.cpp
+++ b/Engine/Plugins/Animation/ControlRig/Source/ControlRigEditor/Private/EditMode/AnimDetailsProxy.cpp
@@ -3053,12 +3053,55 @@ void UControlRigDetailPanelControlProxies::ResetSequencerProxies(TMap<ERigContro
 	}
 }
 
-const TArray<UControlRigControlsProxy*> UControlRigDetailPanelControlProxies::GetAllSelectedProxies() const
+const TArray<UControlRigControlsProxy*> UControlRigDetailPanelControlProxies::GetAllSelectedProxies() 
 {
 	TArray<UControlRigControlsProxy*> SelectedProxies(SelectedControlRigProxies);
+	TArray<USceneComponent*> BoundCRObjects;
+	for (UControlRigControlsProxy* Proxy : SelectedProxies)
+	{
+		if (Proxy && Proxy->OwnerControlRig.IsValid())
+		{
+			if (USceneComponent* SceneComponent = Cast<USceneComponent>(Proxy->OwnerControlRig->GetObjectBinding()->GetBoundObject()))
+			{
+				BoundCRObjects.Add(SceneComponent);
+			}
+		}
+	}
 	if (SelectedSequencerProxies.Num() > 0)
 	{
-		SelectedProxies.Append(SelectedSequencerProxies);
+		if (SelectedProxies.Num() > 0)
+		{
+			for (int32 Index = SelectedSequencerProxies.Num() -1; Index >=0; --Index)
+			{
+				if (UControlRigControlsProxy* Proxy = SelectedSequencerProxies[Index])
+				{
+					if (Proxy->OwnerObject.IsValid())
+					{
+						if (USceneComponent* SceneComponent = Cast<USceneComponent>(Proxy->OwnerObject.Get()))
+						{
+							if (BoundCRObjects.Contains(SceneComponent))
+							{
+								SelectedSequencerProxies.RemoveAt(Index);
+								continue;
+							}
+						}
+						if (AActor* Actor = Cast<AActor>(Proxy->OwnerObject.Get()))
+						{
+							if (BoundCRObjects.Contains(Actor->GetRootComponent()))
+							{
+								SelectedSequencerProxies.RemoveAt(Index);
+								continue;
+							}
+						}
+					}
+					SelectedProxies.Add(Proxy);
+				}
+			}
+		}
+		else //no CR selected so just add them all
+		{
+			SelectedProxies.Append(SelectedSequencerProxies);
+		}
 	}
 	return SelectedProxies;
 }
diff --git a/Engine/Plugins/Animation/ControlRig/Source/ControlRigEditor/Private/EditMode/ControlRigControlsProxy.h b/Engine/Plugins/Animation/ControlRig/Source/ControlRigEditor/Private/EditMode/ControlRigControlsProxy.h
--- a/Engine/Plugins/Animation/ControlRig/Source/ControlRigEditor/Private/EditMode/ControlRigControlsProxy.h
+++ b/Engine/Plugins/Animation/ControlRig/Source/ControlRigEditor/Private/EditMode/ControlRigControlsProxy.h
@@ -205,7 +205,7 @@ public:
 	void RemoveAllProxies();
 	void ProxyChanged(UControlRig* InControlRig, FRigControlElement* RigElement, bool bModify = true);
 	void RecreateAllProxies(UControlRig* InControlRig);
-	const TArray<UControlRigControlsProxy*> GetAllSelectedProxies() const;
+	const TArray<UControlRigControlsProxy*> GetAllSelectedProxies();
 	bool IsSelected(UControlRig* InControlRig, FRigControlElement* RigElement) const;
 	void SetSequencer(TWeakPtr<ISequencer> InSequencer) { Sequencer = InSequencer; }
 	ISequencer* GetSequencer() const { return Sequencer.Pin().Get(); }
