[Sequencer Control Rig picker] show rigs that can work with any skeleton when "filter by skeleton" is turned on

diff --git a/Engine/Plugins/Animation/ControlRig/Source/ControlRigEditor/Private/Sequencer/ControlRigParameterTrackEditor.cpp b/Engine/Plugins/Animation/ControlRig/Source/ControlRigEditor/Private/Sequencer/ControlRigParameterTrackEditor.cpp
--- a/Engine/Plugins/Animation/ControlRig/Source/ControlRigEditor/Private/Sequencer/ControlRigParameterTrackEditor.cpp
+++ b/Engine/Plugins/Animation/ControlRig/Source/ControlRigEditor/Private/Sequencer/ControlRigParameterTrackEditor.cpp
@@ -1794,6 +1794,16 @@ bool FControlRigParameterTrackEditor::IsControlRigAllowed(const FAssetData& Asse
 				return true;
 			}
 		}
+
+		if (!FSoftObjectPath(PreviewSkeletalMesh).IsValid() &&
+			!FSoftObjectPath(PreviewSkeleton).IsValid() &&
+			!FSoftObjectPath(SourceHierarchyImport).IsValid() &&
+			!FSoftObjectPath(SourceCurveImport).IsValid())
+		{
+			// this indicates that the rig can work on any skeleton (for example, utility rigs or deformer rigs)
+			return true;
+		}
+		
 		return false;
 	}
 	return true;	
