Sequencer: Anim Layers: Add Set Passthrough Key as an CR edit mode hotkey.

diff --git a/Engine/Plugins/Animation/ControlRig/Source/ControlRigEditor/Private/EditMode/ControlRigEditMode.cpp b/Engine/Plugins/Animation/ControlRig/Source/ControlRigEditor/Private/EditMode/ControlRigEditMode.cpp
--- a/Engine/Plugins/Animation/ControlRig/Source/ControlRigEditor/Private/EditMode/ControlRigEditMode.cpp
+++ b/Engine/Plugins/Animation/ControlRig/Source/ControlRigEditor/Private/EditMode/ControlRigEditMode.cpp
@@ -2846,6 +2846,10 @@ void FControlRigEditMode::BindCommands()
 		Commands.ToggleControlShapeTransformEdit,
 		FExecuteAction::CreateRaw(this, &FControlRigEditMode::ToggleControlShapeTransformEdit));
 
+	CommandBindings->MapAction(
+		Commands.SetAnimLayerPassthroughKey,
+		FExecuteAction::CreateRaw(this, &FControlRigEditMode::SetAnimLayerPassthroughKey));
+
 	CommandBindings->MapAction(
 		Commands.OpenSpacePickerWidget,
 		FExecuteAction::CreateRaw(this, &FControlRigEditMode::OpenSpacePickerWidget));
@@ -3121,6 +3125,29 @@ bool FControlRigEditMode::IsDragAnimSliderToolPressed(FViewport* InViewport)
 	return false;
 }
 
+void FControlRigEditMode::SetAnimLayerPassthroughKey()
+{
+	ISequencer* Sequencer = WeakSequencer.Pin().Get();
+	if (Sequencer)
+	{
+		if (UAnimLayers* AnimLayers = UAnimLayers::GetAnimLayers(Sequencer))
+		{
+			const FScopedTransaction Transaction(LOCTEXT("SetPassthroughKey_Transaction", "Set Passthrough Key"), !GIsTransacting);
+			for (UAnimLayer* AnimLayer : AnimLayers->AnimLayers)
+			{
+				if (AnimLayer->GetSelectedInList())
+				{
+					int32 Index = AnimLayers->GetAnimLayerIndex(AnimLayer);
+					if (Index != INDEX_NONE)
+					{
+						AnimLayers->SetPassthroughKey(Sequencer, Index);
+					}
+				}
+			}
+		}
+	}
+}
+
 void FControlRigEditMode::OpenSpacePickerWidget()
 {
 	TMap<UControlRig*, TArray<FRigElementKey>> SelectedControlRigsAndControls;
diff --git a/Engine/Plugins/Animation/ControlRig/Source/ControlRigEditor/Private/EditMode/ControlRigEditModeCommands.cpp b/Engine/Plugins/Animation/ControlRig/Source/ControlRigEditor/Private/EditMode/ControlRigEditModeCommands.cpp
--- a/Engine/Plugins/Animation/ControlRig/Source/ControlRigEditor/Private/EditMode/ControlRigEditModeCommands.cpp
+++ b/Engine/Plugins/Animation/ControlRig/Source/ControlRigEditor/Private/EditMode/ControlRigEditModeCommands.cpp
@@ -22,6 +22,7 @@ void FControlRigEditModeCommands::RegisterCommands()
 	UI_COMMAND(DragAnimSliderTool, "Drag Anim Slider Tool", "Drag existing anim slider", EUserInterfaceActionType::Button, FInputChord(EKeys::U));
 	UI_COMMAND(ChangeAnimSliderTool, "Change Anim Slider Tool", "Go to the next anim slider", EUserInterfaceActionType::Button, FInputChord(EKeys::U, EModifierKey::Shift));
 
+	UI_COMMAND(SetAnimLayerPassthroughKey," Set Anim Layer Passthrough Key","Set previous layer value on current selected anim layer", EUserInterfaceActionType::Button, FInputChord());
 
 	UI_COMMAND(ToggleControlShapeTransformEdit, "Toggle Shape Transform Edit", "Toggle Editing Selected Control's Shape Transform", EUserInterfaceActionType::ToggleButton, FInputChord(EKeys::Period, EModifierKey::Control)); 
 	UI_COMMAND(OpenSpacePickerWidget, "Open the Space Picker", "Allows space switching on the control", EUserInterfaceActionType::Button, FInputChord(EKeys::Tab)); 
diff --git a/Engine/Plugins/Animation/ControlRig/Source/ControlRigEditor/Private/EditMode/ControlRigEditModeCommands.h b/Engine/Plugins/Animation/ControlRig/Source/ControlRigEditor/Private/EditMode/ControlRigEditModeCommands.h
--- a/Engine/Plugins/Animation/ControlRig/Source/ControlRigEditor/Private/EditMode/ControlRigEditModeCommands.h
+++ b/Engine/Plugins/Animation/ControlRig/Source/ControlRigEditor/Private/EditMode/ControlRigEditModeCommands.h
@@ -62,6 +62,9 @@ public:
 	/** Toggle Shape Transform Edit*/
 	TSharedPtr< FUICommandInfo > ToggleControlShapeTransformEdit;
 
+	/** Sets Passthrough Key on selected anim layers */
+	TSharedPtr< FUICommandInfo > SetAnimLayerPassthroughKey;
+
 	/** Opens up the space picker widget */
 	TSharedPtr< FUICommandInfo > OpenSpacePickerWidget;
 
diff --git a/Engine/Plugins/Animation/ControlRig/Source/ControlRigEditor/Private/Sequencer/AnimLayers/AnimLayers.cpp b/Engine/Plugins/Animation/ControlRig/Source/ControlRigEditor/Private/Sequencer/AnimLayers/AnimLayers.cpp
--- a/Engine/Plugins/Animation/ControlRig/Source/ControlRigEditor/Private/Sequencer/AnimLayers/AnimLayers.cpp
+++ b/Engine/Plugins/Animation/ControlRig/Source/ControlRigEditor/Private/Sequencer/AnimLayers/AnimLayers.cpp
@@ -1682,6 +1682,7 @@ bool UAnimLayers::SetPassthroughKey(ISequencer* InSequencer, int32 InIndex)
 							int32 SectionIndex = TrackSections.Find((SectionItem.Section.Get()));
 							if (SectionIndex != INDEX_NONE)
 							{
+								const FScopedTransaction Transaction(LOCTEXT("SetPassthroughKey_Transaction", "Set Passthrough Key"), !GIsTransacting);
 								SectionItem.Section->Modify();
 								TrackSections.SetNum(SectionIndex); //this will gives us up to the section 
 								TArray<UMovieSceneSection*> Sections;
diff --git a/Engine/Plugins/Animation/ControlRig/Source/ControlRigEditor/Private/Sequencer/AnimLayers/SAnimLayers.cpp b/Engine/Plugins/Animation/ControlRig/Source/ControlRigEditor/Private/Sequencer/AnimLayers/SAnimLayers.cpp
--- a/Engine/Plugins/Animation/ControlRig/Source/ControlRigEditor/Private/Sequencer/AnimLayers/SAnimLayers.cpp
+++ b/Engine/Plugins/Animation/ControlRig/Source/ControlRigEditor/Private/Sequencer/AnimLayers/SAnimLayers.cpp
@@ -134,7 +134,6 @@ void FAnimLayerSourceUIEntry::AddSelected() const
 {
 	if (AnimLayer)
 	{
-
 		AnimLayer->AddSelectedInSequencer();
 		ClearCaches();
 	}
@@ -143,7 +142,7 @@ void FAnimLayerSourceUIEntry::AddSelected() const
 void FAnimLayerSourceUIEntry::RemoveSelected() const
 {
 	if (AnimLayer)
-{
+	{
 		AnimLayer->RemoveSelectedInSequencer();
 		ClearCaches();
 	}
@@ -934,18 +933,26 @@ void FAnimLayerSourcesView::AddSelected()
 {
 	TArray<FAnimLayerSourceUIEntryPtr> Selected;
 	SourcesListView->GetSelectedItems(Selected);
-	for (const FAnimLayerSourceUIEntryPtr& Ptr : Selected)
+	if (Selected.Num() > 0)
 	{
-		Ptr->AddSelected();
+		const FScopedTransaction Transaction(LOCTEXT("AddSelectedAnimLayer_Transaction", "Add Selected"), !GIsTransacting);
+		for (const FAnimLayerSourceUIEntryPtr& Ptr : Selected)
+		{
+			Ptr->AddSelected();
+		}
 	}
 }
 void FAnimLayerSourcesView::RemoveSelected()
 {
 	TArray<FAnimLayerSourceUIEntryPtr> Selected;
 	SourcesListView->GetSelectedItems(Selected);
-	for (const FAnimLayerSourceUIEntryPtr& Ptr : Selected)
+	if (Selected.Num() > 0)
 	{
-		Ptr->RemoveSelected();
+		const FScopedTransaction Transaction(LOCTEXT("RemoveSelected_Transaction", "Remove Selected"), !GIsTransacting);
+		for (const FAnimLayerSourceUIEntryPtr& Ptr : Selected)
+		{
+			Ptr->RemoveSelected();
+		}
 	}
 }
 
@@ -953,9 +960,13 @@ void FAnimLayerSourcesView::SelectObjects()
 {
 	TArray<FAnimLayerSourceUIEntryPtr> Selected;
 	SourcesListView->GetSelectedItems(Selected);
-	for(const FAnimLayerSourceUIEntryPtr& Ptr: Selected)
+	if (Selected.Num() > 0)
 	{
-		Ptr->SelectObjects();
+		const FScopedTransaction Transaction(LOCTEXT("SetSelected_Transaction", "Set Selection"), !GIsTransacting);
+		for (const FAnimLayerSourceUIEntryPtr& Ptr : Selected)
+		{
+			Ptr->SelectObjects();
+		}
 	}
 }
 
@@ -963,18 +974,26 @@ void FAnimLayerSourcesView::DeleteAnimLayer()
 {
 	TArray<FAnimLayerSourceUIEntryPtr> Selected;
 	SourcesListView->GetSelectedItems(Selected);
-	for (const FAnimLayerSourceUIEntryPtr& Ptr : Selected)
+	if (Selected.Num() > 0)
 	{
-		Ptr->DeleteAnimLayer();
+		const FScopedTransaction Transaction(LOCTEXT("DeleteAnimLayer_Transaction", "Delete Anim Layer"), !GIsTransacting);
+		for (const FAnimLayerSourceUIEntryPtr& Ptr : Selected)
+		{
+			Ptr->DeleteAnimLayer();
+		}
 	}
 }
 void FAnimLayerSourcesView::Duplicate()
 {
 	TArray<FAnimLayerSourceUIEntryPtr> Selected;
 	SourcesListView->GetSelectedItems(Selected);
-	for (const FAnimLayerSourceUIEntryPtr& Ptr : Selected)
+	if (Selected.Num() > 0)
 	{
-		Ptr->Duplicate();
+		const FScopedTransaction Transaction(LOCTEXT("DuplicateAnimLayer_Transaction", "Duplicate Anim Layer"), !GIsTransacting);
+		for (const FAnimLayerSourceUIEntryPtr& Ptr : Selected)
+		{
+			Ptr->Duplicate();
+		}
 	}
 }
 
@@ -1047,9 +1066,13 @@ void FAnimLayerSourcesView::SetPassthroughKey()
 {
 	TArray<FAnimLayerSourceUIEntryPtr> Selected;
 	SourcesListView->GetSelectedItems(Selected);
-	for (const FAnimLayerSourceUIEntryPtr& Ptr : Selected)
+	if (Selected.Num() > 0)
 	{
-		Ptr->SetPassthroughKey();
+		const FScopedTransaction Transaction(LOCTEXT("SetPassthroughKey_Transaction", "Set Passthrough Key"), !GIsTransacting);
+		for (const FAnimLayerSourceUIEntryPtr& Ptr : Selected)
+		{
+			Ptr->SetPassthroughKey();
+		}
 	}
 }
 
diff --git a/Engine/Plugins/Animation/ControlRig/Source/ControlRigEditor/Public/EditMode/ControlRigEditMode.h b/Engine/Plugins/Animation/ControlRig/Source/ControlRigEditor/Public/EditMode/ControlRigEditMode.h
--- a/Engine/Plugins/Animation/ControlRig/Source/ControlRigEditor/Public/EditMode/ControlRigEditMode.h
+++ b/Engine/Plugins/Animation/ControlRig/Source/ControlRigEditor/Public/EditMode/ControlRigEditMode.h
@@ -297,6 +297,9 @@ public:
 	/** Frame a list of provided items*/
    	void FrameItems(const TArray<FRigElementKey>& InItems);
 
+	/** Sets Passthrough Key on selected anim layers */
+	void SetAnimLayerPassthroughKey();
+
 	/** Opens up the space picker widget */
 	void OpenSpacePickerWidget();
 
