[ControlRig Sequencer] minor fix for finding space channels when there are mulitple rigs

diff --git a/Engine/Plugins/Animation/ControlRig/Source/ControlRigEditor/Private/Sequencer/ControlRigSpaceChannelEditors.cpp b/Engine/Plugins/Animation/ControlRig/Source/ControlRigEditor/Private/Sequencer/ControlRigSpaceChannelEditors.cpp
--- a/Engine/Plugins/Animation/ControlRig/Source/ControlRigEditor/Private/Sequencer/ControlRigSpaceChannelEditors.cpp
+++ b/Engine/Plugins/Animation/ControlRig/Source/ControlRigEditor/Private/Sequencer/ControlRigSpaceChannelEditors.cpp
@@ -199,43 +199,57 @@ FSpaceChannelAndSection FControlRigSpaceChannelHelpers::FindSpaceChannelAndSecti
 	{
 		return SpaceChannelAndSection;
 	}
+
+	bool bFoundTrack = false;
+	
 	const TArray<FMovieSceneBinding>& Bindings = MovieScene->GetBindings();
 	bool bRecreateCurves = false;
 	TArray<TPair<UControlRig*, FName>> ControlRigPairsToReselect;
 	for (const FMovieSceneBinding& Binding : Bindings)
 	{
-		UMovieSceneControlRigParameterTrack* ControlRigParameterTrack = Cast<UMovieSceneControlRigParameterTrack>(MovieScene->FindTrack(UMovieSceneControlRigParameterTrack::StaticClass(), Binding.GetObjectGuid(), NAME_None));
-		if (ControlRigParameterTrack && ControlRigParameterTrack->GetControlRig() == ControlRig)
+		TArray<UMovieSceneTrack*> Tracks = MovieScene->FindTracks(UMovieSceneControlRigParameterTrack::StaticClass(), Binding.GetObjectGuid(), NAME_None);
+		for (UMovieSceneTrack* Track : Tracks)
 		{
-			UMovieSceneControlRigParameterSection* ActiveSection = Cast<UMovieSceneControlRigParameterSection>(ControlRigParameterTrack->GetSectionToKey(ControlName));
-			if (ActiveSection)
+			UMovieSceneControlRigParameterTrack* ControlRigParameterTrack = Cast<UMovieSceneControlRigParameterTrack>(Track);
+			if (ControlRigParameterTrack && ControlRigParameterTrack->GetControlRig() == ControlRig)
 			{
-				ActiveSection->Modify();
-				ControlRig->Modify();
-				SpaceChannelAndSection.SectionToKey = ActiveSection;
-				FSpaceControlNameAndChannel* NameAndChannel = ActiveSection->GetSpaceChannel(ControlName);
-				if (NameAndChannel)
+				UMovieSceneControlRigParameterSection* ActiveSection = Cast<UMovieSceneControlRigParameterSection>(ControlRigParameterTrack->GetSectionToKey(ControlName));
+				if (ActiveSection)
 				{
-					SpaceChannelAndSection.SpaceChannel = &NameAndChannel->SpaceCurve;
-				}
-				else if (bCreateIfNeeded)
-				{
-					if (ControlRig->IsControlSelected(ControlName))
-					{
-						TPair<UControlRig*, FName> Pair;
-						Pair.Key = ControlRig;
-						Pair.Value = ControlName;
-						ControlRigPairsToReselect.Add(Pair);
-					}
-					ActiveSection->AddSpaceChannel(ControlName, true /*ReconstructChannelProxy*/);
-					NameAndChannel = ActiveSection->GetSpaceChannel(ControlName);
+					ActiveSection->Modify();
+					ControlRig->Modify();
+					SpaceChannelAndSection.SectionToKey = ActiveSection;
+					FSpaceControlNameAndChannel* NameAndChannel = ActiveSection->GetSpaceChannel(ControlName);
 					if (NameAndChannel)
 					{
 						SpaceChannelAndSection.SpaceChannel = &NameAndChannel->SpaceCurve;
-						bRecreateCurves = true;
+					}
+					else if (bCreateIfNeeded)
+					{
+						if (ControlRig->IsControlSelected(ControlName))
+						{
+							TPair<UControlRig*, FName> Pair;
+							Pair.Key = ControlRig;
+							Pair.Value = ControlName;
+							ControlRigPairsToReselect.Add(Pair);
+						}
+						ActiveSection->AddSpaceChannel(ControlName, true /*ReconstructChannelProxy*/);
+						NameAndChannel = ActiveSection->GetSpaceChannel(ControlName);
+						if (NameAndChannel)
+						{
+							SpaceChannelAndSection.SpaceChannel = &NameAndChannel->SpaceCurve;
+							bRecreateCurves = true;
+						}
 					}
 				}
+				
+				bFoundTrack = true;
+				break;
 			}
+		}
+
+		if (bFoundTrack)
+		{
 			break;
 		}
 	}
