Fixe a crash in Control Rig when undocking the hiearchy and opening multiple editors, due to the rig hiearchy summoner not allowing the destruction of the undocked tab

diff --git a/Engine/Plugins/Animation/ControlRig/Source/ControlRigEditor/Private/Editor/RigHierarchyTabSummoner.cpp b/Engine/Plugins/Animation/ControlRig/Source/ControlRigEditor/Private/Editor/RigHierarchyTabSummoner.cpp
--- a/Engine/Plugins/Animation/ControlRig/Source/ControlRigEditor/Private/Editor/RigHierarchyTabSummoner.cpp
+++ b/Engine/Plugins/Animation/ControlRig/Source/ControlRigEditor/Private/Editor/RigHierarchyTabSummoner.cpp
@@ -59,7 +59,7 @@ TSharedRef<SDockTab> FRigHierarchyTabSummoner::SpawnTab(const FWorkflowTabSpawnI
 				}
 			}
 		}
-		return HierarchyTabCount > 1;
+		return HierarchyTabCount > 0;
     }));
 	DockTab->SetOnTabClosed( SDockTab::FOnTabClosedCallback::CreateLambda([](TSharedRef<SDockTab> DockTab)
 	{
