Fixed an issue where generating lip sync animations could sometimes access an out of bounds array.

#tests Made sure batch file generated lipsync data correctly.

diff --git a/Engine/Plugins/Animation/ControlRig/Source/ControlRig/Private/Rigs/FKControlRig.cpp b/Engine/Plugins/Animation/ControlRig/Source/ControlRig/Private/Rigs/FKControlRig.cpp
--- a/Engine/Plugins/Animation/ControlRig/Source/ControlRig/Private/Rigs/FKControlRig.cpp
+++ b/Engine/Plugins/Animation/ControlRig/Source/ControlRig/Private/Rigs/FKControlRig.cpp
@@ -122,7 +122,7 @@ bool UFKControlRig::Execute_Internal(const FName& InEventName)
 			const FName ControlName = GetControlName(BoneElement->GetFName(), BoneElement->GetType());
 			const FRigElementKey ControlKey(ControlName, ERigElementType::Control);
 			const int32 ControlIndex = GetHierarchy()->GetIndex(ControlKey);
-			if (ControlIndex != INDEX_NONE && IsControlActive[ControlIndex])
+			if (ControlIndex != INDEX_NONE && GetControlActive(ControlIndex))
 			{
 				FRigControlElement* Control = GetHierarchy()->Get<FRigControlElement>(ControlIndex);
 				const FTransform LocalTransform = GetHierarchy()->GetLocalTransform(ControlIndex);
@@ -160,7 +160,7 @@ bool UFKControlRig::Execute_Internal(const FName& InEventName)
 			const FRigElementKey ControlKey(ControlName, ERigElementType::Control);
 			const int32 ControlIndex = GetHierarchy()->GetIndex(ControlKey);
 			
-			if (ControlIndex != INDEX_NONE && IsControlActive[ControlIndex])
+			if (ControlIndex != INDEX_NONE && GetControlActive(ControlIndex))
 			{
 				const float CurveValue = GetHierarchy()->GetControlValue(ControlIndex).Get<float>();
 
