Sequencer: Control Rig: Add option to have non-grouped channels still get sorted when other groups are present. This is needed for a request to get the CR WeightChannel first in the outliner when active, otherwise to hard to find. While doing this noticed we weren't allow the channels to go to the topmost track, which is another request we have heard so put that in also.

diff --git a/Engine/Plugins/Animation/ControlRig/Source/ControlRig/Private/Sequencer/MovieSceneControlRigParameterSection.cpp b/Engine/Plugins/Animation/ControlRig/Source/ControlRig/Private/Sequencer/MovieSceneControlRigParameterSection.cpp
--- a/Engine/Plugins/Animation/ControlRig/Source/ControlRig/Private/Sequencer/MovieSceneControlRigParameterSection.cpp
+++ b/Engine/Plugins/Animation/ControlRig/Source/ControlRig/Private/Sequencer/MovieSceneControlRigParameterSection.cpp
@@ -47,7 +47,7 @@ struct FParameterFloatChannelEditorData
 			MetaData.SetIdentifiers(InName, GroupName, GroupName);
 			MetaData.bEnabled = bEnabledOverride;
 			MetaData.SortOrder = SortStartIndex++;
-			MetaData.bCanCollapseToTrack = false;
+			MetaData.bCanCollapseToTrack = true;
 		}
 
 		ExternalValues.OnGetExternalValue = [InControlRig, InName](UObject& InObject, FTrackInstancePropertyBindings* Bindings) { return GetValue(InControlRig, InName,InObject, Bindings); };
@@ -126,7 +126,7 @@ struct FParameterVectorChannelEditorData
 			MetaData[0].Group = GroupName;
 			MetaData[0].bEnabled = bEnabledOverride;
 			MetaData[0].SortOrder = SortStartIndex++;
-			MetaData[0].bCanCollapseToTrack = false;
+			MetaData[0].bCanCollapseToTrack = true;
 		}
 		{
 			TotalName += ".Y";
@@ -136,7 +136,7 @@ struct FParameterVectorChannelEditorData
 			MetaData[1].Group = GroupName;
 			MetaData[1].bEnabled = bEnabledOverride;
 			MetaData[1].SortOrder = SortStartIndex++;
-			MetaData[1].bCanCollapseToTrack = false;
+			MetaData[1].bCanCollapseToTrack = true;
 		}
 		{
 			TotalName += ".Z";
@@ -146,7 +146,7 @@ struct FParameterVectorChannelEditorData
 			MetaData[2].Group = GroupName;
 			MetaData[2].bEnabled = bEnabledOverride;
 			MetaData[2].SortOrder = SortStartIndex++;
-			MetaData[2].bCanCollapseToTrack = false;
+			MetaData[2].bCanCollapseToTrack = true;
 		}
 		{
 			TotalName += ".W";
@@ -156,7 +156,7 @@ struct FParameterVectorChannelEditorData
 			MetaData[3].Group = GroupName;
 			MetaData[3].bEnabled = bEnabledOverride;
 			MetaData[3].SortOrder = SortStartIndex++;
-			MetaData[3].bCanCollapseToTrack = false;
+			MetaData[3].bCanCollapseToTrack = true;
 		}
 
 		ExternalValues[0].OnGetExternalValue = [InControlRig, InName,NumChannels](UObject& InObject, FTrackInstancePropertyBindings* Bindings) { return ExtractChannelX(InObject, InControlRig, InName, NumChannels); };
@@ -370,7 +370,7 @@ struct FParameterTransformChannelEditorData
 			MetaData[0].bEnabled = bEnabledOverride && EnumHasAllFlags(Mask, EMovieSceneTransformChannel::TranslationX);
 			MetaData[0].Color = FCommonChannelData::RedChannelColor;
 			MetaData[0].SortOrder = SortStartIndex++;
-			MetaData[0].bCanCollapseToTrack = false;
+			MetaData[0].bCanCollapseToTrack = true;
 			for (const FName& PropertyMetaDataKey : PropertyMetaDataKeys)
 			{
 				MetaData[0].PropertyMetaData.Add(PropertyMetaDataKey, RelativeLocationProperty->GetMetaData(PropertyMetaDataKey));
@@ -386,7 +386,7 @@ struct FParameterTransformChannelEditorData
 			MetaData[1].bEnabled = bEnabledOverride && EnumHasAllFlags(Mask, EMovieSceneTransformChannel::TranslationY);
 			MetaData[1].Color = FCommonChannelData::GreenChannelColor;
 			MetaData[1].SortOrder = SortStartIndex++;
-			MetaData[1].bCanCollapseToTrack = false;
+			MetaData[1].bCanCollapseToTrack = true;
 			for (const FName& PropertyMetaDataKey : PropertyMetaDataKeys)
 			{
 				MetaData[1].PropertyMetaData.Add(PropertyMetaDataKey, RelativeLocationProperty->GetMetaData(PropertyMetaDataKey));
@@ -402,7 +402,7 @@ struct FParameterTransformChannelEditorData
 			MetaData[2].bEnabled = bEnabledOverride && EnumHasAllFlags(Mask, EMovieSceneTransformChannel::TranslationZ);
 			MetaData[2].Color = FCommonChannelData::BlueChannelColor;
 			MetaData[2].SortOrder = SortStartIndex++;
-			MetaData[2].bCanCollapseToTrack = false;
+			MetaData[2].bCanCollapseToTrack = true;
 			for (const FName& PropertyMetaDataKey : PropertyMetaDataKeys)
 			{
 				MetaData[2].PropertyMetaData.Add(PropertyMetaDataKey, RelativeLocationProperty->GetMetaData(PropertyMetaDataKey));
@@ -419,7 +419,7 @@ struct FParameterTransformChannelEditorData
 			MetaData[3].bEnabled = bEnabledOverride && EnumHasAllFlags(Mask, EMovieSceneTransformChannel::RotationX);
 			MetaData[3].Color = FCommonChannelData::RedChannelColor;
 			MetaData[3].SortOrder = SortStartIndex++;
-			MetaData[3].bCanCollapseToTrack = false;
+			MetaData[3].bCanCollapseToTrack = true;
 			for (const FName& PropertyMetaDataKey : PropertyMetaDataKeys)
 			{
 				MetaData[3].PropertyMetaData.Add(PropertyMetaDataKey, RelativeRotationProperty->GetMetaData(PropertyMetaDataKey));
@@ -435,7 +435,7 @@ struct FParameterTransformChannelEditorData
 			MetaData[4].bEnabled = bEnabledOverride && EnumHasAllFlags(Mask, EMovieSceneTransformChannel::RotationY);
 			MetaData[4].Color = FCommonChannelData::GreenChannelColor;
 			MetaData[4].SortOrder = SortStartIndex++;
-			MetaData[4].bCanCollapseToTrack = false;
+			MetaData[4].bCanCollapseToTrack = true;
 			for (const FName& PropertyMetaDataKey : PropertyMetaDataKeys)
 			{
 				MetaData[4].PropertyMetaData.Add(PropertyMetaDataKey, RelativeRotationProperty->GetMetaData(PropertyMetaDataKey));
@@ -451,7 +451,7 @@ struct FParameterTransformChannelEditorData
 			MetaData[5].bEnabled = bEnabledOverride && EnumHasAllFlags(Mask, EMovieSceneTransformChannel::RotationZ);
 			MetaData[5].Color = FCommonChannelData::BlueChannelColor;
 			MetaData[5].SortOrder = SortStartIndex++;
-			MetaData[5].bCanCollapseToTrack = false;
+			MetaData[5].bCanCollapseToTrack = true;
 			for (const FName& PropertyMetaDataKey : PropertyMetaDataKeys)
 			{
 				MetaData[5].PropertyMetaData.Add(PropertyMetaDataKey, RelativeRotationProperty->GetMetaData(PropertyMetaDataKey));
@@ -468,7 +468,7 @@ struct FParameterTransformChannelEditorData
 			MetaData[6].bEnabled = bEnabledOverride && EnumHasAllFlags(Mask, EMovieSceneTransformChannel::ScaleX);
 			MetaData[6].Color = FCommonChannelData::RedChannelColor;
 			MetaData[6].SortOrder = SortStartIndex++;
-			MetaData[6].bCanCollapseToTrack = false;
+			MetaData[6].bCanCollapseToTrack = true;
 			for (const FName& PropertyMetaDataKey : PropertyMetaDataKeys)
 			{
 				MetaData[6].PropertyMetaData.Add(PropertyMetaDataKey, RelativeScale3DProperty->GetMetaData(PropertyMetaDataKey));
@@ -484,7 +484,7 @@ struct FParameterTransformChannelEditorData
 			MetaData[7].bEnabled = bEnabledOverride && EnumHasAllFlags(Mask, EMovieSceneTransformChannel::ScaleY);
 			MetaData[7].Color = FCommonChannelData::GreenChannelColor;
 			MetaData[7].SortOrder = SortStartIndex++;
-			MetaData[7].bCanCollapseToTrack = false;
+			MetaData[7].bCanCollapseToTrack = true;
 			for (const FName& PropertyMetaDataKey : PropertyMetaDataKeys)
 			{
 				MetaData[7].PropertyMetaData.Add(PropertyMetaDataKey, RelativeScale3DProperty->GetMetaData(PropertyMetaDataKey));
@@ -500,7 +500,7 @@ struct FParameterTransformChannelEditorData
 			MetaData[8].bEnabled = bEnabledOverride && EnumHasAllFlags(Mask, EMovieSceneTransformChannel::ScaleZ);
 			MetaData[8].Color = FCommonChannelData::BlueChannelColor;
 			MetaData[8].SortOrder = SortStartIndex++;
-			MetaData[8].bCanCollapseToTrack = false;
+			MetaData[8].bCanCollapseToTrack = true;
 			for (const FName& PropertyMetaDataKey : PropertyMetaDataKeys)
 			{
 				MetaData[8].PropertyMetaData.Add(PropertyMetaDataKey, RelativeScale3DProperty->GetMetaData(PropertyMetaDataKey));
@@ -1663,7 +1663,7 @@ EMovieSceneChannelProxyType UMovieSceneControlRigParameterSection::CacheChannelP
 		}
 		int32 ControlIndex = 0; 
 		int32 MaskIndex = 0;
-		int32 TotalIndex = 0; 
+		int32 SortOrder = 1; //start with one so Weight is first 
 		int32 FloatChannelIndex = 0;
 		int32 BoolChannelIndex = 0;
 		int32 EnumChannelIndex = 0;
@@ -1679,7 +1679,7 @@ EMovieSceneChannelProxyType UMovieSceneControlRigParameterSection::CacheChannelP
 
 		// begin constraints
 	// begin constraints
-		auto AddConstrainChannels = [this, &ConstraintsChannelIndex, &TotalIndex, &Channels](
+		auto AddConstrainChannels = [this, &ConstraintsChannelIndex, &SortOrder, &Channels](
 			const FName& InControlName, const FText& InGroup, const bool bEnabled)
 		{
 			
@@ -1718,8 +1718,8 @@ EMovieSceneChannelProxyType UMovieSceneControlRigParameterSection::CacheChannelP
 						const FText DisplayText = FText::FromString(Constraint->GetTypeLabel());
 						FMovieSceneChannelMetaData MetaData(Constraint->GetFName(), DisplayText, InGroup, bEnabled);
 						ConstraintsChannelIndex += 1;
-						MetaData.SortOrder = TotalIndex++;
-						MetaData.bCanCollapseToTrack = false;
+						MetaData.SortOrder = SortOrder++;
+						MetaData.bCanCollapseToTrack = true;
 
 						Channels.Add(ConstraintChannel->ActiveChannel, MetaData, TMovieSceneExternalValue<bool>());
 #else
@@ -1782,7 +1782,7 @@ EMovieSceneChannelProxyType UMovieSceneControlRigParameterSection::CacheChannelP
 						{
 							if (Group.IsEmpty())
 							{
-								ControlChannelMap.Add(Scalar.ParameterName, FChannelMapInfo(ControlIndex, TotalIndex, FloatChannelIndex, INDEX_NONE, NAME_None, MaskIndex, CategoryIndex));
+								ControlChannelMap.Add(Scalar.ParameterName, FChannelMapInfo(ControlIndex, SortOrder, FloatChannelIndex, INDEX_NONE, NAME_None, MaskIndex, CategoryIndex));
 								Group = Hierarchy->GetDisplayNameForUI(ControlElement);
 								if (bEnabled)
 								{
@@ -1793,16 +1793,16 @@ EMovieSceneChannelProxyType UMovieSceneControlRigParameterSection::CacheChannelP
 							{
 								const FChannelMapInfo* pChannelIndex = ControlChannelMap.Find(ParentControlName);
 								const int32 ParentControlIndex = pChannelIndex ? pChannelIndex->ControlIndex : INDEX_NONE; 
-								ControlChannelMap.Add(Scalar.ParameterName, FChannelMapInfo(ControlIndex, TotalIndex, FloatChannelIndex,ParentControlIndex, NAME_None, MaskIndex, CategoryIndex));
+								ControlChannelMap.Add(Scalar.ParameterName, FChannelMapInfo(ControlIndex, SortOrder, FloatChannelIndex,ParentControlIndex, NAME_None, MaskIndex, CategoryIndex));
 							}
 
-							FParameterFloatChannelEditorData EditorData(ControlRig, Scalar.ParameterName, bEnabled, Group, TotalIndex);
+							FParameterFloatChannelEditorData EditorData(ControlRig, Scalar.ParameterName, bEnabled, Group, SortOrder);
 							EditorData.MetaData.DisplayText = Hierarchy->GetDisplayNameForUI(ControlElement);
 							EditorData.MetaData.PropertyMetaData.Add(UIMin, FString::SanitizeFloat(ControlElement->Settings.MinimumValue.Get<float>()));
 							EditorData.MetaData.PropertyMetaData.Add(UIMax, FString::SanitizeFloat(ControlElement->Settings.MaximumValue.Get<float>()));
 							Channels.Add(Scalar.ParameterCurve, EditorData.MetaData, EditorData.ExternalValues);
 							FloatChannelIndex += 1;
-							TotalIndex += 1;
+							SortOrder += 1;
 							ControlIndex += 1;
 							break;
 						}
@@ -1817,7 +1817,7 @@ EMovieSceneChannelProxyType UMovieSceneControlRigParameterSection::CacheChannelP
 						{
 							if (Group.IsEmpty())
 							{
-								ControlChannelMap.Add(Bool.ParameterName, FChannelMapInfo(ControlIndex, TotalIndex, BoolChannelIndex, INDEX_NONE,BoolChannelTypeName,MaskIndex, CategoryIndex));
+								ControlChannelMap.Add(Bool.ParameterName, FChannelMapInfo(ControlIndex, SortOrder, BoolChannelIndex, INDEX_NONE,BoolChannelTypeName,MaskIndex, CategoryIndex));
 								Group = Hierarchy->GetDisplayNameForUI(ControlElement);
 								if (bEnabled)
 								{
@@ -1828,16 +1828,16 @@ EMovieSceneChannelProxyType UMovieSceneControlRigParameterSection::CacheChannelP
 							{
 								const FChannelMapInfo* pChannelIndex = ControlChannelMap.Find(ParentControlName);
 								const int32 ParentControlIndex = pChannelIndex ? pChannelIndex->ControlIndex : INDEX_NONE; 
-								ControlChannelMap.Add(Bool.ParameterName, FChannelMapInfo(ControlIndex, TotalIndex, BoolChannelIndex,ParentControlIndex, BoolChannelTypeName,MaskIndex, CategoryIndex));
+								ControlChannelMap.Add(Bool.ParameterName, FChannelMapInfo(ControlIndex, SortOrder, BoolChannelIndex,ParentControlIndex, BoolChannelTypeName,MaskIndex, CategoryIndex));
 							}
 
 							FMovieSceneChannelMetaData MetaData(Bool.ParameterName, Group, Group, bEnabled);
 							MetaData.DisplayText = Hierarchy->GetDisplayNameForUI(ControlElement);
-							MetaData.SortOrder = TotalIndex++;
+							MetaData.SortOrder = SortOrder++;
 							BoolChannelIndex += 1;
 							ControlIndex += 1;
 							// Prevent single channels from collapsing to the track node
-							MetaData.bCanCollapseToTrack = false;
+							MetaData.bCanCollapseToTrack = true;
 							Channels.Add(Bool.ParameterCurve, MetaData, TMovieSceneExternalValue<bool>());
 							break;
 						}
@@ -1854,7 +1854,7 @@ EMovieSceneChannelProxyType UMovieSceneControlRigParameterSection::CacheChannelP
 							{
 								if (Group.IsEmpty())
 								{
-									ControlChannelMap.Add(Enum.ParameterName, FChannelMapInfo(ControlIndex, TotalIndex, EnumChannelIndex,INDEX_NONE, EnumChannelTypeName,MaskIndex, CategoryIndex));
+									ControlChannelMap.Add(Enum.ParameterName, FChannelMapInfo(ControlIndex, SortOrder, EnumChannelIndex,INDEX_NONE, EnumChannelTypeName,MaskIndex, CategoryIndex));
 									Group = Hierarchy->GetDisplayNameForUI(ControlElement);
 									if (bEnabled)
 									{
@@ -1865,16 +1865,16 @@ EMovieSceneChannelProxyType UMovieSceneControlRigParameterSection::CacheChannelP
 								{
 									const FChannelMapInfo* pChannelIndex = ControlChannelMap.Find(ParentControlName);
 									const int32 ParentControlIndex = pChannelIndex ? pChannelIndex->ControlIndex : INDEX_NONE; 
-									ControlChannelMap.Add(Enum.ParameterName, FChannelMapInfo(ControlIndex, TotalIndex, EnumChannelIndex, ParentControlIndex, EnumChannelTypeName,MaskIndex, CategoryIndex));
+									ControlChannelMap.Add(Enum.ParameterName, FChannelMapInfo(ControlIndex, SortOrder, EnumChannelIndex, ParentControlIndex, EnumChannelTypeName,MaskIndex, CategoryIndex));
 								}
 
 								FMovieSceneChannelMetaData MetaData(Enum.ParameterName, Group, Group, bEnabled);
 								MetaData.DisplayText = Hierarchy->GetDisplayNameForUI(ControlElement);
 								EnumChannelIndex += 1;
 								ControlIndex += 1;
-								MetaData.SortOrder = TotalIndex++;
+								MetaData.SortOrder = SortOrder++;
 								// Prevent single channels from collapsing to the track node
-								MetaData.bCanCollapseToTrack = false;
+								MetaData.bCanCollapseToTrack = true;
 								Channels.Add(Enum.ParameterCurve, MetaData, TMovieSceneExternalValue<uint8>());
 								break;
 							}
@@ -1888,7 +1888,7 @@ EMovieSceneChannelProxyType UMovieSceneControlRigParameterSection::CacheChannelP
 							{
 								if (Group.IsEmpty())
 								{
-									ControlChannelMap.Add(Integer.ParameterName, FChannelMapInfo(ControlIndex, TotalIndex, IntegerChannelIndex,INDEX_NONE,IntegerChannelTypeName,MaskIndex, CategoryIndex));
+									ControlChannelMap.Add(Integer.ParameterName, FChannelMapInfo(ControlIndex, SortOrder, IntegerChannelIndex,INDEX_NONE,IntegerChannelTypeName,MaskIndex, CategoryIndex));
 									Group = Hierarchy->GetDisplayNameForUI(ControlElement);
 									if (bEnabled)
 									{
@@ -1899,16 +1899,16 @@ EMovieSceneChannelProxyType UMovieSceneControlRigParameterSection::CacheChannelP
 								{
 									const FChannelMapInfo* pChannelIndex = ControlChannelMap.Find(ParentControlName);
 									const int32 ParentControlIndex = pChannelIndex ? pChannelIndex->ControlIndex : INDEX_NONE; 
-									ControlChannelMap.Add(Integer.ParameterName, FChannelMapInfo(ControlIndex, TotalIndex, IntegerChannelIndex, ParentControlIndex, IntegerChannelTypeName,MaskIndex, CategoryIndex));
+									ControlChannelMap.Add(Integer.ParameterName, FChannelMapInfo(ControlIndex, SortOrder, IntegerChannelIndex, ParentControlIndex, IntegerChannelTypeName,MaskIndex, CategoryIndex));
 								}
 
 								FMovieSceneChannelMetaData MetaData(Integer.ParameterName, Group, Group, bEnabled);
 								MetaData.DisplayText = Hierarchy->GetDisplayNameForUI(ControlElement);
 								IntegerChannelIndex += 1;
 								ControlIndex += 1;
-								MetaData.SortOrder = TotalIndex++;
+								MetaData.SortOrder = SortOrder++;
 								// Prevent single channels from collapsing to the track node
-								MetaData.bCanCollapseToTrack = false;
+								MetaData.bCanCollapseToTrack = true;
 								MetaData.PropertyMetaData.Add(UIMin, FString::FromInt(ControlElement->Settings.MinimumValue.Get<int32>()));
 								MetaData.PropertyMetaData.Add(UIMax, FString::FromInt(ControlElement->Settings.MaximumValue.Get<int32>()));
 								Channels.Add(Integer.ParameterCurve, MetaData, TMovieSceneExternalValue<int32>());
@@ -1927,7 +1927,7 @@ EMovieSceneChannelProxyType UMovieSceneControlRigParameterSection::CacheChannelP
 						{
 							if(Group.IsEmpty())
 							{
-								ControlChannelMap.Add(Vector2D.ParameterName, FChannelMapInfo(ControlIndex, TotalIndex, FloatChannelIndex, INDEX_NONE, NAME_None, MaskIndex, CategoryIndex));
+								ControlChannelMap.Add(Vector2D.ParameterName, FChannelMapInfo(ControlIndex, SortOrder, FloatChannelIndex, INDEX_NONE, NAME_None, MaskIndex, CategoryIndex));
 								if (bEnabled)
 								{
 									++CategoryIndex;
@@ -1938,15 +1938,15 @@ EMovieSceneChannelProxyType UMovieSceneControlRigParameterSection::CacheChannelP
 							{
 								const FChannelMapInfo* pChannelIndex = ControlChannelMap.Find(ParentControlName);
 								const int32 ParentControlIndex = pChannelIndex ? pChannelIndex->ControlIndex : INDEX_NONE; 
-								ControlChannelMap.Add(Vector2D.ParameterName, FChannelMapInfo(ControlIndex, TotalIndex, FloatChannelIndex, ParentControlIndex, NAME_None, MaskIndex, CategoryIndex));
+								ControlChannelMap.Add(Vector2D.ParameterName, FChannelMapInfo(ControlIndex, SortOrder, FloatChannelIndex, ParentControlIndex, NAME_None, MaskIndex, CategoryIndex));
 							}
-							FParameterVectorChannelEditorData EditorData(ControlRig, Vector2D.ParameterName, bEnabled, Group, TotalIndex, 2);
+							FParameterVectorChannelEditorData EditorData(ControlRig, Vector2D.ParameterName, bEnabled, Group, SortOrder, 2);
 							MaybeApplyChannelMask(EditorData.MetaData[0], ControlElement, ERigControlTransformChannel::TranslationX);
 							MaybeApplyChannelMask(EditorData.MetaData[1], ControlElement, ERigControlTransformChannel::TranslationY);
 							Channels.Add(Vector2D.XCurve, EditorData.MetaData[0], EditorData.ExternalValues[0]);
 							Channels.Add(Vector2D.YCurve, EditorData.MetaData[1], EditorData.ExternalValues[1]);
 							FloatChannelIndex += 2;
-							TotalIndex += 2;
+							SortOrder += 2;
 							ControlIndex += 1;
 							break;
 						}
@@ -1963,7 +1963,7 @@ EMovieSceneChannelProxyType UMovieSceneControlRigParameterSection::CacheChannelP
 						{
 							if(Group.IsEmpty())
 							{
-								ControlChannelMap.Add(Vector.ParameterName, FChannelMapInfo(ControlIndex, TotalIndex, FloatChannelIndex, INDEX_NONE, NAME_None, MaskIndex, CategoryIndex));
+								ControlChannelMap.Add(Vector.ParameterName, FChannelMapInfo(ControlIndex, SortOrder, FloatChannelIndex, INDEX_NONE, NAME_None, MaskIndex, CategoryIndex));
 								if (bEnabled)
 								{
 									++CategoryIndex;
@@ -1974,7 +1974,7 @@ EMovieSceneChannelProxyType UMovieSceneControlRigParameterSection::CacheChannelP
 							{
 								const FChannelMapInfo* pChannelIndex = ControlChannelMap.Find(ParentControlName);
 								const int32 ParentControlIndex = pChannelIndex ? pChannelIndex->ControlIndex : INDEX_NONE; 
-								ControlChannelMap.Add(Vector.ParameterName, FChannelMapInfo(ControlIndex, TotalIndex, FloatChannelIndex, ParentControlIndex, NAME_None, MaskIndex, CategoryIndex));
+								ControlChannelMap.Add(Vector.ParameterName, FChannelMapInfo(ControlIndex, SortOrder, FloatChannelIndex, ParentControlIndex, NAME_None, MaskIndex, CategoryIndex));
 							}
 							if (FSpaceControlNameAndChannel* SpaceChannel = GetSpaceChannel(Vector.ParameterName))
 							{
@@ -1992,14 +1992,14 @@ EMovieSceneChannelProxyType UMovieSceneControlRigParameterSection::CacheChannelP
 								FMovieSceneChannelMetaData SpaceMetaData(FName(*TotalName), Group, Group, bEnabled);
 								SpaceMetaData.DisplayText = FText::FromName(SpaceName);
 								SpaceChannelIndex += 1;
-								SpaceMetaData.SortOrder = TotalIndex++;
+								SpaceMetaData.SortOrder = SortOrder++;
 								// Prevent single channels from collapsing to the track node
-								SpaceMetaData.bCanCollapseToTrack = false;
+								SpaceMetaData.bCanCollapseToTrack = true;
 								Channels.Add(SpaceChannel->SpaceCurve, SpaceMetaData);
 							}
 						
 
-							FParameterVectorChannelEditorData EditorData(ControlRig, Vector.ParameterName, bEnabled, Group, TotalIndex, 3);
+							FParameterVectorChannelEditorData EditorData(ControlRig, Vector.ParameterName, bEnabled, Group, SortOrder, 3);
 
 							if(ControlElement->Settings.ControlType == ERigControlType::Position)
 							{
@@ -2024,7 +2024,7 @@ EMovieSceneChannelProxyType UMovieSceneControlRigParameterSection::CacheChannelP
 							Channels.Add(Vector.YCurve, EditorData.MetaData[1], EditorData.ExternalValues[1]);
 							Channels.Add(Vector.ZCurve, EditorData.MetaData[2], EditorData.ExternalValues[2]);
 							FloatChannelIndex += 3;
-							TotalIndex += 3;
+							SortOrder += 3;
 							ControlIndex += 1;
 							break;
 						}
@@ -2043,7 +2043,7 @@ EMovieSceneChannelProxyType UMovieSceneControlRigParameterSection::CacheChannelP
 							const FName ControlName = ControlElement->GetFName();
 							if(Group.IsEmpty())
 							{
-								ControlChannelMap.Add(Transform.ParameterName, FChannelMapInfo(ControlIndex, TotalIndex, FloatChannelIndex, INDEX_NONE, NAME_None, MaskIndex, CategoryIndex));
+								ControlChannelMap.Add(Transform.ParameterName, FChannelMapInfo(ControlIndex, SortOrder, FloatChannelIndex, INDEX_NONE, NAME_None, MaskIndex, CategoryIndex));
 								if (bEnabled)
 								{
 									++CategoryIndex;
@@ -2054,7 +2054,7 @@ EMovieSceneChannelProxyType UMovieSceneControlRigParameterSection::CacheChannelP
 							{
 								const FChannelMapInfo* pChannelIndex = ControlChannelMap.Find(ParentControlName);
 								const int32 ParentControlIndex = pChannelIndex ? pChannelIndex->ControlIndex : INDEX_NONE; 
-								ControlChannelMap.Add(Transform.ParameterName, FChannelMapInfo(ControlIndex, TotalIndex, FloatChannelIndex, ParentControlIndex, NAME_None, MaskIndex, CategoryIndex));
+								ControlChannelMap.Add(Transform.ParameterName, FChannelMapInfo(ControlIndex, SortOrder, FloatChannelIndex, ParentControlIndex, NAME_None, MaskIndex, CategoryIndex));
 							}
 
 							// constraints
@@ -2076,16 +2076,16 @@ EMovieSceneChannelProxyType UMovieSceneControlRigParameterSection::CacheChannelP
 								FMovieSceneChannelMetaData SpaceMetaData(FName(*TotalName), Group, Group, bEnabled);
 								SpaceMetaData.DisplayText = FText::FromName(SpaceName);
 								SpaceChannelIndex += 1;
-								SpaceMetaData.SortOrder = TotalIndex++;
+								SpaceMetaData.SortOrder = SortOrder++;
 								// Prevent single channels from collapsing to the track node
-								SpaceMetaData.bCanCollapseToTrack = false;
+								SpaceMetaData.bCanCollapseToTrack = true;
 								//TMovieSceneExternalValue<FMovieSceneControlRigSpaceBaseKey> ExternalData;
 								Channels.Add(SpaceChannel->SpaceCurve, SpaceMetaData);
 							}
 					
 
 							FParameterTransformChannelEditorData EditorData(ControlRig, Transform.ParameterName, bEnabled, TransformMask.GetChannels(), Group, 
-								TotalIndex);
+								SortOrder);
 
 							MaybeApplyChannelMask(EditorData.MetaData[0], ControlElement, ERigControlTransformChannel::TranslationX);
 							MaybeApplyChannelMask(EditorData.MetaData[1], ControlElement, ERigControlTransformChannel::TranslationY);
@@ -2115,13 +2115,13 @@ EMovieSceneChannelProxyType UMovieSceneControlRigParameterSection::CacheChannelP
 								Channels.Add(Transform.Scale[1], EditorData.MetaData[7], EditorData.ExternalValues[7]);
 								Channels.Add(Transform.Scale[2], EditorData.MetaData[8], EditorData.ExternalValues[8]);
 								FloatChannelIndex += 9;
-								TotalIndex += 9;
+								SortOrder += 9;
 
 							}
 							else
 							{
 								FloatChannelIndex += 6;
-								TotalIndex += 6;
+								SortOrder += 6;
 
 							}
 							ControlIndex += 1;
@@ -2141,10 +2141,10 @@ EMovieSceneChannelProxyType UMovieSceneControlRigParameterSection::CacheChannelP
 					{
 						if (ControlElement->GetFName() == Scalar.ParameterName)
 						{
-							ControlChannelMap.Add(Scalar.ParameterName, FChannelMapInfo(ControlIndex, TotalIndex,FloatChannelIndex,INDEX_NONE, NAME_None,MaskIndex));
+							ControlChannelMap.Add(Scalar.ParameterName, FChannelMapInfo(ControlIndex, SortOrder,FloatChannelIndex,INDEX_NONE, NAME_None,MaskIndex));
 							Channels.Add(Scalar.ParameterCurve);
 							FloatChannelIndex += 1;
-							TotalIndex += 1;
+							SortOrder += 1;
 							ControlIndex += 1;
 							break;
 						}
@@ -2157,10 +2157,10 @@ EMovieSceneChannelProxyType UMovieSceneControlRigParameterSection::CacheChannelP
 					{
 						if (ControlElement->GetFName() == Bool.ParameterName)
 						{
-							ControlChannelMap.Add(Bool.ParameterName, FChannelMapInfo(ControlIndex, TotalIndex, BoolChannelIndex, INDEX_NONE, NAME_None, MaskIndex));
+							ControlChannelMap.Add(Bool.ParameterName, FChannelMapInfo(ControlIndex, SortOrder, BoolChannelIndex, INDEX_NONE, NAME_None, MaskIndex));
 							Channels.Add(Bool.ParameterCurve);
 							BoolChannelIndex += 1;
-							TotalIndex += 1;
+							SortOrder += 1;
 							ControlIndex += 1;
 							break;
 						}
@@ -2175,10 +2175,10 @@ EMovieSceneChannelProxyType UMovieSceneControlRigParameterSection::CacheChannelP
 						{
 							if (ControlElement->GetFName() == Enum.ParameterName)
 							{
-								ControlChannelMap.Add(Enum.ParameterName, FChannelMapInfo(ControlIndex, TotalIndex, EnumChannelIndex, INDEX_NONE, NAME_None, MaskIndex));
+								ControlChannelMap.Add(Enum.ParameterName, FChannelMapInfo(ControlIndex, SortOrder, EnumChannelIndex, INDEX_NONE, NAME_None, MaskIndex));
 								Channels.Add(Enum.ParameterCurve);
 								EnumChannelIndex += 1;
-								TotalIndex += 1;
+								SortOrder += 1;
 								ControlIndex += 1;
 								break;
 							}
@@ -2190,10 +2190,10 @@ EMovieSceneChannelProxyType UMovieSceneControlRigParameterSection::CacheChannelP
 						{
 							if (ControlElement->GetFName() == Integer.ParameterName)
 							{
-								ControlChannelMap.Add(Integer.ParameterName, FChannelMapInfo(ControlIndex, TotalIndex, IntegerChannelIndex, INDEX_NONE, NAME_None, MaskIndex));
+								ControlChannelMap.Add(Integer.ParameterName, FChannelMapInfo(ControlIndex, SortOrder, IntegerChannelIndex, INDEX_NONE, NAME_None, MaskIndex));
 								Channels.Add(Integer.ParameterCurve);
 								IntegerChannelIndex += 1;
-								TotalIndex += 1;
+								SortOrder += 1;
 								ControlIndex += 1;
 								break;
 							}
@@ -2207,11 +2207,11 @@ EMovieSceneChannelProxyType UMovieSceneControlRigParameterSection::CacheChannelP
 					{
 						if (ControlElement->GetFName() == Vector2D.ParameterName)
 						{
-							ControlChannelMap.Add(Vector2D.ParameterName, FChannelMapInfo(ControlIndex, TotalIndex, FloatChannelIndex, INDEX_NONE, NAME_None, MaskIndex));
+							ControlChannelMap.Add(Vector2D.ParameterName, FChannelMapInfo(ControlIndex, SortOrder, FloatChannelIndex, INDEX_NONE, NAME_None, MaskIndex));
 							Channels.Add(Vector2D.XCurve);
 							Channels.Add(Vector2D.YCurve);
 							FloatChannelIndex += 2;
-							TotalIndex += 2;
+							SortOrder += 2;
 							ControlIndex += 1;
 							break;
 						}
@@ -2226,7 +2226,7 @@ EMovieSceneChannelProxyType UMovieSceneControlRigParameterSection::CacheChannelP
 					{
 						if (ControlElement->GetFName() == Vector.ParameterName)
 						{
-							ControlChannelMap.Add(Vector.ParameterName, FChannelMapInfo(ControlIndex, TotalIndex, FloatChannelIndex, INDEX_NONE, NAME_None, MaskIndex));
+							ControlChannelMap.Add(Vector.ParameterName, FChannelMapInfo(ControlIndex, SortOrder, FloatChannelIndex, INDEX_NONE, NAME_None, MaskIndex));
 							bool bDoSpaceChannel = true;
 							if (bDoSpaceChannel)
 							{
@@ -2247,7 +2247,7 @@ EMovieSceneChannelProxyType UMovieSceneControlRigParameterSection::CacheChannelP
 							Channels.Add(Vector.YCurve);
 							Channels.Add(Vector.ZCurve);
 							FloatChannelIndex += 3;
-							TotalIndex += 3;
+							SortOrder += 3;
 							ControlIndex += 1;
 							break;
 						}
@@ -2272,7 +2272,7 @@ EMovieSceneChannelProxyType UMovieSceneControlRigParameterSection::CacheChannelP
 					{
 						if (ControlElement->GetFName() == Transform.ParameterName)
 						{
-							ControlChannelMap.Add(Transform.ParameterName, FChannelMapInfo(ControlIndex, TotalIndex,FloatChannelIndex, INDEX_NONE, NAME_None, MaskIndex));
+							ControlChannelMap.Add(Transform.ParameterName, FChannelMapInfo(ControlIndex, SortOrder,FloatChannelIndex, INDEX_NONE, NAME_None, MaskIndex));
 						
 							bool bDoSpaceChannel = true;
 							if (bDoSpaceChannel)
@@ -2305,12 +2305,12 @@ EMovieSceneChannelProxyType UMovieSceneControlRigParameterSection::CacheChannelP
 								Channels.Add(Transform.Scale[1]);
 								Channels.Add(Transform.Scale[2]);
 								FloatChannelIndex += 9;
-								TotalIndex += 9;
+								SortOrder += 9;
 							}
 							else
 							{
 								FloatChannelIndex += 6;
-								TotalIndex += 6;
+								SortOrder += 6;
 							}
 
 							ControlIndex += 1;
@@ -2329,8 +2329,9 @@ EMovieSceneChannelProxyType UMovieSceneControlRigParameterSection::CacheChannelP
 		FMovieSceneChannelMetaData      MetaData;
 		MetaData.SetIdentifiers("Weight", NSLOCTEXT("MovieSceneTransformSection", "Weight", "Weight"));
 		MetaData.bEnabled = EnumHasAllFlags(TransformMask.GetChannels(), EMovieSceneTransformChannel::Weight);
-		MetaData.SortOrder = TotalIndex++;
-		MetaData.bCanCollapseToTrack = false;
+		MetaData.SortOrder = 0;
+		MetaData.bSortEmptyGroupsLast = false;
+		MetaData.bCanCollapseToTrack = true;
 		TMovieSceneExternalValue<float> ExVal;
 		Channels.Add(Weight, MetaData, ExVal);
 #else
diff --git a/Engine/Source/Editor/Sequencer/Private/ISequencerSection.cpp b/Engine/Source/Editor/Sequencer/Private/ISequencerSection.cpp
--- a/Engine/Source/Editor/Sequencer/Private/ISequencerSection.cpp
+++ b/Engine/Source/Editor/Sequencer/Private/ISequencerSection.cpp
@@ -35,6 +35,7 @@ struct FGroupData
 		: GroupText(InGroupText)
 		, GetGroupTooltipTextDelegate(InGetGroupTooltipTextDelegate)
 		, SortOrder(-1)
+		, bSortEmptyGroupsLast(true)
 	{}
 
 	void AddChannel(ISequencerSection::FChannelData&& InChannel)
@@ -43,6 +44,7 @@ struct FGroupData
 		{
 			SortOrder = InChannel.MetaData.SortOrder;
 		}
+		bSortEmptyGroupsLast = InChannel.MetaData.bSortEmptyGroupsLast;
 
 		Channels.Add(MoveTemp(InChannel));
 	}
@@ -56,6 +58,9 @@ struct FGroupData
 	/** Sort order of the group */
 	uint32 SortOrder;
 
+	/** By default if a channel has no FText::Group specified, we put it last, by setting this to false we use SortIndex instead */
+	bool bSortEmptyGroupsLast;
+
 	/** Array of channels within this group */
 	TArray<ISequencerSection::FChannelData, TInlineAllocator<4>> Channels;
 };
@@ -162,13 +167,22 @@ void ISequencerSection::GenerateSectionLayout( ISectionLayoutBuilder& LayoutBuil
 	{
 		if (A.IsNone())
 		{
-			return false;
+			const bool bSortEmptyGroupsLast = GroupToChannelsMap.FindChecked(A).bSortEmptyGroupsLast;
+			if(bSortEmptyGroupsLast)
+			{ 
+				return false;
+			}
 		}
 		else if (B.IsNone())
 		{
-			return true;
+			const bool bSortEmptyGroupsLast = GroupToChannelsMap.FindChecked(B).bSortEmptyGroupsLast;
+			if (bSortEmptyGroupsLast)
+			{
+				return true;
+			}
 		}
 
+
 		const int32 SortOrderA = GroupToChannelsMap.FindChecked(A).SortOrder;
 		const int32 SortOrderB = GroupToChannelsMap.FindChecked(B).SortOrder;
 		return SortOrderA < SortOrderB;
diff --git a/Engine/Source/Runtime/MovieScene/Private/Channels/MovieSceneChannelEditorData.cpp b/Engine/Source/Runtime/MovieScene/Private/Channels/MovieSceneChannelEditorData.cpp
--- a/Engine/Source/Runtime/MovieScene/Private/Channels/MovieSceneChannelEditorData.cpp
+++ b/Engine/Source/Runtime/MovieScene/Private/Channels/MovieSceneChannelEditorData.cpp
@@ -27,6 +27,7 @@ FMovieSceneChannelMetaData::FMovieSceneChannelMetaData()
 	, bCanCollapseToTrack(true)
 	, bRelativeToSection(false)
 	, SortOrder(0)
+	, bSortEmptyGroupsLast(true)
 	, Name(NAME_None)
 	, KeyOffset(0)
 {}
@@ -36,6 +37,7 @@ FMovieSceneChannelMetaData::FMovieSceneChannelMetaData(FName InName, FText InDis
 	, bCanCollapseToTrack(true)
 	, bRelativeToSection(false)
 	, SortOrder(0)
+	, bSortEmptyGroupsLast(true)
 	, Name(InName)
 	, DisplayText(InDisplayText)
 	, Group(InGroup)
diff --git a/Engine/Source/Runtime/MovieScene/Public/Channels/MovieSceneChannelEditorData.h b/Engine/Source/Runtime/MovieScene/Public/Channels/MovieSceneChannelEditorData.h
--- a/Engine/Source/Runtime/MovieScene/Public/Channels/MovieSceneChannelEditorData.h
+++ b/Engine/Source/Runtime/MovieScene/Public/Channels/MovieSceneChannelEditorData.h
@@ -74,6 +74,9 @@ struct FMovieSceneChannelMetaData
 	uint8 bRelativeToSection : 1;
 	/** A sort order for this channel. Channels are sorted by this order, then by name. Groups are sorted by the channel with the lowest sort order. */
 	uint32 SortOrder;
+	/** By default if a channel has no FText::Group specified, we put it last, by setting this to false we use SortIndex instead */
+	bool bSortEmptyGroupsLast;
+
 	/** This channel's unique name */
 	FName Name;
 	/**
