[ControlRigEditMode] clear cached transform for shape actor scene proxies, this is a temp fix. Will remove when proper fix is provided by the rendering team

diff --git a/Engine/Plugins/Animation/ControlRig/Source/ControlRigEditor/Private/EditMode/ControlRigEditMode.cpp b/Engine/Plugins/Animation/ControlRig/Source/ControlRigEditor/Private/EditMode/ControlRigEditMode.cpp
--- a/Engine/Plugins/Animation/ControlRig/Source/ControlRigEditor/Private/EditMode/ControlRigEditMode.cpp
+++ b/Engine/Plugins/Animation/ControlRig/Source/ControlRigEditor/Private/EditMode/ControlRigEditMode.cpp
@@ -4063,6 +4063,12 @@ bool FControlRigEditMode::TryUpdatingControlsShapes(UControlRig* InControlRig)
 		{
 			UpdateControlShape(ShapeActor, ControlElement, Params);
 		}
+
+		// workaround for UE-225122, FPrimitiveSceneProxy currently lazily updates the transform, but due to a thread sync issue,
+		// if we are setting the transform to 0 at tick 1 and setting it to the correct value like 100 at tick2, depending on
+		// the value of the cached transform, only one of the two sets would be committed. This call clears the cached transform to 0 such that
+		// set to 0(here) is always ignored and set to 100(TickControlShape) is always accepted.
+		ShapeActor->MarkComponentsRenderStateDirty();
 	}
 
 	return true;
