[ControlRig AnimDetails] added a null check

diff --git a/Engine/Plugins/Animation/ControlRig/Source/ControlRigEditor/Private/EditMode/AnimDetailsProxy.cpp b/Engine/Plugins/Animation/ControlRig/Source/ControlRigEditor/Private/EditMode/AnimDetailsProxy.cpp
--- a/Engine/Plugins/Animation/ControlRig/Source/ControlRigEditor/Private/EditMode/AnimDetailsProxy.cpp
+++ b/Engine/Plugins/Animation/ControlRig/Source/ControlRigEditor/Private/EditMode/AnimDetailsProxy.cpp
@@ -3059,7 +3059,7 @@ const TArray<UControlRigControlsProxy*> UControlRigDetailPanelControlProxies::Ge
 	TArray<USceneComponent*> BoundCRObjects;
 	for (UControlRigControlsProxy* Proxy : SelectedProxies)
 	{
-		if (Proxy && Proxy->OwnerControlRig.IsValid())
+		if (Proxy && Proxy->OwnerControlRig.IsValid() && Proxy->OwnerControlRig->GetObjectBinding().IsValid())
 		{
 			if (USceneComponent* SceneComponent = Cast<USceneComponent>(Proxy->OwnerControlRig->GetObjectBinding()->GetBoundObject()))
 			{
