Anim Outliner: do not try to find the parent bone when switching to world space

NOTE: space switching on FK Control Rig looks buggy

#rnx

diff --git a/Engine/Plugins/Animation/ControlRig/Source/ControlRigEditor/Private/EditMode/SControlRigOutliner.cpp b/Engine/Plugins/Animation/ControlRig/Source/ControlRigEditor/Private/EditMode/SControlRigOutliner.cpp
--- a/Engine/Plugins/Animation/ControlRig/Source/ControlRigEditor/Private/EditMode/SControlRigOutliner.cpp
+++ b/Engine/Plugins/Animation/ControlRig/Source/ControlRigEditor/Private/EditMode/SControlRigOutliner.cpp
@@ -611,7 +611,7 @@ bool SMultiRigHierarchyTreeView::AddElement(UControlRig* InControlRig, const FRi
 
 			if (ParentKey.IsValid())
 			{
-				if(FKControlRig)
+				if(FKControlRig && ParentKey != URigHierarchy::GetWorldSpaceReferenceKey())
 				{
 					if(const FRigControlElement* ControlElement = Cast<FRigControlElement>(InElement))
 					{
