Sequencer: Control Rig: Fix crash after undo add rig then add again, need to remove the stale control rig node from the layer instance proxy

diff --git a/Engine/Plugins/Animation/ControlRig/Source/ControlRig/Private/Sequencer/ControlRigLayerInstanceProxy.cpp b/Engine/Plugins/Animation/ControlRig/Source/ControlRig/Private/Sequencer/ControlRigLayerInstanceProxy.cpp
--- a/Engine/Plugins/Animation/ControlRig/Source/ControlRig/Private/Sequencer/ControlRigLayerInstanceProxy.cpp
+++ b/Engine/Plugins/Animation/ControlRig/Source/ControlRig/Private/Sequencer/ControlRigLayerInstanceProxy.cpp
@@ -42,6 +42,28 @@ void FControlRigLayerInstanceProxy::PreEvaluateAnimation(UAnimInstance* InAnimIn
 
 void FControlRigLayerInstanceProxy::SortControlRigNodes()
 {
+	//remove any invalid nodes
+	TArray<int32> NodesToRemove;
+	for (int32 Index = 0; Index < ControlRigNodes.Num(); ++Index)
+	{
+		FAnimNode_ControlRig_ExternalSource* CurrentNode = ControlRigNodes[Index].Get();
+		if (CurrentNode && CurrentNode->GetControlRig() == nullptr)
+		{
+			//need to find the control rig by manually finding it in the sequencer map
+			for (TPair<int32, FAnimNode_ControlRig_ExternalSource*>& Pair : SequencerToControlRigNodeMap)
+			{
+				if (Pair.Value == CurrentNode)
+				{
+					NodesToRemove.Add(Pair.Key);
+				}
+			}
+		}
+	}
+	for (int32 ID : NodesToRemove)
+	{
+		RemoveControlRigTrack(ID);
+	}
+
 	auto SortPredicate = [](TSharedPtr<FAnimNode_ControlRig_ExternalSource>& A, TSharedPtr<FAnimNode_ControlRig_ExternalSource>& B)
 	{
 		FAnimNode_ControlRig_ExternalSource* APtr = A.Get();
