Sequencer: Control Rig:  Change the control mask from index based to be a set of names of masked out controls.  Indices aren't reliable since the number of controls for rigs now changes for most rigs based upon how it is getting set up, e.g IK for FK settings now change the number of controls. We don't upgrade to use a name set until channel reconstruction since the rigs aren't setup till them. Related to this we always store the control names used to reconstruct the channel proxy, even when empty. Also we now check the mask when evaluating, which is needed for overrides.  Tested old content to make sure the upgade path worked, and new content with procedural/modular rigs.

diff --git a/Engine/Plugins/Animation/ControlRig/Source/ControlRig/Private/Sequencer/MovieSceneControlRigParameterSection.cpp b/Engine/Plugins/Animation/ControlRig/Source/ControlRig/Private/Sequencer/MovieSceneControlRigParameterSection.cpp
--- a/Engine/Plugins/Animation/ControlRig/Source/ControlRig/Private/Sequencer/MovieSceneControlRigParameterSection.cpp
+++ b/Engine/Plugins/Animation/ControlRig/Source/ControlRig/Private/Sequencer/MovieSceneControlRigParameterSection.cpp
@@ -1648,131 +1648,127 @@ EMovieSceneChannelProxyType UMovieSceneControlRigParameterSection::CacheChannelP
 
 	FMovieSceneChannelProxyData Channels;
 	ControlChannelMap.Empty();
+
 	// Need to create the channels in sorted orders, only if we have controls
-	if (ControlRig && ControlRig->AvailableControls().Num() > 0)
+	if (ControlRig)
 	{
 		TArray<FRigControlElement*> SortedControls;
 		ControlRig->GetControlsInOrder(SortedControls);
 		StoreLastControlsUsedToReconstruct(SortedControls);
-		bool bIsInUndo = false;
-		if (ControlsMask.Num() != SortedControls.Num())
-		{
-			TArray<bool> OnArray;
-			OnArray.Init(true, ControlRig->AvailableControls().Num());
-			SetControlsMask(OnArray);
-		}
-		int32 ControlIndex = 0; 
-		int32 MaskIndex = 0;
-		int32 SortOrder = 1; //start with one so Weight is first 
-		int32 FloatChannelIndex = 0;
-		int32 BoolChannelIndex = 0;
-		int32 EnumChannelIndex = 0;
-		int32 IntegerChannelIndex = 0;
-		int32 SpaceChannelIndex = 0;
-		int32 CategoryIndex = 0;
-		int32 ConstraintsChannelIndex = 0;
-		
-		const FName BoolChannelTypeName = FMovieSceneBoolChannel::StaticStruct()->GetFName();
-		const FName EnumChannelTypeName = FMovieSceneByteChannel::StaticStruct()->GetFName();
-		const FName IntegerChannelTypeName = FMovieSceneIntegerChannel::StaticStruct()->GetFName();
-		const FName SpaceName = FName(TEXT("Space"));
-
+		if (SortedControls.Num() > 0)
+		{
+			int32 ControlIndex = 0;
+			int32 MaskIndex = 0;
+			int32 SortOrder = 1; //start with one so Weight is first 
+			int32 FloatChannelIndex = 0;
+			int32 BoolChannelIndex = 0;
+			int32 EnumChannelIndex = 0;
+			int32 IntegerChannelIndex = 0;
+			int32 SpaceChannelIndex = 0;
+			int32 CategoryIndex = 0;
+			int32 ConstraintsChannelIndex = 0;
+
+			const FName BoolChannelTypeName = FMovieSceneBoolChannel::StaticStruct()->GetFName();
+			const FName EnumChannelTypeName = FMovieSceneByteChannel::StaticStruct()->GetFName();
+			const FName IntegerChannelTypeName = FMovieSceneIntegerChannel::StaticStruct()->GetFName();
+			const FName SpaceName = FName(TEXT("Space"));
+
+			// begin constraints
 		// begin constraints
-	// begin constraints
-		auto AddConstrainChannels = [this, &ConstraintsChannelIndex, &SortOrder, &Channels](
-			const FName& InControlName, const FText& InGroup, const bool bEnabled)
-		{
-			
-			const FConstraintsManagerController& Controller = FConstraintsManagerController::Get(ControlRig->GetWorld());
-
-			static constexpr bool bSorted = true;
-			const uint32 ControlHash = UTransformableControlHandle::ComputeHash(ControlRig.Get(), InControlName);
-			TArray<TWeakObjectPtr<UTickableConstraint>> Constraints = Controller.GetParentConstraints(ControlHash, bSorted);
-			for (const TWeakObjectPtr<UTickableConstraint>& Constraint : Constraints)
-			{
-				if (Constraint.IsValid())
+			auto AddConstrainChannels = [this, &ConstraintsChannelIndex, &SortOrder, &Channels](
+				const FName& InControlName, const FText& InGroup, const bool bEnabled)
 				{
-					const FGuid& ConstraintID = Constraint->ConstraintID;
-					if (FConstraintAndActiveChannel* ConstraintChannel = GetConstraintChannel(ConstraintID))
-					{
-						if (FChannelMapInfo* ChannelInfo = ControlChannelMap.Find(InControlName))
-						{
-							ChannelInfo->ConstraintsIndex.Add(ConstraintsChannelIndex);
-						}
 
-#if WITH_EDITOR
-						ConstraintChannel->ActiveChannel.ExtraLabel = [Constraint]
+					const FConstraintsManagerController& Controller = FConstraintsManagerController::Get(ControlRig->GetWorld());
+
+					static constexpr bool bSorted = true;
+					const uint32 ControlHash = UTransformableControlHandle::ComputeHash(ControlRig.Get(), InControlName);
+					TArray<TWeakObjectPtr<UTickableConstraint>> Constraints = Controller.GetParentConstraints(ControlHash, bSorted);
+					for (const TWeakObjectPtr<UTickableConstraint>& Constraint : Constraints)
+					{
+						if (Constraint.IsValid())
 						{
-							if (Constraint.IsValid())
+							const FGuid& ConstraintID = Constraint->ConstraintID;
+							if (FConstraintAndActiveChannel* ConstraintChannel = GetConstraintChannel(ConstraintID))
 							{
-								FString ParentStr; Constraint->GetLabel().Split(TEXT("."), &ParentStr, nullptr);
-								if (!ParentStr.IsEmpty())
+								if (FChannelMapInfo* ChannelInfo = ControlChannelMap.Find(InControlName))
 								{
-									return ParentStr;
+									ChannelInfo->ConstraintsIndex.Add(ConstraintsChannelIndex);
 								}
-							}
-							static const FString DummyStr;
-							return DummyStr;
-						};
 
-						const FText DisplayText = FText::FromString(Constraint->GetTypeLabel());
-						FMovieSceneChannelMetaData MetaData(Constraint->GetFName(), DisplayText, InGroup, bEnabled);
-						ConstraintsChannelIndex += 1;
-						MetaData.SortOrder = SortOrder++;
-						MetaData.bCanCollapseToTrack = true;
+#if WITH_EDITOR
+								ConstraintChannel->ActiveChannel.ExtraLabel = [Constraint]
+									{
+										if (Constraint.IsValid())
+										{
+											FString ParentStr; Constraint->GetLabel().Split(TEXT("."), &ParentStr, nullptr);
+											if (!ParentStr.IsEmpty())
+											{
+												return ParentStr;
+											}
+										}
+										static const FString DummyStr;
+										return DummyStr;
+									};
+
+								const FText DisplayText = FText::FromString(Constraint->GetTypeLabel());
+								FMovieSceneChannelMetaData MetaData(Constraint->GetFName(), DisplayText, InGroup, bEnabled);
+								ConstraintsChannelIndex += 1;
+								MetaData.SortOrder = SortOrder++;
+								MetaData.bCanCollapseToTrack = true;
 
-						Channels.Add(ConstraintChannel->ActiveChannel, MetaData, TMovieSceneExternalValue<bool>());
+								Channels.Add(ConstraintChannel->ActiveChannel, MetaData, TMovieSceneExternalValue<bool>());
 #else
-						Channels.Add(ConstraintChannel->ActiveChannel);
+								Channels.Add(ConstraintChannel->ActiveChannel);
 #endif
+							}
+						}
 					}
-				}
-			}
-		};
-		// end constraints
+				};
+			// end constraints
 
 #if WITH_EDITOR
 		// masking for per control channels based on control filters
-		auto MaybeApplyChannelMask = [](FMovieSceneChannelMetaData& OutMetadata, const FRigControlElement* ControlElement, ERigControlTransformChannel InChannel)
-		{
-			if(!OutMetadata.bEnabled)
-			{
-				return;
-			}
-			
-			const TArray<ERigControlTransformChannel>& FilteredChannels = ControlElement->Settings.FilteredChannels;
-			if(!FilteredChannels.IsEmpty())
-			{
-				OutMetadata.bEnabled = FilteredChannels.Contains(InChannel);
-			}
-		};
+			auto MaybeApplyChannelMask = [](FMovieSceneChannelMetaData& OutMetadata, const FRigControlElement* ControlElement, ERigControlTransformChannel InChannel)
+				{
+					if (!OutMetadata.bEnabled)
+					{
+						return;
+					}
+
+					const TArray<ERigControlTransformChannel>& FilteredChannels = ControlElement->Settings.FilteredChannels;
+					if (!FilteredChannels.IsEmpty())
+					{
+						OutMetadata.bEnabled = FilteredChannels.Contains(InChannel);
+					}
+				};
 #endif
 
-		URigHierarchy* Hierarchy = ControlRig->GetHierarchy();
-		for (FRigControlElement* ControlElement : SortedControls)
-		{
-			if (!Hierarchy->IsAnimatable(ControlElement))
+			URigHierarchy* Hierarchy = ControlRig->GetHierarchy();
+			for (FRigControlElement* ControlElement : SortedControls)
 			{
-				continue;
-			}
+				if (!Hierarchy->IsAnimatable(ControlElement))
+				{
+					continue;
+				}
 
-			FName ParentControlName = NAME_None;
-			FText Group;
+				FName ParentControlName = NAME_None;
+				FText Group;
 
-			if(Hierarchy->ShouldBeGrouped(ControlElement))
-			{
-				if(const FRigControlElement* ParentControlElement = Cast<FRigControlElement>(Hierarchy->GetFirstParent(ControlElement)))
+				if (Hierarchy->ShouldBeGrouped(ControlElement))
 				{
-					ParentControlName = ParentControlElement->GetFName();
-					Group = Hierarchy->GetDisplayNameForUI(ParentControlElement);
+					if (const FRigControlElement* ParentControlElement = Cast<FRigControlElement>(Hierarchy->GetFirstParent(ControlElement)))
+					{
+						ParentControlName = ParentControlElement->GetFName();
+						Group = Hierarchy->GetDisplayNameForUI(ParentControlElement);
+					}
 				}
-			}
 
-			bool bEnabled = ControlsMask[MaskIndex];
+				bool bEnabled = GetControlNameMask(ControlElement->GetFName());
 
 #if WITH_EDITOR
-			switch (ControlElement->Settings.ControlType)
-			{
+				switch (ControlElement->Settings.ControlType)
+				{
 				case ERigControlType::Float:
 				case ERigControlType::ScaleFloat:
 				{
@@ -1792,8 +1788,8 @@ EMovieSceneChannelProxyType UMovieSceneControlRigParameterSection::CacheChannelP
 							else
 							{
 								const FChannelMapInfo* pChannelIndex = ControlChannelMap.Find(ParentControlName);
-								const int32 ParentControlIndex = pChannelIndex ? pChannelIndex->ControlIndex : INDEX_NONE; 
-								ControlChannelMap.Add(Scalar.ParameterName, FChannelMapInfo(ControlIndex, SortOrder, FloatChannelIndex,ParentControlIndex, NAME_None, MaskIndex, CategoryIndex));
+								const int32 ParentControlIndex = pChannelIndex ? pChannelIndex->ControlIndex : INDEX_NONE;
+								ControlChannelMap.Add(Scalar.ParameterName, FChannelMapInfo(ControlIndex, SortOrder, FloatChannelIndex, ParentControlIndex, NAME_None, MaskIndex, CategoryIndex));
 							}
 
 							FParameterFloatChannelEditorData EditorData(ControlRig, Scalar.ParameterName, bEnabled, Group, SortOrder);
@@ -1817,7 +1813,7 @@ EMovieSceneChannelProxyType UMovieSceneControlRigParameterSection::CacheChannelP
 						{
 							if (Group.IsEmpty())
 							{
-								ControlChannelMap.Add(Bool.ParameterName, FChannelMapInfo(ControlIndex, SortOrder, BoolChannelIndex, INDEX_NONE,BoolChannelTypeName,MaskIndex, CategoryIndex));
+								ControlChannelMap.Add(Bool.ParameterName, FChannelMapInfo(ControlIndex, SortOrder, BoolChannelIndex, INDEX_NONE, BoolChannelTypeName, MaskIndex, CategoryIndex));
 								Group = Hierarchy->GetDisplayNameForUI(ControlElement);
 								if (bEnabled)
 								{
@@ -1827,8 +1823,8 @@ EMovieSceneChannelProxyType UMovieSceneControlRigParameterSection::CacheChannelP
 							else
 							{
 								const FChannelMapInfo* pChannelIndex = ControlChannelMap.Find(ParentControlName);
-								const int32 ParentControlIndex = pChannelIndex ? pChannelIndex->ControlIndex : INDEX_NONE; 
-								ControlChannelMap.Add(Bool.ParameterName, FChannelMapInfo(ControlIndex, SortOrder, BoolChannelIndex,ParentControlIndex, BoolChannelTypeName,MaskIndex, CategoryIndex));
+								const int32 ParentControlIndex = pChannelIndex ? pChannelIndex->ControlIndex : INDEX_NONE;
+								ControlChannelMap.Add(Bool.ParameterName, FChannelMapInfo(ControlIndex, SortOrder, BoolChannelIndex, ParentControlIndex, BoolChannelTypeName, MaskIndex, CategoryIndex));
 							}
 
 							FMovieSceneChannelMetaData MetaData(Bool.ParameterName, Group, Group, bEnabled);
@@ -1854,7 +1850,7 @@ EMovieSceneChannelProxyType UMovieSceneControlRigParameterSection::CacheChannelP
 							{
 								if (Group.IsEmpty())
 								{
-									ControlChannelMap.Add(Enum.ParameterName, FChannelMapInfo(ControlIndex, SortOrder, EnumChannelIndex,INDEX_NONE, EnumChannelTypeName,MaskIndex, CategoryIndex));
+									ControlChannelMap.Add(Enum.ParameterName, FChannelMapInfo(ControlIndex, SortOrder, EnumChannelIndex, INDEX_NONE, EnumChannelTypeName, MaskIndex, CategoryIndex));
 									Group = Hierarchy->GetDisplayNameForUI(ControlElement);
 									if (bEnabled)
 									{
@@ -1864,8 +1860,8 @@ EMovieSceneChannelProxyType UMovieSceneControlRigParameterSection::CacheChannelP
 								else
 								{
 									const FChannelMapInfo* pChannelIndex = ControlChannelMap.Find(ParentControlName);
-									const int32 ParentControlIndex = pChannelIndex ? pChannelIndex->ControlIndex : INDEX_NONE; 
-									ControlChannelMap.Add(Enum.ParameterName, FChannelMapInfo(ControlIndex, SortOrder, EnumChannelIndex, ParentControlIndex, EnumChannelTypeName,MaskIndex, CategoryIndex));
+									const int32 ParentControlIndex = pChannelIndex ? pChannelIndex->ControlIndex : INDEX_NONE;
+									ControlChannelMap.Add(Enum.ParameterName, FChannelMapInfo(ControlIndex, SortOrder, EnumChannelIndex, ParentControlIndex, EnumChannelTypeName, MaskIndex, CategoryIndex));
 								}
 
 								FMovieSceneChannelMetaData MetaData(Enum.ParameterName, Group, Group, bEnabled);
@@ -1888,7 +1884,7 @@ EMovieSceneChannelProxyType UMovieSceneControlRigParameterSection::CacheChannelP
 							{
 								if (Group.IsEmpty())
 								{
-									ControlChannelMap.Add(Integer.ParameterName, FChannelMapInfo(ControlIndex, SortOrder, IntegerChannelIndex,INDEX_NONE,IntegerChannelTypeName,MaskIndex, CategoryIndex));
+									ControlChannelMap.Add(Integer.ParameterName, FChannelMapInfo(ControlIndex, SortOrder, IntegerChannelIndex, INDEX_NONE, IntegerChannelTypeName, MaskIndex, CategoryIndex));
 									Group = Hierarchy->GetDisplayNameForUI(ControlElement);
 									if (bEnabled)
 									{
@@ -1898,8 +1894,8 @@ EMovieSceneChannelProxyType UMovieSceneControlRigParameterSection::CacheChannelP
 								else
 								{
 									const FChannelMapInfo* pChannelIndex = ControlChannelMap.Find(ParentControlName);
-									const int32 ParentControlIndex = pChannelIndex ? pChannelIndex->ControlIndex : INDEX_NONE; 
-									ControlChannelMap.Add(Integer.ParameterName, FChannelMapInfo(ControlIndex, SortOrder, IntegerChannelIndex, ParentControlIndex, IntegerChannelTypeName,MaskIndex, CategoryIndex));
+									const int32 ParentControlIndex = pChannelIndex ? pChannelIndex->ControlIndex : INDEX_NONE;
+									ControlChannelMap.Add(Integer.ParameterName, FChannelMapInfo(ControlIndex, SortOrder, IntegerChannelIndex, ParentControlIndex, IntegerChannelTypeName, MaskIndex, CategoryIndex));
 								}
 
 								FMovieSceneChannelMetaData MetaData(Integer.ParameterName, Group, Group, bEnabled);
@@ -1925,7 +1921,7 @@ EMovieSceneChannelProxyType UMovieSceneControlRigParameterSection::CacheChannelP
 					{
 						if (ControlElement->GetFName() == Vector2D.ParameterName)
 						{
-							if(Group.IsEmpty())
+							if (Group.IsEmpty())
 							{
 								ControlChannelMap.Add(Vector2D.ParameterName, FChannelMapInfo(ControlIndex, SortOrder, FloatChannelIndex, INDEX_NONE, NAME_None, MaskIndex, CategoryIndex));
 								if (bEnabled)
@@ -1937,7 +1933,7 @@ EMovieSceneChannelProxyType UMovieSceneControlRigParameterSection::CacheChannelP
 							else
 							{
 								const FChannelMapInfo* pChannelIndex = ControlChannelMap.Find(ParentControlName);
-								const int32 ParentControlIndex = pChannelIndex ? pChannelIndex->ControlIndex : INDEX_NONE; 
+								const int32 ParentControlIndex = pChannelIndex ? pChannelIndex->ControlIndex : INDEX_NONE;
 								ControlChannelMap.Add(Vector2D.ParameterName, FChannelMapInfo(ControlIndex, SortOrder, FloatChannelIndex, ParentControlIndex, NAME_None, MaskIndex, CategoryIndex));
 							}
 							FParameterVectorChannelEditorData EditorData(ControlRig, Vector2D.ParameterName, bEnabled, Group, SortOrder, 2);
@@ -1961,7 +1957,7 @@ EMovieSceneChannelProxyType UMovieSceneControlRigParameterSection::CacheChannelP
 					{
 						if (ControlElement->GetFName() == Vector.ParameterName)
 						{
-							if(Group.IsEmpty())
+							if (Group.IsEmpty())
 							{
 								ControlChannelMap.Add(Vector.ParameterName, FChannelMapInfo(ControlIndex, SortOrder, FloatChannelIndex, INDEX_NONE, NAME_None, MaskIndex, CategoryIndex));
 								if (bEnabled)
@@ -1973,7 +1969,7 @@ EMovieSceneChannelProxyType UMovieSceneControlRigParameterSection::CacheChannelP
 							else
 							{
 								const FChannelMapInfo* pChannelIndex = ControlChannelMap.Find(ParentControlName);
-								const int32 ParentControlIndex = pChannelIndex ? pChannelIndex->ControlIndex : INDEX_NONE; 
+								const int32 ParentControlIndex = pChannelIndex ? pChannelIndex->ControlIndex : INDEX_NONE;
 								ControlChannelMap.Add(Vector.ParameterName, FChannelMapInfo(ControlIndex, SortOrder, FloatChannelIndex, ParentControlIndex, NAME_None, MaskIndex, CategoryIndex));
 							}
 							if (FSpaceControlNameAndChannel* SpaceChannel = GetSpaceChannel(Vector.ParameterName))
@@ -1997,29 +1993,29 @@ EMovieSceneChannelProxyType UMovieSceneControlRigParameterSection::CacheChannelP
 								SpaceMetaData.bCanCollapseToTrack = true;
 								Channels.Add(SpaceChannel->SpaceCurve, SpaceMetaData);
 							}
-						
+
 
 							FParameterVectorChannelEditorData EditorData(ControlRig, Vector.ParameterName, bEnabled, Group, SortOrder, 3);
 
-							if(ControlElement->Settings.ControlType == ERigControlType::Position)
+							if (ControlElement->Settings.ControlType == ERigControlType::Position)
 							{
 								MaybeApplyChannelMask(EditorData.MetaData[0], ControlElement, ERigControlTransformChannel::TranslationX);
 								MaybeApplyChannelMask(EditorData.MetaData[1], ControlElement, ERigControlTransformChannel::TranslationY);
 								MaybeApplyChannelMask(EditorData.MetaData[2], ControlElement, ERigControlTransformChannel::TranslationZ);
 							}
-							else if(ControlElement->Settings.ControlType == ERigControlType::Rotator)
+							else if (ControlElement->Settings.ControlType == ERigControlType::Rotator)
 							{
 								MaybeApplyChannelMask(EditorData.MetaData[0], ControlElement, ERigControlTransformChannel::Pitch);
 								MaybeApplyChannelMask(EditorData.MetaData[1], ControlElement, ERigControlTransformChannel::Yaw);
 								MaybeApplyChannelMask(EditorData.MetaData[2], ControlElement, ERigControlTransformChannel::Roll);
 							}
-							else if(ControlElement->Settings.ControlType == ERigControlType::Scale)
+							else if (ControlElement->Settings.ControlType == ERigControlType::Scale)
 							{
 								MaybeApplyChannelMask(EditorData.MetaData[0], ControlElement, ERigControlTransformChannel::ScaleX);
 								MaybeApplyChannelMask(EditorData.MetaData[1], ControlElement, ERigControlTransformChannel::ScaleY);
 								MaybeApplyChannelMask(EditorData.MetaData[2], ControlElement, ERigControlTransformChannel::ScaleZ);
 							}
-							
+
 							Channels.Add(Vector.XCurve, EditorData.MetaData[0], EditorData.ExternalValues[0]);
 							Channels.Add(Vector.YCurve, EditorData.MetaData[1], EditorData.ExternalValues[1]);
 							Channels.Add(Vector.ZCurve, EditorData.MetaData[2], EditorData.ExternalValues[2]);
@@ -2041,7 +2037,7 @@ EMovieSceneChannelProxyType UMovieSceneControlRigParameterSection::CacheChannelP
 						if (ControlElement->GetFName() == Transform.ParameterName)
 						{
 							const FName ControlName = ControlElement->GetFName();
-							if(Group.IsEmpty())
+							if (Group.IsEmpty())
 							{
 								ControlChannelMap.Add(Transform.ParameterName, FChannelMapInfo(ControlIndex, SortOrder, FloatChannelIndex, INDEX_NONE, NAME_None, MaskIndex, CategoryIndex));
 								if (bEnabled)
@@ -2053,7 +2049,7 @@ EMovieSceneChannelProxyType UMovieSceneControlRigParameterSection::CacheChannelP
 							else
 							{
 								const FChannelMapInfo* pChannelIndex = ControlChannelMap.Find(ParentControlName);
-								const int32 ParentControlIndex = pChannelIndex ? pChannelIndex->ControlIndex : INDEX_NONE; 
+								const int32 ParentControlIndex = pChannelIndex ? pChannelIndex->ControlIndex : INDEX_NONE;
 								ControlChannelMap.Add(Transform.ParameterName, FChannelMapInfo(ControlIndex, SortOrder, FloatChannelIndex, ParentControlIndex, NAME_None, MaskIndex, CategoryIndex));
 							}
 
@@ -2082,15 +2078,15 @@ EMovieSceneChannelProxyType UMovieSceneControlRigParameterSection::CacheChannelP
 								//TMovieSceneExternalValue<FMovieSceneControlRigSpaceBaseKey> ExternalData;
 								Channels.Add(SpaceChannel->SpaceCurve, SpaceMetaData);
 							}
-					
 
-							FParameterTransformChannelEditorData EditorData(ControlRig, Transform.ParameterName, bEnabled, TransformMask.GetChannels(), Group, 
+
+							FParameterTransformChannelEditorData EditorData(ControlRig, Transform.ParameterName, bEnabled, TransformMask.GetChannels(), Group,
 								SortOrder);
 
 							MaybeApplyChannelMask(EditorData.MetaData[0], ControlElement, ERigControlTransformChannel::TranslationX);
 							MaybeApplyChannelMask(EditorData.MetaData[1], ControlElement, ERigControlTransformChannel::TranslationY);
 							MaybeApplyChannelMask(EditorData.MetaData[2], ControlElement, ERigControlTransformChannel::TranslationZ);
-							
+
 							Channels.Add(Transform.Translation[0], EditorData.MetaData[0], EditorData.ExternalValues[0]);
 							Channels.Add(Transform.Translation[1], EditorData.MetaData[1], EditorData.ExternalValues[1]);
 							Channels.Add(Transform.Translation[2], EditorData.MetaData[2], EditorData.ExternalValues[2]);
@@ -2110,7 +2106,7 @@ EMovieSceneChannelProxyType UMovieSceneControlRigParameterSection::CacheChannelP
 								MaybeApplyChannelMask(EditorData.MetaData[6], ControlElement, ERigControlTransformChannel::ScaleX);
 								MaybeApplyChannelMask(EditorData.MetaData[7], ControlElement, ERigControlTransformChannel::ScaleY);
 								MaybeApplyChannelMask(EditorData.MetaData[8], ControlElement, ERigControlTransformChannel::ScaleZ);
-								
+
 								Channels.Add(Transform.Scale[0], EditorData.MetaData[6], EditorData.ExternalValues[6]);
 								Channels.Add(Transform.Scale[1], EditorData.MetaData[7], EditorData.ExternalValues[7]);
 								Channels.Add(Transform.Scale[2], EditorData.MetaData[8], EditorData.ExternalValues[8]);
@@ -2132,7 +2128,7 @@ EMovieSceneChannelProxyType UMovieSceneControlRigParameterSection::CacheChannelP
 				default:
 					break;
 				}
-	#else
+#else
 				switch (ControlElement->Settings.ControlType)
 				{
 				case ERigControlType::Float:
@@ -2141,7 +2137,7 @@ EMovieSceneChannelProxyType UMovieSceneControlRigParameterSection::CacheChannelP
 					{
 						if (ControlElement->GetFName() == Scalar.ParameterName)
 						{
-							ControlChannelMap.Add(Scalar.ParameterName, FChannelMapInfo(ControlIndex, SortOrder,FloatChannelIndex,INDEX_NONE, NAME_None,MaskIndex));
+							ControlChannelMap.Add(Scalar.ParameterName, FChannelMapInfo(ControlIndex, SortOrder, FloatChannelIndex, INDEX_NONE, NAME_None, MaskIndex));
 							Channels.Add(Scalar.ParameterCurve);
 							FloatChannelIndex += 1;
 							SortOrder += 1;
@@ -2272,8 +2268,8 @@ EMovieSceneChannelProxyType UMovieSceneControlRigParameterSection::CacheChannelP
 					{
 						if (ControlElement->GetFName() == Transform.ParameterName)
 						{
-							ControlChannelMap.Add(Transform.ParameterName, FChannelMapInfo(ControlIndex, SortOrder,FloatChannelIndex, INDEX_NONE, NAME_None, MaskIndex));
-						
+							ControlChannelMap.Add(Transform.ParameterName, FChannelMapInfo(ControlIndex, SortOrder, FloatChannelIndex, INDEX_NONE, NAME_None, MaskIndex));
+
 							bool bDoSpaceChannel = true;
 							if (bDoSpaceChannel)
 							{
@@ -2289,8 +2285,8 @@ EMovieSceneChannelProxyType UMovieSceneControlRigParameterSection::CacheChannelP
 									SpaceChannelIndex += 1;
 									Channels.Add(SpaceChannel->SpaceCurve);
 								}
-							}	
-						
+							}
+
 							Channels.Add(Transform.Translation[0]);
 							Channels.Add(Transform.Translation[1]);
 							Channels.Add(Transform.Translation[2]);
@@ -2315,32 +2311,32 @@ EMovieSceneChannelProxyType UMovieSceneControlRigParameterSection::CacheChannelP
 
 							ControlIndex += 1;
 							break;
-						}
 					}
+				}
 					break;
 				}
 			}
 #endif
-			++MaskIndex;
+				++MaskIndex;
 
 		}
-		
+
 #if WITH_EDITOR
-		FMovieSceneChannelMetaData      MetaData;
-		MetaData.SetIdentifiers("Weight", NSLOCTEXT("MovieSceneTransformSection", "Weight", "Weight"));
-		MetaData.bEnabled = EnumHasAllFlags(TransformMask.GetChannels(), EMovieSceneTransformChannel::Weight);
-		MetaData.SortOrder = 0;
-		MetaData.bSortEmptyGroupsLast = false;
-		MetaData.bCanCollapseToTrack = true;
-		TMovieSceneExternalValue<float> ExVal;
-		Channels.Add(Weight, MetaData, ExVal);
+			FMovieSceneChannelMetaData      MetaData;
+			MetaData.SetIdentifiers("Weight", NSLOCTEXT("MovieSceneTransformSection", "Weight", "Weight"));
+			MetaData.bEnabled = EnumHasAllFlags(TransformMask.GetChannels(), EMovieSceneTransformChannel::Weight);
+			MetaData.SortOrder = 0;
+			MetaData.bSortEmptyGroupsLast = false;
+			MetaData.bCanCollapseToTrack = true;
+			TMovieSceneExternalValue<float> ExVal;
+			Channels.Add(Weight, MetaData, ExVal);
 #else
-		Channels.Add(Weight);
+			Channels.Add(Weight);
 
 #endif
+		}
 	}
 
-
 	ChannelProxy = MakeShared<FMovieSceneChannelProxy>(MoveTemp(Channels));
 	
 	return EMovieSceneChannelProxyType::Dynamic;
@@ -2479,13 +2475,24 @@ void UMovieSceneControlRigParameterSection::RecreateWithThisControlRig(UControlR
 	TransformParameterNamesAndCurves.Empty();
 	*/
 
+	//update the mask array to the new mask name set
+	//need to do it here since we won't get controls until here
+	const int32 NumControls = ControlRig->AvailableControls().Num();
+	const int32 MaskNum = ControlsMask.Num();
+	if (NumControls > 0 && NumControls == MaskNum)
+	{
+		ConvertMaskArrayToNameSet();
+	}
+	/*
 	//if we had the same with same number of controls keep the mask otherwise reset it.
-	if (!bSameControlRig || (ControlRig->AvailableControls().Num() > 0  && ControlRig->AvailableControls().Num() != ControlsMask.Num()))
+
+	if (!bSameControlRig || (NumControls > 0  && (NumControls != MaskNum)))
 	{
 		TArray<bool> OnArray;
 		OnArray.Init(true, ControlRig->AvailableControls().Num());
 		SetControlsMask(OnArray);
 	}
+	*/
 	
 	TArray<FRigControlElement*> SortedControls;
 	ControlRig->GetControlsInOrder(SortedControls);
@@ -2741,6 +2748,68 @@ void UMovieSceneControlRigParameterSection::ChangeControlRotationOrder(const FNa
 	}
 }
 
+void UMovieSceneControlRigParameterSection::ConvertMaskArrayToNameSet()
+{
+	if (ControlRig && ControlsMask.Num() > 0)
+	{
+		TArray<FRigControlElement*> SortedControls;
+		ControlRig->GetControlsInOrder(SortedControls);
+		if (SortedControls.Num() == ControlsMask.Num())
+		{
+			ControlNameMask.Empty();
+			for (int32 Index = 0; Index < SortedControls.Num(); ++Index)
+			{
+				if (ControlsMask[Index] == false)
+				{
+					ControlNameMask.Add(SortedControls[Index]->GetKey().Name);
+				}
+			}
+		}
+		//empty ControlsMask, no longer needed
+		ControlsMask.Empty();
+	}
+}
+
+void UMovieSceneControlRigParameterSection::FillControlNameMask(bool bValue)
+{
+	if (ControlRig)
+	{
+		ControlNameMask.Empty();
+		if (bValue == false)
+		{
+			TArray<FRigControlElement*> SortedControls;
+			ControlRig->GetControlsInOrder(SortedControls);
+			for (FRigControlElement* ControlElement : SortedControls)
+			{
+				if (ControlElement)
+				{
+					ControlNameMask.Add(ControlElement->GetKey().Name);
+				}
+			}
+		}
+		ReconstructChannelProxy();
+	}
+}
+
+void UMovieSceneControlRigParameterSection::SetControlNameMask(const FName& Name, bool bValue)
+{
+	if (bValue == false)
+	{
+		ControlNameMask.Add(Name);
+	}
+	else
+	{
+		ControlNameMask.Remove(Name);
+	}
+	ReconstructChannelProxy();
+}
+
+//get value, will return false if not found
+bool UMovieSceneControlRigParameterSection::GetControlNameMask(const FName& Name) const
+{
+	return (ControlNameMask.Find(Name) == nullptr);
+}
+
 void UMovieSceneControlRigParameterSection::FixRotationWinding(const FName& ControlName, FFrameNumber StartFrame, FFrameNumber EndFrame)
 {
 	FChannelMapInfo* pChannelIndex = ControlChannelMap.Find(ControlName);
@@ -3466,7 +3535,38 @@ const TArray<FIntegerParameterNameAndCurve>& UMovieSceneControlRigParameterSecti
 {
 	return IntegerParameterNamesAndCurves;
 }
+/*
+void FillControlNameMask(bool bValue)
+{
+	TArray<FRigControlElement*> SortedControls;
+	ControlRig->GetControlsInOrder(SortedControls);
+	ControlNameMask.Empty();
+	for (FRigControlElement* ControlElement : SortedControls)
+	{
+		if (ControlElement)
+		{
+			ControlNameMask.Add(ControlElement->GetKey().Name, bValue);
+		}
+	}
+}
 
+void SetControlNameMask(const FName& Name, bool bValue)
+{
+	ControlNameMask.Add(Name, bValue);
+}
+
+//get value, will return false if not found
+bool GetControlNameMask(const FName& Name, bool& OutValue)
+{
+	bool* bValue = ControlNameMask.Find(Name);
+	if (bValue)
+	{
+		OutValue = *bValue;
+		return true;
+	}
+	return false;
+}
+*/
 void UMovieSceneControlRigParameterSection::ClearAllParameters()
 {
 	BoolParameterNamesAndCurves.SetNum(0);
@@ -3587,7 +3687,7 @@ int32 UMovieSceneControlRigParameterSection::GetActiveCategoryIndex(FName Contro
 {
 	int32 CategoryIndex = INDEX_NONE;
 	const FChannelMapInfo* pChannelIndex = ControlChannelMap.Find(ControlName);
-	if (pChannelIndex != nullptr && ControlsMask[pChannelIndex->MaskIndex])
+	if (pChannelIndex != nullptr && GetControlNameMask(ControlName))
 	{
 		CategoryIndex = pChannelIndex->CategoryIndex;
 	}
diff --git a/Engine/Plugins/Animation/ControlRig/Source/ControlRig/Private/Sequencer/MovieSceneControlRigParameterTemplate.cpp b/Engine/Plugins/Animation/ControlRig/Source/ControlRig/Private/Sequencer/MovieSceneControlRigParameterTemplate.cpp
--- a/Engine/Plugins/Animation/ControlRig/Source/ControlRig/Private/Sequencer/MovieSceneControlRigParameterTemplate.cpp
+++ b/Engine/Plugins/Animation/ControlRig/Source/ControlRig/Private/Sequencer/MovieSceneControlRigParameterTemplate.cpp
@@ -294,8 +294,6 @@ struct FEvaluatedControlRigParameterSectionChannelMasks : IPersistentEvaluationD
 		TArrayView<const FTransformParameterNameAndCurves> Transforms
 		)
 	{
-		const TArray<bool>& ControlsMask = Section->GetControlsMask();
-
 		const FChannelMapInfo* ChannelInfo = nullptr;
 
 		ScalarCurveMask.Add(false, Scalars.Num());
@@ -311,49 +309,49 @@ struct FEvaluatedControlRigParameterSectionChannelMasks : IPersistentEvaluationD
 		{
 			const FScalarParameterNameAndCurve& Scalar = Scalars[Index];
 			ChannelInfo = Section->ControlChannelMap.Find(Scalar.ParameterName);
-			ScalarCurveMask[Index] = (!ChannelInfo || ControlsMask[ChannelInfo->MaskIndex]);
+			ScalarCurveMask[Index] = (!ChannelInfo || Section->GetControlNameMask(Scalar.ParameterName));
 		}
 		for (int32 Index = 0; Index < Bools.Num(); ++Index)
 		{
 			const FBoolParameterNameAndCurve& Bool = Bools[Index];
 			ChannelInfo = Section->ControlChannelMap.Find(Bool.ParameterName);
-			BoolCurveMask[Index] = (!ChannelInfo || ControlsMask[ChannelInfo->MaskIndex]);
+			BoolCurveMask[Index] = (!ChannelInfo || Section->GetControlNameMask(Bool.ParameterName));
 		}
 		for (int32 Index = 0; Index < Integers.Num(); ++Index)
 		{
 			const FIntegerParameterNameAndCurve& Integer = Integers[Index];
 			ChannelInfo = Section->ControlChannelMap.Find(Integer.ParameterName);
-			IntegerCurveMask[Index] = (!ChannelInfo || ControlsMask[ChannelInfo->MaskIndex]);
+			IntegerCurveMask[Index] = (!ChannelInfo || Section->GetControlNameMask(Integer.ParameterName));
 		}
 		for (int32 Index = 0; Index < Enums.Num(); ++Index)
 		{
 			const FEnumParameterNameAndCurve& Enum = Enums[Index];
 			ChannelInfo = Section->ControlChannelMap.Find(Enum.ParameterName);
-			EnumCurveMask[Index] = (!ChannelInfo || ControlsMask[ChannelInfo->ControlIndex]);
+			EnumCurveMask[Index] = (!ChannelInfo || Section->GetControlNameMask(Enum.ParameterName));
 		}
 		for (int32 Index = 0; Index < Vector2Ds.Num(); ++Index)
 		{
 			const FVector2DParameterNameAndCurves& Vector2D = Vector2Ds[Index];
 			ChannelInfo = Section->ControlChannelMap.Find(Vector2D.ParameterName);
-			Vector2DCurveMask[Index] = (!ChannelInfo || ControlsMask[ChannelInfo->MaskIndex]);
+			Vector2DCurveMask[Index] = (!ChannelInfo || Section->GetControlNameMask(Vector2D.ParameterName));
 		}
 		for (int32 Index = 0; Index < Vectors.Num(); ++Index)
 		{
 			const FVectorParameterNameAndCurves& Vector = Vectors[Index];
 			ChannelInfo = Section->ControlChannelMap.Find(Vector.ParameterName);
-			VectorCurveMask[Index] = (!ChannelInfo || ControlsMask[ChannelInfo->MaskIndex]);
+			VectorCurveMask[Index] = (!ChannelInfo || Section->GetControlNameMask(Vector.ParameterName));
 		}
 		for (int32 Index = 0; Index < Colors.Num(); ++Index)
 		{
 			const FColorParameterNameAndCurves& Color = Colors[Index];
 			ChannelInfo = Section->ControlChannelMap.Find(Color.ParameterName);
-			ColorCurveMask[Index] = (!ChannelInfo || ControlsMask[ChannelInfo->MaskIndex]);
+			ColorCurveMask[Index] = (!ChannelInfo || Section->GetControlNameMask(Color.ParameterName));
 		}
 		for (int32 Index = 0; Index < Transforms.Num(); ++Index)
 		{
 			const FTransformParameterNameAndCurves& Transform = Transforms[Index];
 			ChannelInfo = Section->ControlChannelMap.Find(Transform.ParameterName);
-			TransformCurveMask[Index] = (!ChannelInfo || ControlsMask[ChannelInfo->MaskIndex]);
+			TransformCurveMask[Index] = (!ChannelInfo || Section->GetControlNameMask(Transform.ParameterName));
 		}
 	}
 };
@@ -1834,7 +1832,6 @@ struct TControlRigParameterActuatorTransform : TMovieSceneBlendingActuator<FCont
 
 };
 
-
 void FMovieSceneControlRigParameterTemplate::Evaluate(const FMovieSceneEvaluationOperand& Operand, const FMovieSceneContext& Context, const FPersistentEvaluationData& PersistentData, FMovieSceneExecutionTokens& ExecutionTokens) const
 {
 	const FFrameTime Time = Context.GetTime();
@@ -1894,6 +1891,10 @@ void FMovieSceneControlRigParameterTemplate::Evaluate(const FMovieSceneEvaluatio
 
 		for (const FScalarParameterStringAndValue& ScalarNameAndValue : Values.ScalarValues)
 		{
+			if (Section->GetControlNameMask(ScalarNameAndValue.ParameterName) == false)
+			{
+				continue;
+			}
 			FMovieSceneAnimTypeID AnimTypeID = TypeIDs->FindScalar(ScalarNameAndValue.ParameterName);
 			FMovieSceneBlendingActuatorID ActuatorTypeID(AnimTypeID);
 
@@ -1907,6 +1908,10 @@ void FMovieSceneControlRigParameterTemplate::Evaluate(const FMovieSceneEvaluatio
 		UE::MovieScene::TMultiChannelValue<float, 3> VectorData;
 		for (const FVectorParameterStringAndValue& VectorNameAndValue : Values.VectorValues)
 		{
+			if (Section->GetControlNameMask(VectorNameAndValue.ParameterName) == false)
+			{
+				continue;
+			}
 			FMovieSceneAnimTypeID AnimTypeID = TypeIDs->FindVector(VectorNameAndValue.ParameterName);
 			FMovieSceneBlendingActuatorID ActuatorTypeID(AnimTypeID);
 
@@ -1924,6 +1929,10 @@ void FMovieSceneControlRigParameterTemplate::Evaluate(const FMovieSceneEvaluatio
 		UE::MovieScene::TMultiChannelValue<float, 2> Vector2DData;
 		for (const FVector2DParameterStringAndValue& Vector2DNameAndValue : Values.Vector2DValues)
 		{
+			if (Section->GetControlNameMask(Vector2DNameAndValue.ParameterName) == false)
+			{
+				continue;
+			}
 			FMovieSceneAnimTypeID AnimTypeID = TypeIDs->FindVector2D(Vector2DNameAndValue.ParameterName);
 			FMovieSceneBlendingActuatorID ActuatorTypeID(AnimTypeID);
 
@@ -1940,6 +1949,10 @@ void FMovieSceneControlRigParameterTemplate::Evaluate(const FMovieSceneEvaluatio
 		UE::MovieScene::TMultiChannelValue<float, 9> TransformData;
 		for (const FEulerTransformParameterStringAndValue& TransformNameAndValue : Values.TransformValues)
 		{
+			if (Section->GetControlNameMask(TransformNameAndValue.ParameterName) == false)
+			{
+				continue;
+			}
 			FMovieSceneAnimTypeID AnimTypeID = TypeIDs->FindTransform(TransformNameAndValue.ParameterName);
 			FMovieSceneBlendingActuatorID ActuatorTypeID(AnimTypeID);
 
@@ -2386,9 +2399,12 @@ void FMovieSceneControlRigParameterTemplate::Interrogate(const FMovieSceneContex
 			Weight *= ManualWeight;
 		}
 
-
 		for (const FScalarParameterStringAndValue& ScalarNameAndValue : Values.ScalarValues)
 		{
+			if (Section->GetControlNameMask(ScalarNameAndValue.ParameterName) == false)
+			{
+				continue;
+			}
 			FMovieSceneAnimTypeID AnimTypeID = TypeIDs->FindScalar(ScalarNameAndValue.ParameterName);
 			FMovieSceneBlendingActuatorID ActuatorTypeID(AnimTypeID);
 
@@ -2402,6 +2418,10 @@ void FMovieSceneControlRigParameterTemplate::Interrogate(const FMovieSceneContex
 		UE::MovieScene::TMultiChannelValue<float, 2> Vector2DData;
 		for (const FVector2DParameterStringAndValue& Vector2DNameAndValue : Values.Vector2DValues)
 		{
+			if (Section->GetControlNameMask(Vector2DNameAndValue.ParameterName) == false)
+			{
+				continue;
+			}
 			FMovieSceneAnimTypeID AnimTypeID = TypeIDs->FindVector2D(Vector2DNameAndValue.ParameterName);
 			FMovieSceneBlendingActuatorID ActuatorTypeID(AnimTypeID);
 
@@ -2418,6 +2438,10 @@ void FMovieSceneControlRigParameterTemplate::Interrogate(const FMovieSceneContex
 		UE::MovieScene::TMultiChannelValue<float, 3> VectorData;
 		for (const FVectorParameterStringAndValue& VectorNameAndValue : Values.VectorValues)
 		{
+			if (Section->GetControlNameMask(VectorNameAndValue.ParameterName) == false)
+			{
+				continue;
+			}
 			FMovieSceneAnimTypeID AnimTypeID = TypeIDs->FindVector(VectorNameAndValue.ParameterName);
 			FMovieSceneBlendingActuatorID ActuatorTypeID(AnimTypeID);
 
@@ -2435,6 +2459,10 @@ void FMovieSceneControlRigParameterTemplate::Interrogate(const FMovieSceneContex
 		UE::MovieScene::TMultiChannelValue<float, 9> TransformData;
 		for (const FEulerTransformParameterStringAndValue& TransformNameAndValue : Values.TransformValues)
 		{
+			if (Section->GetControlNameMask(TransformNameAndValue.ParameterName) == false)
+			{
+				continue;
+			}
 			FMovieSceneAnimTypeID AnimTypeID = TypeIDs->FindTransform(TransformNameAndValue.ParameterName);
 			FMovieSceneBlendingActuatorID ActuatorTypeID(AnimTypeID);
 
diff --git a/Engine/Plugins/Animation/ControlRig/Source/ControlRig/Private/Sequencer/MovieSceneControlRigParameterTrack.cpp b/Engine/Plugins/Animation/ControlRig/Source/ControlRig/Private/Sequencer/MovieSceneControlRigParameterTrack.cpp
--- a/Engine/Plugins/Animation/ControlRig/Source/ControlRig/Private/Sequencer/MovieSceneControlRigParameterTrack.cpp
+++ b/Engine/Plugins/Animation/ControlRig/Source/ControlRig/Private/Sequencer/MovieSceneControlRigParameterTrack.cpp
@@ -459,7 +459,7 @@ void UMovieSceneControlRigParameterTrack::SetSectionToKey(UMovieSceneSection* In
 					FChannelMapInfo* pChannelIndex = CRSection->ControlChannelMap.Find(SectionToKeyItem.Key);
 					if (pChannelIndex)
 					{
-						if (CRSection->GetControlsMask(pChannelIndex->MaskIndex))
+						if (CRSection->GetControlNameMask(SectionToKeyItem.Key))
 						{
 							SectionToKeyItem.Value = InSection;
 						}
diff --git a/Engine/Plugins/Animation/ControlRig/Source/ControlRig/Public/Sequencer/MovieSceneControlRigParameterSection.h b/Engine/Plugins/Animation/ControlRig/Source/ControlRig/Public/Sequencer/MovieSceneControlRigParameterSection.h
--- a/Engine/Plugins/Animation/ControlRig/Source/ControlRig/Public/Sequencer/MovieSceneControlRigParameterSection.h
+++ b/Engine/Plugins/Animation/ControlRig/Source/ControlRig/Public/Sequencer/MovieSceneControlRigParameterSection.h
@@ -228,10 +228,14 @@ public:
 	UPROPERTY(EditAnywhere, Category = "Animation")
 	TSubclassOf<UControlRig> ControlRigClass;
 
-	/** Mask for controls themselves*/
+	/** Deprecrated, use ControlNameMask*/
 	UPROPERTY()
 	TArray<bool> ControlsMask;
 
+	/** Names of Controls that are masked out on this section*/
+	UPROPERTY()
+	TSet<FName> ControlNameMask;
+
 	/** Mask for Transform Mask*/
 	UPROPERTY()
 	FMovieSceneTransformMask TransformMask;
@@ -244,6 +248,7 @@ public:
 	UPROPERTY()
 	TMap<FName, FChannelMapInfo> ControlChannelMap;
 
+
 protected:
 	/** Enum Curves*/
 	UPROPERTY()
@@ -320,6 +325,20 @@ public:
 	UE_DEPRECATED(5.5, "LoadAnimSequenceIntoThisSection without taking a sequence start frame is deprecated, use version that takes start frame instead")
 	virtual bool LoadAnimSequenceIntoThisSection(UAnimSequence* Sequence, UMovieScene* MovieScene, UObject* BoundObject, bool bKeyReduce, float Tolerance, bool bResetControls, FFrameNumber InStartFrame , EMovieSceneKeyInterpolation InInterpolation);
 #endif
+	
+	void FillControlNameMask(bool bValue);
+
+	void SetControlNameMask(const FName& Name, bool bValue);
+
+	bool GetControlNameMask(const FName& Name) const;
+
+	UE_DEPRECATED(5.5, "Use GetControlNameMask")
+	const TArray<bool>& GetControlsMask() const
+	{
+		return ControlsMask;
+	}
+	
+	UE_DEPRECATED(5.5, "Use GetControlNameMask")
 	const TArray<bool>& GetControlsMask() 
 	{
 		if (ChannelProxy.IsValid() == false)
@@ -329,6 +348,7 @@ public:
 		return ControlsMask;
 	}
 
+	UE_DEPRECATED(5.5, "Use GetControlNameMask")
 	bool GetControlsMask(int32 Index)  
 	{
 		if (ChannelProxy.IsValid() == false)
@@ -342,12 +362,14 @@ public:
 		return false;
 	}
 
+	UE_DEPRECATED(5.5, "Use SetControlNameMask")
 	void SetControlsMask(const TArray<bool>& InMask)
 	{
 		ControlsMask = InMask;
 		ReconstructChannelProxy();
 	}
 
+	UE_DEPRECATED(5.5, "Use SetControlNameMask")
 	void SetControlsMask(int32 Index, bool Val)
 	{
 		if (Index >= 0 && Index < ControlsMask.Num())
@@ -357,11 +379,13 @@ public:
 		ReconstructChannelProxy();
 	}
 
+	UE_DEPRECATED(5.5, "Use FillControlNameMask")
 	void FillControlsMask(bool Val)
 	{
 		ControlsMask.Init(Val, ControlsMask.Num());
 		ReconstructChannelProxy();
 	}
+	
 	/**
 	* This function returns the active category index of the control, based upon what controls are active/masked or not
 	* If itself is masked it returns INDEX_NONE
@@ -514,6 +538,8 @@ public:
 
 protected:
 
+	void ConvertMaskArrayToNameSet();
+
 	//~ UMovieSceneSection interface
 	virtual void Serialize(FArchive& Ar) override;
 	virtual void PostEditImport() override;
diff --git a/Engine/Plugins/Animation/ControlRig/Source/ControlRigEditor/Private/ControlRigSequencerEditorLibrary.cpp b/Engine/Plugins/Animation/ControlRig/Source/ControlRigEditor/Private/ControlRigSequencerEditorLibrary.cpp
--- a/Engine/Plugins/Animation/ControlRig/Source/ControlRigEditor/Private/ControlRigSequencerEditorLibrary.cpp
+++ b/Engine/Plugins/Animation/ControlRig/Source/ControlRigEditor/Private/ControlRigSequencerEditorLibrary.cpp
@@ -3078,21 +3078,7 @@ bool UControlRigSequencerEditorLibrary::GetControlsMask(UMovieSceneSection* InSe
 		FFrame::KismetExecutionMessage(TEXT("Section does not have a control rig"), ELogVerbosity::Error);
 		return false;
 	}
-
-	TArray<FRigControlElement*> Controls;
-	ControlRig->GetControlsInOrder(Controls);
-	int32 Index = 0;
-	for (const FRigControlElement* RigControl : Controls)
-	{
-		if (RigControl->GetFName() == ControlName)
-		{
-			return ParameterSection->GetControlsMask(Index);
-		}
-		++Index;
-	}
-
-	FFrame::KismetExecutionMessage(*FString::Printf(TEXT("Control Name ('%s') not found"), *ControlName.ToString()), ELogVerbosity::Error);
-	return false;
+	return ParameterSection->GetControlNameMask(ControlName);
 }
 
 void UControlRigSequencerEditorLibrary::SetControlsMask(UMovieSceneSection* InSection, const TArray<FName>& ControlNames, bool bVisible)
@@ -3112,18 +3098,11 @@ void UControlRigSequencerEditorLibrary::SetControlsMask(UMovieSceneSection* InSe
 	}
 
 	ParameterSection->Modify();
-
-	TArray<FRigControlElement*> Controls;
-	ControlRig->GetControlsInOrder(Controls);
-	int32 Index = 0;
-	for (const FRigControlElement* RigControl : Controls)
+	for (const FName& ControlName : ControlNames)
 	{
-		if (ControlNames.Contains(RigControl->GetFName()))
-		{
-			ParameterSection->SetControlsMask(Index, bVisible);
-		}
-		++Index;
+		ParameterSection->SetControlNameMask(ControlName, bVisible);
 	}
+	
 }
 
 void UControlRigSequencerEditorLibrary::ShowAllControls(UMovieSceneSection* InSection)
@@ -3136,7 +3115,7 @@ void UControlRigSequencerEditorLibrary::ShowAllControls(UMovieSceneSection* InSe
 	}
 
 	ParameterSection->Modify();
-	ParameterSection->FillControlsMask(true);
+	ParameterSection->FillControlNameMask(true);
 }
 
 void UControlRigSequencerEditorLibrary::HideAllControls(UMovieSceneSection* InSection)
@@ -3149,7 +3128,7 @@ void UControlRigSequencerEditorLibrary::HideAllControls(UMovieSceneSection* InSe
 	}
 
 	ParameterSection->Modify();
-	ParameterSection->FillControlsMask(false);
+	ParameterSection->FillControlNameMask(false);
 }
 
 bool UControlRigSequencerEditorLibrary::IsFKControlRig(UControlRig* InControlRig)
diff --git a/Engine/Plugins/Animation/ControlRig/Source/ControlRigEditor/Private/EditMode/ControlRigEditMode.cpp b/Engine/Plugins/Animation/ControlRig/Source/ControlRigEditor/Private/EditMode/ControlRigEditMode.cpp
--- a/Engine/Plugins/Animation/ControlRig/Source/ControlRigEditor/Private/EditMode/ControlRigEditMode.cpp
+++ b/Engine/Plugins/Animation/ControlRig/Source/ControlRigEditor/Private/EditMode/ControlRigEditMode.cpp
@@ -813,8 +813,7 @@ TSet<FName> FControlRigEditMode::GetActiveControlsFromSequencer(UControlRig* Con
 							UMovieSceneControlRigParameterSection* ActiveSection = Cast<UMovieSceneControlRigParameterSection>(ControlRigParameterTrack->GetSectionToKey(ControlElement->GetFName()));
 							if (ActiveSection)
 							{
-								TArray<bool> Mask = ActiveSection->GetControlsMask();
-								if (Mask[Index])
+								if (ActiveSection->GetControlNameMask(ControlElement->GetFName()))
 								{
 									ActiveControls.Add(ControlElement->GetFName());
 								}
diff --git a/Engine/Plugins/Animation/ControlRig/Source/ControlRigEditor/Private/Sequencer/AnimLayers/AnimLayers.cpp b/Engine/Plugins/Animation/ControlRig/Source/ControlRigEditor/Private/Sequencer/AnimLayers/AnimLayers.cpp
--- a/Engine/Plugins/Animation/ControlRig/Source/ControlRigEditor/Private/Sequencer/AnimLayers/AnimLayers.cpp
+++ b/Engine/Plugins/Animation/ControlRig/Source/ControlRigEditor/Private/Sequencer/AnimLayers/AnimLayers.cpp
@@ -1747,7 +1747,6 @@ bool UAnimLayers::SetPassthroughKey(ISequencer* InSequencer, int32 InIndex)
 									}
 								}
 								InSequencer->NotifyMovieSceneDataChanged(EMovieSceneDataChangeType::TrackValueChanged);
-
 							}
 						}
 						else
@@ -1766,6 +1765,292 @@ bool UAnimLayers::SetPassthroughKey(ISequencer* InSequencer, int32 InIndex)
 	return true;
 }
 
+static void MergeControlRigSections(UMovieSceneControlRigParameterSection* BaseSection, UMovieSceneControlRigParameterSection* Section, const TRange<FFrameNumber>& Range)
+{
+	if (!BaseSection || !Section)
+	{
+		return;
+	}
+	TArrayView<FMovieSceneFloatChannel*> BaseFloatChannels = BaseSection->GetChannelProxy().GetChannels<FMovieSceneFloatChannel>();
+	if (BaseFloatChannels.Num() > 0)
+	{
+		//need to go through each control and merge that
+		TArray<FRigControlElement*> Controls;
+		UControlRig* ControlRig = BaseSection->GetControlRig();
+		if (ControlRig == nullptr)
+		{
+			return;
+		}
+		ControlRig->GetControlsInOrder(Controls);
+		URigHierarchy* Hierarchy = ControlRig->GetHierarchy();
+		if (Hierarchy == nullptr)
+		{
+			return;
+		}
+
+		for (int32 LocalControlIndex = 0; LocalControlIndex < Controls.Num(); ++LocalControlIndex)
+		{
+			FRigControlElement* ControlElement = Controls[LocalControlIndex];
+			check(ControlElement);
+			if (!Hierarchy->IsAnimatable(ControlElement))
+			{
+				continue;
+			}
+
+			if (FChannelMapInfo* pChannelIndex = BaseSection->ControlChannelMap.Find(ControlElement->GetFName()))
+			{
+				const int32 ChannelIndex = pChannelIndex->ChannelIndex;
+				bool bMaskKeyOut = (BaseSection->GetControlNameMask(ControlElement->GetFName()) == false)
+					|| (Section->GetControlNameMask(ControlElement->GetFName()) == false);
+
+				if (bMaskKeyOut)
+				{
+					continue;
+				}
+
+				switch (ControlElement->Settings.ControlType)
+				{
+
+					case ERigControlType::Float:
+					case ERigControlType::ScaleFloat:
+					{
+						int32 StartIndex = ChannelIndex;
+						int32 EndIndex = ChannelIndex;
+						MovieSceneToolHelpers::MergeSections<FMovieSceneFloatChannel>(BaseSection,
+							Section, StartIndex, EndIndex, Range);
+						break;
+					}
+					case ERigControlType::Vector2D:
+					{
+						int32 StartIndex = ChannelIndex;
+						int32 EndIndex = ChannelIndex + 1;
+						MovieSceneToolHelpers::MergeSections<FMovieSceneFloatChannel>(BaseSection,
+							Section, StartIndex, EndIndex, Range);
+						break;
+					}
+					case ERigControlType::Position:
+					case ERigControlType::Scale:
+					case ERigControlType::Rotator:
+					{
+						int32 StartIndex = ChannelIndex;
+						int32 EndIndex = ChannelIndex + 2;
+						MovieSceneToolHelpers::MergeSections<FMovieSceneFloatChannel>(BaseSection,
+							Section, StartIndex, EndIndex, Range);
+						break;
+					}
+
+					case ERigControlType::Transform:
+					case ERigControlType::TransformNoScale:
+					case ERigControlType::EulerTransform:
+					{
+						EMovieSceneTransformChannel BaseChannelMask = BaseSection->GetTransformMask().GetChannels();
+						EMovieSceneTransformChannel ChannelMask = Section->GetTransformMask().GetChannels();
+
+						const bool bDoAllTransform = EnumHasAllFlags(ChannelMask, EMovieSceneTransformChannel::AllTransform);
+						if (bDoAllTransform)
+						{
+							if (ControlElement->Settings.ControlType == ERigControlType::TransformNoScale)
+							{
+								int32 StartIndex = ChannelIndex;
+								int32 EndIndex = ChannelIndex + 5;
+								MovieSceneToolHelpers::MergeSections<FMovieSceneFloatChannel>(BaseSection,
+									Section, StartIndex, EndIndex, Range);
+							}
+							else
+							{
+								int32 StartIndex = ChannelIndex;
+								int32 EndIndex = ChannelIndex + 8;
+								MovieSceneToolHelpers::MergeSections<FMovieSceneFloatChannel>(BaseSection,
+									Section, StartIndex, EndIndex, Range);
+							}
+						}
+						else
+						{
+							if (EnumHasAllFlags(BaseChannelMask, EMovieSceneTransformChannel::TranslationX)
+								 && EnumHasAllFlags(ChannelMask, EMovieSceneTransformChannel::TranslationX))
+							{
+								int32 StartIndex = ChannelIndex;
+								int32 EndIndex = StartIndex;
+								MovieSceneToolHelpers::MergeSections<FMovieSceneFloatChannel>(BaseSection,
+									Section, StartIndex, EndIndex, Range);
+							}
+							if (EnumHasAllFlags(BaseChannelMask, EMovieSceneTransformChannel::TranslationY)
+								&& EnumHasAllFlags(ChannelMask, EMovieSceneTransformChannel::TranslationY))
+							{
+								int32 StartIndex = ChannelIndex + 1;
+								int32 EndIndex = StartIndex;
+								MovieSceneToolHelpers::MergeSections<FMovieSceneFloatChannel>(BaseSection,
+									Section, StartIndex, EndIndex, Range);
+							}
+							if (EnumHasAllFlags(BaseChannelMask, EMovieSceneTransformChannel::TranslationZ)
+								&& EnumHasAllFlags(ChannelMask, EMovieSceneTransformChannel::TranslationZ))
+							{
+								int32 StartIndex = ChannelIndex + 2;
+								int32 EndIndex = StartIndex;
+								MovieSceneToolHelpers::MergeSections<FMovieSceneFloatChannel>(BaseSection,
+									Section, StartIndex, EndIndex, Range);
+							}
+							if (EnumHasAllFlags(BaseChannelMask, EMovieSceneTransformChannel::RotationX)
+								&& EnumHasAllFlags(ChannelMask, EMovieSceneTransformChannel::RotationX))
+							{
+								int32 StartIndex = ChannelIndex + 3;
+								int32 EndIndex = StartIndex;
+								MovieSceneToolHelpers::MergeSections<FMovieSceneFloatChannel>(BaseSection,
+									Section, StartIndex, EndIndex, Range);
+							}
+							if (EnumHasAllFlags(BaseChannelMask, EMovieSceneTransformChannel::RotationY)
+								&& EnumHasAllFlags(ChannelMask, EMovieSceneTransformChannel::RotationY))
+							{
+								int32 StartIndex = ChannelIndex + 4;
+								int32 EndIndex = StartIndex;
+								MovieSceneToolHelpers::MergeSections<FMovieSceneFloatChannel>(BaseSection,
+									Section, StartIndex, EndIndex, Range);
+							}
+							if (EnumHasAllFlags(BaseChannelMask, EMovieSceneTransformChannel::RotationZ)
+								&& EnumHasAllFlags(ChannelMask, EMovieSceneTransformChannel::RotationZ))
+							{
+								int32 StartIndex = ChannelIndex + 5;
+								int32 EndIndex = StartIndex;
+								MovieSceneToolHelpers::MergeSections<FMovieSceneFloatChannel>(BaseSection,
+									Section, StartIndex, EndIndex, Range);
+							}
+							if (ControlElement->Settings.ControlType != ERigControlType::TransformNoScale)
+							{
+								if (EnumHasAllFlags(BaseChannelMask, EMovieSceneTransformChannel::ScaleX)
+									&& EnumHasAllFlags(ChannelMask, EMovieSceneTransformChannel::ScaleX))
+								{
+									int32 StartIndex = ChannelIndex + 6;
+									int32 EndIndex = StartIndex;
+									MovieSceneToolHelpers::MergeSections<FMovieSceneFloatChannel>(BaseSection,
+										Section, StartIndex, EndIndex, Range);
+								}
+								if (EnumHasAllFlags(BaseChannelMask, EMovieSceneTransformChannel::ScaleY)
+									&& EnumHasAllFlags(ChannelMask, EMovieSceneTransformChannel::ScaleY))
+								{
+									int32 StartIndex = ChannelIndex + 7;
+									int32 EndIndex = StartIndex;
+									MovieSceneToolHelpers::MergeSections<FMovieSceneFloatChannel>(BaseSection,
+										Section, StartIndex, EndIndex, Range);
+								}
+								if (EnumHasAllFlags(BaseChannelMask, EMovieSceneTransformChannel::ScaleZ)
+									&& EnumHasAllFlags(ChannelMask, EMovieSceneTransformChannel::ScaleZ))
+								{
+									int32 StartIndex = ChannelIndex + 8;
+									int32 EndIndex = StartIndex;
+									MovieSceneToolHelpers::MergeSections<FMovieSceneFloatChannel>(BaseSection,
+										Section, StartIndex, EndIndex, Range);
+								}
+							}
+						}
+						break;
+					}
+					default:
+						break;
+				}
+			}
+		}
+	}
+}
+
+static void MergeTransformSections(UMovieScene3DTransformSection* BaseSection, UMovieScene3DTransformSection* Section, const TRange<FFrameNumber>& Range)
+{
+	if (!BaseSection || !Section)
+	{
+		return;
+	}
+	TArrayView<FMovieSceneDoubleChannel*> BaseDoubleChannels = BaseSection->GetChannelProxy().GetChannels<FMovieSceneDoubleChannel>();
+	if (BaseDoubleChannels.Num() > 0)
+	{
+		EMovieSceneTransformChannel BaseChannelMask = BaseSection->GetMask().GetChannels();
+		EMovieSceneTransformChannel ChannelMask = Section->GetMask().GetChannels();
+
+		const bool bDoAllTransform = EnumHasAllFlags(ChannelMask, EMovieSceneTransformChannel::AllTransform);
+		if (bDoAllTransform)
+		{
+			int32 StartIndex = 0;
+			int32 EndIndex = BaseDoubleChannels.Num() - 1;
+			MovieSceneToolHelpers::MergeSections<FMovieSceneDoubleChannel>(BaseSection,
+				Section, StartIndex, EndIndex, Range);
+		}
+		else
+		{
+			const int32 ChannelIndex = 0;
+			if (EnumHasAllFlags(BaseChannelMask, EMovieSceneTransformChannel::TranslationX)
+				&& EnumHasAllFlags(ChannelMask, EMovieSceneTransformChannel::TranslationX))
+			{
+				int32 StartIndex = ChannelIndex;
+				int32 EndIndex = StartIndex;
+				MovieSceneToolHelpers::MergeSections<FMovieSceneDoubleChannel>(BaseSection,
+					Section, StartIndex, EndIndex, Range);
+			}
+			if (EnumHasAllFlags(BaseChannelMask, EMovieSceneTransformChannel::TranslationY)
+				&& EnumHasAllFlags(ChannelMask, EMovieSceneTransformChannel::TranslationY))
+			{
+				int32 StartIndex = ChannelIndex + 1;
+				int32 EndIndex = StartIndex;
+				MovieSceneToolHelpers::MergeSections<FMovieSceneDoubleChannel>(BaseSection,
+					Section, StartIndex, EndIndex, Range);
+			}
+			if (EnumHasAllFlags(BaseChannelMask, EMovieSceneTransformChannel::TranslationZ)
+				&& EnumHasAllFlags(ChannelMask, EMovieSceneTransformChannel::TranslationZ))
+			{
+				int32 StartIndex = ChannelIndex + 2;
+				int32 EndIndex = StartIndex;
+				MovieSceneToolHelpers::MergeSections<FMovieSceneDoubleChannel>(BaseSection,
+					Section, StartIndex, EndIndex, Range);
+			}
+			if (EnumHasAllFlags(BaseChannelMask, EMovieSceneTransformChannel::RotationX)
+				&& EnumHasAllFlags(ChannelMask, EMovieSceneTransformChannel::RotationX))
+			{
+				int32 StartIndex = ChannelIndex + 3;
+				int32 EndIndex = StartIndex;
+				MovieSceneToolHelpers::MergeSections<FMovieSceneDoubleChannel>(BaseSection,
+					Section, StartIndex, EndIndex, Range);
+			}
+			if (EnumHasAllFlags(BaseChannelMask, EMovieSceneTransformChannel::RotationY)
+				&& EnumHasAllFlags(ChannelMask, EMovieSceneTransformChannel::RotationY))
+			{
+				int32 StartIndex = ChannelIndex + 4;
+				int32 EndIndex = StartIndex;
+				MovieSceneToolHelpers::MergeSections<FMovieSceneDoubleChannel>(BaseSection,
+					Section, StartIndex, EndIndex, Range);
+			}
+			if (EnumHasAllFlags(BaseChannelMask, EMovieSceneTransformChannel::RotationZ)
+				&& EnumHasAllFlags(ChannelMask, EMovieSceneTransformChannel::RotationZ))
+			{
+				int32 StartIndex = ChannelIndex +  5;
+				int32 EndIndex = StartIndex;
+				MovieSceneToolHelpers::MergeSections<FMovieSceneDoubleChannel>(BaseSection,
+					Section, StartIndex, EndIndex, Range);
+			}
+			if (EnumHasAllFlags(BaseChannelMask, EMovieSceneTransformChannel::ScaleX)
+				&& EnumHasAllFlags(ChannelMask, EMovieSceneTransformChannel::ScaleX))
+			{
+				int32 StartIndex = ChannelIndex + 6;
+				int32 EndIndex = StartIndex;
+				MovieSceneToolHelpers::MergeSections<FMovieSceneFloatChannel>(BaseSection,
+					Section, StartIndex, EndIndex, Range);
+			}
+			if (EnumHasAllFlags(BaseChannelMask, EMovieSceneTransformChannel::ScaleY)
+				&& EnumHasAllFlags(ChannelMask, EMovieSceneTransformChannel::ScaleY))
+			{
+				int32 StartIndex = ChannelIndex + 7;
+				int32 EndIndex = StartIndex;
+				MovieSceneToolHelpers::MergeSections<FMovieSceneDoubleChannel>(BaseSection,
+					Section, StartIndex, EndIndex, Range);
+			}
+			if (EnumHasAllFlags(BaseChannelMask, EMovieSceneTransformChannel::ScaleZ)
+				&& EnumHasAllFlags(ChannelMask, EMovieSceneTransformChannel::ScaleZ))
+			{
+				int32 StartIndex = ChannelIndex + 8;
+				int32 EndIndex = StartIndex;
+				MovieSceneToolHelpers::MergeSections<FMovieSceneFloatChannel>(BaseSection,
+					Section, StartIndex, EndIndex, Range);
+			}
+		}
+	}
+}
+
 bool UAnimLayers::MergeAnimLayers(ISequencer* InSequencer, const TArray<int32>& Indices, const FBakingAnimationKeySettings* InSettings)
 {
 	if (InSequencer == nullptr)
@@ -1823,36 +2108,51 @@ bool UAnimLayers::MergeAnimLayers(ISequencer* InSequencer, const TArray<int32>&
 						{
 							if (SectionItem.Section->IsActive())//active sections merge them
 							{
-								
-								TArrayView<FMovieSceneFloatChannel*> BaseFloatChannels = BaseSectionItem->Section->GetChannelProxy().GetChannels<FMovieSceneFloatChannel>();
-								TArrayView<FMovieSceneDoubleChannel*> BaseDoubleChannels = BaseSectionItem->Section->GetChannelProxy().GetChannels<FMovieSceneDoubleChannel>();
-								if (BaseDoubleChannels.Num() > 0)
+								//if transform or control rig section we need to handle masking
+								if (UMovieSceneControlRigParameterSection* BaseCRSection = Cast< UMovieSceneControlRigParameterSection>(BaseSectionItem->Section.Get()))
 								{
-									int32 StartIndex = 0;
-									int32 EndIndex = BaseDoubleChannels.Num() - 1;
-									MovieSceneToolHelpers::MergeSections<FMovieSceneDoubleChannel>(BaseSectionItem->Section.Get(),
-										SectionItem.Section.Get(), StartIndex, EndIndex, Range);
-
+									UMovieSceneControlRigParameterSection* CRSection = Cast< UMovieSceneControlRigParameterSection>(SectionItem.Section.Get());
+									MergeControlRigSections(BaseCRSection, CRSection, Range);
 								}
-								else if (BaseFloatChannels.Num() > 0)
+								else if (UMovieScene3DTransformSection* BaseTRSection = Cast< UMovieScene3DTransformSection>(BaseSectionItem->Section.Get()))
 								{
-									int32 StartIndex = 0;
-									int32 EndIndex = BaseFloatChannels.Num() - 1;
-									MovieSceneToolHelpers::MergeSections<FMovieSceneFloatChannel>(BaseSectionItem->Section.Get(),
-										SectionItem.Section.Get(), StartIndex, EndIndex, Range);
+									UMovieScene3DTransformSection* TRSection = Cast< UMovieScene3DTransformSection>(SectionItem.Section.Get());
+									MergeTransformSections(BaseTRSection, TRSection, Range);
+								}
+								else
+								{
+									TArrayView<FMovieSceneFloatChannel*> BaseFloatChannels = BaseSectionItem->Section->GetChannelProxy().GetChannels<FMovieSceneFloatChannel>();
+									TArrayView<FMovieSceneDoubleChannel*> BaseDoubleChannels = BaseSectionItem->Section->GetChannelProxy().GetChannels<FMovieSceneDoubleChannel>();
+									if (BaseDoubleChannels.Num() > 0)
+									{
+										int32 StartIndex = 0;
+										int32 EndIndex = BaseDoubleChannels.Num() - 1;
+										MovieSceneToolHelpers::MergeSections<FMovieSceneDoubleChannel>(BaseSectionItem->Section.Get(),
+											SectionItem.Section.Get(), StartIndex, EndIndex, Range);
+
+									}
+									else if (BaseFloatChannels.Num() > 0)
+									{
+										int32 StartIndex = 0;
+										int32 EndIndex = BaseFloatChannels.Num() - 1;
+										MovieSceneToolHelpers::MergeSections<FMovieSceneFloatChannel>(BaseSectionItem->Section.Get(),
+											SectionItem.Section.Get(), StartIndex, EndIndex, Range);
+									}
 								}
 							}
-			
-							if (UMovieSceneControlRigParameterSection* CRSection = Cast<UMovieSceneControlRigParameterSection>(BaseSectionItem->Section))
+							if (BaseLayer != AnimLayers[0]) //if not base layer
 							{
-								if (SectionItem.AnimLayerSet.Names.Num() > 0)
+								if (UMovieSceneControlRigParameterSection* CRSection = Cast<UMovieSceneControlRigParameterSection>(BaseSectionItem->Section))
 								{
-									AddNamesToMask(BaseSectionItem, CRSection, SectionItem.AnimLayerSet);
-									if (SortedIndices[0] != 0) //if not base then make sure mask is set up
+									if (SectionItem.AnimLayerSet.Names.Num() > 0)
 									{
-										TArray<FName> AllControls;
-										BaseSectionItem->AnimLayerSet.Names.GenerateKeyArray(AllControls);
-										UAnimLayers::SetUpControlRigSection(CRSection, AllControls);
+										AddNamesToMask(BaseSectionItem, CRSection, SectionItem.AnimLayerSet);
+										if (SortedIndices[0] != 0) //if not base then make sure mask is set up
+										{
+											TArray<FName> AllControls;
+											BaseSectionItem->AnimLayerSet.Names.GenerateKeyArray(AllControls);
+											UAnimLayers::SetUpControlRigSection(CRSection, AllControls);
+										}
 									}
 								}
 							}
@@ -1893,14 +2193,21 @@ bool UAnimLayers::MergeAnimLayers(ISequencer* InSequencer, const TArray<int32>&
 		}
 	}
 	UAnimLayer* BaseLayer = LayersToMerge[LayersToMerge.Num() - 1];
-	FString Merged(TEXT("Merged"));
-	FString ExistingName = BaseLayer->GetName().ToString();
-	if (ExistingName.Contains(Merged) == false)
+	if (BaseLayer != AnimLayers[0]) //if not base layer
 	{
-		FString NewLayerName = FString::Printf(TEXT("%s_Merged"), *ExistingName);
-		FText LayerText;
-		LayerText = LayerText.FromString(NewLayerName);
-		BaseLayer->SetName(LayerText); //need items/sections to be added so we can change their track row names
+		FString Merged(TEXT("Merged"));
+		FString ExistingName = BaseLayer->GetName().ToString();
+		if (ExistingName.Contains(Merged) == false)
+		{
+			FString NewLayerName = FString::Printf(TEXT("%s_Merged"), *ExistingName);
+			FText LayerText;
+			LayerText = LayerText.FromString(NewLayerName);
+			BaseLayer->SetName(LayerText); //need items/sections to be added so we can change their track row names
+		}
+	}
+	else
+	{
+		SetUpBaseLayerSections(); //if it is the base reset it
 	}
 	InSequencer->NotifyMovieSceneDataChanged(EMovieSceneDataChangeType::MovieSceneStructureItemAdded);
 	return true;
@@ -2282,17 +2589,13 @@ void UAnimLayers::SetUpControlRigSection(UMovieSceneControlRigParameterSection*
 		if (ParameterSection && ControlRig)
 		{
 			ParameterSection->Modify();
-			ParameterSection->FillControlsMask(false);
+			ParameterSection->FillControlNameMask(false);
 
 			TArray<FRigControlElement*> Controls;
 			ControlRig->GetControlsInOrder(Controls);
 			for (const FName& RigName : ControlNames)
 			{
-				FChannelMapInfo* pChannelIndex = ParameterSection->ControlChannelMap.Find(RigName);
-				if (pChannelIndex)
-				{
-					ParameterSection->SetControlsMask(pChannelIndex->MaskIndex, true);
-				}
+				ParameterSection->SetControlNameMask(RigName, true);
 			}
 		}
 	}
diff --git a/Engine/Plugins/Animation/ControlRig/Source/ControlRigEditor/Private/Sequencer/ControlRigParameterTrackEditor.cpp b/Engine/Plugins/Animation/ControlRig/Source/ControlRigEditor/Private/Sequencer/ControlRigParameterTrackEditor.cpp
--- a/Engine/Plugins/Animation/ControlRig/Source/ControlRigEditor/Private/Sequencer/ControlRigParameterTrackEditor.cpp
+++ b/Engine/Plugins/Animation/ControlRig/Source/ControlRigEditor/Private/Sequencer/ControlRigParameterTrackEditor.cpp
@@ -3900,7 +3900,6 @@ void FControlRigParameterTrackEditor::GetControlRigKeys(
 	FGeneratedTrackKeys& OutGeneratedKeys,
 	const bool bInConstraintSpace)
 {
-	const TArray<bool>& ControlsMask = SectionToKey->GetControlsMask();
 	EMovieSceneTransformChannel TransformMask = SectionToKey->GetTransformMask().GetChannels();
 
 	TArray<FRigControlElement*> Controls;
@@ -3929,7 +3928,7 @@ void FControlRigParameterTrackEditor::GetControlRigKeys(
 			int32 ChannelIndex = pChannelIndex->ChannelIndex;
 			const int32 MaskIndex = pChannelIndex->MaskIndex;
 
-			bool bMaskKeyOut = (MaskIndex >= ControlsMask.Num() || ControlsMask[MaskIndex] == false);
+			bool bMaskKeyOut = (SectionToKey->GetControlNameMask(ControlElement->GetFName()) == false);
 			bool bSetKey = ParameterName.IsNone() || (ControlElement->GetFName() == ParameterName && !bMaskKeyOut);
 
 			FRigControlValue ControlValue = InControlRig->GetControlValue(ControlElement, ERigControlValueType::Current);
@@ -4922,20 +4921,16 @@ private:
 			}
 			if (Track && Sequencer)
 			{
-				TArray<bool> Mask;
-				Mask.SetNum(BoneCheckArray.Num());
-				for (const FFKBoneCheckInfo& Info : BoneCheckArray)
-				{
-					Mask[Info.BoneID] = Info.bActive;
-				}
-
 				TArray<UMovieSceneSection*> Sections = Track->GetAllSections();
 				for (UMovieSceneSection* IterSection : Sections)
 				{
 					UMovieSceneControlRigParameterSection* Section = Cast< UMovieSceneControlRigParameterSection>(IterSection);
 					if (Section)
 					{
-						Section->SetControlsMask(Mask);
+						for (const FFKBoneCheckInfo& Info : BoneCheckArray)
+						{
+							Section->SetControlNameMask(Info.BoneName, Info.bActive);
+						}
 					}
 				}
 				Sequencer->NotifyMovieSceneDataChanged(EMovieSceneDataChangeType::MovieSceneStructureItemsChanged);
@@ -5487,79 +5482,7 @@ void FControlRigParameterSection::BuildSectionContextMenu(FMenuBuilder& MenuBuil
 				})
 			);
 	};
-	auto ToggleControls = [this](int32 Index)
-	{
-		return FUIAction(
-			FExecuteAction::CreateLambda([this, Index]
-				{
-					const TSharedPtr<ISequencer> SequencerPtr = WeakSequencer.Pin();
-					if (!SequencerPtr)
-					{
-						return;
-					}
-
-					UMovieSceneControlRigParameterSection* const ParameterSection = CastChecked<UMovieSceneControlRigParameterSection>(WeakSection.Get());
-					if (!IsValid(ParameterSection))
-					{
-						return;
-					}
-
-					FScopedTransaction Transaction(LOCTEXT("ToggleRigControlFiltersTransaction", "Toggle Rig Control Filters"));
-					ParameterSection->Modify();
-					if (Index >= 0)
-					{
-						ParameterSection->SetControlsMask(Index, !ParameterSection->GetControlsMask(Index));
-					}
-					else
-					{
-						ParameterSection->FillControlsMask(!ParameterSection->GetControlsMask(0));
-					}
-					SequencerPtr->NotifyMovieSceneDataChanged(EMovieSceneDataChangeType::MovieSceneStructureItemsChanged);
-				}
-			),
-			FCanExecuteAction(),
-			FGetActionCheckState::CreateLambda([this, Index]
-				{
-					UMovieSceneControlRigParameterSection* const ParameterSection = CastChecked<UMovieSceneControlRigParameterSection>(WeakSection.Get());
-					if (!IsValid(ParameterSection))
-					{
-						return ECheckBoxState::Unchecked;
-					}
-
-					TArray<bool> ControlBool = ParameterSection->GetControlsMask();
-					if (Index >= 0)
-					{
-						if (ControlBool[Index])
-						{
-							return ECheckBoxState::Checked;
-						}
-						else
-						{
-							return ECheckBoxState::Unchecked;
-						}
-					}
-					else
-					{
-						TOptional<bool> FirstVal;
-						for (bool Val : ControlBool)
-						{
-							if (FirstVal.IsSet())
-							{
-								if (Val != FirstVal)
-								{
-									return ECheckBoxState::Undetermined;
-								}
-							}
-							else
-							{
-								FirstVal = Val;
-							}
-						}
-						return (FirstVal.IsSet() && FirstVal.GetValue()) ? ECheckBoxState::Checked : ECheckBoxState::Unchecked;
-					}
-				})
-			);
-	};
+	
 	UMovieSceneControlRigParameterTrack* Track = ParameterSection->GetTypedOuter<UMovieSceneControlRigParameterTrack>();
 	if (Track)
 	{
@@ -5692,7 +5615,7 @@ void FControlRigParameterSection::ShowSelectedControlsChannels()
 	{
 		FScopedTransaction Transaction(LOCTEXT("ShowSelecedControlChannels", "Show Selected Control Channels"));
 		ParameterSection->Modify();
-		ParameterSection->FillControlsMask(false);
+		ParameterSection->FillControlNameMask(false);
 
 		TArray<FRigControlElement*> Controls;
 		ControlRig->GetControlsInOrder(Controls);
@@ -5704,7 +5627,7 @@ void FControlRigParameterSection::ShowSelectedControlsChannels()
 				FChannelMapInfo* pChannelIndex = ParameterSection->ControlChannelMap.Find(RigName);
 				if (pChannelIndex)
 				{
-					ParameterSection->SetControlsMask(pChannelIndex->MaskIndex, true);
+					ParameterSection->SetControlNameMask(RigName, true);
 				}
 			}
 		}
@@ -5720,7 +5643,7 @@ void FControlRigParameterSection::ShowAllControlsChannels()
 	{
 		FScopedTransaction Transaction(LOCTEXT("ShowAllControlChannels", "Show All Control Channels"));
 		ParameterSection->Modify();
-		ParameterSection->FillControlsMask(true);
+		ParameterSection->FillControlNameMask(true);
 		SequencerPtr->NotifyMovieSceneDataChanged(EMovieSceneDataChangeType::MovieSceneStructureItemsChanged);
 	}
 }
