[misc] more FString::Printf format fixes

diff --git a/Engine/Plugins/Animation/ControlRig/Source/ControlRig/Private/Rigs/RigHierarchy.cpp b/Engine/Plugins/Animation/ControlRig/Source/ControlRig/Private/Rigs/RigHierarchy.cpp
--- a/Engine/Plugins/Animation/ControlRig/Source/ControlRig/Private/Rigs/RigHierarchy.cpp
+++ b/Engine/Plugins/Animation/ControlRig/Source/ControlRig/Private/Rigs/RigHierarchy.cpp
@@ -52,7 +52,6 @@ TAutoConsoleVariable<bool> CVarControlRigHierarchyEnableRotationOrder(TEXT("Cont
 TAutoConsoleVariable<bool> CVarControlRigHierarchyEnableModules(TEXT("ControlRig.Hierarchy.Modules"), true, TEXT("enables the modular rigging functionality"));
 TAutoConsoleVariable<bool> CVarControlRigHierarchyEnablePhysics(TEXT("ControlRig.Hierarchy.Physics"), false, TEXT("enables physics support for the rig hierarchy"));
 static int32 sRigHierarchyLastTrace = INDEX_NONE;
-static TCHAR sRigHierarchyTraceFormat[16];
 
 // A console command to trace a single frame / single execution for a control rig anim node / control rig component
 FAutoConsoleCommandWithWorldAndArgs FCmdControlRigHierarchyTraceFrames
@@ -5172,10 +5171,6 @@ void URigHierarchy::CheckTraceFormatIfRequired()
 	if(sRigHierarchyLastTrace != CVarControlRigHierarchyTracePrecision->GetInt())
 	{
 		sRigHierarchyLastTrace = CVarControlRigHierarchyTracePrecision->GetInt();
-		const FString Format = FString::Printf(TEXT("%%.%df"), sRigHierarchyLastTrace);
-		check(Format.Len() < 16);
-		sRigHierarchyTraceFormat[Format.Len()] = '\0';
-		FMemory::Memcpy(sRigHierarchyTraceFormat, *Format, Format.Len() * sizeof(TCHAR));
 	}
 }
 
@@ -5186,7 +5181,7 @@ struct TRigHierarchyJsonPrintPolicy
 	static inline void WriteDouble(  FArchive* Stream, double Value )
 	{
 		URigHierarchy::CheckTraceFormatIfRequired();
-		TJsonPrintPolicy<CharType>::WriteString(Stream, FString::Printf(sRigHierarchyTraceFormat, Value));
+		TJsonPrintPolicy<CharType>::WriteString(Stream, FString::Printf(TEXT("%.*f"), sRigHierarchyLastTrace, Value));
 	}
 };
 
diff --git a/Engine/Plugins/Runtime/RigVM/Source/RigVMDeveloper/Private/RigVMModel/RigVMController.cpp b/Engine/Plugins/Runtime/RigVM/Source/RigVMDeveloper/Private/RigVMModel/RigVMController.cpp
--- a/Engine/Plugins/Runtime/RigVM/Source/RigVMDeveloper/Private/RigVMModel/RigVMController.cpp
+++ b/Engine/Plugins/Runtime/RigVM/Source/RigVMDeveloper/Private/RigVMModel/RigVMController.cpp
@@ -14020,7 +14020,7 @@ bool URigVMController::AddTagToFunctionVariant(const FName& InFunctionName, cons
 		return Tag.Name == InTag.Name; 
 	}))
 	{
-		ReportErrorf(TEXT("Cannot add tag %s to function %s - tag already applied."), *InFunctionName.ToString());
+		ReportErrorf(TEXT("Cannot add tag %s to function %s - tag already applied."), *InTag.Name.ToString(), *InFunctionName.ToString());
 		return false;
 	}
 
diff --git a/Engine/Source/Developer/TraceInsights/Private/Insights/MemoryProfiler/ViewModels/MemoryGraphTrack.cpp b/Engine/Source/Developer/TraceInsights/Private/Insights/MemoryProfiler/ViewModels/MemoryGraphTrack.cpp
--- a/Engine/Source/Developer/TraceInsights/Private/Insights/MemoryProfiler/ViewModels/MemoryGraphTrack.cpp
+++ b/Engine/Source/Developer/TraceInsights/Private/Insights/MemoryProfiler/ViewModels/MemoryGraphTrack.cpp
@@ -1175,11 +1175,7 @@ FString FMemoryGraphTrack::FormatValue(const double InValue, const double InUnit
 		return FText::AsNumber(static_cast<int64>(InValue)).ToString();
 	}
 
-	FString OutText;
-
-	TCHAR FormatString[32];
-	FCString::Snprintf(FormatString, sizeof(FormatString), TEXT("%%.%df"), FMath::Abs(InDecimalDigitCount));
-	OutText = FString::Printf(FormatString, InValue / InUnitValue);
+	FString OutText = FString::Printf(TEXT("%.*f"), FMath::Abs(InDecimalDigitCount), InValue / InUnitValue);
 
 	if (InDecimalDigitCount < 0)
 	{
diff --git a/Engine/Source/Developer/TraceInsights/Private/Insights/TimingProfiler/Widgets/STimersViewTooltip.cpp b/Engine/Source/Developer/TraceInsights/Private/Insights/TimingProfiler/Widgets/STimersViewTooltip.cpp
--- a/Engine/Source/Developer/TraceInsights/Private/Insights/TimingProfiler/Widgets/STimersViewTooltip.cpp
+++ b/Engine/Source/Developer/TraceInsights/Private/Insights/TimingProfiler/Widgets/STimersViewTooltip.cpp
@@ -376,10 +376,9 @@ TSharedPtr<SToolTip> STimersViewTooltip::GetRowTooltip(const TSharedPtr<FTimerNo
 
 	if (true)
 	{
-		TCHAR FormatString[32];
-		FCString::Snprintf(FormatString, sizeof(FormatString), TEXT("%%.%dfs (%%s)"), NumDigits);
-		const FText TotalInclusiveTimeText = FText::FromString(FString::Printf(FormatString, Stats.TotalInclusiveTime, *FormatTimeAuto(Stats.TotalInclusiveTime, 2)));
-		const FText TotalExclusiveTimeText = FText::FromString(FString::Printf(FormatString, Stats.TotalExclusiveTime, *FormatTimeAuto(Stats.TotalExclusiveTime, 2)));
+		static constexpr TCHAR FormatString[] = TEXT("%.*fs (%s)");
+		const FText TotalInclusiveTimeText = FText::FromString(FString::Printf(FormatString, NumDigits, Stats.TotalInclusiveTime, *FormatTimeAuto(Stats.TotalInclusiveTime, 2)));
+		const FText TotalExclusiveTimeText = FText::FromString(FString::Printf(FormatString, NumDigits, Stats.TotalExclusiveTime, *FormatTimeAuto(Stats.TotalExclusiveTime, 2)));
 		AddStatsRow(GridPanel, Row, LOCTEXT("TT_TotalTime", "Total Time:"), TotalInclusiveTimeText, TotalExclusiveTimeText);
 	}
 	if (!std::isnan(Stats.MaxInclusiveTime) || !std::isnan(Stats.MaxExclusiveTime))
