Constraints: handle multiple ControlRigs per bound object

#rnx

diff --git a/Engine/Plugins/Animation/ControlRig/Source/ControlRig/Private/Constraints/ControlRigTransformableHandle.cpp b/Engine/Plugins/Animation/ControlRig/Source/ControlRig/Private/Constraints/ControlRigTransformableHandle.cpp
--- a/Engine/Plugins/Animation/ControlRig/Source/ControlRig/Private/Constraints/ControlRigTransformableHandle.cpp
+++ b/Engine/Plugins/Animation/ControlRig/Source/ControlRig/Private/Constraints/ControlRigTransformableHandle.cpp
@@ -631,19 +631,28 @@ void UTransformableControlHandle::ResolveBoundObjects(FMovieSceneSequenceID Loca
 {
 	if (const UControlRig* InControlRig = Cast<UControlRig>(SubObject))
 	{
-		if (ControlRig != InControlRig)
+		// nothing to do
+		if (ControlRig == InControlRig)
 		{
-			for (const TWeakObjectPtr<> ParentObject : ConstraintBindingID.ResolveBoundObjects(LocalSequenceID, SharedPlaybackState))
+			return;
+		}
+
+		// skip resolving if the rigs don't share the same class type 
+		if (ControlRig && ControlRig->GetClass() != InControlRig->GetClass())
+		{
+			return;
+		}
+		
+		for (const TWeakObjectPtr<> ParentObject : ConstraintBindingID.ResolveBoundObjects(LocalSequenceID, SharedPlaybackState))
+		{
+			const UObject* Bindable = FControlRigObjectBinding::GetBindableObject(ParentObject.Get());
+			if (InControlRig->GetObjectBinding() && InControlRig->GetObjectBinding()->GetBoundObject() == Bindable)
 			{
-				const UObject* Bindable = FControlRigObjectBinding::GetBindableObject(ParentObject.Get());
-				if (InControlRig->GetObjectBinding() && InControlRig->GetObjectBinding()->GetBoundObject() == Bindable)
-				{
-					UnregisterDelegates();
-					ControlRig = const_cast<UControlRig*>(InControlRig);
-					RegisterDelegates();
-				}
-				break; //just do one
+				UnregisterDelegates();
+				ControlRig = const_cast<UControlRig*>(InControlRig);
+				RegisterDelegates();
 			}
+			break; //just do one
 		}
 	}
 }
diff --git a/Engine/Plugins/Animation/ControlRig/Source/ControlRig/Private/Sequencer/MovieSceneControlRigParameterTemplate.cpp b/Engine/Plugins/Animation/ControlRig/Source/ControlRig/Private/Sequencer/MovieSceneControlRigParameterTemplate.cpp
--- a/Engine/Plugins/Animation/ControlRig/Source/ControlRig/Private/Sequencer/MovieSceneControlRigParameterTemplate.cpp
+++ b/Engine/Plugins/Animation/ControlRig/Source/ControlRig/Private/Sequencer/MovieSceneControlRigParameterTemplate.cpp
@@ -25,6 +25,7 @@
 #include "TransformConstraint.h"
 #include "TransformableHandle.h"
 #include "AnimationCoreLibrary.h"
+#include "Constraints/ControlRigTransformableHandle.h"
 
 #include UE_INLINE_GENERATED_CPP_BY_NAME(MovieSceneControlRigParameterTemplate)
 
@@ -1412,38 +1413,28 @@ struct FControlRigParameterExecutionToken : IMovieSceneExecutionToken
 				}
 				if (BoundObject)
 				{
+					UWorld* BoundObjectWorld = BoundObject->GetWorld();
+					
 					TSharedRef<UE::MovieScene::FSharedPlaybackState> SharedPlaybackState = Player.GetSharedPlaybackState();
 					for (FConstraintAndActiveValue& ConstraintValue : ConstraintsValues)
 					{
 						UMovieSceneControlRigParameterSection* NonConstSection = const_cast<UMovieSceneControlRigParameterSection*>(Section);
-						CreateConstraintIfNeeded(BoundObject->GetWorld(), ConstraintValue, NonConstSection);
+						CreateConstraintIfNeeded(BoundObjectWorld, ConstraintValue, NonConstSection);
 
 						if (ConstraintValue.Constraint.IsValid())
 						{
-							//For Control Rig we may need to explicitly set the control rig
 							if (UTickableTransformConstraint* TransformConstraint = Cast<UTickableTransformConstraint>(ConstraintValue.Constraint))
 							{
-								TransformConstraint->InitConstraint(BoundObject->GetWorld());
+								TransformConstraint->InitConstraint(BoundObjectWorld);
 							}
 							ConstraintValue.Constraint->ResolveBoundObjects(Operand.SequenceID, SharedPlaybackState, ControlRig);
 							ConstraintValue.Constraint->SetActive(ConstraintValue.Value);
 						}
 					}
-					//unfortunately for Constraints with ControlRig we need to resolve all Parents also. Don't need to do children since they wil be handled by
-					//the channel resolve above
-					const FConstraintsManagerController& Controller = FConstraintsManagerController::Get(BoundObject->GetWorld());
-					const TArray< TWeakObjectPtr<UTickableConstraint>> Constraints = Controller.GetAllConstraints();
-					for (const TWeakObjectPtr<UTickableConstraint>& TickConstraint : Constraints)
-					{
-						if (UTickableTransformConstraint* TransformConstraint = Cast< UTickableTransformConstraint>(TickConstraint.Get()))
-						{
-							if (TransformConstraint->ParentTRSHandle)
-							{
-								TransformConstraint->ParentTRSHandle->ResolveBoundObjects(Operand.SequenceID, SharedPlaybackState, ControlRig);
-								TransformConstraint->EnsurePrimaryDependency(BoundObject->GetWorld());
-							}
-						}
-					}
+
+					// for Constraints with ControlRig we need to resolve all Parents also
+					// Don't need to do children since they wil be handled by the channel resolve above
+					ResolveParentHandles(BoundObject, ControlRig, Operand, SharedPlaybackState);
 				}
 				else  //no bound object so turn off constraint
 				{
@@ -1461,6 +1452,81 @@ struct FControlRigParameterExecutionToken : IMovieSceneExecutionToken
 
 	}
 
+	void ResolveParentHandles(
+		const UObject* InBoundObject, UControlRig* InControlRigInstance,
+		const FMovieSceneEvaluationOperand& InOperand,
+		const TSharedRef<UE::MovieScene::FSharedPlaybackState>& InSharedPlaybackState) const
+	{
+		if (!InBoundObject)
+		{
+			return;
+		}
+
+		UWorld* BoundObjectWorld = InBoundObject->GetWorld();
+		const bool bIsGameWorld = InBoundObject->GetWorld() ? InBoundObject->GetWorld()->IsGameWorld() : false;
+		
+		UMovieSceneControlRigParameterTrack* ControlRigTrack = Section->GetTypedOuter<UMovieSceneControlRigParameterTrack>();
+
+		// is this control rig a game world instance of this section's rig?
+		auto WasAGameInstance = [ControlRigTrack](const UControlRig* InRigToTest)
+		{
+			return ControlRigTrack ? ControlRigTrack->IsAGameInstance(InRigToTest) : false;
+		};
+
+		// is the parent handle of this constraint related to this section?
+		// this return true if the handle's control rig has been spawned by the ControlRigTrack (whether in Editor or Game)
+		// if false, it means that the handle represents another control on another control rig so we don't need to resolve it here
+		// note that it returns true if ControlRigTrack is null (is this possible?!) or if the ControlRig is null (we can't infer anything from this)
+		auto ShouldResolveParent = [ControlRigTrack](const UTransformableControlHandle* ParentControlHandle)
+		{
+			if (!ParentControlHandle)
+			{
+				return false;
+			}
+
+			if (!ControlRigTrack)
+			{
+				// cf. UObjectBaseUtility::IsInOuter
+				return true;	
+			}
+			
+			return ParentControlHandle->ControlRig ? ParentControlHandle->ControlRig->IsInOuter(ControlRigTrack) : true;
+		};
+
+		// this is the default's section rig. when bIsGameWorld is false, InControlRigInstance should be equal to SectionRig
+		const UControlRig* SectionRig = Section->GetControlRig();
+
+		const FConstraintsManagerController& Controller = FConstraintsManagerController::Get(BoundObjectWorld);
+		const TArray< TWeakObjectPtr<UTickableConstraint>> Constraints = Controller.GetAllConstraints();
+
+		for (const TWeakObjectPtr<UTickableConstraint>& TickConstraint : Constraints)
+		{
+			UTickableTransformConstraint* TransformConstraint = Cast< UTickableTransformConstraint>(TickConstraint.Get());
+			UTransformableControlHandle* ParentControlHandle = TransformConstraint ? Cast<UTransformableControlHandle>(TransformConstraint->ParentTRSHandle) : nullptr;
+			if (ParentControlHandle && ShouldResolveParent(ParentControlHandle))
+			{
+				if (bIsGameWorld)
+				{
+					// switch from section's rig to the game instance
+					if (ParentControlHandle->ControlRig == SectionRig)
+					{
+						ParentControlHandle->ResolveBoundObjects(InOperand.SequenceID, InSharedPlaybackState, InControlRigInstance);
+						TransformConstraint->EnsurePrimaryDependency(BoundObjectWorld);
+					}
+				}
+				else
+				{
+					// switch from the game instance to the section's rig
+					if (WasAGameInstance(ParentControlHandle->ControlRig.Get()))
+					{
+						ParentControlHandle->ResolveBoundObjects(InOperand.SequenceID, InSharedPlaybackState, InControlRigInstance);
+						TransformConstraint->EnsurePrimaryDependency(BoundObjectWorld);
+					}
+				}
+			}
+		}
+	}
+	
 	const UMovieSceneControlRigParameterSection* Section;
 	/** Array of evaluated bool values */
 	TArray<FBoolParameterStringAndValue, TInlineAllocator<2>> BoolValues;
diff --git a/Engine/Plugins/Animation/ControlRig/Source/ControlRig/Private/Sequencer/MovieSceneControlRigParameterTrack.cpp b/Engine/Plugins/Animation/ControlRig/Source/ControlRig/Private/Sequencer/MovieSceneControlRigParameterTrack.cpp
--- a/Engine/Plugins/Animation/ControlRig/Source/ControlRig/Private/Sequencer/MovieSceneControlRigParameterTrack.cpp
+++ b/Engine/Plugins/Animation/ControlRig/Source/ControlRig/Private/Sequencer/MovieSceneControlRigParameterTrack.cpp
@@ -976,6 +976,24 @@ UControlRig* UMovieSceneControlRigParameterTrack::GetGameWorldControlRig(UWorld*
 	return nullptr;
 }
 
+bool UMovieSceneControlRigParameterTrack::IsAGameInstance(const UControlRig* InControlRig, const bool bCheckValidWorld) const
+{
+	if (!InControlRig || GameWorldControlRigs.IsEmpty())
+	{
+		return false;
+	}
+
+	for (const TPair<TWeakObjectPtr<UWorld>, TObjectPtr<UControlRig>>& WorldAndControlRig: GameWorldControlRigs)
+	{
+		if (WorldAndControlRig.Value == InControlRig)
+		{
+			return bCheckValidWorld ? WorldAndControlRig.Key.IsValid() : true;
+		}
+	}
+	
+	return false;
+}
+
 TArray<FRigControlFBXNodeAndChannels>* UMovieSceneControlRigParameterTrack::GetNodeAndChannelMappings(UMovieSceneSection* InSection )
 {
 #if WITH_EDITOR
diff --git a/Engine/Plugins/Animation/ControlRig/Source/ControlRig/Public/Sequencer/MovieSceneControlRigParameterTrack.h b/Engine/Plugins/Animation/ControlRig/Source/ControlRig/Public/Sequencer/MovieSceneControlRigParameterTrack.h
--- a/Engine/Plugins/Animation/ControlRig/Source/ControlRig/Public/Sequencer/MovieSceneControlRigParameterTrack.h
+++ b/Engine/Plugins/Animation/ControlRig/Source/ControlRig/Public/Sequencer/MovieSceneControlRigParameterTrack.h
@@ -230,6 +230,8 @@ public:
 
 public:
 	UControlRig* GetGameWorldControlRig(UWorld* InWorld);
+	bool IsAGameInstance(const UControlRig* InControlRig, const bool bCheckValidWorld = false) const;
+    
 private:
 	/** copy of the controlled control rig that we use in the game world so editor control rig doesn't conflict*/
 	UPROPERTY(transient)
