Sequencer: Anim Layers: Add missing 'static' identifier so python is usable

diff --git a/Engine/Plugins/Animation/ControlRig/Source/ControlRigEditor/Public/ControlRigSequencerEditorLibrary.h b/Engine/Plugins/Animation/ControlRig/Source/ControlRigEditor/Public/ControlRigSequencerEditorLibrary.h
--- a/Engine/Plugins/Animation/ControlRig/Source/ControlRigEditor/Public/ControlRigSequencerEditorLibrary.h
+++ b/Engine/Plugins/Animation/ControlRig/Source/ControlRigEditor/Public/ControlRigSequencerEditorLibrary.h
@@ -1132,7 +1132,7 @@ public:
 	* @return Returns true if successful, false otherwise
 	*/
 	UFUNCTION(BlueprintCallable, Category = "Editor Scripting | Sequencer Tools | Control Rig | Animation Layers")
-	bool DeleteAnimLayer(int32 Index);
+	static bool DeleteAnimLayer(int32 Index);
 
 	/**
 	* Duplicate anim layer at specified index
@@ -1140,14 +1140,14 @@ public:
 	* @return Returns index of new layer, -1 if none created
 	*/
 	UFUNCTION(BlueprintCallable, Category = "Editor Scripting | Sequencer Tools | Control Rig | Animation Layers")
-	int32 DuplicateAnimLayer(int32 Index);
+	static int32 DuplicateAnimLayer(int32 Index);
 
 	/**
 	* Add anim layer from objects selected in Sequencer
 	* @return Returns Index of created anim layer
 	*/
 	UFUNCTION(BlueprintCallable, Category = "Editor Scripting | Sequencer Tools | Control Rig | Animation Layers")
-	int32 AddAnimLayerFromSelection();
+	static int32 AddAnimLayerFromSelection();
 
 	/**
 	* Merge specified anim layers into one layer. Will merge onto the anim layer with the lowest index
@@ -1155,14 +1155,14 @@ public:
 	* @return Returns true if successful, false otherwise
 	*/
 	UFUNCTION(BlueprintCallable, Category = "Editor Scripting | Sequencer Tools | Control Rig | Animation Layers")
-	bool MergeAnimLayers(const TArray<int32>& Indices);
+	static bool MergeAnimLayers(const TArray<int32>& Indices);
 
 	/**
 	* Get the animation layer objects
 	* @return Returns array of anim layer objects if they exist on active Sequencer
 	*/
 	UFUNCTION(BlueprintCallable, Category = "Editor Scripting | Sequencer Tools | Control Rig | Animation Layers")
-	TArray<UAnimLayer*> GetAnimLayers();
+	static TArray<UAnimLayer*> GetAnimLayers();
 	
 	/**
 	* Helper function to get the index in the anim layer array from the anim layer
@@ -1170,6 +1170,6 @@ public:
 	* @return Returns index for the anim layer or INDEX_NONE(-1) if it doesn't exist
 	*/
 	UFUNCTION(BlueprintCallable, Category = "Editor Scripting | Sequencer Tools | Control Rig | Animation Layers")
-	int32 GetAnimLayerIndex(UAnimLayer* AnimLayer);
+	static int32 GetAnimLayerIndex(UAnimLayer* AnimLayer);
 
 };
