Move Animation Mode extensions of the viewport toolbar "Select" menu to instead extend the "Transform" menu

I'm doing this because we are removing the "Select" menu from the new Level Editor viewport toolbar.

diff --git a/Engine/Plugins/Animation/ControlRig/Source/ControlRigEditor/Private/EditMode/ControlRigEditMode.cpp b/Engine/Plugins/Animation/ControlRig/Source/ControlRigEditor/Private/EditMode/ControlRigEditMode.cpp
--- a/Engine/Plugins/Animation/ControlRig/Source/ControlRigEditor/Private/EditMode/ControlRigEditMode.cpp
+++ b/Engine/Plugins/Animation/ControlRig/Source/ControlRigEditor/Private/EditMode/ControlRigEditMode.cpp
@@ -483,8 +483,7 @@ void FControlRigEditMode::Enter()
 		}
 	}
 
-	UE::ControlRig::PopulateControlRigViewportToolbarTransformsSubmenu("LevelEditor.ViewportToolbar.Transform");
-	UE::ControlRig::PopulateControlRigViewportToolbarSelectionSubmenu("LevelEditor.ViewportToolbar.Select");
+	UE::ControlRig::PopulateControlRigViewportToolbarTransformSubmenu("LevelEditor.ViewportToolbar.Transform");
 	UE::ControlRig::PopulateControlRigViewportToolbarShowSubmenu(
 		"LevelEditor.ViewportToolbar.Show", GetToolkit()->GetToolkitCommands()
 	);
diff --git a/Engine/Plugins/Animation/ControlRig/Source/ControlRigEditor/Private/Editor/ControlRigViewportToolbarExtensions.cpp b/Engine/Plugins/Animation/ControlRig/Source/ControlRigEditor/Private/Editor/ControlRigViewportToolbarExtensions.cpp
--- a/Engine/Plugins/Animation/ControlRig/Source/ControlRigEditor/Private/Editor/ControlRigViewportToolbarExtensions.cpp
+++ b/Engine/Plugins/Animation/ControlRig/Source/ControlRigEditor/Private/Editor/ControlRigViewportToolbarExtensions.cpp
@@ -20,7 +20,7 @@ const FName ControlRigOwnerName = "ControlRigViewportToolbar";
 
 } // namespace UE::ControlRig::Private
 
-void UE::ControlRig::PopulateControlRigViewportToolbarTransformsSubmenu(const FName InMenuName)
+void UE::ControlRig::PopulateControlRigViewportToolbarTransformSubmenu(const FName InMenuName)
 {
 	FToolMenuOwnerScoped ScopeOwner(UE::ControlRig::Private::ControlRigOwnerName);
 
@@ -103,46 +103,43 @@ void UE::ControlRig::PopulateControlRigViewportToolbarTransformsSubmenu(const FN
 			GizmoSection.AddEntry(Entry);
 		}
 	}
-}
 
-void UE::ControlRig::PopulateControlRigViewportToolbarSelectionSubmenu(const FName InMenuName)
-{
-	FToolMenuOwnerScoped ScopeOwner(UE::ControlRig::Private::ControlRigOwnerName);
+	{
+		FToolMenuSection& SelectionSection = Menu->FindOrAddSection("Selection");
 
-	UToolMenu* const Menu = UToolMenus::Get()->ExtendMenu(InMenuName);
-	FToolMenuSection& OptionsSection = Menu->FindOrAddSection("Options");
-	UControlRigEditModeSettings* const Settings = GetMutableDefault<UControlRigEditModeSettings>();
+		UControlRigEditModeSettings* const Settings = GetMutableDefault<UControlRigEditModeSettings>();
 
-	// Add "Restore Coordinate Space on Switch" checkbox.
-	{
-		FUIAction Action;
+		// Add "Restore Coordinate Space on Switch" checkbox.
+		{
+			FUIAction Action;
 
-		Action.ExecuteAction = FExecuteAction::CreateLambda(
-			[Settings]() -> void
-			{
-				Settings->bOnlySelectRigControls = !Settings->bOnlySelectRigControls;
-				Settings->PostEditChange();
-			}
-		);
-		Action.GetActionCheckState = FGetActionCheckState::CreateLambda(
-			[Settings]() -> ECheckBoxState
-			{
-				return Settings->bOnlySelectRigControls ? ECheckBoxState::Checked : ECheckBoxState::Unchecked;
-			}
-		);
-
-		FToolMenuEntry Entry = FToolMenuEntry::InitMenuEntry(
-			"OnlySelectRigControls",
-			LOCTEXT("OnlySelectRigControlsLabel", "Select Only Control Rig Controls"),
-			LOCTEXT("OnlySelectRigControlsTooltip", "Whether or not only Rig Controls can be selected."),
-			FSlateIcon(),
-			Action,
-			EUserInterfaceActionType::ToggleButton
-		);
-		// We want to appear late in the section.
-		Entry.InsertPosition.Position = EToolMenuInsertType::Last;
-		Entry.SetShowInToolbarTopLevel(true);
-		OptionsSection.AddEntry(Entry);
+			Action.ExecuteAction = FExecuteAction::CreateLambda(
+				[Settings]() -> void
+				{
+					Settings->bOnlySelectRigControls = !Settings->bOnlySelectRigControls;
+					Settings->PostEditChange();
+				}
+			);
+			Action.GetActionCheckState = FGetActionCheckState::CreateLambda(
+				[Settings]() -> ECheckBoxState
+				{
+					return Settings->bOnlySelectRigControls ? ECheckBoxState::Checked : ECheckBoxState::Unchecked;
+				}
+			);
+
+			FToolMenuEntry Entry = FToolMenuEntry::InitMenuEntry(
+				"OnlySelectRigControls",
+				LOCTEXT("OnlySelectRigControlsLabel", "Select Only Control Rig Controls"),
+				LOCTEXT("OnlySelectRigControlsTooltip", "Whether or not only Rig Controls can be selected."),
+				FSlateIcon(),
+				Action,
+				EUserInterfaceActionType::ToggleButton
+			);
+			// We want to appear late in the section.
+			Entry.InsertPosition.Position = EToolMenuInsertType::Last;
+			Entry.SetShowInToolbarTopLevel(true);
+			SelectionSection.AddEntry(Entry);
+		}
 	}
 }
 
diff --git a/Engine/Plugins/Animation/ControlRig/Source/ControlRigEditor/Private/Editor/ControlRigViewportToolbarExtensions.h b/Engine/Plugins/Animation/ControlRig/Source/ControlRigEditor/Private/Editor/ControlRigViewportToolbarExtensions.h
--- a/Engine/Plugins/Animation/ControlRig/Source/ControlRigEditor/Private/Editor/ControlRigViewportToolbarExtensions.h
+++ b/Engine/Plugins/Animation/ControlRig/Source/ControlRigEditor/Private/Editor/ControlRigViewportToolbarExtensions.h
@@ -10,8 +10,7 @@ class FUICommandList;
 namespace UE::ControlRig
 {
 
-void PopulateControlRigViewportToolbarTransformsSubmenu(const FName InMenuName);
-void PopulateControlRigViewportToolbarSelectionSubmenu(const FName InMenuName);
+void PopulateControlRigViewportToolbarTransformSubmenu(const FName InMenuName);
 void PopulateControlRigViewportToolbarShowSubmenu(const FName InMenuName, const TSharedPtr<const FUICommandList>& InCommandList);
 
 void RemoveControlRigViewportToolbarExtensions();
