Sequencer: Control Rig: No longer remove parents if no space keys, procedural rigs also add them via AddParent so can't remove them.

diff --git a/Engine/Plugins/Animation/ControlRig/Source/ControlRig/Private/Sequencer/MovieSceneControlRigParameterTrack.cpp b/Engine/Plugins/Animation/ControlRig/Source/ControlRig/Private/Sequencer/MovieSceneControlRigParameterTrack.cpp
--- a/Engine/Plugins/Animation/ControlRig/Source/ControlRig/Private/Sequencer/MovieSceneControlRigParameterTrack.cpp
+++ b/Engine/Plugins/Animation/ControlRig/Source/ControlRig/Private/Sequencer/MovieSceneControlRigParameterTrack.cpp
@@ -97,34 +97,6 @@ UMovieSceneSection* UMovieSceneControlRigParameterTrack::CreateNewSection()
 void UMovieSceneControlRigParameterTrack::HandleOnSpaceAdded(UMovieSceneControlRigParameterSection* Section, const FName& InControlName, FMovieSceneControlRigSpaceChannel* Channel)
 {
 	OnSpaceChannelAdded.Broadcast(Section, InControlName, Channel);
-
-	Channel->OnSpaceNoLongerUsed().RemoveAll(this);
-	Channel->OnSpaceNoLongerUsed().AddUObject(this, &UMovieSceneControlRigParameterTrack::HandleOnSpaceNoLongerUsed, InControlName);
-}
-
-void UMovieSceneControlRigParameterTrack::HandleOnSpaceNoLongerUsed(FMovieSceneControlRigSpaceChannel* InChannel, const TArray<FRigElementKey>& InSpaces, FName InControlName)
-{
-	if (InChannel && GetControlRig())
-	{
-		if(URigHierarchy* Hierarchy = GetControlRig()->GetHierarchy())
-		{
-			if(URigHierarchyController* Controller = Hierarchy->GetController())
-			{
-				const FRigElementKey ControlKey(InControlName, ERigElementType::Control);
-				if(Hierarchy->Find<FRigControlElement>(ControlKey))
-				{
-					FRigElementKey DefaultParent = Hierarchy->GetFirstParent(ControlKey);
-					for(const FRigElementKey& ParentToRemove : InSpaces)
-					{
-						if(DefaultParent != ParentToRemove)
-						{
-							Controller->RemoveParent(ControlKey, ParentToRemove, false, false, false);
-						}
-					}
-				}
-			}
-		}
-	}
 }
 
 void UMovieSceneControlRigParameterTrack::HandleOnConstraintAdded(
diff --git a/Engine/Plugins/Animation/ControlRig/Source/ControlRig/Public/Sequencer/MovieSceneControlRigParameterTrack.h b/Engine/Plugins/Animation/ControlRig/Source/ControlRig/Public/Sequencer/MovieSceneControlRigParameterTrack.h
--- a/Engine/Plugins/Animation/ControlRig/Source/ControlRig/Public/Sequencer/MovieSceneControlRigParameterTrack.h
+++ b/Engine/Plugins/Animation/ControlRig/Source/ControlRig/Public/Sequencer/MovieSceneControlRigParameterTrack.h
@@ -187,8 +187,7 @@ private:
 
 	//we register this with sections.
 	void HandleOnSpaceAdded(UMovieSceneControlRigParameterSection* Section, const FName& InControlName, FMovieSceneControlRigSpaceChannel* Channel);
-	//we'll register this against all space channels
-	void HandleOnSpaceNoLongerUsed(FMovieSceneControlRigSpaceChannel* InChannel, const TArray<FRigElementKey>& InSpaces, FName InControlName);
+	
 	//then send this event out to the track editor.
 	UMovieSceneControlRigParameterSection::FSpaceChannelAddedEvent OnSpaceChannelAdded;
 
