Setting EulerAngle in SetLocalControlRigRotator(s) functions

#rnx

diff --git a/Engine/Plugins/Animation/ControlRig/Source/ControlRigEditor/Private/ControlRigSequencerEditorLibrary.cpp b/Engine/Plugins/Animation/ControlRig/Source/ControlRigEditor/Private/ControlRigSequencerEditorLibrary.cpp
--- a/Engine/Plugins/Animation/ControlRig/Source/ControlRigEditor/Private/ControlRigSequencerEditorLibrary.cpp
+++ b/Engine/Plugins/Animation/ControlRig/Source/ControlRigEditor/Private/ControlRigSequencerEditorLibrary.cpp
@@ -2100,13 +2100,16 @@ void UControlRigSequencerEditorLibrary::SetLocalControlRigRotator(ULevelSequence
 	{
 		return;
 	}
-	if (FRigControlElement* Element = ControlRig->FindControl(ControlName))
+	FRigControlElement* Element = ControlRig->FindControl(ControlName);
+	if (!Element)
 	{
-		if (Element->Settings.ControlType != ERigControlType::Rotator)
-		{
-			UE_LOG(LogControlRig, Error, TEXT("Control Rig Wrong Type"));
-			return;
-		}
+		UE_LOG(LogControlRig, Error, TEXT("Control name not found."));
+		return;
+	}
+	if (Element->Settings.ControlType != ERigControlType::Rotator)
+	{
+		UE_LOG(LogControlRig, Error, TEXT("Control Rig Wrong Type"));
+		return;
 	}
 	if (UMovieScene* MovieScene = LevelSequence->GetMovieScene())
 	{
@@ -2118,6 +2121,9 @@ void UControlRigSequencerEditorLibrary::SetLocalControlRigRotator(ULevelSequence
 		FRigControlModifiedContext Context;
 		Context.SetKey = bSetKey ? EControlRigSetKey::Always : EControlRigSetKey::DoNotCare;
 		Context.LocalTime = TickResolution.AsSeconds(FFrameTime(Frame));
+
+		const FVector EulerAngle(Value.Roll, Value.Pitch, Value.Yaw);
+		ControlRig->GetHierarchy()->SetControlSpecifiedEulerAngle(Element, EulerAngle);
 		ControlRig->SetControlValue<FVector3f>(ControlName, (FVector3f)Value.Euler(), true, Context);
 	}
 }
@@ -2129,13 +2135,16 @@ void UControlRigSequencerEditorLibrary::SetLocalControlRigRotators(ULevelSequenc
 	{
 		return;
 	}
-	if (FRigControlElement* Element = ControlRig->FindControl(ControlName))
+	FRigControlElement* Element = ControlRig->FindControl(ControlName);
+	if (!Element)
 	{
-		if (Element->Settings.ControlType != ERigControlType::Rotator)
-		{
-			UE_LOG(LogControlRig, Error, TEXT("Control Rig Wrong Type"));
-			return;
-		}
+		UE_LOG(LogControlRig, Error, TEXT("Control name not found."));
+		return;
+	}
+	if (Element->Settings.ControlType != ERigControlType::Rotator)
+	{
+		UE_LOG(LogControlRig, Error, TEXT("Control Rig Wrong Type"));
+		return;
 	}
 	if (UMovieScene* MovieScene = LevelSequence->GetMovieScene())
 	{
@@ -2149,8 +2158,10 @@ void UControlRigSequencerEditorLibrary::SetLocalControlRigRotators(ULevelSequenc
 			{
 				Frame = FFrameRate::TransformTime(FFrameTime(Frame, 0), MovieScene->GetDisplayRate(), MovieScene->GetTickResolution()).RoundToFrame();
 			}
-			FRotator Value = Values[Index];
+			const FRotator Value = Values[Index];
 			Context.LocalTime = TickResolution.AsSeconds(FFrameTime(Frame));
+			const FVector EulerAngle(Value.Roll, Value.Pitch, Value.Yaw);
+			ControlRig->GetHierarchy()->SetControlSpecifiedEulerAngle(Element, EulerAngle);
 			ControlRig->SetControlValue<FVector3f>(ControlName, (FVector3f)Value.Euler(), true, Context);
 		}
 	}
