move to using trasctionally safe scope locks and critical sections

diff --git a/Engine/Plugins/Animation/ControlRig/Source/ControlRig/Private/ControlRigComponent.cpp b/Engine/Plugins/Animation/ControlRig/Source/ControlRig/Private/ControlRigComponent.cpp
--- a/Engine/Plugins/Animation/ControlRig/Source/ControlRig/Private/ControlRigComponent.cpp
+++ b/Engine/Plugins/Animation/ControlRig/Source/ControlRig/Private/ControlRigComponent.cpp
@@ -5,6 +5,7 @@
 #include "Units/Execution/RigUnit_BeginExecution.h"
 #include "Units/Execution/RigUnit_InverseExecution.h"
 #include "Units/Execution/RigUnit_Hierarchy.h"
+#include "Misc/TransactionallySafeScopeLock.h"
 
 #include "SkeletalDebugRendering.h"
 #include "Components/InstancedStaticMeshComponent.h"
@@ -45,7 +46,7 @@ struct FSkeletalMeshToMap
 	TArray<FControlRigComponentMappedCurve> Curves;
 };
 
-FCriticalSection gPendingSkeletalMeshesLock;
+FTransactionallySafeCriticalSection gPendingSkeletalMeshesLock;
 TMap<TWeakObjectPtr<UControlRigComponent>, TArray<FSkeletalMeshToMap> > gPendingSkeletalMeshes;
 
 UControlRigComponent::UControlRigComponent(const FObjectInitializer& ObjectInitializer)
@@ -103,7 +104,7 @@ void UControlRigComponent::BeginDestroy()
 {
 	Super::BeginDestroy();
 
-	FScopeLock Lock(&gPendingSkeletalMeshesLock);
+	FTransactionallySafeScopeLock Lock(&gPendingSkeletalMeshesLock);
 	gPendingSkeletalMeshes.Remove(this);
 }
 
@@ -130,7 +131,7 @@ void UControlRigComponent::OnRegister()
 	ControlRig = nullptr;
 
 	{
-		FScopeLock Lock(&gPendingSkeletalMeshesLock);
+		FTransactionallySafeScopeLock Lock(&gPendingSkeletalMeshesLock);
 		gPendingSkeletalMeshes.FindOrAdd(this);
 	}
 
@@ -187,7 +188,7 @@ void UControlRigComponent::OnUnregister()
 	}
 	else
 	{
-		FScopeLock Lock(&gPendingSkeletalMeshesLock);
+		FTransactionallySafeScopeLock Lock(&gPendingSkeletalMeshesLock);
 		gPendingSkeletalMeshes.Remove(this);
 	}
 
@@ -209,7 +210,7 @@ void UControlRigComponent::TickComponent(float DeltaTime, enum ELevelTick TickTy
 
 	if (UControlRig* CR = SetupControlRigIfRequired())
 	{
-		FScopeLock Lock(&gPendingSkeletalMeshesLock);
+		FTransactionallySafeScopeLock Lock(&gPendingSkeletalMeshesLock);
 		TArray<FSkeletalMeshToMap>* PendingSkeletalMeshes = gPendingSkeletalMeshes.Find(this);
 
 		if (PendingSkeletalMeshes != nullptr && PendingSkeletalMeshes->Num() > 0)
@@ -656,7 +657,7 @@ void UControlRigComponent::AddMappedSkeletalMesh(USkeletalMeshComponent* Skeleta
 		PendingMesh.Bones = Bones;
 		PendingMesh.Curves = Curves;
 
-		FScopeLock Lock(&gPendingSkeletalMeshesLock);
+		FTransactionallySafeScopeLock Lock(&gPendingSkeletalMeshesLock);
 		TArray<FSkeletalMeshToMap>& PendingSkeletalMeshes = gPendingSkeletalMeshes.FindOrAdd(this);
 		PendingSkeletalMeshes.Add(PendingMesh);
 		return;
