Bones rendering: fixed FAnimationViewportClient::UpdateBonesToDraw using an invalid ref skeleton
    - NOTE: UpdateBonesToDraw should probably be called lazily when ticking to give other notification systems some time to update things before caching

#rnx

diff --git a/Engine/Plugins/Animation/ControlRig/Source/ControlRigEditor/Private/Editor/ControlRigSkeletalMeshComponent.cpp b/Engine/Plugins/Animation/ControlRig/Source/ControlRigEditor/Private/Editor/ControlRigSkeletalMeshComponent.cpp
--- a/Engine/Plugins/Animation/ControlRig/Source/ControlRigEditor/Private/Editor/ControlRigSkeletalMeshComponent.cpp
+++ b/Engine/Plugins/Animation/ControlRig/Source/ControlRigEditor/Private/Editor/ControlRigSkeletalMeshComponent.cpp
@@ -72,11 +72,15 @@ void UControlRigSkeletalMeshComponent::RebuildDebugDrawSkeleton()
 	
 	UControlRigLayerInstance* ControlRigInstance = Cast<UControlRigLayerInstance>(GetAnimInstance());
 
+	bool bHasValidControlRig = false;
+	
 	if (ControlRigInstance)
 	{
 		UControlRig* ControlRig = ControlRigInstance->GetFirstAvailableControlRig();
 		if (ControlRig)
 		{
+			bHasValidControlRig = true;
+			
 			// we are trying to poke into running instances of Control Rigs
 			// on the anim thread and query data, using a lock here to make sure
 			// we don't get an inconsistent view of the rig at some
@@ -145,6 +149,14 @@ void UControlRigSkeletalMeshComponent::RebuildDebugDrawSkeleton()
 		}
 	}
 	
+	if (!bHasValidControlRig)
+	{
+		DebugDrawSkeleton.Empty();
+		DebugDrawBones.Reset();
+		DebugDrawBoneIndexInHierarchy.Reset();
+		return;
+	}
+	
 	bRebuildDebugDrawSkeletonRequired = false;
 }
 
