Control Rig: Use only one method for schematic visibility

#rnx

diff --git a/Engine/Plugins/Animation/ControlRig/Source/ControlRigEditor/Private/Editor/ControlRigEditor.cpp b/Engine/Plugins/Animation/ControlRig/Source/ControlRigEditor/Private/Editor/ControlRigEditor.cpp
--- a/Engine/Plugins/Animation/ControlRig/Source/ControlRigEditor/Private/Editor/ControlRigEditor.cpp
+++ b/Engine/Plugins/Animation/ControlRig/Source/ControlRigEditor/Private/Editor/ControlRigEditor.cpp
@@ -136,6 +136,7 @@ FControlRigEditor::FControlRigEditor()
 	, bIsConstructionEventRunning(false)
 	, LastHierarchyHash(INDEX_NONE)
 	, bRefreshDirectionManipulationTargetsRequired(false)
+	, bSchematicViewPortIsHidden(false)
 {
 	LastEventQueue = ConstructionEventQueue;
 }
@@ -2536,7 +2537,7 @@ void FControlRigEditor::HandleToggleSchematicViewport()
 	if(SchematicViewport.IsValid())
 	{
 		SchematicModel.UpdateControlRigContent();
-		SchematicViewport->ToggleVisibility();
+		bSchematicViewPortIsHidden = !bSchematicViewPortIsHidden; 
 	}
 }
 
@@ -2551,6 +2552,11 @@ bool FControlRigEditor::IsSchematicViewportActive() const
 
 EVisibility FControlRigEditor::GetSchematicOverlayVisibility() const
 {
+	if(bSchematicViewPortIsHidden)
+	{
+		return EVisibility::Hidden;
+	}
+	
 	if(const URigHierarchy* Hierarchy = GetHierarchyBeingDebugged())
 	{
 		TArray<const FRigBaseElement*> SelectedElements = Hierarchy->GetSelectedElements();
@@ -2562,7 +2568,7 @@ EVisibility FControlRigEditor::GetSchematicOverlayVisibility() const
 			return EVisibility::Hidden;
 		}
 	}
-	return EVisibility::HitTestInvisible;
+	return EVisibility::SelfHitTestInvisible;
 }
 
 bool FControlRigEditor::GetToolbarDrawAxesOnSelection() const
diff --git a/Engine/Plugins/Animation/ControlRig/Source/ControlRigEditor/Private/Editor/ControlRigEditor.h b/Engine/Plugins/Animation/ControlRig/Source/ControlRigEditor/Private/Editor/ControlRigEditor.h
--- a/Engine/Plugins/Animation/ControlRig/Source/ControlRigEditor/Private/Editor/ControlRigEditor.h
+++ b/Engine/Plugins/Animation/ControlRig/Source/ControlRigEditor/Private/Editor/ControlRigEditor.h
@@ -369,6 +369,7 @@ protected:
 	bool bRefreshDirectionManipulationTargetsRequired;
 	FSimpleMulticastDelegate RequestNavigateToConnectorWarningDelegate;
 	TSharedPtr<SSchematicGraphPanel> SchematicViewport;
+	bool bSchematicViewPortIsHidden;
 
 	static const TArray<FName> ForwardsSolveEventQueue;
 	static const TArray<FName> BackwardsSolveEventQueue;
