Control Rig: Allow linker mismatch when packaging with old assets

#rnx

diff --git a/Engine/Plugins/Animation/ControlRig/Source/ControlRig/Private/ControlRig.cpp b/Engine/Plugins/Animation/ControlRig/Source/ControlRig/Private/ControlRig.cpp
--- a/Engine/Plugins/Animation/ControlRig/Source/ControlRig/Private/ControlRig.cpp
+++ b/Engine/Plugins/Animation/ControlRig/Source/ControlRig/Private/ControlRig.cpp
@@ -3442,7 +3442,7 @@ void UControlRig::PostInitInstance(URigVMHost* InCDO)
 			if (ObjectFound)
 			{
 				FName NewName = MakeUniqueObjectName(GetTransientPackage(), URigHierarchy::StaticClass(), TEXT("DynamicHierarchy_Deleted"));
-				ObjectFound->Rename(*NewName.ToString(), GetTransientPackage(), REN_DoNotDirty | REN_DontCreateRedirectors | REN_NonTransactional);
+				ObjectFound->Rename(*NewName.ToString(), GetTransientPackage(), REN_DoNotDirty | REN_DontCreateRedirectors | REN_NonTransactional | REN_AllowPackageLinkerMismatch);
 				ObjectFound->MarkAsGarbage();
 			}
 		}
