Control Rig Edit Mode: fixed control shapes in game visibility

#rnx

diff --git a/Engine/Plugins/Animation/ControlRig/Source/ControlRigEditor/Private/EditMode/ControlRigEditMode.cpp b/Engine/Plugins/Animation/ControlRig/Source/ControlRigEditor/Private/EditMode/ControlRigEditMode.cpp
--- a/Engine/Plugins/Animation/ControlRig/Source/ControlRigEditor/Private/EditMode/ControlRigEditMode.cpp
+++ b/Engine/Plugins/Animation/ControlRig/Source/ControlRigEditor/Private/EditMode/ControlRigEditMode.cpp
@@ -5490,15 +5490,7 @@ void FControlRigEditMode::PostPoseUpdate(const FEditorViewportClient* InViewport
 	const bool bIsGameView = InViewportClient ? InViewportClient->IsInGameView() : false;
 	if (bIsGameView)
 	{
-		// no need to update the shape actors in game view, they just need to be hidden
-		static constexpr bool bHiddenInEditor = true;
-		for (const TPair<UControlRig*,TArray<TObjectPtr<AControlRigShapeActor>>>& ShapeActors : ControlRigShapeActors)
-		{
-			for (AControlRigShapeActor* ShapeActor : ShapeActors.Value)
-			{
-				ShapeActor->SetIsTemporarilyHiddenInEditor(bHiddenInEditor);
-			}
-		}
+		// no need to update the shape actors in game view (shapes are already hidden in game using SetActorHiddenInGame(true))
 		return;
 	}
 
