Control Rig: Catch invalid pose adapter pointers during animbp compilation

#rnx

diff --git a/Engine/Plugins/Animation/ControlRig/Source/ControlRig/Private/AnimNode_ControlRigBase.cpp b/Engine/Plugins/Animation/ControlRig/Source/ControlRig/Private/AnimNode_ControlRigBase.cpp
--- a/Engine/Plugins/Animation/ControlRig/Source/ControlRig/Private/AnimNode_ControlRigBase.cpp
+++ b/Engine/Plugins/Animation/ControlRig/Source/ControlRig/Private/AnimNode_ControlRigBase.cpp
@@ -892,11 +892,6 @@ void FAnimNode_ControlRigBase::Evaluate_AnyThread(FPoseContext& Output)
 {
 	DECLARE_SCOPE_HIERARCHICAL_COUNTER_FUNC()
 
-	if(bEnablePoseAdapter && !PoseAdapter.IsValid())
-	{
-		return;
-	}
-
 	FPoseContext SourcePose(Output);
 
 	if (Source.GetLinkNode())
@@ -909,6 +904,13 @@ void FAnimNode_ControlRigBase::Evaluate_AnyThread(FPoseContext& Output)
 		SourcePose.ResetToRefPose();
 	}
 
+	if(bEnablePoseAdapter && !PoseAdapter.IsValid())
+	{
+		Output = SourcePose;
+		ensureMsgf(false, TEXT("At this point the pose adapter needs to be set!"));
+		return;
+	}
+
 	if (CanExecute() && FAnimWeight::IsRelevant(InternalBlendAlpha) && GetControlRig())
 	{
 		if (FAnimWeight::IsFullWeight(InternalBlendAlpha))
@@ -1119,6 +1121,10 @@ void FAnimNode_ControlRigBase::CacheBones_AnyThread(const FAnimationCacheBonesCo
 		{
 			if(URigHierarchy* Hierarchy = ControlRig->GetHierarchy())
 			{
+				if(!PoseAdapter.IsValid())
+				{
+					PoseAdapter = MakeShareable(new FAnimNode_ControlRig_PoseAdapter());
+				}
 				Hierarchy->LinkPoseAdapter(PoseAdapter);
 			}
 		}
