Sequencer: Anim Details:  Fix for regression from 36487487, better test to check for selected actors in sequencer. Found in production.

diff --git a/Engine/Plugins/Animation/ControlRig/Source/ControlRigEditor/Private/EditMode/AnimDetailsProxy.cpp b/Engine/Plugins/Animation/ControlRig/Source/ControlRigEditor/Private/EditMode/AnimDetailsProxy.cpp
--- a/Engine/Plugins/Animation/ControlRig/Source/ControlRigEditor/Private/EditMode/AnimDetailsProxy.cpp
+++ b/Engine/Plugins/Animation/ControlRig/Source/ControlRigEditor/Private/EditMode/AnimDetailsProxy.cpp
@@ -3053,10 +3053,11 @@ void UControlRigDetailPanelControlProxies::ResetSequencerProxies(TMap<ERigContro
 	}
 }
 
-const TArray<UControlRigControlsProxy*> UControlRigDetailPanelControlProxies::GetAllSelectedProxies() 
+const TArray<UControlRigControlsProxy*> UControlRigDetailPanelControlProxies::GetAllSelectedProxies()
 {
 	TArray<UControlRigControlsProxy*> SelectedProxies(SelectedControlRigProxies);
 	TArray<USceneComponent*> BoundCRObjects;
+	TArray<AActor*> BoundCRActors;
 	for (UControlRigControlsProxy* Proxy : SelectedProxies)
 	{
 		if (Proxy && Proxy->OwnerControlRig.IsValid() && Proxy->OwnerControlRig->GetObjectBinding().IsValid())
@@ -3064,6 +3065,10 @@ const TArray<UControlRigControlsProxy*> UControlRigDetailPanelControlProxies::Ge
 			if (USceneComponent* SceneComponent = Cast<USceneComponent>(Proxy->OwnerControlRig->GetObjectBinding()->GetBoundObject()))
 			{
 				BoundCRObjects.Add(SceneComponent);
+				if (AActor* ParentActor = SceneComponent->GetOwner())
+				{
+					BoundCRActors.Add(ParentActor);
+				}
 			}
 		}
 	}
@@ -3071,7 +3076,7 @@ const TArray<UControlRigControlsProxy*> UControlRigDetailPanelControlProxies::Ge
 	{
 		if (SelectedProxies.Num() > 0)
 		{
-			for (int32 Index = SelectedSequencerProxies.Num() -1; Index >=0; --Index)
+			for (int32 Index = SelectedSequencerProxies.Num() - 1; Index >= 0; --Index)
 			{
 				if (UControlRigControlsProxy* Proxy = SelectedSequencerProxies[Index])
 				{
@@ -3087,7 +3092,7 @@ const TArray<UControlRigControlsProxy*> UControlRigDetailPanelControlProxies::Ge
 						}
 						if (AActor* Actor = Cast<AActor>(Proxy->OwnerObject.Get()))
 						{
-							if (BoundCRObjects.Contains(Actor->GetRootComponent()))
+							if (BoundCRActors.Contains(Actor))
 							{
 								SelectedSequencerProxies.RemoveAt(Index);
 								continue;
