fix static analysis warning

#rnx

diff --git a/Engine/Plugins/Animation/ControlRig/Source/ControlRigEditor/Private/Sequencer/ControlConstraintChannelInterface.cpp b/Engine/Plugins/Animation/ControlRig/Source/ControlRigEditor/Private/Sequencer/ControlConstraintChannelInterface.cpp
--- a/Engine/Plugins/Animation/ControlRig/Source/ControlRigEditor/Private/Sequencer/ControlConstraintChannelInterface.cpp
+++ b/Engine/Plugins/Animation/ControlRig/Source/ControlRigEditor/Private/Sequencer/ControlConstraintChannelInterface.cpp
@@ -71,7 +71,12 @@ bool FControlConstraintChannelInterface::SmartConstraintKey(
 	UTickableTransformConstraint* InConstraint, const TOptional<bool>& InOptActive,
 	const FFrameNumber& InTime, const TSharedPtr<ISequencer>& InSequencer)
 {
-	const UTransformableControlHandle* ControlHandle = InConstraint ? Cast<UTransformableControlHandle>(InConstraint->ChildTRSHandle) : nullptr;
+	if (!InConstraint)
+	{
+		return false;
+	}
+
+	const UTransformableControlHandle* ControlHandle = Cast<UTransformableControlHandle>(InConstraint->ChildTRSHandle);
 	if (!ControlHandle)
 	{
 		return false;
