Sequencer: Anim Layers: When we select an anim layer we also want it selected in Sequencer so it's curves show up.

diff --git a/Engine/Plugins/Animation/ControlRig/Source/ControlRigEditor/Private/Sequencer/AnimLayers/AnimLayers.cpp b/Engine/Plugins/Animation/ControlRig/Source/ControlRigEditor/Private/Sequencer/AnimLayers/AnimLayers.cpp
--- a/Engine/Plugins/Animation/ControlRig/Source/ControlRigEditor/Private/Sequencer/AnimLayers/AnimLayers.cpp
+++ b/Engine/Plugins/Animation/ControlRig/Source/ControlRigEditor/Private/Sequencer/AnimLayers/AnimLayers.cpp
@@ -3,10 +3,12 @@
 #include "ControlRig.h"
 #include "ISequencer.h"
 #include "MVVM/Selection/Selection.h"
+#include "MVVM/TrackRowModelStorageExtension.h"
+#include "MVVM/ViewModels/TrackModel.h"
+#include "MVVM/ViewModels/SectionModel.h"
+#include "MVVM/ViewModels/TrackRowModel.h"
 #include "MVVM/ViewModels/SequencerEditorViewModel.h"
 #include "MVVM/ViewModels/SequencerOutlinerViewModel.h"
-#include "MVVM/ViewModels/SequencerTrackAreaViewModel.h"
-#include "MVVM/ViewModels/SequencerEditorViewModel.h"
 #include "MVVM/ViewModels/ObjectBindingModel.h"
 #include "LevelSequence.h"
 #include "ControlRigSequencerEditorLibrary.h"
@@ -586,6 +588,71 @@ bool UAnimLayer::AddSelectedInSequencer()
 	return true;
 }
 
+void UAnimLayer::GetSections(TArray<UMovieSceneSection*>& OutSections) const
+{
+	for (const TPair<TWeakObjectPtr<UObject>, FAnimLayerItem>& Pair : AnimLayerItems)
+	{
+		if (Pair.Key != nullptr)
+		{
+			for (const FAnimLayerSectionItem& SectionItem : Pair.Value.SectionItems)
+			{
+				if (SectionItem.Section.IsValid())
+				{
+					OutSections.Add(SectionItem.Section.Get());
+				}
+			}
+		}
+	}
+}
+
+void UAnimLayer::SetSelectedInList(bool bInValue) 
+{
+	bIsSelectedInList = bInValue;
+	if (TSharedPtr<ISequencer> SequencerPtr = UAnimLayers::GetSequencerFromAsset())
+	{
+		TArray<UMovieSceneSection*> Sections;
+		GetSections(Sections);
+		using namespace UE::Sequencer;
+		FOutlinerSelection& SelectedOutlinerItems = SequencerPtr->GetViewModel()->GetSelection()->Outliner;
+		if (bIsSelectedInList) //selected so select it's track row model
+		{
+			FTrackRowModelStorageExtension* TrackRowModelStorage = SequencerPtr->GetViewModel()->GetRootModel()->CastDynamic<FTrackRowModelStorageExtension>();
+			check(TrackRowModelStorage);
+			for (UMovieSceneSection* Section : Sections)
+			{
+				if (UMovieSceneTrack* Track =Section->GetTypedOuter<UMovieSceneTrack>())
+				{
+					const int32 RowIndex = Section->GetRowIndex();
+					TSharedPtr<FTrackRowModel> TrackRowModel = TrackRowModelStorage->FindModelForTrackRow(Track, RowIndex);
+					SelectedOutlinerItems.Select(TrackRowModel);
+				}
+			}
+		}
+		else //not selected, so clear it from outliner
+		{ 
+			TArray <TViewModelPtr<FTrackRowModel>> TrackRowModelsToDeselect;
+			for (FViewModelPtr ViewModel : SelectedOutlinerItems)
+			{
+				if (TViewModelPtr<FTrackRowModel> TrackRowModel = ViewModel.ImplicitCast())
+				{
+					for (UMovieSceneSection* Section : TrackRowModel->GetSections())
+					{
+						if (Sections.Contains(Section))
+						{
+							TrackRowModelsToDeselect.Add(TrackRowModel);
+							break;
+						}
+					}
+				}
+			}
+			for (TViewModelPtr<FTrackRowModel>& TrackRowModel:  TrackRowModelsToDeselect)
+			{
+				SelectedOutlinerItems.Deselect(TrackRowModel);
+			}
+		}
+	}
+}
+
 bool UAnimLayer::RemoveAnimLayerItem(UObject* InObject)
 {
 	if (FAnimLayerItem* Item =  AnimLayerItems.Find(InObject))
diff --git a/Engine/Plugins/Animation/ControlRig/Source/ControlRigEditor/Private/Sequencer/AnimLayers/AnimLayers.h b/Engine/Plugins/Animation/ControlRig/Source/ControlRigEditor/Private/Sequencer/AnimLayers/AnimLayers.h
--- a/Engine/Plugins/Animation/ControlRig/Source/ControlRigEditor/Private/Sequencer/AnimLayers/AnimLayers.h
+++ b/Engine/Plugins/Animation/ControlRig/Source/ControlRigEditor/Private/Sequencer/AnimLayers/AnimLayers.h
@@ -188,6 +188,7 @@ public:
 	UAnimLayer() {};
 
 	void GetAnimLayerObjects(FAnimLayerObjects& InLayerObjects) const;
+	void GetSections(TArray<UMovieSceneSection*>& OutSections) const;
 
 	UFUNCTION(BlueprintCallable, Category = "State")
 	ECheckBoxState GetKeyed() const;
@@ -230,7 +231,7 @@ public:
 	bool RemoveSelectedInSequencer();
 
 	bool RemoveAnimLayerItem(UObject* InObject);
-	void SetSelectedInList(bool bInValue) { bIsSelectedInList = bInValue; }
+	void SetSelectedInList(bool bInValue); 
 	bool GetSelectedInList() const { return bIsSelectedInList; }
 
 	//for key property status for weights
diff --git a/Engine/Source/Editor/Sequencer/Public/MVVM/TrackRowModelStorageExtension.h b/Engine/Source/Editor/Sequencer/Public/MVVM/TrackRowModelStorageExtension.h
--- a/Engine/Source/Editor/Sequencer/Public/MVVM/TrackRowModelStorageExtension.h
+++ b/Engine/Source/Editor/Sequencer/Public/MVVM/TrackRowModelStorageExtension.h
@@ -21,7 +21,7 @@ class FSequenceModel;
 class FTrackRowModel;
 class FViewModel;
 
-class FTrackRowModelStorageExtension
+class SEQUENCER_API FTrackRowModelStorageExtension
 	: public IDynamicExtension
 {
 public:
