Sequencer: Speculative fix for crash in SimpleConstructionScript->GetAllNodes

#rnx

diff --git a/Engine/Plugins/Animation/ControlRig/Source/ControlRigEditor/Private/Sequencer/ControlRigParameterTrackEditor.cpp b/Engine/Plugins/Animation/ControlRig/Source/ControlRigEditor/Private/Sequencer/ControlRigParameterTrackEditor.cpp
--- a/Engine/Plugins/Animation/ControlRig/Source/ControlRigEditor/Private/Sequencer/ControlRigParameterTrackEditor.cpp
+++ b/Engine/Plugins/Animation/ControlRig/Source/ControlRigEditor/Private/Sequencer/ControlRigParameterTrackEditor.cpp
@@ -195,7 +195,7 @@ static USkeleton* AcquireSkeletonFromObjectGuid(const FGuid& Guid, UObject** Obj
 		}
 
 		UBlueprintGeneratedClass* ActorBlueprintGeneratedClass = Cast<UBlueprintGeneratedClass>(Actor->GetClass());
-		if (ActorBlueprintGeneratedClass)
+		if (ActorBlueprintGeneratedClass && ActorBlueprintGeneratedClass->SimpleConstructionScript)
 		{
 			const TArray<USCS_Node*>& ActorBlueprintNodes = ActorBlueprintGeneratedClass->SimpleConstructionScript->GetAllNodes();
 
diff --git a/Engine/Source/Editor/MovieSceneTools/Private/TrackEditors/SkeletalAnimationTrackEditor.cpp b/Engine/Source/Editor/MovieSceneTools/Private/TrackEditors/SkeletalAnimationTrackEditor.cpp
--- a/Engine/Source/Editor/MovieSceneTools/Private/TrackEditors/SkeletalAnimationTrackEditor.cpp
+++ b/Engine/Source/Editor/MovieSceneTools/Private/TrackEditors/SkeletalAnimationTrackEditor.cpp
@@ -401,7 +401,7 @@ TArray<USkeletalMeshComponent*> AcquireSkeletalMeshComponentsFromObjectGuid(cons
 		}
 
 		UBlueprintGeneratedClass* ActorBlueprintGeneratedClass = Cast<UBlueprintGeneratedClass>(Actor->GetClass());
-		if (ActorBlueprintGeneratedClass)
+		if (ActorBlueprintGeneratedClass && ActorBlueprintGeneratedClass->SimpleConstructionScript)
 		{
 			const TArray<USCS_Node*>& ActorBlueprintNodes = ActorBlueprintGeneratedClass->SimpleConstructionScript->GetAllNodes();
 
