Control Rig: Rely on oodle compression when serializing URigHierarchy to disk

#rnx

diff --git a/Engine/Plugins/Animation/ControlRig/Source/ControlRig/Private/Rigs/RigHierarchy.cpp b/Engine/Plugins/Animation/ControlRig/Source/ControlRig/Private/Rigs/RigHierarchy.cpp
--- a/Engine/Plugins/Animation/ControlRig/Source/ControlRig/Private/Rigs/RigHierarchy.cpp
+++ b/Engine/Plugins/Animation/ControlRig/Source/ControlRig/Private/Rigs/RigHierarchy.cpp
@@ -10,11 +10,11 @@
 #include "Math/ControlRigMathLibrary.h"
 #include "UObject/AnimObjectVersion.h"
 #include "ControlRigObjectVersion.h"
+#include "IControlRigObjectBinding.h"
 #include "ModularRig.h"
 #include "Algo/Count.h"
-#include "Serialization/MemoryReader.h"
-#include "Serialization/MemoryWriter.h"
 #include "HAL/LowLevelMemTracker.h"
+#include "Misc/Compression.h"
 
 LLM_DEFINE_TAG(Animation_ControlRig);
 
@@ -209,9 +209,31 @@ void URigHierarchy::Save(FArchive& Ar)
 
 	ComputeAllTransforms();
 
+	FArchive* ArchiveForElements = &Ar;
+
+	const bool bIsDuplicating = (Ar.GetPortFlags() & PPF_Duplicate) != 0;
+	bool bUseCompressedArchive =
+		!bIsDuplicating &&
+		Ar.IsPersistent() &&
+		!Ar.IsObjectReferenceCollector() &&
+		!Ar.ShouldSkipBulkData() &&
+		!Ar.IsTransacting();
+
+	Ar << bUseCompressedArchive;
+
+	TArray<uint8> UncompressedBytes;
+	TArray<FName> UniqueNames;
+	FRigHierarchyMemoryWriter MemoryWriter(UncompressedBytes, UniqueNames, Ar.IsPersistent());
+	if(bUseCompressedArchive)
+	{
+		MemoryWriter.UsingCustomVersion(FAnimObjectVersion::GUID);
+		MemoryWriter.UsingCustomVersion(FControlRigObjectVersion::GUID);
+		
+		ArchiveForElements = &MemoryWriter;
+	}
+
 	int32 ElementCount = Elements.Num();
-	Ar << ElementCount;
-	UE_RIGVM_ARCHIVETRACE_ENTRY(Ar, TEXT("ElementCount"));
+	*ArchiveForElements << ElementCount;
 
 	for(int32 ElementIndex = 0; ElementIndex < ElementCount; ElementIndex++)
 	{
@@ -219,26 +241,23 @@ void URigHierarchy::Save(FArchive& Ar)
 
 		// store the key
 		FRigElementKey Key = Element->GetKey();
-		Ar << Key;
+		*ArchiveForElements << Key;
 	}
 
 	for(int32 ElementIndex = 0; ElementIndex < ElementCount; ElementIndex++)
 	{
 		FRigBaseElement* Element = Elements[ElementIndex];
-		Element->Serialize(Ar, FRigBaseElement::StaticData);
+		Element->Serialize(*ArchiveForElements, FRigBaseElement::StaticData);
 	}
 
 	for(int32 ElementIndex = 0; ElementIndex < ElementCount; ElementIndex++)
 	{
 		FRigBaseElement* Element = Elements[ElementIndex];
-		Element->Serialize(Ar, FRigBaseElement::InterElementData);
+		Element->Serialize(*ArchiveForElements, FRigBaseElement::InterElementData);
 	}
-	UE_RIGVM_ARCHIVETRACE_ENTRY(Ar, TEXT("Elements"));
 
-	Ar << PreviousNameMap;
-	UE_RIGVM_ARCHIVETRACE_ENTRY(Ar, TEXT("PreviousNameMap"));
-	Ar << PreviousParentMap;
-	UE_RIGVM_ARCHIVETRACE_ENTRY(Ar, TEXT("PreviousParentMap"));
+	*ArchiveForElements << PreviousNameMap;
+	*ArchiveForElements << PreviousParentMap;
 
 	{
 		TMap<FRigElementKey, FMetadataStorage> ElementMetadataToSave;
@@ -251,8 +270,43 @@ void URigHierarchy::Save(FArchive& Ar)
 			}
 		}
 		
-		Ar << ElementMetadataToSave;
-		UE_RIGVM_ARCHIVETRACE_ENTRY(Ar, TEXT("ElementMetadataToSave"));
+		*ArchiveForElements << ElementMetadataToSave;
+	}
+
+	if(bUseCompressedArchive)
+	{
+		// It is possible for compression to actually increase the size of the data, so we over allocate here to handle that.
+		int32 UncompressedSize = UncompressedBytes.Num();
+		int32 CompressedSize = UncompressedSize * 4 / 3;
+
+		TArray<uint8> CompressedBytes;
+		CompressedBytes.SetNumUninitialized(CompressedSize);
+
+		bool bStoreCompressedBytes = FCompression::CompressMemory(NAME_Oodle, CompressedBytes.GetData(), CompressedSize, UncompressedBytes.GetData(), UncompressedBytes.Num(), COMPRESS_BiasMemory);
+		if(bStoreCompressedBytes)
+		{
+			if(CompressedSize < UncompressedSize)
+			{
+				CompressedBytes.SetNum(CompressedSize);
+			}
+			else
+			{
+				bStoreCompressedBytes = false;
+			}
+		}
+
+		Ar << UniqueNames;
+		Ar << UncompressedSize;
+		Ar << bStoreCompressedBytes;
+		if(bStoreCompressedBytes)
+		{
+			Ar << CompressedBytes;
+		}
+		else
+		{
+			Ar << UncompressedBytes;
+		}
+		UE_RIGVM_ARCHIVETRACE_ENTRY(Ar, TEXT("CompressedElements"));
 	}
 }
 
@@ -295,8 +349,47 @@ void URigHierarchy::Load(FArchive& Ar)
 	
 	Reset();
 
+	bool bUseCompressedArchive = false;
+	if(Ar.CustomVer(FControlRigObjectVersion::GUID) >= FControlRigObjectVersion::RigHierarchyCompressElements)
+	{
+		Ar << bUseCompressedArchive;
+	}
+
+	FArchive* ArchiveForElements = &Ar;
+
+	TUniquePtr<FRigHierarchyMemoryReader> MemoryReader;
+	TArray<uint8> UncompressedBytes;
+	TArray<FName> UniqueNames;
+	if(bUseCompressedArchive)
+	{
+		// this block is not going to be entered if we are in an older asset
+		// ( < FControlRigObjectVersion::RigHierarchyCompressElements ) since
+		// bUseCompressedArchive can only be true in that case.
+		Ar << UniqueNames;
+		
+		int32 UncompressedSize = 0;
+		Ar << UncompressedSize;
+
+		bool bStoreCompressedBytes = false;
+		Ar << bStoreCompressedBytes;
+
+		TArray<uint8> CompressedBytes;
+		Ar << CompressedBytes;
+
+		if(bStoreCompressedBytes)
+		{
+			UncompressedBytes.SetNumUninitialized(UncompressedSize);
+
+			bool bSuccess = FCompression::UncompressMemory(NAME_Oodle, UncompressedBytes.GetData(), UncompressedSize, CompressedBytes.GetData(), CompressedBytes.Num());
+			verify(bSuccess);
+		}
+
+		MemoryReader = MakeUnique<FRigHierarchyMemoryReader>(bStoreCompressedBytes ? UncompressedBytes : CompressedBytes, UniqueNames, Ar.IsPersistent());
+		ArchiveForElements = MemoryReader.Get();
+	}
+
 	int32 ElementCount = 0;
-	Ar << ElementCount;
+	*ArchiveForElements << ElementCount;
 
 	PoseVersionPerElement.Reset();
 
@@ -304,11 +397,11 @@ void URigHierarchy::Load(FArchive& Ar)
 	int32 NumDirtyStates = 0;
 	int32 NumCurves = 0;
 
-	const bool bAllocateStoragePerElement = Ar.CustomVer(FControlRigObjectVersion::GUID) < FControlRigObjectVersion::RigHierarchyIndirectElementStorage;
+	const bool bAllocateStoragePerElement = ArchiveForElements->CustomVer(FControlRigObjectVersion::GUID) < FControlRigObjectVersion::RigHierarchyIndirectElementStorage;
 	for(int32 ElementIndex = 0; ElementIndex < ElementCount; ElementIndex++)
 	{
 		FRigElementKey Key;
-		Ar << Key;
+		*ArchiveForElements << Key;
 
 		FRigBaseElement* Element = MakeElement(Key.Type);
 		check(Element);
@@ -321,7 +414,7 @@ void URigHierarchy::Load(FArchive& Ar)
 		if(bAllocateStoragePerElement)
 		{
 			AllocateDefaultElementStorage(Element, false);
-			Element->Load(Ar, FRigBaseElement::StaticData);
+			Element->Load(*ArchiveForElements, FRigBaseElement::StaticData);
 		}
 		else
 		{
@@ -391,7 +484,7 @@ void URigHierarchy::Load(FArchive& Ar)
 			}
 
 			Element->LinkStorage(ElementTransforms.Storage, ElementDirtyStates.Storage, ElementCurves.Storage);
-			Element->Load(Ar, FRigBaseElement::StaticData);
+			Element->Load(*ArchiveForElements, FRigBaseElement::StaticData);
 		}
 	}
 	IncrementTopologyVersion();
@@ -399,7 +492,7 @@ void URigHierarchy::Load(FArchive& Ar)
 	for(int32 ElementIndex = 0; ElementIndex < ElementCount; ElementIndex++)
 	{
 		FRigBaseElement* Element = Elements[ElementIndex];
-		Element->Load(Ar, FRigBaseElement::InterElementData);
+		Element->Load(*ArchiveForElements, FRigBaseElement::InterElementData);
 	}
 
 	IncrementTopologyVersion();
@@ -431,10 +524,10 @@ void URigHierarchy::Load(FArchive& Ar)
 		}
 	}
 
-	if (Ar.CustomVer(FControlRigObjectVersion::GUID) >= FControlRigObjectVersion::RigHierarchyStoringPreviousNames)
+	if (ArchiveForElements->CustomVer(FControlRigObjectVersion::GUID) >= FControlRigObjectVersion::RigHierarchyStoringPreviousNames)
 	{
-		Ar << PreviousNameMap;
-		Ar << PreviousParentMap;
+		*ArchiveForElements << PreviousNameMap;
+		*ArchiveForElements << PreviousParentMap;
 	}
 	else
 	{
@@ -442,12 +535,12 @@ void URigHierarchy::Load(FArchive& Ar)
 		PreviousParentMap.Reset();
 	}
 
-	if (Ar.CustomVer(FControlRigObjectVersion::GUID) >= FControlRigObjectVersion::RigHierarchyStoresElementMetadata)
+	if (ArchiveForElements->CustomVer(FControlRigObjectVersion::GUID) >= FControlRigObjectVersion::RigHierarchyStoresElementMetadata)
 	{
 		ElementMetadata.Reset();
 		TMap<FRigElementKey, FMetadataStorage> LoadedElementMetadata;
 		
-		Ar << LoadedElementMetadata;
+		*ArchiveForElements << LoadedElementMetadata;
 		for (TPair<FRigElementKey, FMetadataStorage>& Entry: LoadedElementMetadata)
 		{
 			FRigBaseElement* Element = Find(Entry.Key);
@@ -8140,4 +8233,56 @@ FRigHierarchyRedirectorGuard::FRigHierarchyRedirectorGuard(UControlRig* InContro
 : Guard(InControlRig->GetHierarchy()->ElementKeyRedirector, &InControlRig->GetElementKeyRedirector())
 {
 }
+
+FRigHierarchyMemoryWriter::FRigHierarchyMemoryWriter(TArray<uint8>& InOutBuffer, TArray<FName>& InOutNames, bool bIsPersistent)
+: FMemoryWriter(InOutBuffer, bIsPersistent)
+, Names(InOutNames)
+{
+}
+
+FArchive& FRigHierarchyMemoryWriter::operator<<(FName& Value)
+{
+	if(const int32* ExistingIndexPtr = NameToIndex.Find(Value))
+	{
+		int32 ExistingIndex = *ExistingIndexPtr;
+		*this << ExistingIndex;
+	}
+	else
+	{
+		int32 NewIndex = Names.Add(Value);
+		NameToIndex.Add(Value, NewIndex);
+		*this << NewIndex;
+	}
+	return *this;
+}
+
+FArchive& FRigHierarchyMemoryWriter::operator<<(FText& Value)
+{
+	FString ValueString = Value.ToString();
+	*this << ValueString;
+	return *this;
+}
+
+FRigHierarchyMemoryReader::FRigHierarchyMemoryReader(TArray<uint8>& InOutBuffer, TArray<FName>& InOutNames, bool bIsPersistent)
+: FMemoryReader(InOutBuffer, bIsPersistent)
+, Names(InOutNames)
+{
+}
+
+FArchive& FRigHierarchyMemoryReader::operator<<(FName& Value)
+{
+	int32 NameIndex = INDEX_NONE;
+	*this << NameIndex;
+	check(Names.IsValidIndex(NameIndex));
+	Value = Names[NameIndex];
+	return *this;
+}
+
+FArchive& FRigHierarchyMemoryReader::operator<<(FText& Value)
+{
+	FString ValueString;
+	*this << ValueString;
+	Value = FText::FromString(ValueString);
+	return *this;
+}
 ;
diff --git a/Engine/Plugins/Animation/ControlRig/Source/ControlRig/Public/ControlRigObjectVersion.h b/Engine/Plugins/Animation/ControlRig/Source/ControlRig/Public/ControlRigObjectVersion.h
--- a/Engine/Plugins/Animation/ControlRig/Source/ControlRig/Public/ControlRigObjectVersion.h
+++ b/Engine/Plugins/Animation/ControlRig/Source/ControlRig/Public/ControlRigObjectVersion.h
@@ -107,6 +107,9 @@ struct CONTROLRIG_API FControlRigObjectVersion
 		// Moved the element storage into separate buffers
 		RigHierarchyIndirectElementStorage,
 
+		// Compress the rig hierarchy when storing to disk
+		RigHierarchyCompressElements,
+
 		// -----<new versions can be added above this line>-------------------------------------------------
 		VersionPlusOne,
 		LatestVersion = VersionPlusOne - 1
diff --git a/Engine/Plugins/Animation/ControlRig/Source/ControlRig/Public/Rigs/RigHierarchy.h b/Engine/Plugins/Animation/ControlRig/Source/ControlRig/Public/Rigs/RigHierarchy.h
--- a/Engine/Plugins/Animation/ControlRig/Source/ControlRig/Public/Rigs/RigHierarchy.h
+++ b/Engine/Plugins/Animation/ControlRig/Source/ControlRig/Public/Rigs/RigHierarchy.h
@@ -17,6 +17,8 @@
 #include "RigVMPythonUtils.h"
 #endif
 #include "Containers/Queue.h"
+#include "Serialization/MemoryReader.h"
+#include "Serialization/MemoryWriter.h"
 #include "RigHierarchy.generated.h"
 
 class UControlRig;
@@ -5331,3 +5333,30 @@ public:
 
 	virtual URigHierarchy* GetHierarchy() const = 0;
 };
+
+class CONTROLRIG_API FRigHierarchyMemoryWriter : public FMemoryWriter
+{
+public:
+	FRigHierarchyMemoryWriter() = delete;
+	FRigHierarchyMemoryWriter(TArray<uint8>& InOutBuffer, TArray<FName>& InOutNames, bool bIsPersistent = false);
+	using FMemoryWriter::operator<<; // For visibility of the overloads we don't override
+	virtual FArchive& operator<<(FName& Value) override;
+	virtual FArchive& operator<<(FText& Value) override;
+
+protected:
+	TArray<FName>& Names;
+	TMap<FName,int32> NameToIndex;
+};
+
+class CONTROLRIG_API FRigHierarchyMemoryReader : public FMemoryReader
+{
+public:
+	FRigHierarchyMemoryReader() = delete;
+	FRigHierarchyMemoryReader(TArray<uint8>& InOutBuffer, TArray<FName>& InOutNames, bool bIsPersistent = false);
+	using FMemoryReader::operator<<; // For visibility of the overloads we don't override
+	virtual FArchive& operator<<(FName& Value) override;
+	virtual FArchive& operator<<(FText& Value) override;
+
+protected:
+	TArray<FName>& Names;
+};
