#UE Use CompressMemoryBound to calculate compressed buffer size in URigHierarchy::Save

diff --git a/Engine/Plugins/Animation/ControlRig/Source/ControlRig/Private/Rigs/RigHierarchy.cpp b/Engine/Plugins/Animation/ControlRig/Source/ControlRig/Private/Rigs/RigHierarchy.cpp
--- a/Engine/Plugins/Animation/ControlRig/Source/ControlRig/Private/Rigs/RigHierarchy.cpp
+++ b/Engine/Plugins/Animation/ControlRig/Source/ControlRig/Private/Rigs/RigHierarchy.cpp
@@ -277,7 +277,7 @@ void URigHierarchy::Save(FArchive& Ar)
 	{
 		// It is possible for compression to actually increase the size of the data, so we over allocate here to handle that.
 		int32 UncompressedSize = UncompressedBytes.Num();
-		int32 CompressedSize = UncompressedSize * 4 / 3;
+		int32 CompressedSize = FCompression::CompressMemoryBound(NAME_Oodle, UncompressedBytes);
 
 		TArray<uint8> CompressedBytes;
 		CompressedBytes.SetNumUninitialized(CompressedSize);
