Control Rig Edit Mode: general improvements
- shape actors: skip hovering when the shapes are being recreated or when the user is interacting
- shape actors: spawn the shape actor using the right transform directly to avoid useless later notifications when calling SetActorTransform
- anim details: avoid useless find proxy + constraints evaluation when building proxies from scratch

diff --git a/Engine/Plugins/Animation/ControlRig/Source/ControlRig/Private/ControlRigGizmoActor.cpp b/Engine/Plugins/Animation/ControlRig/Source/ControlRig/Private/ControlRigGizmoActor.cpp
--- a/Engine/Plugins/Animation/ControlRig/Source/ControlRig/Private/ControlRigGizmoActor.cpp
+++ b/Engine/Plugins/Animation/ControlRig/Source/ControlRig/Private/ControlRigGizmoActor.cpp
@@ -190,9 +190,9 @@ bool AControlRigShapeActor::UpdateControlSettings(
 
 namespace FControlRigShapeHelper
 {
-	FActorSpawnParameters GetDefaultSpawnParameter()
+	const FActorSpawnParameters& GetDefaultSpawnParameter()
 	{
-		FActorSpawnParameters ActorSpawnParameters;
+		static FActorSpawnParameters ActorSpawnParameters;
 #if WITH_EDITOR
 		ActorSpawnParameters.bTemporaryEditorActor = true;
 		ActorSpawnParameters.bHideFromSceneOutliner = true;
@@ -238,7 +238,7 @@ namespace FControlRigShapeHelper
 
 	AControlRigShapeActor* CreateDefaultShapeActor(UWorld* InWorld, const FControlShapeActorCreationParam& CreationParam)
 	{
-		AControlRigShapeActor* ShapeActor = InWorld->SpawnActor<AControlRigShapeActor>(AControlRigShapeActor::StaticClass(), GetDefaultSpawnParameter());
+		AControlRigShapeActor* ShapeActor = InWorld->SpawnActor<AControlRigShapeActor>(AControlRigShapeActor::StaticClass(), CreationParam.SpawnTransform, GetDefaultSpawnParameter());
 		if (ShapeActor)
 		{
 			ShapeActor->ControlRigIndex = CreationParam.ControlRigIndex;
@@ -246,10 +246,6 @@ namespace FControlRigShapeHelper
 			ShapeActor->ControlName = CreationParam.ControlName;
 			ShapeActor->ShapeName = CreationParam.ShapeName;
 			ShapeActor->SetSelectable(CreationParam.bSelectable);
-			ShapeActor->SetActorTransform(CreationParam.SpawnTransform);
-#if WITH_EDITOR
-			ShapeActor->SetActorLabel(CreationParam.ControlName.ToString(), false);
-#endif // WITH_EDITOR
 
 			UStaticMeshComponent* MeshComponent = ShapeActor->StaticMeshComponent;
 
diff --git a/Engine/Plugins/Animation/ControlRig/Source/ControlRig/Public/ControlRigGizmoActor.h b/Engine/Plugins/Animation/ControlRig/Source/ControlRig/Public/ControlRigGizmoActor.h
--- a/Engine/Plugins/Animation/ControlRig/Source/ControlRig/Public/ControlRigGizmoActor.h
+++ b/Engine/Plugins/Animation/ControlRig/Source/ControlRig/Public/ControlRigGizmoActor.h
@@ -191,5 +191,5 @@ namespace FControlRigShapeHelper
 	AControlRigShapeActor* CreateShapeActor(UWorld* InWorld, TSubclassOf<AControlRigShapeActor> InClass, const FControlShapeActorCreationParam& CreationParam);
 	extern CONTROLRIG_API AControlRigShapeActor* CreateDefaultShapeActor(UWorld* InWorld, const FControlShapeActorCreationParam& CreationParam);
 
-	FActorSpawnParameters GetDefaultSpawnParameter();
+	const FActorSpawnParameters& GetDefaultSpawnParameter();
 }
\ No newline at end of file
diff --git a/Engine/Plugins/Animation/ControlRig/Source/ControlRigEditor/Private/EditMode/AnimDetailsProxy.cpp b/Engine/Plugins/Animation/ControlRig/Source/ControlRigEditor/Private/EditMode/AnimDetailsProxy.cpp
--- a/Engine/Plugins/Animation/ControlRig/Source/ControlRigEditor/Private/EditMode/AnimDetailsProxy.cpp
+++ b/Engine/Plugins/Animation/ControlRig/Source/ControlRigEditor/Private/EditMode/AnimDetailsProxy.cpp
@@ -3134,7 +3134,33 @@ void UControlRigDetailPanelControlProxies::ValuesChanged()
 	}
 }
 
-UControlRigControlsProxy* UControlRigDetailPanelControlProxies::AddProxy(UControlRig* ControlRig, FRigControlElement* ControlElement)
+void UControlRigDetailPanelControlProxies::AddNewProxy(const FRigControlElement* InControlElement, UControlRig* InControlRig, FNameToProxyMap& InOutNameToProxy)
+{
+	if (!InControlElement || !InControlRig)
+	{
+		return;
+	}
+
+	const FRigControlSettings& Settings = InControlElement->Settings;
+	UControlRigControlsProxy* Proxy = NewProxyFromType(Settings.ControlType, Settings.ControlEnum);
+	if (!Proxy)
+	{
+		return;
+	}
+
+	// add it to the map
+	InOutNameToProxy.NameToProxy.Add(InControlElement->GetFName(), Proxy);
+
+	Proxy->Type = Settings.ControlType;
+	Proxy->OwnerControlElement.UpdateCache(InControlElement->GetKey(), InControlRig->GetHierarchy());
+	Proxy->OwnerControlRig = InControlRig;
+	Proxy->AddControlRigControl(InControlRig, InControlElement->GetKey().Name);
+	Proxy->bIsIndividual = (InControlElement->IsAnimationChannel()) || (Settings.AnimationType == ERigControlAnimationType::ProxyControl);
+	Proxy->SetFlags(RF_Transactional);
+	Proxy->Modify();
+}
+
+UControlRigControlsProxy* UControlRigDetailPanelControlProxies::FindOrAddProxy(UControlRig* ControlRig, FRigControlElement* ControlElement)
 {
 	if (ControlRig == nullptr || ControlElement == nullptr)
 	{
@@ -3145,7 +3171,7 @@ UControlRigControlsProxy* UControlRigDetailPanelControlProxies::AddProxy(UContro
 		(ControlElement->Settings.AnimationType == ERigControlAnimationType::ProxyControl);
 
 	UControlRigControlsProxy* Proxy = FindProxy(ControlRig, ControlElement);
-	if (!Proxy && ControlElement != nullptr)
+	if (!Proxy)
 	{
 		Proxy = NewProxyFromType(ControlElement->Settings.ControlType, ControlElement->Settings.ControlEnum);
 		//proxy was created so add it
@@ -3162,9 +3188,7 @@ UControlRigControlsProxy* UControlRigDetailPanelControlProxies::AddProxy(UContro
 				NewControlRigProxies.NameToProxy.Add(ControlElement->GetFName(), Proxy);
 				ControlRigOnlyProxies.Add(ControlRig, NewControlRigProxies);
 			}
-		}
-		if (Proxy)
-		{
+
 			Proxy->Type = ControlElement->Settings.ControlType;;
 			Proxy->OwnerControlElement.UpdateCache(ControlElement->GetKey(), ControlRig->GetHierarchy());
 			Proxy->OwnerControlRig = ControlRig;
@@ -3232,7 +3256,7 @@ UControlRigControlsProxy* UControlRigDetailPanelControlProxies::AddProxy(UObject
 }
 
 
-UControlRigControlsProxy* UControlRigDetailPanelControlProxies::NewProxyFromType(ERigControlType ControlType, TObjectPtr<UEnum>& EnumPtr)
+UControlRigControlsProxy* UControlRigDetailPanelControlProxies::NewProxyFromType(ERigControlType ControlType, const TObjectPtr<UEnum>& InEnumPtr)
 {
 	UControlRigControlsProxy* Proxy = nullptr;
 	switch (ControlType)
@@ -3252,14 +3276,14 @@ UControlRigControlsProxy* UControlRigDetailPanelControlProxies::NewProxyFromType
 	}
 	case ERigControlType::Integer:
 	{
-		if (EnumPtr == nullptr)
+		if (InEnumPtr == nullptr)
 		{
 			Proxy = NewObject<UAnimDetailControlsProxyInteger>(this, NAME_None); //was GetTransientPackage....? whay
 		}
 		else
 		{
 			UAnimDetailControlsProxyEnum* EnumProxy = NewObject<UAnimDetailControlsProxyEnum>(this, NAME_None);
-			EnumProxy->Enum.EnumType = EnumPtr;
+			EnumProxy->Enum.EnumType = InEnumPtr;
 			Proxy = EnumProxy;
 		}
 		break;
@@ -3387,15 +3411,36 @@ void UControlRigDetailPanelControlProxies::RemoveControlRigProxies(UControlRig*
 	}
 }
 
-void UControlRigDetailPanelControlProxies::RecreateAllProxies(UControlRig* ControlRig)
+void UControlRigDetailPanelControlProxies::RecreateAllProxies(UControlRig* InControlRig)
 {
-	RemoveControlRigProxies(ControlRig);
-	TArray<FRigControlElement*> Controls = ControlRig->AvailableControls();
-	for (FRigControlElement* ControlElement : Controls)
+	RemoveControlRigProxies(InControlRig);
+
+	if (!InControlRig)
+	{
+		return;
+	}
+
+	const TArray<FRigControlElement*> Controls = InControlRig->AvailableControls();
+	if (Controls.IsEmpty())
+	{
+		return;
+	}
+
+	if (ensure(!ControlRigOnlyProxies.Contains(InControlRig)))
 	{
-		if (ControlElement->Settings.AnimationType != ERigControlAnimationType::VisualCue)
+		FNameToProxyMap& NameToProxies = ControlRigOnlyProxies.Add(InControlRig);
+	
+		for (FRigControlElement* ControlElement : Controls)
+		{
+			if (ControlElement->Settings.AnimationType != ERigControlAnimationType::VisualCue)
+			{
+				AddNewProxy(ControlElement, InControlRig, NameToProxies);
+			}
+		}
+
+		if (NameToProxies.NameToProxy.IsEmpty())
 		{
-			AddProxy(ControlRig, ControlElement);
+			ControlRigOnlyProxies.Remove(InControlRig);
 		}
 	}
 }
diff --git a/Engine/Plugins/Animation/ControlRig/Source/ControlRigEditor/Private/EditMode/ControlRigControlsProxy.h b/Engine/Plugins/Animation/ControlRig/Source/ControlRigEditor/Private/EditMode/ControlRigControlsProxy.h
--- a/Engine/Plugins/Animation/ControlRig/Source/ControlRigEditor/Private/EditMode/ControlRigControlsProxy.h
+++ b/Engine/Plugins/Animation/ControlRig/Source/ControlRigEditor/Private/EditMode/ControlRigControlsProxy.h
@@ -198,7 +198,9 @@ protected:
 public:
 
 	void SelectProxy(UControlRig* InControlRig, FRigControlElement* RigElement, bool bSelected);
-	UControlRigControlsProxy* AddProxy(UControlRig* InControlRig, FRigControlElement* RigElement);
+
+	void AddNewProxy(const FRigControlElement* InControlElement, UControlRig* InControlRig, FNameToProxyMap& InOutNameToProxy);
+	UControlRigControlsProxy* FindOrAddProxy(UControlRig* InControlRig, FRigControlElement* InRigElement);
 	UControlRigControlsProxy* AddProxy(UObject* InObject, ERigControlType Type, TWeakObjectPtr<UMovieSceneTrack>& Track, TSharedPtr<FTrackInstancePropertyBindings>& Binding);
 	void RemoveControlRigProxies(UControlRig* InControlRig);
 	void RemoveSequencerProxies(UObject* InObject);
@@ -222,7 +224,7 @@ private:
 	UControlRigControlsProxy* FindProxy(UControlRig* InControlRig, FRigControlElement* RigElement) const;
 	UControlRigControlsProxy* FindProxy(UObject* InObject, FName PropertyName) const;
 
-	UControlRigControlsProxy* NewProxyFromType(ERigControlType Type, TObjectPtr<UEnum>& EnumPtr);
+	UControlRigControlsProxy* NewProxyFromType(ERigControlType Type, const TObjectPtr<UEnum>& InEnumPtr);
 
 	//if SelectedProxy == nullptr then always clear it
 	void ClearSelectedProperty(UControlRigControlsProxy* SelectedProxy = nullptr);
diff --git a/Engine/Plugins/Animation/ControlRig/Source/ControlRigEditor/Private/EditMode/ControlRigEditMode.cpp b/Engine/Plugins/Animation/ControlRig/Source/ControlRigEditor/Private/EditMode/ControlRigEditMode.cpp
--- a/Engine/Plugins/Animation/ControlRig/Source/ControlRigEditor/Private/EditMode/ControlRigEditMode.cpp
+++ b/Engine/Plugins/Animation/ControlRig/Source/ControlRigEditor/Private/EditMode/ControlRigEditMode.cpp
@@ -621,7 +621,7 @@ void FControlRigEditMode::Tick(FEditorViewportClient* ViewportClient, float Delt
 		// HandleSelectionChanged() will already update the pivots 
 		UpdatePivotTransforms();
 	}
-	
+
 	if (!AreEditingControlRigDirectly() == false)
 	{
 		ViewportClient->Invalidate();
@@ -631,6 +631,7 @@ void FControlRigEditMode::Tick(FEditorViewportClient* ViewportClient, float Delt
 	if (RecreateControlShapesRequired != ERecreateControlRigShape::RecreateNone && !(FSlateApplication::Get().HasAnyMouseCaptor() || GUnrealEd->IsUserInteracting()))
 	{
 		RecreateControlShapeActors();
+		const bool bAreEditingControlRigDirectly = AreEditingControlRigDirectly();
 		for (TWeakObjectPtr<UControlRig>& RuntimeRigPtr : RuntimeControlRigs)
 		{
 			if (UControlRig* ControlRig = RuntimeRigPtr.Get())
@@ -646,14 +647,14 @@ void FControlRigEditMode::Tick(FEditorViewportClient* ViewportClient, float Delt
 							ShapeActor->SetSelected(true);
 						}
 
-						if (!AreEditingControlRigDirectly())
+						if (!bAreEditingControlRigDirectly)
 						{
 							FRigControlElement* ControlElement = ControlRig->FindControl(SelectedKey.Name);
 							if (ControlElement)
 							{
 								if (!ControlRig->IsCurveControl(ControlElement))
 								{
-									ControlProxy->AddProxy(ControlRig, ControlElement);
+									ControlProxy->FindOrAddProxy(ControlRig, ControlElement);
 								}
 							}
 						}
@@ -665,6 +666,10 @@ void FControlRigEditMode::Tick(FEditorViewportClient* ViewportClient, float Delt
 		HandleSelectionChanged();
 		RecreateControlShapesRequired = ERecreateControlRigShape::RecreateNone;
 		ControlRigsToRecreate.SetNum(0);
+		if (DetailKeyFrameCache)
+        {
+        	DetailKeyFrameCache->ResetCachedData();
+        }
 	}
 
 	{
@@ -3669,6 +3674,16 @@ bool FControlRigEditMode::MouseMove(FEditorViewportClient* InViewportClient, FVi
 	// NOTE: with synthesized mouse moves, this is being called a lot sadly so playing in sequencer with the mouse over the viewport leads to fps drop
 	auto HasAnyHoverableShapeActor = [this, InViewportClient]()
 	{
+		if (RecreateControlShapesRequired != ERecreateControlRigShape::RecreateNone)
+		{
+			return false;
+		}
+		
+		if (!InteractionScopes.IsEmpty())
+		{
+			return false;
+		}
+		
 		if (!InViewportClient || InViewportClient->IsInGameView())
 		{
 			return false;
@@ -4220,6 +4235,10 @@ void FControlRigEditMode::CreateShapeActors(UControlRig* InControlRig)
 		{
 			ControlProxy->RecreateAllProxies(InControlRig);
 		}
+		if (DetailKeyFrameCache)
+		{
+			DetailKeyFrameCache->ResetCachedData();
+		}
 	}
 }
 
@@ -4664,6 +4683,10 @@ void FControlRigEditMode::OnPostConstruction_AnyThread(UControlRig* InRig, const
 			RequestToRecreateControlShapeActors(InRig);
 			RecreateControlShapeActors();
 			HandleSelectionChanged();
+			if (DetailKeyFrameCache)
+			{
+				DetailKeyFrameCache->ResetCachedData();
+			}
 		};
 				
 		if(IsInGameThread())
