Control Rig Edit Mode: disable compensation when evaluating constraints

#rnx

diff --git a/Engine/Plugins/Animation/ControlRig/Source/ControlRigEditor/Private/EditMode/ControlRigEditMode.cpp b/Engine/Plugins/Animation/ControlRig/Source/ControlRigEditor/Private/EditMode/ControlRigEditMode.cpp
--- a/Engine/Plugins/Animation/ControlRig/Source/ControlRigEditor/Private/EditMode/ControlRigEditMode.cpp
+++ b/Engine/Plugins/Animation/ControlRig/Source/ControlRigEditor/Private/EditMode/ControlRigEditMode.cpp
@@ -4864,7 +4864,8 @@ void FControlRigEditMode::SetControlShapeTransform(
 	
 	ControlRig->SetControlLocalTransform(InShapeActor->ControlName, LocalTransform, bNotify, Context, bUndo, bFixEulerFlips);
 	EvaluateRigIfAdditive();
-	
+
+	TGuardValue<bool> CompensateGuard(FMovieSceneConstraintChannelHelper::bDoNotCompensate, true);
 	const FConstraintsManagerController& Controller = FConstraintsManagerController::Get(ControlRig->GetWorld());
 	Controller.EvaluateAllConstraints();
 }
