Sequencer: Control Rig: Make sure to only rename float type controls to remapped curve names when the version is old, add missing version number in serialize, plus some minor optimizations during the naming mapping process to make it more efficient when we know the control type.

diff --git a/Engine/Plugins/Animation/ControlRig/Source/ControlRig/Private/Sequencer/MovieSceneControlRigParameterSection.cpp b/Engine/Plugins/Animation/ControlRig/Source/ControlRig/Private/Sequencer/MovieSceneControlRigParameterSection.cpp
--- a/Engine/Plugins/Animation/ControlRig/Source/ControlRig/Private/Sequencer/MovieSceneControlRigParameterSection.cpp
+++ b/Engine/Plugins/Animation/ControlRig/Source/ControlRig/Private/Sequencer/MovieSceneControlRigParameterSection.cpp
@@ -863,7 +863,7 @@ void UMovieSceneControlRigParameterSection::PreSave(FObjectPreSaveContext SaveCo
 	Super::PreSave(SaveContext);
 }
 
-bool UMovieSceneControlRigParameterSection::RenameParameterName(const FName& OldParameterName, const FName& NewParameterName)
+bool UMovieSceneControlRigParameterSection::RenameParameterName(const FName& OldParameterName, const FName& NewParameterName, TOptional<ERigControlType> ControlType)
 {
 	bool bWasReplaced = false;
 
@@ -883,15 +883,61 @@ bool UMovieSceneControlRigParameterSection::RenameParameterName(const FName& Old
 			}
 		}
 	};
-
-	RenameParameterNameInner(ScalarParameterNamesAndCurves);
-	RenameParameterNameInner(BoolParameterNamesAndCurves);
-	RenameParameterNameInner(EnumParameterNamesAndCurves);
-	RenameParameterNameInner(IntegerParameterNamesAndCurves);
-	RenameParameterNameInner(Vector2DParameterNamesAndCurves);
-	RenameParameterNameInner(VectorParameterNamesAndCurves);
-	RenameParameterNameInner(ColorParameterNamesAndCurves);
-	RenameParameterNameInner(TransformParameterNamesAndCurves);
+	if(ControlType.IsSet())
+	{
+		switch (ControlType.GetValue())
+		{
+			case ERigControlType::Float:
+			case ERigControlType::ScaleFloat:
+			{
+				RenameParameterNameInner(ScalarParameterNamesAndCurves);
+				break;
+			}
+			case ERigControlType::Bool:
+			{
+				RenameParameterNameInner(BoolParameterNamesAndCurves);
+				break;
+			}
+			case ERigControlType::Integer:
+			{
+				RenameParameterNameInner(IntegerParameterNamesAndCurves);
+				RenameParameterNameInner(EnumParameterNamesAndCurves);
+				break;
+			}
+			case ERigControlType::Vector2D:
+			{
+				RenameParameterNameInner(Vector2DParameterNamesAndCurves);
+				break;
+			}
+			case ERigControlType::Position:
+			case ERigControlType::Rotator:
+			case ERigControlType::Scale:
+			{
+				RenameParameterNameInner(VectorParameterNamesAndCurves);
+				break;
+			}
+			break;
+			case ERigControlType::Transform:
+			case ERigControlType::EulerTransform:
+			case ERigControlType::TransformNoScale:
+			{
+				RenameParameterNameInner(TransformParameterNamesAndCurves);
+				break;
+			}
+		
+		};
+	}
+	else
+	{
+		RenameParameterNameInner(ScalarParameterNamesAndCurves);
+		RenameParameterNameInner(BoolParameterNamesAndCurves);
+		RenameParameterNameInner(EnumParameterNamesAndCurves);
+		RenameParameterNameInner(IntegerParameterNamesAndCurves);
+		RenameParameterNameInner(Vector2DParameterNamesAndCurves);
+		RenameParameterNameInner(VectorParameterNamesAndCurves);
+		RenameParameterNameInner(ColorParameterNamesAndCurves);
+		RenameParameterNameInner(TransformParameterNamesAndCurves);
+	}
 
 	if (bWasReplaced)
 	{
@@ -1051,6 +1097,7 @@ void UMovieSceneControlRigParameterSection::SetBlendType(EMovieSceneBlendType In
 
 void UMovieSceneControlRigParameterSection::Serialize(FArchive& Ar)
 {
+	Ar.UsingCustomVersion(FUE5MainStreamObjectVersion::GUID);
 	Super::Serialize(Ar);
 }
 #if WITH_EDITOR
@@ -2529,12 +2576,17 @@ void UMovieSceneControlRigParameterSection::RecreateWithThisControlRig(UControlR
 		FName PreviousName = Hierarchy->GetPreviousName(ControlElement->GetKey());
 		if (PreviousName != NAME_None && PreviousName != ControlElement->GetKey().Name)
 		{
-			RenameParameterName(PreviousName, ControlElement->GetKey().Name);
+			TOptional<ERigControlType> ControlType = ControlElement->Settings.ControlType;
+			RenameParameterName(PreviousName, ControlElement->GetKey().Name, ControlType);
 		}
-
-		if (const FName* OldCurveControlName = CurveControlNameRemapping.Find(ControlElement->GetFName()))
+		if (ControlElement->Settings.ControlType == ERigControlType::Float ||
+			ControlElement->Settings.ControlType == ERigControlType::ScaleFloat)
 		{
-			RenameParameterName(*OldCurveControlName, ControlElement->GetKey().Name);
+			if (const FName* OldCurveControlName = CurveControlNameRemapping.Find(ControlElement->GetFName()))
+			{
+				TOptional<ERigControlType> ControlType = ControlElement->Settings.ControlType;
+				RenameParameterName(*OldCurveControlName, ControlElement->GetKey().Name, ControlType);
+			}
 		}
 		
 		switch (ControlElement->Settings.ControlType)
diff --git a/Engine/Plugins/Animation/ControlRig/Source/ControlRig/Private/Sequencer/MovieSceneControlRigParameterTrack.cpp b/Engine/Plugins/Animation/ControlRig/Source/ControlRig/Private/Sequencer/MovieSceneControlRigParameterTrack.cpp
--- a/Engine/Plugins/Animation/ControlRig/Source/ControlRig/Private/Sequencer/MovieSceneControlRigParameterTrack.cpp
+++ b/Engine/Plugins/Animation/ControlRig/Source/ControlRig/Private/Sequencer/MovieSceneControlRigParameterTrack.cpp
@@ -691,11 +691,12 @@ void UMovieSceneControlRigParameterTrack::RenameParameterName(const FName& OldPa
 {
 	if (OldParameterName != NewParameterName)
 	{
+		TOptional<ERigControlType> ControlType; //not set so need to check all 
 		for (UMovieSceneSection* Section : Sections)
 		{
 			if (UMovieSceneControlRigParameterSection* CRSection = Cast<UMovieSceneControlRigParameterSection>(Section))
 			{
-				CRSection->RenameParameterName(OldParameterName, NewParameterName);
+				CRSection->RenameParameterName(OldParameterName, NewParameterName, ControlType);
 			}
 		}
 		if (FControlRotationOrder* RotationOrder = ControlsRotationOrder.Find(OldParameterName))
diff --git a/Engine/Plugins/Animation/ControlRig/Source/ControlRig/Public/Sequencer/MovieSceneControlRigParameterSection.h b/Engine/Plugins/Animation/ControlRig/Source/ControlRig/Public/Sequencer/MovieSceneControlRigParameterSection.h
--- a/Engine/Plugins/Animation/ControlRig/Source/ControlRig/Public/Sequencer/MovieSceneControlRigParameterSection.h
+++ b/Engine/Plugins/Animation/ControlRig/Source/ControlRig/Public/Sequencer/MovieSceneControlRigParameterSection.h
@@ -209,8 +209,8 @@ public:
 	FSpaceChannelAddedEvent& SpaceChannelAdded() { return OnSpaceChannelAdded; }
 
 	const FName& FindControlNameFromConstraintChannel(const FMovieSceneConstraintChannel* InConstraintChannel) const;
-		
-	bool RenameParameterName(const FName& OldParameterName, const FName& NewParameterName);
+	
+	bool RenameParameterName(const FName& OldParameterName, const FName& NewParameterName, TOptional<ERigControlType> ControlType);
 
 	void ChangeControlRotationOrder(const FName& InControlName, const TOptional<EEulerRotationOrder>& CurrentOrder, 
 		const  TOptional<EEulerRotationOrder>& NewOrder, EMovieSceneKeyInterpolation Interpolation);
