Control Rig Edit Mode: fixed actor label

diff --git a/Engine/Plugins/Animation/ControlRig/Source/ControlRig/Private/ControlRigGizmoActor.cpp b/Engine/Plugins/Animation/ControlRig/Source/ControlRig/Private/ControlRigGizmoActor.cpp
--- a/Engine/Plugins/Animation/ControlRig/Source/ControlRig/Private/ControlRigGizmoActor.cpp
+++ b/Engine/Plugins/Animation/ControlRig/Source/ControlRig/Private/ControlRigGizmoActor.cpp
@@ -247,6 +247,10 @@ namespace FControlRigShapeHelper
 			ShapeActor->ShapeName = CreationParam.ShapeName;
 			ShapeActor->SetSelectable(CreationParam.bSelectable);
 
+#if WITH_EDITOR
+			ShapeActor->SetActorLabel(CreationParam.ControlName.ToString(), false);
+#endif // WITH_EDITOR
+
 			UStaticMeshComponent* MeshComponent = ShapeActor->StaticMeshComponent;
 
 			if (!CreationParam.StaticMesh.IsValid())
