AnimLayer: Fix anim layer merges with overrids when the sections may have the control masked out.

diff --git a/Engine/Plugins/Animation/ControlRig/Source/ControlRigEditor/Private/Sequencer/AnimLayers/AnimLayers.cpp b/Engine/Plugins/Animation/ControlRig/Source/ControlRigEditor/Private/Sequencer/AnimLayers/AnimLayers.cpp
--- a/Engine/Plugins/Animation/ControlRig/Source/ControlRigEditor/Private/Sequencer/AnimLayers/AnimLayers.cpp
+++ b/Engine/Plugins/Animation/ControlRig/Source/ControlRigEditor/Private/Sequencer/AnimLayers/AnimLayers.cpp
@@ -1832,8 +1832,16 @@ bool UAnimLayers::SetPassthroughKey(ISequencer* InSequencer, int32 InIndex)
 	return true;
 }
 
+
 static void MergeControlRigSections(UMovieSceneControlRigParameterSection* BaseSection, UMovieSceneControlRigParameterSection* Section, const TRange<FFrameNumber>& Range)
 {
+	//disable section if no control in the section and override, need to make sure we skip it's eval or it will incorrectly
+	//override when merging
+	auto ShouldDisableSection = [](UMovieSceneControlRigParameterSection* CRSection, FRigControlElement* ControlElement)
+	{
+		return (CRSection->GetControlNameMask(ControlElement->GetFName()) == false &&
+				CRSection->GetBlendType().IsValid() && CRSection->GetBlendType() == EMovieSceneBlendType::Override);
+	};
 	if (!BaseSection || !Section)
 	{
 		return;
@@ -1867,12 +1875,19 @@ static void MergeControlRigSections(UMovieSceneControlRigParameterSection* BaseS
 			if (FChannelMapInfo* pChannelIndex = BaseSection->ControlChannelMap.Find(ControlElement->GetFName()))
 			{
 				const int32 ChannelIndex = pChannelIndex->ChannelIndex;
-				bool bMaskKeyOut = (BaseSection->GetControlNameMask(ControlElement->GetFName()) == false)
-					|| (Section->GetControlNameMask(ControlElement->GetFName()) == false);
-
-				if (bMaskKeyOut)
+				TOptional<bool> bBaseSectionResetActive;
+				TOptional<bool> bSectionResetActive; 
+				
+				if (ShouldDisableSection(BaseSection, ControlElement))
 				{
-					continue;
+					bBaseSectionResetActive = BaseSection->IsActive();
+					BaseSection->SetIsActive(false);
+				}
+				
+				if (ShouldDisableSection(Section,ControlElement))
+				{
+					bSectionResetActive = Section->IsActive();
+					Section->SetIsActive(false);
 				}
 
 				switch (ControlElement->Settings.ControlType)
@@ -2014,6 +2029,14 @@ static void MergeControlRigSections(UMovieSceneControlRigParameterSection* BaseS
 					default:
 						break;
 				}
+				if (bBaseSectionResetActive.IsSet())
+				{
+					BaseSection->SetIsActive(bBaseSectionResetActive.GetValue());
+				}
+				if (bSectionResetActive.IsSet())
+				{
+					Section->SetIsActive(bSectionResetActive.GetValue());
+				} 
 			}
 		}
 	}
diff --git a/Engine/Source/Editor/MovieSceneTools/Public/MovieSceneToolHelpers.h b/Engine/Source/Editor/MovieSceneTools/Public/MovieSceneToolHelpers.h
--- a/Engine/Source/Editor/MovieSceneTools/Public/MovieSceneToolHelpers.h
+++ b/Engine/Source/Editor/MovieSceneTools/Public/MovieSceneToolHelpers.h
@@ -759,7 +759,7 @@ public:
 				else if (Section->GetBlendType().Get() == EMovieSceneBlendType::Override)
 				{
 					Value = (Value * (1.0 - Weight)) +
-						(Value * Weight);
+						(WeightedValue * Weight);
 				}
 			}
 			Values.Add(Value);
@@ -883,18 +883,15 @@ public:
 				}
 				else if (MergeAlgorithm == FChannelMergeAlgorithm::Override)
 				{
+					//when doing an override merge we need to just get the full value since the new layer will also be
+					//an override layer
 					if (Sections[1]->GetRange().Contains(Frame))
 					{
-						float Weight = Sections[1]->GetTotalWeightValue(FrameTime);
 						TArray<CurveValueType> ChannelValues = MovieSceneToolHelpers::GetChannelValues<ChannelType,
 							CurveValueType>(OverrideChannelIndex, OverrideChannelIndex, OtherSections, AbsoluteSections, Frame);
 						if (ChannelValues.Num() == 1)
 						{
-							CurveValueType WeightedValue = 0.0;
-							ChannelType* EachChannel = Channels[1];
-							EachChannel->Evaluate(FrameTime, WeightedValue);
-							WeightedValue *= Weight;
-							Value.Value = WeightedValue + (1.0 - Weight) * ChannelValues[0];
+							Value.Value = ChannelValues[0];
 						}
 					}
 				}
@@ -944,13 +941,15 @@ public:
 		Sections.Add(BaseSection);
 		Sections.Add(TopSection);
 		FChannelMergeAlgorithm MergeAlgorithm = FChannelMergeAlgorithm::Add;
-		if (TopSection->GetBlendType().IsValid() && TopSection->GetBlendType().Get() == EMovieSceneBlendType::Absolute)
+		//if either setion is override we do an override blend since the result will be override
+		if ((TopSection->GetBlendType().IsValid() && TopSection->GetBlendType().Get() == EMovieSceneBlendType::Override) ||
+			(BaseSection->GetBlendType().IsValid() && BaseSection->GetBlendType().Get() == EMovieSceneBlendType::Override))
 		{
-			MergeAlgorithm = FChannelMergeAlgorithm::Average;
+			MergeAlgorithm = FChannelMergeAlgorithm::Override;
 		}
-		else if (TopSection->GetBlendType().IsValid() && TopSection->GetBlendType().Get() == EMovieSceneBlendType::Override)
+		else
 		{
-			MergeAlgorithm = FChannelMergeAlgorithm::Override;
+			MergeAlgorithm = FChannelMergeAlgorithm::Average;
 		}
 
 		for (ChannelIndex = StartIndex; ChannelIndex <= EndIndex; ++ChannelIndex)
