Control Rig Edit Mode: general improvements

the outliner is now lazily built and this has several advantages
    - there's no need to rebuild the outliner every time a control rig is added to sequencer. this can be done once they've all been added
    - the outliner is not actually built when it's useless (especially when it is hidden) but it still receives notifications of pending refresh, so that it is correctly built when needed

diff --git a/Engine/Plugins/Animation/ControlRig/Source/ControlRigEditor/Private/EditMode/SControlRigOutliner.cpp b/Engine/Plugins/Animation/ControlRig/Source/ControlRigEditor/Private/EditMode/SControlRigOutliner.cpp
--- a/Engine/Plugins/Animation/ControlRig/Source/ControlRigEditor/Private/EditMode/SControlRigOutliner.cpp
+++ b/Engine/Plugins/Animation/ControlRig/Source/ControlRigEditor/Private/EditMode/SControlRigOutliner.cpp
@@ -1,8 +1,6 @@
 // Copyright Epic Games, Inc. All Rights Reserved.
 
 #include "EditMode/SControlRigOutliner.h"
-#include "Framework/Notifications/NotificationManager.h"
-#include "Widgets/Notifications/SNotificationList.h"
 #include "Framework/MultiBox/MultiBoxBuilder.h"
 #include "Widgets/Input/SButton.h"
 #include "Widgets/Layout/SExpandableArea.h"
@@ -13,32 +11,24 @@
 #include "Styling/CoreStyle.h"
 #include "ScopedTransaction.h"
 #include "ControlRig.h"
-#include "UnrealEdGlobals.h"
 #include "EditMode/ControlRigEditMode.h"
 #include "EditorModeManager.h"
 #include "ISequencer.h"
 #include "MVVM/ViewModels/SequencerEditorViewModel.h"
 #include "MVVM/Selection/Selection.h"
-#include "LevelSequence.h"
-#include "Selection.h"
 #include "Editor.h"
-#include "LevelEditor.h"
 #include "Widgets/Text/STextBlock.h"
 #include "EditMode/ControlRigControlsProxy.h"
-#include "EditMode/ControlRigEditModeSettings.h"
 #include "TimerManager.h"
 #include "Editor/SRigHierarchyTreeView.h"
 #include "Rigs/RigHierarchyController.h"
 #include "ControlRigObjectBinding.h"
-#include "ControlRigEditorStyle.h"
 #include "Settings/ControlRigSettings.h"
 #include "Widgets/Layout/SSpacer.h"
 #include "Widgets/Images/SImage.h"
 #include "Widgets/Input/SSearchBox.h"
-#include "PropertyCustomizationHelpers.h"
 #include "Framework/Application/SlateApplication.h"
 #include "Editor/EditorEngine.h"
-#include "HelperUtil.h"
 #include "ScopedTransaction.h"
 #include "Rigs/FKControlRig.h"
 
@@ -347,6 +337,11 @@ TPair<const FSlateBrush*, FSlateColor> SMultiRigHierarchyItem::GetBrushForElemen
 /// SRigHierarchyTreeView
 ///////////////////////////////////////////////////////////
 
+SMultiRigHierarchyTreeView::~SMultiRigHierarchyTreeView()
+{
+	UnregisterPendingRefresh();
+}
+
 void SMultiRigHierarchyTreeView::Construct(const FArguments& InArgs)
 {
 	Delegates = InArgs._RigTreeDelegates;
@@ -725,108 +720,172 @@ bool SMultiRigHierarchyTreeView::RemoveElement(const FMultiRigData& InData)
 	return ElementMap.Remove(InData) > 0;
 }
 
-void SMultiRigHierarchyTreeView::RefreshTreeView(bool bRebuildContent)
+void SMultiRigHierarchyTreeView::RequestTreeViewRefresh(TFunction<void()>&& InRefreshFunction)
 {
-	TMap<FMultiRigData, bool> ExpansionState;
+	UnregisterPendingRefresh();
+
+	if (InRefreshFunction)
+	{
+		PendingTreeViewRefreshHandle = RegisterActiveTimer( 0.f,
+			FWidgetActiveTimerDelegate::CreateLambda([this, RefreshFunction = MoveTemp(InRefreshFunction)](double, float)
+			{
+				RefreshFunction();
+				return EActiveTimerReturnType::Stop;
+			})
+		);		
+	}
+	else
+	{
+		PendingTreeViewRefreshHandle = RegisterActiveTimer( 0.f,
+			FWidgetActiveTimerDelegate::CreateLambda([this](double, float)
+			{
+				RefreshTreeView();
+				return EActiveTimerReturnType::Stop;
+			})
+		);
+	}
+}
 
-	if (bRebuildContent)
+void SMultiRigHierarchyTreeView::UnregisterPendingRefresh()
+{
+	if (PendingTreeViewRefreshHandle.IsValid())
 	{
-		for (TPair<FMultiRigData, TSharedPtr<FMultiRigTreeElement>> Pair : ElementMap)
+		if (TSharedPtr<FActiveTimerHandle> ActiveTimerHandle = PendingTreeViewRefreshHandle.Pin())
 		{
-			ExpansionState.FindOrAdd(Pair.Key) = IsItemExpanded(Pair.Value);
+			UnRegisterActiveTimer(ActiveTimerHandle.ToSharedRef());
 		}
+		PendingTreeViewRefreshHandle.Reset();
+	}
+}
 
-		// internally save expansion states before rebuilding the tree, so the states can be restored later
-		SaveAndClearSparseItemInfos();
-
-		RootElements.Reset();
-		ElementMap.Reset();
-		ParentMap.Reset();
+void SMultiRigHierarchyTreeView::RefreshTreeView()
+{
+	// store expansion state
+	TMap<FMultiRigData, bool> ExpansionState;
+	for (TPair<FMultiRigData, TSharedPtr<FMultiRigTreeElement>> Pair : ElementMap)
+	{
+		ExpansionState.FindOrAdd(Pair.Key) = IsItemExpanded(Pair.Value);
 	}
 
-	if (bRebuildContent)
+	// internally save expansion states before rebuilding the tree, so the states can be restored later
+	SaveAndClearSparseItemInfos();
+
+	RootElements.Reset();
+	ElementMap.Reset();
+	ParentMap.Reset();
+
+	// rebuild elements
+	static FMultiRigData EmptyParentData(nullptr, FRigElementKey());
+	for (const TWeakObjectPtr<UControlRig>& ControlRigPtr : ControlRigs)
 	{
-		for (const TWeakObjectPtr<UControlRig>& ControlRigPtr : ControlRigs)
+		if (UControlRig* ControlRig  = ControlRigPtr.Get())
 		{
-			if (UControlRig* ControlRig  = ControlRigPtr.Get())
-			{
-				FMultiRigData Empty(nullptr, FRigElementKey());
-				FMultiRigData CRData;
-				CRData.ControlRig = ControlRig; //leave key unset so it's valid
+			FMultiRigData CRData;
+			CRData.ControlRig = ControlRig; //leave key unset so it's valid
 
-				AddElement(CRData, Empty);
-				if (const URigHierarchy* Hierarchy = ControlRig->GetHierarchy())
-				{
-					Hierarchy->Traverse([&](FRigBaseElement* Element, bool& bContinue)
-						{
-							AddElement(ControlRig,Element);
-							bContinue = true;
-						});
+			// add root element
+			AddElement(CRData, EmptyParentData);
 
-					// expand all elements upon the initial construction of the tree
-					if (ExpansionState.Num() == 0)
-					{
-						for (TSharedPtr<FMultiRigTreeElement> RootElement : RootElements)
-						{
-							SetExpansionRecursive(RootElement, false, true);
-						}
-					}
-					else if (ExpansionState.Num() < ElementMap.Num())
-					{
-						for (const TPair<FMultiRigData, TSharedPtr<FMultiRigTreeElement>>& Element : ElementMap)
-						{
-							if (!ExpansionState.Contains(Element.Key))
-							{
-								SetItemExpansion(Element.Value, true);
-							}
-						}
-					}
-
-					for (const auto& Pair : ElementMap)
-					{
-						RestoreSparseItemInfos(Pair.Value);
-					}
+			// add children
+			if (const URigHierarchy* Hierarchy = ControlRig->GetHierarchy())
+			{
+				Hierarchy->Traverse([&](FRigBaseElement* Element, bool& bContinue)
+				{
+					AddElement(ControlRig, Element);
+					bContinue = true;
+				});
+			}
+		}
+	}
 
-				}
+	// expand all elements upon the initial construction of the tree
+	if (ExpansionState.Num() == 0)
+	{
+		for (TSharedPtr<FMultiRigTreeElement> RootElement : RootElements)
+		{
+			SetExpansionRecursive(RootElement, false, true);
+		}
+	}
+	else if (ExpansionState.Num() < ElementMap.Num())
+	{
+		for (const TPair<FMultiRigData, TSharedPtr<FMultiRigTreeElement>>& Element : ElementMap)
+		{
+			if (!ExpansionState.Contains(Element.Key))
+			{
+				SetItemExpansion(Element.Value, true);
 			}
 		}
 	}
-	else
+
+	// restore infos	
+	for (const auto& Pair : ElementMap)
 	{
-		if (RootElements.Num() > 0)
+		RestoreSparseItemInfos(Pair.Value);
+	}
+
+	RequestTreeRefresh();
+
+	// update selection
+	
+	TMap<UControlRig*, TArray<FRigElementKey>> RigAndSelection;
+	for (const TWeakObjectPtr<UControlRig>& ControlRigPtr : ControlRigs)
+	{
+		if (UControlRig* ControlRig = ControlRigPtr.Get())
 		{
-			// elements may be added at the end of the list after a spacer element
-			// we need to remove the spacer element and re-add it at the end
-			RootElements.RemoveAll([](TSharedPtr<FMultiRigTreeElement> InElement)
+			if (const URigHierarchy* Hierarchy = ControlRig->GetHierarchy())
+			{
+				TArray<FRigElementKey> Selection = Hierarchy->GetSelectedKeys();
+				if (!Selection.IsEmpty())
 				{
-					return (InElement.Get()->Data.ControlRig == nullptr && InElement.Get()->Data.Key == FRigElementKey());
-				});
+					RigAndSelection.Emplace(ControlRig, MoveTemp(Selection));
+				}
+			}
 		}
 	}
+	
+	ClearSelection();
 
-	RequestTreeRefresh();
+	for (const TPair<UControlRig*, TArray<FRigElementKey>>& RigAndKeys : RigAndSelection)
 	{
-		ClearSelection();
+		UControlRig* ControlRig = RigAndKeys.Key;
+		
+		// look for the root item referencing this rig
+		FMultiRigData ControlRigData;
+		ControlRigData.ControlRig = ControlRig;
+
+		const int32 RootElementIndex = RootElements.IndexOfByPredicate([ControlRigData](const TSharedPtr<FMultiRigTreeElement>& Root)
+		{
+			return Root->Data == ControlRigData;
+		});
 
-		for (const TWeakObjectPtr<UControlRig>& ControlRigPtr : ControlRigs)
+		if (RootElementIndex != INDEX_NONE)
 		{
-			if (UControlRig* ControlRig = ControlRigPtr.Get())
+			// look for the child item referencing this key
+			const TSharedPtr<FMultiRigTreeElement>& Root = RootElements[RootElementIndex];
+			for (const FRigElementKey& Key : RigAndKeys.Value)
 			{
-				if (const URigHierarchy* Hierarchy = ControlRig->GetHierarchy())
+				const FMultiRigData Data(ControlRig, Key);
+				const TSharedPtr<FMultiRigTreeElement> Found = FindElement(Data, Root);
+				if (Found.IsValid())
 				{
-
-					TArray<FRigElementKey> Selection = Hierarchy->GetSelectedKeys();
-					for (const FRigElementKey& Key : Selection)
+					SetItemSelection(Found, true, ESelectInfo::OnNavigation);
+				}
+			}
+		}
+		else 
+		{
+			// otherwise, iterate thru the elements if there's a filter as the root might have been skipped because of it  
+			const FRigTreeDisplaySettings& Settings = Delegates.GetDisplaySettings();
+			const FString FilteredString = Settings.FilterText.ToString();
+			if (!FilteredString.IsEmpty())
+			{
+				for (const FRigElementKey& Key : RigAndKeys.Value)
+				{
+					const FMultiRigData Data(ControlRig, Key);
+					TSharedPtr<FMultiRigTreeElement>* Found = ElementMap.Find(Data);
+					if (Found && Found->IsValid())
 					{
-						for (int32 RootIndex = 0; RootIndex < RootElements.Num(); ++RootIndex)
-						{
-							FMultiRigData Data(ControlRig, Key);
-							TSharedPtr<FMultiRigTreeElement> Found = FindElement(Data, RootElements[RootIndex]);
-							if (Found.IsValid())
-							{
-								SetItemSelection(Found, true, ESelectInfo::OnNavigation);
-							}
-						}
+						SetItemSelection(*Found, true, ESelectInfo::OnNavigation);
 					}
 				}
 			}
@@ -895,28 +954,18 @@ TArray<URigHierarchy*> SMultiRigHierarchyTreeView::GetHierarchy() const
 	return RigHierarchy;
 }
 
-void SMultiRigHierarchyTreeView::SetControlRigs(TArrayView < TWeakObjectPtr<UControlRig>>& InControlRigs)
+void SMultiRigHierarchyTreeView::SetControlRigs(const TArrayView<TWeakObjectPtr<UControlRig>>& InControlRigs)
 {
 	ControlRigs.SetNum(0);
-	TArray < TPair<UControlRig*, TArray<FName>>> SelectedControls;
 	for (TWeakObjectPtr<UControlRig>& ControlRig : InControlRigs)
 	{
 		if (ControlRig.IsValid())
 		{
 			ControlRigs.Add(ControlRig.Get());
-			SelectedControls.Add(TPair<UControlRig*, TArray<FName>>(ControlRig.Get(), ControlRig->CurrentControlSelection()));
 		}
 	}
 	
-	RefreshTreeView(true);
-	//reselect controls that will have gotten cleared by the refresh, this situation can happen on save
-	for (TPair<UControlRig*, TArray<FName>>& CRS : SelectedControls)
-	{
-		for (const FName& Name : CRS.Value)
-		{
-			CRS.Key->SelectControl(Name, true);
-		}
-	}
+	RequestTreeViewRefresh();
 }
 
 //////////////////////////////////////////////////////////////
@@ -978,7 +1027,7 @@ const FRigTreeDisplaySettings& SSearchableMultiRigHierarchyTreeView::GetDisplayS
 void SSearchableMultiRigHierarchyTreeView::OnFilterTextChanged(const FText& SearchText)
 {
 	FilterText = SearchText;
-	GetTreeView()->RefreshTreeView();
+	GetTreeView()->RequestTreeViewRefresh();
 }
 
 //////////////////////////////////////////////////////////////
@@ -1048,7 +1097,7 @@ void SControlRigOutliner::OnObjectsReplaced(const TMap<UObject*, UObject*>& OldT
 	}
 	if (bNewControlRig)
 	{
-		HierarchyTreeView->GetTreeView()->RefreshTreeView(true);
+		HierarchyTreeView->GetTreeView()->RequestTreeViewRefresh();
 	}
 }
 
@@ -1103,7 +1152,6 @@ void SControlRigOutliner::HandleControlSelected(UControlRig* Subject, FRigContro
 	}
 }
 
-
 void SControlRigOutliner::HandleSelectionChanged(TSharedPtr<FMultiRigTreeElement> Selection, ESelectInfo::Type SelectInfo)
 {
 	if (bIsChangingRigHierarchy)
diff --git a/Engine/Plugins/Animation/ControlRig/Source/ControlRigEditor/Private/EditMode/SControlRigOutliner.h b/Engine/Plugins/Animation/ControlRig/Source/ControlRigEditor/Private/EditMode/SControlRigOutliner.h
--- a/Engine/Plugins/Animation/ControlRig/Source/ControlRigEditor/Private/EditMode/SControlRigOutliner.h
+++ b/Engine/Plugins/Animation/ControlRig/Source/ControlRigEditor/Private/EditMode/SControlRigOutliner.h
@@ -141,7 +141,7 @@ public:
 	SLATE_END_ARGS()
 
 	void Construct(const FArguments& InArgs);
-	virtual ~SMultiRigHierarchyTreeView() {}
+	virtual ~SMultiRigHierarchyTreeView();
 
 	virtual FReply OnFocusReceived(const FGeometry& MyGeometry, const FFocusEvent& InFocusEvent) override
 	{
@@ -180,7 +180,6 @@ public:
 	bool AddElement(UControlRig* InControlRig, const FRigBaseElement* InElement);
 	bool ReparentElement(const FMultiRigData& InData, const FMultiRigData& InParentData);
 	bool RemoveElement(const FMultiRigData& InData);
-	void RefreshTreeView(bool bRebuildContent = true);
 	void SetExpansionRecursive(TSharedPtr<FMultiRigTreeElement> InElement, bool bTowardsParent, bool bShouldBeExpanded);
 	TSharedRef<ITableRow> MakeTableRowWidget(TSharedPtr<FMultiRigTreeElement> InItem, const TSharedRef<STableViewBase>& OwnerTable, bool bPinned);
 	void HandleGetChildrenForTree(TSharedPtr<FMultiRigTreeElement> InItem, TArray<TSharedPtr<FMultiRigTreeElement>>& OutChildren);
@@ -190,8 +189,13 @@ public:
 	FMultiRigTreeDelegates& GetTreeDelegates() { return Delegates; }
 
 	TArray<URigHierarchy*> GetHierarchy() const;
-	void SetControlRigs(TArrayView < TWeakObjectPtr<UControlRig>>& InControlRigs);
+	void SetControlRigs(const TArrayView<TWeakObjectPtr<UControlRig>>& InControlRigs);
 
+	/** Requests a tree view refresh. The actual refresh will be done lazily.
+	 * Note that the refresh function will be set to RefreshTreeView if no function is passed. 
+	 * InRefreshFunction can capture what's necessary if anything has to be done once the tree rebuilt. */
+	void RequestTreeViewRefresh(TFunction<void()>&& InRefreshFunction = nullptr);
+	
 private:
 
 	/** A temporary snapshot of the SparseItemInfos in STreeView, used during RefreshTreeView() */
@@ -211,6 +215,11 @@ private:
 	friend class SRigHierarchy;
 
 	TArray <TWeakObjectPtr<UControlRig>> ControlRigs;
+
+	/** Pending function to refresh the tree view. */
+	TWeakPtr<FActiveTimerHandle> PendingTreeViewRefreshHandle;
+	void UnregisterPendingRefresh();
+	void RefreshTreeView();
 };
 
 class SSearchableMultiRigHierarchyTreeView : public SCompoundWidget
