AnimLayer: Fix anim layer merges with overrids when the sections may have the control masked out.

diff --git a/Engine/Plugins/Animation/ControlRig/Source/ControlRigEditor/Private/Sequencer/AnimLayers/AnimLayers.cpp b/Engine/Plugins/Animation/ControlRig/Source/ControlRigEditor/Private/Sequencer/AnimLayers/AnimLayers.cpp
--- a/Engine/Plugins/Animation/ControlRig/Source/ControlRigEditor/Private/Sequencer/AnimLayers/AnimLayers.cpp
+++ b/Engine/Plugins/Animation/ControlRig/Source/ControlRigEditor/Private/Sequencer/AnimLayers/AnimLayers.cpp
@@ -1875,9 +1875,18 @@ static void MergeControlRigSections(UMovieSceneControlRigParameterSection* BaseS
 			if (FChannelMapInfo* pChannelIndex = BaseSection->ControlChannelMap.Find(ControlElement->GetFName()))
 			{
 				const int32 ChannelIndex = pChannelIndex->ChannelIndex;
-				TOptional<bool> bBaseSectionResetActive;
-				TOptional<bool> bSectionResetActive; 
 				
+				//if section is additive an control not there, skip it. we need to keep overrides since we may be blending over
+				//an additive and we want the full override value (with prpoer scale)
+				const bool bMaskKeyOut = (Section->GetBlendType().IsValid() && Section->GetBlendType() == EMovieSceneBlendType::Additive 
+					&& Section->GetControlNameMask(ControlElement->GetFName()) == false);
+				if (bMaskKeyOut)
+				{
+					continue;
+				}
+
+				TOptional<bool> bBaseSectionResetActive;
+				TOptional<bool> bSectionResetActive;
 				if (ShouldDisableSection(BaseSection, ControlElement))
 				{
 					bBaseSectionResetActive = BaseSection->IsActive();
@@ -1890,6 +1899,8 @@ static void MergeControlRigSections(UMovieSceneControlRigParameterSection* BaseS
 					Section->SetIsActive(false);
 				}
 
+				
+
 				switch (ControlElement->Settings.ControlType)
 				{
 
diff --git a/Engine/Source/Editor/MovieSceneTools/Public/MovieSceneToolHelpers.h b/Engine/Source/Editor/MovieSceneTools/Public/MovieSceneToolHelpers.h
--- a/Engine/Source/Editor/MovieSceneTools/Public/MovieSceneToolHelpers.h
+++ b/Engine/Source/Editor/MovieSceneTools/Public/MovieSceneToolHelpers.h
@@ -947,7 +947,8 @@ public:
 		{
 			MergeAlgorithm = FChannelMergeAlgorithm::Override;
 		}
-		else
+		else if (TopSection->GetBlendType().IsValid() && TopSection->GetBlendType().Get() == EMovieSceneBlendType::Absolute)
+
 		{
 			MergeAlgorithm = FChannelMergeAlgorithm::Average;
 		}
