AnimLayer: Open Anim Layer UI automaticaly if there are anim layers in it.

diff --git a/Engine/Plugins/Animation/ControlRig/Source/ControlRigEditor/Private/EditMode/ControlRigEditModeToolkit.cpp b/Engine/Plugins/Animation/ControlRig/Source/ControlRigEditor/Private/EditMode/ControlRigEditModeToolkit.cpp
--- a/Engine/Plugins/Animation/ControlRig/Source/ControlRigEditor/Private/EditMode/ControlRigEditModeToolkit.cpp
+++ b/Engine/Plugins/Animation/ControlRig/Source/ControlRigEditor/Private/EditMode/ControlRigEditModeToolkit.cpp
@@ -25,6 +25,7 @@
 #include "LevelEditor.h"
 #include "SLevelViewport.h"
 #include "Sequencer/AnimLayers/SAnimLayers.h"
+#include "Sequencer/AnimLayers/AnimLayers.h"
 
 #define LOCTEXT_NAMESPACE "FControlRigEditModeToolkit"
 
@@ -458,6 +459,14 @@ void FControlRigEditModeToolkit::InvokeUI()
 		{
 			TryInvokeToolkitUI(AnimLayerTabName);
 		}
+		else if (TSharedPtr<ISequencer> SequencerPtr = UAnimLayers::GetSequencerFromAsset())
+		{
+			//wasn't open but there are layers in open level sequence we try to open it
+			if (UAnimLayers* AnimLayers = UAnimLayers::GetAnimLayers(SequencerPtr.Get()))
+			{
+				TryInvokeToolkitUI(AnimLayerTabName);
+			}
+		}
 		if (bSnapperTabOpen)
 		{
 			TryInvokeToolkitUI(SnapperTabName);
