Control Rig: Support for per module event API

#rnx

diff --git a/Engine/Plugins/Animation/ControlRig/Source/ControlRig/Private/ModularRig.cpp b/Engine/Plugins/Animation/ControlRig/Source/ControlRig/Private/ModularRig.cpp
--- a/Engine/Plugins/Animation/ControlRig/Source/ControlRig/Private/ModularRig.cpp
+++ b/Engine/Plugins/Animation/ControlRig/Source/ControlRig/Private/ModularRig.cpp
@@ -562,17 +562,7 @@ void UModularRig::ExecuteQueue()
 				}
 
 				// Copy variable bindings
-				for (TPair<FName, FRigVMExternalVariable>& Pair : ExecutionElement.ModuleInstance->VariableBindings)
-				{
-					const FRigVMExternalVariable TargetVariable = ExecutionElement.ModuleInstance->GetRig()->GetPublicVariableByName(Pair.Key);
-					if(ensure(TargetVariable.Property))
-					{
-						if (RigVMTypeUtils::AreCompatible(Pair.Value.Property, TargetVariable.Property))
-						{
-							Pair.Value.Property->CopyCompleteValue(TargetVariable.Memory, Pair.Value.Memory);
-						}
-					}
-				}
+				UpdateModuleVariables(ExecutionElement.ModuleInstance);
 			
 				ModuleRig->Execute_Internal(ExecutionElement.EventName);
 				ExecutionElement.bExecuted = true;
@@ -714,6 +704,20 @@ void UModularRig::UpdateSupportedEvents() const
 	});
 }
 
+TArray<FString> UModularRig::GetModulePaths() const
+{
+	TArray<FString> Paths;
+	Paths.Reserve(Modules.Num());
+	ForEachModule([&Paths](const FRigModuleInstance* Module) -> bool
+	{
+		// don't need to use AddUnique since module paths are
+		// guaranteed to be unique already.
+		Paths.Add(Module->GetPath());
+		return true;
+	});
+	return Paths;
+}
+
 void UModularRig::RunPostConstructionEvent()
 {
 	RecomputeShortestDisplayPathCache();
@@ -853,6 +857,24 @@ bool UModularRig::SetModuleVariableBindings(const FString& InModulePath, const T
 	return false;
 }
 
+void UModularRig::UpdateModuleVariables(const FRigModuleInstance* InModule)
+{
+	if(const UControlRig* ModuleRig = InModule->GetRig())
+	{
+		for (const TPair<FName, FRigVMExternalVariable>& Pair : InModule->VariableBindings)
+		{
+			const FRigVMExternalVariable TargetVariable = ModuleRig->GetPublicVariableByName(Pair.Key);
+			if(ensure(TargetVariable.Property))
+			{
+				if (RigVMTypeUtils::AreCompatible(Pair.Value.Property, TargetVariable.Property))
+				{
+					Pair.Value.Property->CopyCompleteValue(TargetVariable.Memory, Pair.Value.Memory);
+				}
+			}
+		}
+	}
+}
+
 void UModularRig::DiscardModuleRig(UControlRig* InControlRig)
 {
 	if(InControlRig)
@@ -953,6 +975,11 @@ void UModularRig::RecomputeShortestDisplayPathCache() const
 }
 
 const FRigModuleInstance* UModularRig::FindModule(const FString& InPath) const
+{
+	return const_cast<UModularRig*>(this)->FindModule(InPath);
+}
+
+FRigModuleInstance* UModularRig::FindModule(const FString& InPath)
 {
 	if(InPath.EndsWith(NamespaceSeparator))
 	{
@@ -966,8 +993,13 @@ const FRigModuleInstance* UModularRig::FindModule(const FString& InPath) const
 
 const FRigModuleInstance* UModularRig::FindModule(const UControlRig* InModuleInstance) const
 {
-	const FRigModuleInstance* FoundModule = nullptr;
-	ForEachModule([InModuleInstance, &FoundModule](const FRigModuleInstance* Module) -> bool
+	return const_cast<UModularRig*>(this)->FindModule(InModuleInstance);
+}
+
+FRigModuleInstance* UModularRig::FindModule(const UControlRig* InModuleInstance)
+{
+	FRigModuleInstance* FoundModule = nullptr;
+	ForEachModule([InModuleInstance, &FoundModule](FRigModuleInstance* Module) -> bool
 	{
 		if(Module->GetRig() == InModuleInstance)
 		{
@@ -982,6 +1014,11 @@ const FRigModuleInstance* UModularRig::FindModule(const UControlRig* InModuleIns
 }
 
 const FRigModuleInstance* UModularRig::FindModule(const FRigBaseElement* InElement) const
+{
+	return const_cast<UModularRig*>(this)->FindModule(InElement);
+}
+
+FRigModuleInstance* UModularRig::FindModule(const FRigBaseElement* InElement)
 {
 	if(InElement)
 	{
@@ -991,6 +1028,11 @@ const FRigModuleInstance* UModularRig::FindModule(const FRigBaseElement* InEleme
 }
 
 const FRigModuleInstance* UModularRig::FindModule(const FRigElementKey& InElementKey) const
+{
+	return const_cast<UModularRig*>(this)->FindModule(InElementKey);
+}
+
+FRigModuleInstance* UModularRig::FindModule(const FRigElementKey& InElementKey)
 {
 	if(const URigHierarchy* Hierarchy = GetHierarchy())
 	{
@@ -1003,9 +1045,35 @@ const FRigModuleInstance* UModularRig::FindModule(const FRigElementKey& InElemen
 	return nullptr;
 }
 
-FString UModularRig::GetParentPath(const FString& InPath) const
+const UControlRig* UModularRig::GetModuleRig(FString InModulePath) const
 {
-	if (const FRigModuleInstance* Element = FindModule(InPath))
+	return const_cast<UModularRig*>(this)->GetModuleRig(InModulePath);
+}
+
+UControlRig* UModularRig::GetModuleRig(FString InModulePath)
+{
+	if(const FRigModuleInstance* Module = FindModule(InModulePath))
+	{
+		UpdateModuleVariables(Module);
+		if(UControlRig* ModuleRig = Module->GetRig())
+		{
+			return ModuleRig;
+		}
+		UE_LOG(LogControlRig, Error, TEXT("Module '%s' doesn't contain a rig instance."), *InModulePath);
+		return nullptr;
+	}
+	UE_LOG(LogControlRig, Error, TEXT("Module '%s' doesn't exist."), *InModulePath);
+	return nullptr;
+}
+
+FString UModularRig::GetParentPathForBP(FString InModulePath) const
+{
+	return GetParentPath(InModulePath);
+}
+
+FString UModularRig::GetParentPath(const FString& InModulePath) const
+{
+	if (const FRigModuleInstance* Element = FindModule(InModulePath))
 	{
 		return Element->ParentPath;
 	}
@@ -1071,7 +1139,7 @@ void UModularRig::ExecuteConnectorEvent(const FRigElementKey& InConnector, const
 		return;
 	}
 
-	TArray<FRigElementResolveResult> Candidates = InOutCandidates;
+	const TArray<FRigElementResolveResult> Candidates = InOutCandidates;
 	
 	FControlRigExecuteContext& PublicContext = GetRigVMExtendedExecuteContext().GetPublicDataSafe<FControlRigExecuteContext>();
 
@@ -1081,13 +1149,112 @@ void UModularRig::ExecuteConnectorEvent(const FRigElementKey& InConnector, const
 	TGuardValue<TArray<FRigElementResolveResult>> CandidatesGuard(PublicContext.UnitContext.ConnectionResolve.Matches, Candidates);
 	TGuardValue<TArray<FRigElementResolveResult>> MatchesGuard(PublicContext.UnitContext.ConnectionResolve.Excluded, {});
 	
-	FRigModuleExecutionElement ExecutionElement(Module, FRigUnit_ConnectorExecution::EventName);
-	TGuardValue<TArray<FRigModuleExecutionElement>> ExecutionGuard(ExecutionQueue, {ExecutionElement});
-	TGuardValue<int32> ExecutionFrontGuard(ExecutionQueueFront, 0);
 	TGuardValue<FRigElementKeyRedirector> RedirectorGuard(ElementKeyRedirector, *InRedirector);
+	ExecuteEventOnModule(FRigUnit_ConnectorExecution::EventName, Module);
 	ExecuteQueue();
 
 	InOutCandidates = PublicContext.UnitContext.ConnectionResolve.Matches;
 }
 
+TArray<FName> UModularRig::GetEventsForAllModules() const
+{
+	TArray<FName> Events;
+	ForEachModule([&Events](const FRigModuleInstance* Module) -> bool
+	{
+		if(const UControlRig* ModuleRig = Module->GetRig())
+		{
+			const TArray<FName> SupportedEvents = ModuleRig->GetSupportedEvents();
+			for(const FName& SupportedEvent : SupportedEvents)
+			{
+				Events.AddUnique(SupportedEvent);
+			}
+		}
+		return true;
+	});
+	return Events;
+}
+
+TArray<FName> UModularRig::GetEventsForModule(FString InModulePath) const
+{
+	if(const UControlRig* ModuleRig = GetModuleRig(InModulePath))
+	{
+		return ModuleRig->GetSupportedEvents();
+	}
+	return TArray<FName>();
+}
+
+TArray<FString> UModularRig::ExecuteEventOnAllModules(FName InEvent)
+{
+	TArray<FRigModuleExecutionElement> QueueForEvent;
+
+	ForEachModule([&QueueForEvent, InEvent](FRigModuleInstance* Module) -> bool
+	{
+		if(const UControlRig* ModuleRig = Module->GetRig())
+		{
+			if(ModuleRig->SupportsEvent(InEvent))
+			{
+				QueueForEvent.Emplace(Module, InEvent);
+			}
+		}
+		return true;
+	});
+
+	TArray<FString> ModulesWhichRanEvent;
+	if(QueueForEvent.IsEmpty())
+	{
+		UE_LOG(LogControlRig, Error, TEXT("Event '%s' is not supported by any module on this modular rig."), *InEvent.ToString());
+		return ModulesWhichRanEvent;
+	}
+	
+	TGuardValue<TArray<FRigModuleExecutionElement>> ExecutionGuard(ExecutionQueue, QueueForEvent);
+	TGuardValue<int32> ExecutionFrontGuard(ExecutionQueueFront, 0);
+	ExecuteQueue();
+
+	for(const FRigModuleExecutionElement& ExecutionElement : ExecutionQueue)
+	{
+		if(ExecutionElement.bExecuted)
+		{
+			ModulesWhichRanEvent.Add(ExecutionElement.ModulePath);
+		}
+		else
+		{
+			UE_LOG(LogControlRig, Error, TEXT("Module '%s' did not run event '%s' successfully."), *ExecutionElement.ModulePath, *InEvent.ToString());
+		}
+	}
+		
+	return ModulesWhichRanEvent;
+}
+
+bool UModularRig::ExecuteEventOnModuleForBP(FName InEvent, FString InModulePath)
+{
+	if(FRigModuleInstance* Module = FindModule(InModulePath))
+	{
+		if(const UControlRig* ModuleRig = Module->GetRig())
+		{
+			if(ModuleRig->SupportsEvent(InEvent))
+			{
+				return ExecuteEventOnModule(InEvent, Module); 
+			}
+			UE_LOG(LogControlRig, Error, TEXT("Module '%s' doesn't support the event '%s'."), *InModulePath, *InEvent.ToString());
+			return false;
+		}
+		UE_LOG(LogControlRig, Error, TEXT("Module '%s' doesn't contain a rig instance."), *InModulePath);
+		return false;
+	}
+	UE_LOG(LogControlRig, Error, TEXT("Module '%s' doesn't exist."), *InModulePath);
+	return false;
+}
+
+bool UModularRig::ExecuteEventOnModule(const FName& InEvent, FRigModuleInstance* InModule)
+{
+	check(InModule);
+
+	TArray<FRigModuleExecutionElement> QueueForEvent;
+	QueueForEvent.Emplace(InModule, InEvent);
+	TGuardValue<TArray<FRigModuleExecutionElement>> ExecutionGuard(ExecutionQueue, QueueForEvent);
+	TGuardValue<int32> ExecutionFrontGuard(ExecutionQueueFront, 0);
+	ExecuteQueue();
+	return ExecutionQueue[0].bExecuted;
+}
+
 #undef LOCTEXT_NAMESPACE
diff --git a/Engine/Plugins/Animation/ControlRig/Source/ControlRig/Private/ModularRigController.cpp b/Engine/Plugins/Animation/ControlRig/Source/ControlRig/Private/ModularRigController.cpp
--- a/Engine/Plugins/Animation/ControlRig/Source/ControlRig/Private/ModularRigController.cpp
+++ b/Engine/Plugins/Animation/ControlRig/Source/ControlRig/Private/ModularRigController.cpp
@@ -108,6 +108,42 @@ const FRigModuleReference* UModularRigController::FindModule(const FString& InPa
 	return const_cast<UModularRigController*>(this)->FindModule(InPath);
 }
 
+FRigModuleReference UModularRigController::GetModuleReference(FString InModulePath) const
+{
+	if(const FRigModuleReference* Module = FindModule(InModulePath))
+	{
+		return *Module;
+	}
+	return FRigModuleReference();
+}
+
+TArray<FRigElementKey> UModularRigController::GetConnectorsForModule(FString InModulePath) const
+{
+	TArray<FRigElementKey> ConnectorsForModule;
+
+	if(const FRigModuleReference* Module = FindModule(InModulePath))
+	{
+		UControlRig* RigCDO = Module->Class->GetDefaultObject<UControlRig>();
+		if (!RigCDO)
+		{
+			UE_LOG(LogControlRig, Error, TEXT("Invalid rig module class %s"), *Module->Class->GetPathName());
+			return ConnectorsForModule;
+		}
+
+		const TArray<FRigModuleConnector>& ExposedConnectors = RigCDO->GetRigModuleSettings().ExposedConnectors;
+		for(const FRigModuleConnector& ExposedConnector : ExposedConnectors)
+		{
+			ConnectorsForModule.Emplace(*URigHierarchy::JoinNameSpace(InModulePath, ExposedConnector.Name), ERigElementType::Connector);
+		}
+	}
+	else
+	{
+		UE_LOG(LogControlRig, Error, TEXT("Could not find module %s"), *InModulePath);
+	}
+
+	return ConnectorsForModule;
+}
+
 bool UModularRigController::CanConnectConnectorToElement(const FRigElementKey& InConnectorKey, const FRigElementKey& InTargetKey, FText& OutErrorMessage)
 {
 	FString ConnectorModulePath, ConnectorName;
@@ -183,7 +219,7 @@ bool UModularRigController::CanConnectConnectorToElement(const FRigElementKey& I
 
 	// Make sure the connection is valid
 	{
-		UModularRig* ModularRig = Cast<UModularRig>(Blueprint->GetObjectBeingDebugged());
+		UModularRig* ModularRig = GetDebuggedModularRig();
 		if (!ModularRig)
 		{
 			OutErrorMessage = FText::FromString(FString::Printf(TEXT("Could not find debugged modular rig in %s"), *Blueprint->GetPathName()));
@@ -243,7 +279,7 @@ bool UModularRigController::ConnectConnectorToElement(const FRigElementKey& InCo
 
 #if WITH_EDITOR
 	FName TargetModulePathName = NAME_None;
-	UModularRig* ModularRig = Cast<UModularRig>(Blueprint->GetObjectBeingDebugged());
+	UModularRig* ModularRig = GetDebuggedModularRig();
 	if(ModularRig)
 	{
 		if (URigHierarchy* Hierarchy = ModularRig->GetHierarchy())
@@ -468,7 +504,7 @@ TArray<FRigElementKey> UModularRigController::DisconnectCyclicConnectors(bool bS
 	const UBlueprint* Blueprint = Cast<UBlueprint>(GetOuter());
 	check(Blueprint);
 
-	const UModularRig* ModularRig = Cast<UModularRig>(Blueprint->GetObjectBeingDebugged());
+	const UModularRig* ModularRig = GetDebuggedModularRig();
 	if (!ModularRig)
 	{
 		return DisconnectedConnectors;
@@ -528,7 +564,7 @@ bool UModularRigController::AutoConnectSecondaryConnectors(const TArray<FRigElem
 		return false;
 	}
 
-	const UModularRig* ModularRig = Cast<UModularRig>(Blueprint->GetObjectBeingDebugged());
+	const UModularRig* ModularRig = GetDebuggedModularRig();
 	if (!ModularRig)
 	{
 		UE_LOG(LogControlRig, Error, TEXT("Could not find debugged modular rig in %s"), *Blueprint->GetPathName());
@@ -679,7 +715,7 @@ bool UModularRigController::AutoConnectModules(const TArray<FString>& InModulePa
 		return false;
 	}
 
-	const UModularRig* ModularRig = Cast<UModularRig>(Blueprint->GetObjectBeingDebugged());
+	const UModularRig* ModularRig = GetDebuggedModularRig();
 	if (!ModularRig)
 	{
 		UE_LOG(LogControlRig, Error, TEXT("Could not find debugged modular rig in %s"), *Blueprint->GetPathName());
@@ -1753,6 +1789,17 @@ bool UModularRigController::SetModuleSelection(const TArray<FString>& InModulePa
 	return bResult;
 }
 
+TArray<FString> UModularRigController::GetAllModules() const
+{
+	TArray<FString> Paths;
+	Model->ForEachModule([&Paths](const FRigModuleReference* Module) -> bool
+	{
+		Paths.Add(Module->GetPath());
+		return true;
+	});
+	return Paths;
+}
+
 TArray<FString> UModularRigController::GetSelectedModules() const
 {
 	return Model->SelectedModulePaths;
@@ -2103,6 +2150,32 @@ void UModularRigController::UpdateShortNames()
 	}
 }
 
+UModularRig* UModularRigController::GetDebuggedModularRig()
+{
+#if WITH_EDITOR
+	if(const UBlueprint* Blueprint = Cast<UBlueprint>(GetOuter()))
+	{
+		if(UModularRig* ModularRig = Cast<UModularRig>(Blueprint->GetObjectBeingDebugged()))
+		{
+			return ModularRig;
+		}
+		if(const UClass* Class = Blueprint->GeneratedClass)
+		{
+			if(UModularRig* CDO = Cast<UModularRig>(Class->GetDefaultObject(true)))
+			{
+				return CDO;
+			}
+		}
+	}
+#endif
+	return nullptr;
+}
+
+const UModularRig* UModularRigController::GetDebuggedModularRig() const
+{
+	return const_cast<UModularRigController*>(this)->GetDebuggedModularRig();
+}
+
 FModularRigControllerCompileBracketScope::FModularRigControllerCompileBracketScope(UModularRigController* InController)
 	: Controller(InController), bSuspendNotifications(InController->bSuspendNotifications)
 {
diff --git a/Engine/Plugins/Animation/ControlRig/Source/ControlRig/Public/ModularRig.h b/Engine/Plugins/Animation/ControlRig/Source/ControlRig/Public/ModularRig.h
--- a/Engine/Plugins/Animation/ControlRig/Source/ControlRig/Public/ModularRig.h
+++ b/Engine/Plugins/Animation/ControlRig/Source/ControlRig/Public/ModularRig.h
@@ -181,11 +181,37 @@ public:
 	void UpdateCachedChildren();
 	void UpdateSupportedEvents() const;
 
+	/**
+	 * @return Returns all of the module paths within this rig
+	 */
+	UFUNCTION(BlueprintPure, Category=ModularRig)
+	TArray<FString> GetModulePaths() const;
+	
 	const FRigModuleInstance* FindModule(const FString& InPath) const;
 	const FRigModuleInstance* FindModule(const UControlRig* InModuleInstance) const;
 	const FRigModuleInstance* FindModule(const FRigBaseElement* InElement) const;
 	const FRigModuleInstance* FindModule(const FRigElementKey& InElementKey) const;
-	FString GetParentPath(const FString& InPath) const;
+
+	FRigModuleInstance* FindModule(const FString& InPath);
+	FRigModuleInstance* FindModule(const UControlRig* InModuleInstance);
+	FRigModuleInstance* FindModule(const FRigBaseElement* InElement);
+	FRigModuleInstance* FindModule(const FRigElementKey& InElementKey);
+
+	/**
+	 * @param InModulePath The path of the module to run the event on
+	 * @return Returns the rig instance for a given module path
+	 */
+	UFUNCTION(BlueprintCallable, Category=ModularRig)
+	UControlRig* GetModuleRig(FString InModulePath);
+	const UControlRig* GetModuleRig(FString InModulePath) const;
+
+	/**
+	 * @param InModulePath The path of the module to receive the parent path for
+	 * @return Returns the parent path for a given module path (or an empty string)
+	 */
+	UFUNCTION(BlueprintPure, Category=ModularRig, meta=(ScriptName="GetParentPath", DisplayName="Get Parent Path"))
+	FString GetParentPathForBP(FString InModulePath) const;
+	FString GetParentPath(const FString& InModulePath) const;
 	FString GetShortestDisplayPathForElement(const FRigElementKey& InElementKey, bool bAlwaysShowNameSpace) const;
 
 	void ForEachModule(TFunctionRef<bool(FRigModuleInstance*)> PerModuleFunction);
@@ -193,6 +219,35 @@ public:
 
 	void ExecuteConnectorEvent(const FRigElementKey& InConnector, const FRigModuleInstance* InModuleInstance, const FRigElementKeyRedirector* InRedirector, TArray<FRigElementResolveResult>& InOutCandidates);
 
+	/**
+	 * @return Returns all of the events supported by the modules in this modular rig
+	 */
+	UFUNCTION(BlueprintPure, Category=ModularRig)
+	TArray<FName> GetEventsForAllModules() const;
+
+	/**
+	 * @param InModulePath The path of the module to receive the events for
+	 * @return Returns the names of all supported events for a given module path
+	 */
+	UFUNCTION(BlueprintPure, Category=ModularRig)
+	TArray<FName> GetEventsForModule(FString InModulePath) const;
+
+	/**
+	 * @param InEvent The name of the event to run
+	 * @return Returns the paths of all modules which ran the event successfully
+	 */
+	UFUNCTION(BlueprintCallable, Category=ModularRig)
+	TArray<FString> ExecuteEventOnAllModules(FName InEvent);
+
+	/**
+	 * @param InEvent The name of the event to run
+	 * @param InModulePath The path of the module to run the event on
+	 * @return Returns true if the event was run successfully
+	 */
+	UFUNCTION(BlueprintCallable, Category=ModularRig, meta=(ScriptName="ExecuteEventOnModule", DisplayName="Execute Event On Module"))
+	bool ExecuteEventOnModuleForBP(FName InEvent, FString InModulePath);
+	bool ExecuteEventOnModule(const FName& InEvent, FRigModuleInstance* InModule);
+
 	/**
 	 * Returns a handle to an existing module
 	 * @param InPath The path of the module to retrieve a handle for.
@@ -221,6 +276,9 @@ private:
 	/** Updates the module's variable bindings */
 	bool SetModuleVariableBindings(const FString& InModulePath, const TMap<FName, FString>& InVariableBindings);
 
+	/** Updates the module's variable values based on its binding */
+	void UpdateModuleVariables(const FRigModuleInstance* InModule);
+
 	/** Destroys / discards a module instance rig */
 	static void DiscardModuleRig(UControlRig* InControlRig);
 
diff --git a/Engine/Plugins/Animation/ControlRig/Source/ControlRig/Public/ModularRigController.h b/Engine/Plugins/Animation/ControlRig/Source/ControlRig/Public/ModularRigController.h
--- a/Engine/Plugins/Animation/ControlRig/Source/ControlRig/Public/ModularRigController.h
+++ b/Engine/Plugins/Animation/ControlRig/Source/ControlRig/Public/ModularRigController.h
@@ -6,6 +6,7 @@
 #include "ModularRigController.generated.h"
 
 struct FRigModuleReference;
+class UModularRig;
 
 DECLARE_MULTICAST_DELEGATE_TwoParams(FModularRigModifiedEvent, EModularRigNotification /* type */, const FRigModuleReference* /* element */);
 
@@ -89,6 +90,9 @@ class CONTROLRIG_API UModularRigController : public UObject
 	UFUNCTION(BlueprintCallable, Category = "ControlRig | Modules")
 	bool SetModuleSelection(const TArray<FString>& InModulePaths);
 
+	UFUNCTION(BlueprintPure, Category = "ControlRig | Modules")
+	TArray<FString> GetAllModules() const;
+
 	UFUNCTION(BlueprintPure, Category = "ControlRig | Modules")
 	TArray<FString> GetSelectedModules() const;
 
@@ -106,16 +110,33 @@ class CONTROLRIG_API UModularRigController : public UObject
 	const FRigModuleReference* FindModule(const FString& InPath) const;
 	FModularRigModifiedEvent& OnModified() { return ModifiedEvent; }
 
+	/**
+	 * @param InModulePath The path of the module reference to return
+	 * @return Returns all of the module paths within this rig
+	 */
+	UFUNCTION(BlueprintPure, Category = "ControlRig | Modules")
+	FRigModuleReference GetModuleReference(FString InModulePath) const;
+
+	/**
+	 * @param InModulePath The path of the module reference to return
+	 * @return Returns all of associated connectors for a given module path
+	 */
+	UFUNCTION(BlueprintPure, Category = "ControlRig | Modules")
+	TArray<FRigElementKey> GetConnectorsForModule(FString InModulePath) const;
+
 private:
 
 	void SetModel(FModularRigModel* InModel) { Model = InModel; }
 	void Notify(const EModularRigNotification& InNotification, const FRigModuleReference* InElement);
 	void UpdateShortNames();
+	UModularRig* GetDebuggedModularRig();
+	const UModularRig* GetDebuggedModularRig() const;
 
 	bool bAutomaticReparenting;
-	
+
 	friend struct FModularRigModel;
 	friend class FModularRigControllerCompileBracketScope;
+	friend class UControlRigBlueprint;
 };
 
 class CONTROLRIG_API FModularRigControllerCompileBracketScope
diff --git a/Engine/Plugins/Animation/ControlRig/Source/ControlRigDeveloper/Private/ControlRigBlueprint.cpp b/Engine/Plugins/Animation/ControlRig/Source/ControlRigDeveloper/Private/ControlRigBlueprint.cpp
--- a/Engine/Plugins/Animation/ControlRig/Source/ControlRigDeveloper/Private/ControlRigBlueprint.cpp
+++ b/Engine/Plugins/Animation/ControlRig/Source/ControlRigDeveloper/Private/ControlRigBlueprint.cpp
@@ -1878,6 +1878,14 @@ void UControlRigBlueprint::RecompileModularRig()
 		}
 	}
 	UpdateModularDependencyDelegates();
+
+#if WITH_EDITOR
+	if(GetObjectBeingDebugged() == nullptr)
+	{
+		SetObjectBeingDebugged(CreateControlRig());
+	}
+#endif
+
 	OnModularRigCompiled().Broadcast(this);
 }
 
diff --git a/Engine/Plugins/Animation/ControlRig/Source/ControlRigDeveloper/Public/ControlRigBlueprint.h b/Engine/Plugins/Animation/ControlRig/Source/ControlRigDeveloper/Public/ControlRigBlueprint.h
--- a/Engine/Plugins/Animation/ControlRig/Source/ControlRigDeveloper/Public/ControlRigBlueprint.h
+++ b/Engine/Plugins/Animation/ControlRig/Source/ControlRigDeveloper/Public/ControlRigBlueprint.h
@@ -211,14 +211,15 @@ public:
 	static EControlRigType GetRigType(const FAssetData& InAsset);
 	static TArray<FSoftObjectPath> GetReferencesToRigModule(const FAssetData& InModuleAsset);
 
+	UFUNCTION(BlueprintCallable, Category = "Control Rig Blueprint")
+	void UpdateExposedModuleConnectors() const;
+	
 protected:
 
 	TArray<FModuleReferenceData> GetModuleReferenceData() const;
 
 	FOnRigTypeChanged OnRigTypeChangedDelegate;
 	
-	void UpdateExposedModuleConnectors() const;
-
 	bool ResolveConnector(const FRigElementKey& DraggedKey, const FRigElementKey& TargetKey, bool bSetupUndoRedo = true);
 
 	void UpdateConnectionMapFromModel();
