Control Rig: Fix CIS breakage in ModularRig.cpp

#rnx

diff --git a/Engine/Plugins/Animation/ControlRig/Source/ControlRig/Private/ModularRig.cpp b/Engine/Plugins/Animation/ControlRig/Source/ControlRig/Private/ModularRig.cpp
--- a/Engine/Plugins/Animation/ControlRig/Source/ControlRig/Private/ModularRig.cpp
+++ b/Engine/Plugins/Animation/ControlRig/Source/ControlRig/Private/ModularRig.cpp
@@ -1163,8 +1163,8 @@ TArray<FName> UModularRig::GetEventsForAllModules() const
 	{
 		if(const UControlRig* ModuleRig = Module->GetRig())
 		{
-			const TArray<FName> SupportedEvents = ModuleRig->GetSupportedEvents();
-			for(const FName& SupportedEvent : SupportedEvents)
+			const TArray<FName> ModuleSupportedEvents = ModuleRig->GetSupportedEvents();
+			for(const FName& SupportedEvent : ModuleSupportedEvents)
 			{
 				Events.AddUnique(SupportedEvent);
 			}
