Add the Parent coordinate system to the new viewport toolbar when entering Animation Mode in the Level Editor

* Add the RelativeCoordinateSystem_Parent command next to _World and _Local.
* Bind an icon for the new command.
* When entering Animation Mode, extend the new level editor viewport toolbar's "Transform > Coordinate System" submenu with the Parent coorinate system. Remove it again when exiting Animation Mode.

diff --git a/Engine/Plugins/Animation/ControlRig/Source/ControlRigEditor/Private/Editor/ControlRigViewportToolbarExtensions.cpp b/Engine/Plugins/Animation/ControlRig/Source/ControlRigEditor/Private/Editor/ControlRigViewportToolbarExtensions.cpp
--- a/Engine/Plugins/Animation/ControlRig/Source/ControlRigEditor/Private/Editor/ControlRigViewportToolbarExtensions.cpp
+++ b/Engine/Plugins/Animation/ControlRig/Source/ControlRigEditor/Private/Editor/ControlRigViewportToolbarExtensions.cpp
@@ -1,9 +1,10 @@
-﻿// Copyright Epic Games, Inc. All Rights Reserved.
+// Copyright Epic Games, Inc. All Rights Reserved.
 
 #include "Editor/ControlRigViewportToolbarExtensions.h"
 
 #include "ControlRigEditorCommands.h"
 #include "DetailCategoryBuilder.h"
+#include "EditorViewportCommands.h"
 #include "EditMode/ControlRigEditModeSettings.h"
 #include "ToolMenus.h"
 #include "Tools/MotionTrailOptions.h"
@@ -141,6 +142,20 @@ void UE::ControlRig::PopulateControlRigViewportToolbarTransformSubmenu(const FNa
 			SelectionSection.AddEntry(Entry);
 		}
 	}
+
+	// Coordinate space submenu.
+	{
+		const FName SubmenuName = UToolMenus::JoinMenuPaths(InMenuName, "CoordinateSystem");
+
+		UToolMenu* const Submenu = UToolMenus::Get()->ExtendMenu(SubmenuName);
+		if (Submenu)
+		{
+			FToolMenuSection& UnnamedSection = Submenu->FindOrAddSection(NAME_None);
+			UnnamedSection.AddEntry(
+				FToolMenuEntry::InitMenuEntry(FEditorViewportCommands::Get().RelativeCoordinateSystem_Parent)
+			);
+		}
+	}
 }
 
 void UE::ControlRig::PopulateControlRigViewportToolbarShowSubmenu(
diff --git a/Engine/Source/Editor/EditorStyle/Private/StarshipStyle.cpp b/Engine/Source/Editor/EditorStyle/Private/StarshipStyle.cpp
--- a/Engine/Source/Editor/EditorStyle/Private/StarshipStyle.cpp
+++ b/Engine/Source/Editor/EditorStyle/Private/StarshipStyle.cpp
@@ -4518,6 +4518,8 @@ void FStarshipEditorStyle::FStyle::SetupLevelEditorStyle()
 		Set( "EditorViewport.RelativeCoordinateSystem_World", new IMAGE_BRUSH_SVG( "Starship/EditorViewport/globe", Icon16x16 ) );
 		Set("EditorViewport.RelativeCoordinateSystem_Local",
 			new IMAGE_BRUSH_SVG("Starship/Common/transform-local", Icon16x16));
+		Set("EditorViewport.RelativeCoordinateSystem_Parent",
+			new CORE_IMAGE_BRUSH_SVG("Starship/Common/ParentHierarchy", Icon16x16));
 
 		Set( "EditorViewport.CamSpeedSetting", new IMAGE_BRUSH_SVG( "Starship/EditorViewport/camera", Icon16x16) );
 		
diff --git a/Engine/Source/Editor/UnrealEd/Private/EditorViewportCommands.cpp b/Engine/Source/Editor/UnrealEd/Private/EditorViewportCommands.cpp
--- a/Engine/Source/Editor/UnrealEd/Private/EditorViewportCommands.cpp
+++ b/Engine/Source/Editor/UnrealEd/Private/EditorViewportCommands.cpp
@@ -167,6 +167,13 @@ void FEditorViewportCommands::RegisterCommands()
 		EUserInterfaceActionType::RadioButton,
 		FInputChord()
 	);
+	UI_COMMAND(
+		RelativeCoordinateSystem_Parent,
+		"Parent Space",
+		"Move and rotate objects relative to the parent's axes",
+		EUserInterfaceActionType::RadioButton,
+		FInputChord()
+	);
 
 #if PLATFORM_MAC
 	UI_COMMAND( CycleTransformGizmoCoordSystem, "Cycle Transform Coordinate System", "Cycles the transform gizmo coordinate systems between world and local (object) space", EUserInterfaceActionType::Button, FInputChord(EModifierKey::Command, EKeys::Tilde));
diff --git a/Engine/Source/Editor/UnrealEd/Private/SEditorViewport.cpp b/Engine/Source/Editor/UnrealEd/Private/SEditorViewport.cpp
--- a/Engine/Source/Editor/UnrealEd/Private/SEditorViewport.cpp
+++ b/Engine/Source/Editor/UnrealEd/Private/SEditorViewport.cpp
@@ -382,6 +382,13 @@ void SEditorViewport::BindCommands()
 		FIsActionChecked::CreateSP( this, &SEditorViewport::IsCoordSystemActive, COORD_Local )
 		);
 
+	CommandListRef.MapAction(
+		Commands.RelativeCoordinateSystem_Parent,
+		FExecuteAction::CreateSP(ClientRef, &FEditorViewportClient::SetWidgetCoordSystemSpace, COORD_Parent),
+		FCanExecuteAction(),
+		FIsActionChecked::CreateSP(this, &SEditorViewport::IsCoordSystemActive, COORD_Parent)
+	);
+
 	CommandListRef.MapAction(
 		Commands.CycleTransformGizmos,
 		FExecuteAction::CreateSP( this, &SEditorViewport::OnCycleWidgetMode ),
diff --git a/Engine/Source/Editor/UnrealEd/Public/EditorViewportCommands.h b/Engine/Source/Editor/UnrealEd/Public/EditorViewportCommands.h
--- a/Engine/Source/Editor/UnrealEd/Public/EditorViewportCommands.h
+++ b/Engine/Source/Editor/UnrealEd/Public/EditorViewportCommands.h
@@ -212,6 +212,9 @@ public:
 	/** Local relative coordinate system */
 	TSharedPtr< FUICommandInfo > RelativeCoordinateSystem_Local;
 
+	/** Parent-relative coordinate system */
+	TSharedPtr<FUICommandInfo> RelativeCoordinateSystem_Parent;
+
 	TSharedPtr< FUICommandInfo > CycleTransformGizmos;
 	TSharedPtr< FUICommandInfo > CycleTransformGizmoCoordSystem;
 
