Control Rig: Introduce reusable storage container and support contiguous allocation

#rnx

diff --git a/Engine/Plugins/Animation/ControlRig/Source/ControlRig/Private/Rigs/RigHierarchy.cpp b/Engine/Plugins/Animation/ControlRig/Source/ControlRig/Private/Rigs/RigHierarchy.cpp
--- a/Engine/Plugins/Animation/ControlRig/Source/ControlRig/Private/Rigs/RigHierarchy.cpp
+++ b/Engine/Plugins/Animation/ControlRig/Source/ControlRig/Private/Rigs/RigHierarchy.cpp
@@ -485,7 +485,7 @@ void URigHierarchy::Load(FArchive& Ar)
 				CurveElement->StorageIndex = CurveIndices[UsedCurveIndex++];
 			}
 
-			Element->LinkStorage(ElementTransforms.Storage, ElementDirtyStates.Storage, ElementCurves.Storage);
+			Element->LinkStorage(ElementTransforms.GetStorage(), ElementDirtyStates.GetStorage(), ElementCurves.GetStorage());
 			Element->Load(*ArchiveForElements, FRigBaseElement::StaticData);
 		}
 	}
@@ -547,7 +547,7 @@ void URigHierarchy::Load(FArchive& Ar)
 		{
 			FRigBaseElement* Element = Find(Entry.Key);
 			Element->MetadataStorageIndex = ElementMetadata.Num();
-			ElementMetadata.Storage.Add(MoveTemp(Entry.Value));
+			ElementMetadata.Add(MoveTemp(Entry.Value));
 		}
 	}
 
@@ -612,9 +612,9 @@ void URigHierarchy::Reset_Impl(bool bResetElements)
 		ElementDirtyStates.Reset();
 		ElementCurves.Reset();
 
-		ElementMetadata.Reset([](int32 InIndex, FMetadataStorage& MetadataStorage)
+		ElementMetadata.Reset([](int32 InIndex, FMetadataStorage* MetadataStorage)
 		{
-			for (TTuple<FName, FRigBaseMetadata*>& Item: MetadataStorage.MetadataMap)
+			for (TTuple<FName, FRigBaseMetadata*>& Item: MetadataStorage->MetadataMap)
 			{
 				FRigBaseMetadata::DestroyMetadata(&Item.Value);
 			} 
@@ -898,7 +898,7 @@ void URigHierarchy::CopyHierarchy(URigHierarchy* InHierarchy)
 			{
 				CurveElement->StorageIndex = CurveIndices[UsedCurveIndex++];
 			}
-			Target->LinkStorage(ElementTransforms.Storage, ElementDirtyStates.Storage, ElementCurves.Storage);
+			Target->LinkStorage(ElementTransforms.GetStorage(), ElementDirtyStates.GetStorage(), ElementCurves.GetStorage());
 
 			Target->CopyFrom(Source);
 
@@ -6157,7 +6157,7 @@ void URigHierarchy::AllocateDefaultElementStorage(FRigBaseElement* InElement, bo
 	}
 	else
 	{
-		InElement->LinkStorage(ElementTransforms.Storage, ElementDirtyStates.Storage, ElementCurves.Storage);
+		InElement->LinkStorage(ElementTransforms.GetStorage(), ElementDirtyStates.GetStorage(), ElementCurves.GetStorage());
 	}
 	ElementTransformRanges.Reset();
 }
@@ -6171,7 +6171,7 @@ void URigHierarchy::UpdateElementStorage()
 {
 	for(FRigBaseElement* Element : Elements)
 	{
-		Element->LinkStorage(ElementTransforms.Storage, ElementDirtyStates.Storage, ElementCurves.Storage);
+		Element->LinkStorage(ElementTransforms.GetStorage(), ElementDirtyStates.GetStorage(), ElementCurves.GetStorage());
 	}
 	ElementTransformRanges.Reset();
 }
@@ -6315,8 +6315,8 @@ bool URigHierarchy::SortElementStorage()
 
 	FRigReusableElementStorage<FTransform> NewElementTransforms;
 	FRigReusableElementStorage<bool> NewElementDirtyStates;
-	NewElementDirtyStates.Storage.Reserve(ElementDirtyStates.Storage.Num() - ElementDirtyStates.FreeList.Num());
-	NewElementDirtyStates.Storage.Reserve(ElementDirtyStates.Storage.Num() - ElementDirtyStates.FreeList.Num());
+	NewElementDirtyStates.Reserve(ElementDirtyStates.Num() - ElementDirtyStates.GetFreeList().Num());
+	NewElementDirtyStates.Reserve(ElementDirtyStates.Num() - ElementDirtyStates.GetFreeList().Num());
 
 	bHasMatchingRanges = true;
 	NewRanges.Reset();
@@ -6330,13 +6330,13 @@ bool URigHierarchy::SortElementStorage()
 
 		if(CurrentTransformIndex != INDEX_NONE)
 		{
-			NewTransformIndex = NewElementTransforms.Storage.Add(ElementTransforms[CurrentTransformIndex]);
+			NewTransformIndex = NewElementTransforms.Add(ElementTransforms[CurrentTransformIndex]);
 			TransformInfos[Index].ComputedTransform->StorageIndex = NewTransformIndex;
 		}
 
 		if(CurrentDirtyStateIndex != INDEX_NONE)
 		{
-			NewDirtyStateIndex = NewElementDirtyStates.Storage.Add(ElementDirtyStates[CurrentDirtyStateIndex]);
+			NewDirtyStateIndex = NewElementDirtyStates.Add(ElementDirtyStates[CurrentDirtyStateIndex]);
 			TransformInfos[Index].DirtyState->StorageIndex = NewDirtyStateIndex;
 		}
 
@@ -6595,7 +6595,7 @@ FRigBaseMetadata* URigHierarchy::GetMetadataForElement(FRigBaseElement* InElemen
 {
 	if (!ElementMetadata.IsValidIndex(InElement->MetadataStorageIndex))
 	{
-		InElement->MetadataStorageIndex = ElementMetadata.Allocate(1, FMetadataStorage())[0];
+		InElement->MetadataStorageIndex = ElementMetadata.Add(FMetadataStorage());
 	}
 	
 	FMetadataStorage& Storage = ElementMetadata[InElement->MetadataStorageIndex];
diff --git a/Engine/Plugins/Animation/ControlRig/Source/ControlRig/Private/Rigs/RigHierarchyElements.cpp b/Engine/Plugins/Animation/ControlRig/Source/ControlRig/Private/Rigs/RigHierarchyElements.cpp
--- a/Engine/Plugins/Animation/ControlRig/Source/ControlRig/Private/Rigs/RigHierarchyElements.cpp
+++ b/Engine/Plugins/Animation/ControlRig/Source/ControlRig/Private/Rigs/RigHierarchyElements.cpp
@@ -9,6 +9,8 @@
 #include "AnimationCoreLibrary.h"
 #include "Algo/Transform.h"
 
+#include UE_INLINE_GENERATED_CPP_BY_NAME(RigHierarchyElements)
+
 ////////////////////////////////////////////////////////////////////////////////
 // FRigBaseElement
 ////////////////////////////////////////////////////////////////////////////////
@@ -23,50 +25,6 @@ FRigBaseElement::~FRigBaseElement()
 	}
 }
 
-UScriptStruct* FRigBaseElement::GetElementStruct() const
-{
-	switch(GetType())
-	{
-		case ERigElementType::Bone:
-		{
-			return FRigBoneElement::StaticStruct();
-		}
-		case ERigElementType::Null:
-		{
-			return FRigNullElement::StaticStruct();
-		}
-		case ERigElementType::Control:
-		{
-			return FRigControlElement::StaticStruct();
-		}
-		case ERigElementType::Curve:
-		{
-			return FRigCurveElement::StaticStruct();
-		}
-		case ERigElementType::Reference:
-		{
-			return FRigReferenceElement::StaticStruct();
-		}
-		case ERigElementType::Physics:
-		{
-			return FRigPhysicsElement::StaticStruct();
-		}
-		case ERigElementType::Connector:
-		{
-			return FRigConnectorElement::StaticStruct();
-		}
-		case ERigElementType::Socket:
-		{
-			return FRigSocketElement::StaticStruct();
-		}
-		default:
-		{
-			break;
-		}
-	}
-	return FRigBaseElement::StaticStruct();
-}
-
 void FRigBaseElement::Serialize(FArchive& Ar, ESerializationPhase SerializationPhase)
 {
 	Ar.UsingCustomVersion(FControlRigObjectVersion::GUID);
diff --git a/Engine/Plugins/Animation/ControlRig/Source/ControlRig/Public/Rigs/RigHierarchyElements.h b/Engine/Plugins/Animation/ControlRig/Source/ControlRig/Public/Rigs/RigHierarchyElements.h
--- a/Engine/Plugins/Animation/ControlRig/Source/ControlRig/Public/Rigs/RigHierarchyElements.h
+++ b/Engine/Plugins/Animation/ControlRig/Source/ControlRig/Public/Rigs/RigHierarchyElements.h
@@ -6,6 +6,7 @@
 #include "RigHierarchyDefines.h"
 #include "RigHierarchyMetadata.h"
 #include "RigConnectionRules.h"
+#include "RigReusableElementStorage.h"
 #include "RigHierarchyElements.generated.h"
 
 struct FRigVMExecuteContext;
@@ -20,6 +21,7 @@ DECLARE_DELEGATE_ThreeParams(FRigElementMetadataTagChangedDelegate, const FRigEl
 DECLARE_DELEGATE_RetVal_ThreeParams(FTransform, FRigReferenceGetWorldTransformDelegate, const FRigVMExecuteContext*, const FRigElementKey& /* Key */, bool /* bInitial */);
 
 #define DECLARE_RIG_ELEMENT_METHODS(ElementType) \
+virtual UScriptStruct* GetElementStruct() const override { return ElementType::StaticStruct(); } \
 template<typename T> \
 friend const T* Cast(const ElementType* InElement) \
 { \
@@ -254,187 +256,6 @@ namespace ERigTransformStorageType
 	};
 }
 
-template<typename T>
-struct CONTROLRIG_API FRigReusableElementStorage
-{
-	TArray<T> Storage;
-	TArray<int32> FreeList;
-
-	bool IsValidIndex(const int32& InIndex) const
-	{
-		return Storage.IsValidIndex(InIndex);
-	}
-
-	int32 Num() const
-	{
-		return Storage.Num();
-	}
-
-	T& operator[](int InIndex)
-	{
-		return Storage[InIndex];
-	}
-
-	const T& operator[](int InIndex) const
-	{
-		return Storage[InIndex];
-	}
-
-	typename TArray<T>::RangedForIteratorType begin() { return Storage.begin(); }
-	typename TArray<T>::RangedForIteratorType end() { return Storage.end(); }
-	typename TArray<T>::RangedForConstIteratorType begin() const { return Storage.begin(); }
-	typename TArray<T>::RangedForConstIteratorType end() const { return Storage.end(); }
-
-	const T* GetData() const
-	{
-		return Storage.GetData();
-	}
-
-	TArray<int32, TInlineAllocator<4>> Allocate(int32 InCount, const T& InDefault)
-	{
-		TArray<int32, TInlineAllocator<4>> Indices;
-	
-		const int32 NumToAllocate = InCount - FMath::Min(InCount, FreeList.Num());
-		if(NumToAllocate > 0)
-		{
-			Storage.Reserve(Storage.Num() + NumToAllocate);
-		}
-	
-		for(int32 Index = 0; Index < InCount; Index++)
-		{
-			if(FreeList.IsEmpty())
-			{
-				Indices.Push(Storage.Add(InDefault));
-			}
-			else
-			{
-				Indices.Push(FreeList.Pop(EAllowShrinking::No));
-				Storage[Indices.Last()] = InDefault;
-			}
-		}
-
-		return Indices;
-	}
-
-	void Deallocate(int32& InIndex, T** InStorage)
-	{
-		if(InIndex == INDEX_NONE)
-		{
-			return;
-		}
-#if WITH_EDITOR
-		check(Storage.IsValidIndex(InIndex));
-		check(!FreeList.Contains(InIndex));
-#endif
-		FreeList.Add(InIndex);
-		InIndex = INDEX_NONE;
-		if(InStorage)
-		{
-			*InStorage = nullptr;
-		}
-	}
-
-	void Deallocate(const TConstArrayView<int32>& InIndices)
-	{
-		if(InIndices.IsEmpty())
-		{
-			return;
-		}
-		FreeList.Reserve(FreeList.Num() + InIndices.Num());
-		for(int32 Index : InIndices)
-		{
-			if(Index != INDEX_NONE && !FreeList.Contains(Index))
-			{
-				Deallocate(Index, nullptr);
-			}
-		}
-	}
-
-	template<typename OwnerType>
-	void Deallocate(OwnerType* InOwner)
-	{
-		check(InOwner);
-		Deallocate(InOwner->StorageIndex, &InOwner->Storage);
-	}
-
-	void Reset(TFunction<void(int32, T&)> OnDestroyCallback = nullptr)
-	{
-		if(OnDestroyCallback)
-		{
-			for(int32 Index = 0; Index < Storage.Num(); Index++)
-			{
-				OnDestroyCallback(Index, Storage[Index]);
-			}
-		}
-		Storage.Reset();
-		FreeList.Reset();
-	}
-
-	bool Contains(int32 InIndex, const T* InStorage)
-	{
-		if(!IsValidIndex(InIndex))
-		{
-			return false;
-		}
-		return GetData() + InIndex == InStorage;
-	}
-
-	template<typename OwnerType>
-	bool Contains(const OwnerType* InOwner)
-	{
-		check(InOwner);
-		return Contains(InOwner->StorageIndex, InOwner->Storage);
-	}
-
-	TMap<int32, int32> Shrink(TFunction<void(int32, T&)> OnDestroyCallback = nullptr)
-	{
-		TMap<int32, int32> OldToNew;
-		
-		if(!FreeList.IsEmpty())
-		{
-			TArray<bool> ToRemove;
-
-			if(FreeList.Num() != Storage.Num() || OnDestroyCallback != nullptr)
-			{
-				ToRemove.AddZeroed(Storage.Num());
-				for(int32 FreeIndex : FreeList)
-				{
-					ToRemove[FreeIndex] = true;
-					if(OnDestroyCallback)
-					{
-						OnDestroyCallback(FreeIndex, Storage[FreeIndex]);
-					}
-				}
-			}
-
-			if(FreeList.Num() != Storage.Num())
-			{
-				TArray<T> NewStorage;
-				NewStorage.Reserve(FMath::Max(Storage.Num() - FreeList.Num(), 0));
-				for(int32 OldIndex = 0; OldIndex < Storage.Num(); OldIndex++)
-				{
-					if(!ToRemove[OldIndex])
-					{
-						OldToNew.Add(OldIndex, NewStorage.Add(Storage[OldIndex]));
-					}
-				}
-				Storage = NewStorage;
-			}
-			else
-			{
-				Storage.Reset();
-			}
-
-			FreeList.Reset();
-		}
-
-		FreeList.Shrink();
-		Storage.Shrink();
-
-		return OldToNew;
-	}
-};
-
 USTRUCT(BlueprintType)
 struct CONTROLRIG_API FRigTransformDirtyState
 {
@@ -927,6 +748,7 @@ public:
 		// with BP and details view wrappers. 
 		// These copies are solely intended for UControlRig::OnControlSelected_BP
 		Key = InOther.Key;
+		CachedNameString.Reset();
 		Index = InOther.Index;
 		SubIndex = InOther.SubIndex;
 		CreatedAtInstructionIndex = InOther.CreatedAtInstructionIndex;
@@ -988,7 +810,7 @@ protected:
 
 public:
 
-	UScriptStruct* GetElementStruct() const;
+	virtual UScriptStruct* GetElementStruct() const { return FRigBaseElement::StaticStruct(); }
 	void Serialize(FArchive& Ar, ESerializationPhase SerializationPhase);
 	virtual void Save(FArchive& Ar, ESerializationPhase SerializationPhase);
 	virtual void Load(FArchive& Ar, ESerializationPhase SerializationPhase);
diff --git a/Engine/Plugins/Animation/ControlRig/Source/ControlRig/Public/Rigs/RigReusableElementStorage.h b/Engine/Plugins/Animation/ControlRig/Source/ControlRig/Public/Rigs/RigReusableElementStorage.h
new file mode 100644
--- /dev/null
+++ b/Engine/Plugins/Animation/ControlRig/Source/ControlRig/Public/Rigs/RigReusableElementStorage.h
@@ -0,0 +1,366 @@
+// Copyright Epic Games, Inc. All Rights Reserved.
+
+#pragma once
+
+#include "CoreMinimal.h"
+
+template<typename T>
+struct FRigReusableElementStorage
+{
+private:
+	TArray<T> Storage;
+	TArray<int32> FreeList;
+	bool bFreeListIsSorted = true;
+
+public:
+	bool IsValidIndex(const int32& InIndex) const
+	{
+		return Storage.IsValidIndex(InIndex);
+	}
+
+	int32 Num() const
+	{
+		return Storage.Num();
+	}
+
+	void Reset(TFunction<void(int32, T*)> OnDestroyCallback = nullptr)
+	{
+		if(OnDestroyCallback)
+		{
+			for(int32 Index = 0; Index < Storage.Num(); Index++)
+			{
+				OnDestroyCallback(Index, &Storage[Index]);
+			}
+		}
+		Storage.Reset();
+		FreeList.Reset();
+		bFreeListIsSorted = true;
+	}
+
+	void Empty(TFunction<void(int32, T*)> OnDestroyCallback = nullptr)
+	{
+		if(OnDestroyCallback)
+		{
+			for(int32 Index = 0; Index < Storage.Num(); Index++)
+			{
+				OnDestroyCallback(Index, &Storage[Index]);
+			}
+		}
+		Storage.Empty();
+		FreeList.Empty();
+		bFreeListIsSorted = true;
+	}
+
+	void Reserve(int32 InCount)
+	{
+		check(FreeList.IsEmpty());
+		Storage.Reserve(InCount);
+	}
+
+	T& operator[](int InIndex)
+	{
+		return Storage[InIndex];
+	}
+
+	const T& operator[](int InIndex) const
+	{
+		return Storage[InIndex];
+	}
+
+	typename TArray<T>::RangedForIteratorType begin() { return Storage.begin(); }
+	typename TArray<T>::RangedForIteratorType end() { return Storage.end(); }
+	typename TArray<T>::RangedForConstIteratorType begin() const { return Storage.begin(); }
+	typename TArray<T>::RangedForConstIteratorType end() const { return Storage.end(); }
+
+	const T* GetData() const
+	{
+		return Storage.GetData();
+	}
+
+	T* GetData()
+	{
+		return Storage.GetData();
+	}
+
+	TArrayView<const T> GetStorage() const
+	{
+		return TArrayView<const T>(Storage.GetData(), Storage.Num());
+	}
+
+	TArrayView<T> GetStorage()
+	{
+		return TArrayView<T>(Storage.GetData(), Storage.Num());
+	}
+
+	TArrayView<const int32> GetFreeList() const
+	{
+		return TArrayView<const int32>(FreeList.GetData(), FreeList.Num());
+	}
+
+	TArrayView<int32> GetFreeList()
+	{
+		return TArrayView<int32>(FreeList.GetData(), FreeList.Num());
+	}
+
+	int32 Add(const T& InDefault)
+	{
+		TArray<int32, TInlineAllocator<4>> Indices = Allocate(1, InDefault);
+		check(Indices.Num() == 1);
+		return Indices[0];
+	}
+
+	int32 Add(const T&& InDefault)
+	{
+		TArray<int32, TInlineAllocator<4>> Indices = Allocate(1, InDefault);
+		check(Indices.Num() == 1);
+		return Indices[0];
+	}
+
+	TArray<int32, TInlineAllocator<4>> Allocate(int32 InCount, const T& InDefault, bool bContiguous = false)
+	{
+		TArray<int32, TInlineAllocator<4>> Indices;
+
+		if(bContiguous)
+		{
+			if(InCount > 0)
+			{
+				const int32 FirstIndex = AllocateContiguous(InCount, InDefault);
+				const int32 UpperBound = FirstIndex + InCount;
+				Indices.Reserve(InCount);
+				for(int32 Index = FirstIndex; Index < UpperBound; Index++)
+				{
+					Indices.Add(Index);
+				}
+			}
+			return Indices;
+		}
+
+		const int32 NumToAllocate = InCount - FMath::Min(InCount, FreeList.Num());
+		if(NumToAllocate > 0)
+		{
+			Storage.Reserve(Storage.Num() + NumToAllocate);
+		}
+	
+		for(int32 Index = 0; Index < InCount; Index++)
+		{
+			if(FreeList.IsEmpty())
+			{
+				Indices.Push(Storage.Add(InDefault));
+			}
+			else
+			{
+				Indices.Push(FreeList.Pop(EAllowShrinking::No));
+				Storage[Indices.Last()] = InDefault;
+			}
+		}
+
+		return Indices;
+	}
+
+	int32 AllocateContiguous(int32 InCount, const T& InDefault)
+	{
+		int32 FirstIndex = INDEX_NONE;
+		if(InCount == 0)
+		{
+			return FirstIndex;
+		}
+
+		// if the freelist has enough room for us to allocate
+		// let's look for a candidate
+		if(FreeList.Num() >= InCount)
+		{
+			if(!bFreeListIsSorted)
+			{
+				FreeList.Sort();
+				bFreeListIsSorted = true;
+			}
+			
+			FirstIndex = FreeList[0];
+			int32 Remainder = InCount - 1;
+
+			for(int32 Index = 1; Index < FreeList.Num(); Index++)
+			{
+				if(Remainder == 0)
+				{
+					break;
+				}
+				
+				// make sure they are consecutive
+				if(FreeList[Index - 1] == FreeList[Index] - 1)
+				{
+					Remainder--;
+				}
+				else
+				{
+					FirstIndex = FreeList[Index];
+					Remainder = InCount - 1;
+				}
+			}
+
+			// if we didn't find enough room for all elements....
+			if(Remainder > 0)
+			{
+				FirstIndex = INDEX_NONE;
+			}
+		}
+
+		if(FirstIndex == INDEX_NONE)
+		{
+			FirstIndex = Storage.AddUninitialized(InCount);
+		}
+		else
+		{
+			const int32 UpperBound = FirstIndex + InCount;
+			FreeList.RemoveAll([FirstIndex, UpperBound](int32 Index) -> bool
+			{
+				return (Index >= FirstIndex) && (Index < UpperBound);
+			});
+		}
+
+		const int32 UpperBound = FirstIndex + InCount;
+		for(int32 Index = FirstIndex; Index < UpperBound; Index++)
+		{
+			Storage[Index] = InDefault;
+		}
+
+		return FirstIndex;
+	}
+
+	void Deallocate(const int32& InIndex)
+	{
+		int32 MutableIndex = InIndex;
+		Deallocate(MutableIndex, nullptr);
+	}
+
+	void Deallocate(int32& InIndex, T** InStorage)
+	{
+		if(InIndex == INDEX_NONE)
+		{
+			return;
+		}
+#if WITH_EDITOR
+		check(Storage.IsValidIndex(InIndex));
+		check(!FreeList.Contains(InIndex));
+#endif
+		if(bFreeListIsSorted && !FreeList.IsEmpty())
+		{
+			if(FreeList.Last() > InIndex)
+			{
+				bFreeListIsSorted = false;
+			}
+		}
+		FreeList.Add(InIndex);
+		InIndex = INDEX_NONE;
+		if(InStorage)
+		{
+			*InStorage = nullptr;
+		}
+	}
+
+	void Deallocate(const TConstArrayView<int32>& InIndices)
+	{
+		if(InIndices.IsEmpty())
+		{
+			return;
+		}
+		FreeList.Reserve(FreeList.Num() + InIndices.Num());
+		for(int32 Index : InIndices)
+		{
+			if(Index != INDEX_NONE && !FreeList.Contains(Index))
+			{
+				Deallocate(Index);
+			}
+		}
+	}
+
+	void Deallocate(int32 InStartIndex, int32 InCount)
+	{
+		const int32 UpperBound = InStartIndex + InCount;
+		const int32 EndIndex = UpperBound - 1;
+		
+		if(InCount == 0 || !Storage.IsValidIndex(InStartIndex) || !Storage.IsValidIndex(EndIndex))
+		{
+			return;
+		}
+		
+		FreeList.Reserve(FreeList.Num() + InCount);
+		for(int32 Index = InStartIndex; Index < UpperBound; Index++)
+		{
+			if(Index != INDEX_NONE && !FreeList.Contains(Index))
+			{
+				Deallocate(Index);
+			}
+		}
+	}
+
+	template<typename OwnerType>
+	void Deallocate(OwnerType* InOwner)
+	{
+		check(InOwner);
+		Deallocate(InOwner->StorageIndex, &InOwner->Storage);
+	}
+
+	bool Contains(int32 InIndex, const T* InStorage)
+	{
+		if(!IsValidIndex(InIndex))
+		{
+			return false;
+		}
+		return GetData() + InIndex == InStorage;
+	}
+
+	template<typename OwnerType>
+	bool Contains(const OwnerType* InOwner)
+	{
+		check(InOwner);
+		return Contains(InOwner->StorageIndex, InOwner->Storage);
+	}
+
+	TMap<int32, int32> Shrink(TFunction<void(int32, T&)> OnDestroyCallback = nullptr)
+	{
+		TMap<int32, int32> OldToNew;
+		
+		if(!FreeList.IsEmpty())
+		{
+			TArray<bool> ToRemove;
+
+			if(FreeList.Num() != Storage.Num() || OnDestroyCallback != nullptr)
+			{
+				ToRemove.AddZeroed(Storage.Num());
+				for(int32 FreeIndex : FreeList)
+				{
+					ToRemove[FreeIndex] = true;
+					if(OnDestroyCallback)
+					{
+						OnDestroyCallback(FreeIndex, Storage[FreeIndex]);
+					}
+				}
+			}
+
+			if(FreeList.Num() != Storage.Num())
+			{
+				TArray<T> NewStorage;
+				NewStorage.Reserve(FMath::Max(Storage.Num() - FreeList.Num(), 0));
+				for(int32 OldIndex = 0; OldIndex < Storage.Num(); OldIndex++)
+				{
+					if(!ToRemove[OldIndex])
+					{
+						OldToNew.Add(OldIndex, NewStorage.Add(Storage[OldIndex]));
+					}
+				}
+				Storage = NewStorage;
+			}
+			else
+			{
+				Storage.Reset();
+			}
+
+			FreeList.Reset();
+		}
+
+		FreeList.Shrink();
+		Storage.Shrink();
+
+		return OldToNew;
+	}
+};
