Update icons for Animation Mode viewport toolbar extensions

diff --git a/Engine/Plugins/Animation/ControlRig/Source/ControlRigEditor/Private/Editor/ControlRigViewportToolbarExtensions.cpp b/Engine/Plugins/Animation/ControlRig/Source/ControlRigEditor/Private/Editor/ControlRigViewportToolbarExtensions.cpp
--- a/Engine/Plugins/Animation/ControlRig/Source/ControlRigEditor/Private/Editor/ControlRigViewportToolbarExtensions.cpp
+++ b/Engine/Plugins/Animation/ControlRig/Source/ControlRigEditor/Private/Editor/ControlRigViewportToolbarExtensions.cpp
@@ -57,7 +57,7 @@ void UE::ControlRig::PopulateControlRigViewportToolbarTransformSubmenu(const FNa
 				LOCTEXT(
 					"LocalTransformsInEachLocalSpaceTooltip", "When multiple objects are selected, whether or not to transform each invidual object along its own local transform space."
 				),
-				FSlateIcon(),
+				FSlateIcon(FAppStyle::GetAppStyleSetName(), "EditorViewport.LocalTransformsInEachLocalSpace"),
 				Action,
 				EUserInterfaceActionType::ToggleButton
 			);
@@ -94,7 +94,7 @@ void UE::ControlRig::PopulateControlRigViewportToolbarTransformSubmenu(const FNa
 					"RestoreCoordinateSpaceOnSwitchTooltip",
 					"Whether to restore the coordinate space when changing Widget Modes in the Viewport."
 				),
-				FSlateIcon(),
+				FSlateIcon(FAppStyle::GetAppStyleSetName(), "EditorViewport.RestoreCoordinateSpaceOnSwitch"),
 				Action,
 				EUserInterfaceActionType::ToggleButton
 			);
@@ -132,7 +132,7 @@ void UE::ControlRig::PopulateControlRigViewportToolbarTransformSubmenu(const FNa
 				"OnlySelectRigControls",
 				LOCTEXT("OnlySelectRigControlsLabel", "Select Only Control Rig Controls"),
 				LOCTEXT("OnlySelectRigControlsTooltip", "Whether or not only Rig Controls can be selected."),
-				FSlateIcon(),
+				FSlateIcon(FAppStyle::GetAppStyleSetName(), "EditorViewport.OnlySelectRigControls"),
 				Action,
 				EUserInterfaceActionType::ToggleButton
 			);
diff --git a/Engine/Source/Editor/EditorStyle/Private/StarshipStyle.cpp b/Engine/Source/Editor/EditorStyle/Private/StarshipStyle.cpp
--- a/Engine/Source/Editor/EditorStyle/Private/StarshipStyle.cpp
+++ b/Engine/Source/Editor/EditorStyle/Private/StarshipStyle.cpp
@@ -4520,6 +4520,12 @@ void FStarshipEditorStyle::FStyle::SetupLevelEditorStyle()
 			new IMAGE_BRUSH_SVG("Starship/Common/transform-local", Icon16x16));
 		Set("EditorViewport.RelativeCoordinateSystem_Parent",
 			new CORE_IMAGE_BRUSH_SVG("Starship/Common/ParentHierarchy", Icon16x16));
+		Set("EditorViewport.RestoreCoordinateSpaceOnSwitch",
+			new IMAGE_BRUSH_SVG("Starship/EditorViewport/RestoreCoordinateSpace_16", Icon16x16));
+		Set("EditorViewport.LocalTransformsInEachLocalSpace",
+			new IMAGE_BRUSH_SVG("Starship/EditorViewport/LocalTransformsLocalSpace_16", Icon16x16));
+		Set("EditorViewport.OnlySelectRigControls",
+			new IMAGE_BRUSH_SVG("Starship/Animation/AnimationSelectOnlyControlRig_16", Icon16x16));
 
 		Set( "EditorViewport.CamSpeedSetting", new IMAGE_BRUSH_SVG( "Starship/EditorViewport/camera", Icon16x16) );
 		
