Control Rig Edit Mode: fixed use of garbaged control rig object when removing delegates

#rnx

diff --git a/Engine/Plugins/Animation/ControlRig/Source/ControlRigEditor/Private/EditMode/ControlRigEditMode.cpp b/Engine/Plugins/Animation/ControlRig/Source/ControlRigEditor/Private/EditMode/ControlRigEditMode.cpp
--- a/Engine/Plugins/Animation/ControlRig/Source/ControlRigEditor/Private/EditMode/ControlRigEditMode.cpp
+++ b/Engine/Plugins/Animation/ControlRig/Source/ControlRigEditor/Private/EditMode/ControlRigEditMode.cpp
@@ -4265,30 +4265,31 @@ bool FControlRigEditMode::ShapeSelectionOverride(const UPrimitiveComponent* InCo
 
 void FControlRigEditMode::OnObjectsReplaced(const TMap<UObject*, UObject*>& OldToNewInstanceMap)
 {
-	for (int32 RigIndex = 0; RigIndex < RuntimeControlRigs.Num(); RigIndex++)
+	bool bHasAnyRigBeenReplaced = false;
+	
+	const TArray<TWeakObjectPtr<UControlRig>> PreviousRuntimeRigs(RuntimeControlRigs);
+	for (const TWeakObjectPtr<UControlRig>& WeakControlRig: PreviousRuntimeRigs)
 	{
-		UObject* OldObject = RuntimeControlRigs[RigIndex].Get();
-		UObject* NewObject = OldToNewInstanceMap.FindRef(OldObject);
-		if (NewObject)
+		if (UControlRig* OldControlRig = WeakControlRig.Get())
 		{
-			TArray<TWeakObjectPtr<UControlRig>> PreviousRuntimeRigs = RuntimeControlRigs;
-			for (int32 PreviousRuntimeRigIndex = 0; PreviousRuntimeRigIndex < PreviousRuntimeRigs.Num(); PreviousRuntimeRigIndex++)
+			if (UControlRig* NewControlRig = Cast<UControlRig>(OldToNewInstanceMap.FindRef(OldControlRig)))
 			{
-				if (PreviousRuntimeRigs[PreviousRuntimeRigIndex].IsValid())
-				{
-					RemoveControlRig(PreviousRuntimeRigs[PreviousRuntimeRigIndex].Get());
-				}
-			}
-			RuntimeControlRigs.Reset();
-
-			UControlRig* NewRig = Cast<UControlRig>(NewObject);
-			AddControlRigInternal(NewRig);
+				// remove old rig (this will also remove it from RuntimeControlRigs)
+				RemoveControlRig(OldControlRig);
 
-			NewRig->Initialize();
+				// add new rig
+				AddControlRigInternal(NewControlRig);
+				NewControlRig->Initialize();
 
-			SetObjects_Internal();
+				bHasAnyRigBeenReplaced = true;
+			}
 		}
 	}
+
+	if (bHasAnyRigBeenReplaced)
+	{
+		SetObjects_Internal();
+	}
 }
 
 bool FControlRigEditMode::IsTransformDelegateAvailable() const
@@ -5154,7 +5155,10 @@ void FControlRigEditMode::AddControlRigInternal(UControlRig* InControlRig)
 
 	InControlRig->SetControlsVisible(true);
 	InControlRig->PostInitInstanceIfRequired();
+
 	InControlRig->GetHierarchy()->OnModified().RemoveAll(this);
+	InControlRig->OnPostConstruction_AnyThread().RemoveAll(this);
+	
 	InControlRig->GetHierarchy()->OnModified().AddSP(this, &FControlRigEditMode::OnHierarchyModified_AnyThread);
 	InControlRig->OnPostConstruction_AnyThread().AddSP(this, &FControlRigEditMode::OnPostConstruction_AnyThread);
 
@@ -5211,10 +5215,14 @@ void FControlRigEditMode::RemoveControlRig(UControlRig* InControlRig)
 	{
 		return;
 	}
-	InControlRig->ControlModified().RemoveAll(this);
-	InControlRig->GetHierarchy()->OnModified().RemoveAll(this);
-	InControlRig->OnPreConstructionForUI_AnyThread().RemoveAll(this);
-	InControlRig->OnPostConstruction_AnyThread().RemoveAll(this);
+
+	if (!URigVMHost::IsGarbageOrDestroyed(InControlRig))
+	{
+		InControlRig->ControlModified().RemoveAll(this);
+		InControlRig->GetHierarchy()->OnModified().RemoveAll(this);
+		InControlRig->OnPreConstructionForUI_AnyThread().RemoveAll(this);
+		InControlRig->OnPostConstruction_AnyThread().RemoveAll(this);
+	}
 	
 	int32 Index = RuntimeControlRigs.Find(InControlRig);
 	TStrongObjectPtr<UControlRigEditModeDelegateHelper>* DelegateHelper = DelegateHelpers.Find(InControlRig);
@@ -5238,6 +5246,8 @@ void FControlRigEditMode::RemoveControlRig(UControlRig* InControlRig)
 		LastShapeLibraryHash.RemoveAt(Index);
 	}
 
+	ControlRigsToRecreate.Remove(InControlRig);
+
 	//needed for the control rig track editor delegates to get removed
 	if (WeakSequencer.IsValid())
 	{
