Control Rig: initial steps for hierarchy element component support

#rnx

diff --git a/Engine/Plugins/Animation/ControlRig/Source/ControlRig/Private/Constraints/ControlRigTransformableHandle.cpp b/Engine/Plugins/Animation/ControlRig/Source/ControlRig/Private/Constraints/ControlRigTransformableHandle.cpp
--- a/Engine/Plugins/Animation/ControlRig/Source/ControlRig/Private/Constraints/ControlRigTransformableHandle.cpp
+++ b/Engine/Plugins/Animation/ControlRig/Source/ControlRig/Private/Constraints/ControlRigTransformableHandle.cpp
@@ -416,7 +416,7 @@ void UTransformableControlHandle::RegisterDelegates()
 void UTransformableControlHandle::OnHierarchyModified(
 	ERigHierarchyNotification InNotif,
 	URigHierarchy* InHierarchy,
-	const FRigBaseElement* InElement)
+	const FRigNotificationSubject& InSubject)
 {
 	if (!ControlRig.IsValid())
 	{
@@ -429,6 +429,8 @@ void UTransformableControlHandle::OnHierarchyModified(
 		return;
 	}
 
+	const FRigBaseElement* InElement = InSubject.Element;
+
 	switch (InNotif)
 	{
 		case ERigHierarchyNotification::ElementRemoved:
diff --git a/Engine/Plugins/Animation/ControlRig/Source/ControlRig/Private/ControlRig.cpp b/Engine/Plugins/Animation/ControlRig/Source/ControlRig/Private/ControlRig.cpp
--- a/Engine/Plugins/Animation/ControlRig/Source/ControlRig/Private/ControlRig.cpp
+++ b/Engine/Plugins/Animation/ControlRig/Source/ControlRig/Private/ControlRig.cpp
@@ -604,7 +604,7 @@ void UControlRig::InitializeFromCDO()
 #endif
 
 			// notify clients that the hierarchy has changed
-			Hierarchy->Notify(ERigHierarchyNotification::HierarchyReset, nullptr);
+			Hierarchy->Notify(ERigHierarchyNotification::HierarchyReset, {});
 
 			// copy hierarchy settings
 			HierarchySettings = CDO->HierarchySettings;
@@ -2377,8 +2377,11 @@ void UControlRig::RunPostConstructionEvent()
 }
 
 void UControlRig::HandleHierarchyModified(ERigHierarchyNotification InNotification, URigHierarchy* InHierarchy,
-                                          const FRigBaseElement* InElement)
+                                          const FRigNotificationSubject& InSubject)
 {
+	const FRigBaseElement* InElement = InSubject.Element;
+	const FRigBaseComponent* InComponent = InSubject.Component;
+	
 	switch(InNotification)
 	{
 		case ERigHierarchyNotification::ElementSelected:
diff --git a/Engine/Plugins/Animation/ControlRig/Source/ControlRig/Private/Rigs/RigHierarchy.cpp b/Engine/Plugins/Animation/ControlRig/Source/ControlRig/Private/Rigs/RigHierarchy.cpp
--- a/Engine/Plugins/Animation/ControlRig/Source/ControlRig/Private/Rigs/RigHierarchy.cpp
+++ b/Engine/Plugins/Animation/ControlRig/Source/ControlRig/Private/Rigs/RigHierarchy.cpp
@@ -179,7 +179,7 @@ void URigHierarchy::AddReferencedObjects(UObject* InpThis, FReferenceCollector&
 	FScopeLock Lock(&pThis->ElementsLock);
 	for (FRigBaseElement* Element : pThis->Elements)
 	{
-		Collector.AddPropertyReferencesWithStructARO(Element->GetElementStruct(), Element, pThis);
+		Collector.AddPropertyReferencesWithStructARO(Element->GetScriptStruct(), Element, pThis);
 	}
 }
 
@@ -275,6 +275,24 @@ void URigHierarchy::Save(FArchive& Ar)
 		*ArchiveForElements << ElementMetadataToSave;
 	}
 
+	// save all components
+	int32 NumComponents = 0;
+	for(FInstancedStruct& InstancedStruct : ElementComponents)
+	{
+		if(InstancedStruct.IsValid())
+		{
+			NumComponents++;
+		}
+	}
+	*ArchiveForElements << NumComponents;
+	for(FInstancedStruct& InstancedStruct : ElementComponents)
+	{
+		if(InstancedStruct.IsValid())
+		{
+			InstancedStruct.Serialize(*ArchiveForElements);
+		}
+	}
+
 	if(bUseCompressedArchive)
 	{
 		// It is possible for compression to actually increase the size of the data, so we over allocate here to handle that.
@@ -551,6 +569,30 @@ void URigHierarchy::Load(FArchive& Ar)
 		}
 	}
 
+	if (ArchiveForElements->CustomVer(FControlRigObjectVersion::GUID) >= FControlRigObjectVersion::RigHierarchyStoresComponents)
+	{
+		ElementComponents.Reset();
+
+		int32 NumComponents = 0;
+		*ArchiveForElements << NumComponents;
+
+		ElementComponents.AddZeroed(NumComponents);
+		for (int32 ComponentIndex = 0; ComponentIndex < ElementComponents.Num(); ComponentIndex++)
+		{
+			FInstancedStruct& InstancedStruct = ElementComponents[ComponentIndex];
+			InstancedStruct.Serialize(*ArchiveForElements);
+			
+			FRigBaseComponent* Component = InstancedStruct.GetMutablePtr<FRigBaseComponent>();
+			FRigBaseElement* Element = Find(Component->GetElementKey());
+			check(Element);
+			Component->IndexInHierarchy = ComponentIndex;
+			Component->Element = Element;
+			Component->IndexInElement = Element->ComponentIndices.Add(Component->IndexInHierarchy);
+		}
+
+		IncrementTopologyVersion();
+	}
+
 	(void)SortElementStorage();
 }
 
@@ -562,7 +604,7 @@ void URigHierarchy::PostLoad()
 
 	CleanupInvalidCaches();
 
-	Notify(ERigHierarchyNotification::HierarchyReset, nullptr);
+	Notify(ERigHierarchyNotification::HierarchyReset, {});
 }
 
 #if WITH_EDITORONLY_DATA
@@ -637,7 +679,7 @@ void URigHierarchy::Reset_Impl(bool bResetElements)
 
 	{
 		FGCScopeGuard Guard;
-		Notify(ERigHierarchyNotification::HierarchyReset, nullptr);
+		Notify(ERigHierarchyNotification::HierarchyReset, {});
 	}
 
 	if(HierarchyForCacheValidation)
@@ -794,6 +836,7 @@ void URigHierarchy::CopyHierarchy(URigHierarchy* InHierarchy)
 
 				Target->SubIndex = SubIndex;
 				Target->Index = Elements.Add(Target);
+				Target->ComponentIndices.Reset();
 
 				ElementsPerType[ElementTypeIndex].Add(Target);
 				IndexLookup.Add(Key, Target->Index);
@@ -831,6 +874,7 @@ void URigHierarchy::CopyHierarchy(URigHierarchy* InHierarchy)
 
 				check(Target->Key.Type == Source->Key.Type);
 				Target->InitializeFrom(Source);
+				Target->ComponentIndices.Reset();
 
 				NumTransforms += Target->GetNumTransforms();
 				NumDirtyStates += Target->GetNumTransforms();
@@ -912,6 +956,22 @@ void URigHierarchy::CopyHierarchy(URigHierarchy* InHierarchy)
 		// make sure the curves are marked as unset after copying
 		UnsetCurveValues();
 
+		// Components
+		ElementComponents.Reset();
+		ElementComponents.Reserve(InHierarchy->ElementComponents.Num());
+		for(int32 SourceComponentIndex = 0; SourceComponentIndex < ElementComponents.Num(); SourceComponentIndex++)
+		{
+			if(!InHierarchy->ElementComponents[SourceComponentIndex].IsValid())
+			{
+				continue;
+			}
+			const int32 TargetComponentIndex = ElementComponents.Add(InHierarchy->ElementComponents[SourceComponentIndex]);
+			FRigBaseComponent* TargetComponent = ElementComponents[TargetComponentIndex].GetMutablePtr<FRigBaseComponent>();
+			TargetComponent->Element = Find(TargetComponent->ElementKey);
+			TargetComponent->IndexInHierarchy = TargetComponentIndex;
+			TargetComponent->IndexInElement = TargetComponent->Element->ComponentIndices.Add(TargetComponent->IndexInHierarchy);
+		}
+
 		// copy the topology version from the hierarchy.
 		// for this we'll include the hash of the hierarchy to make sure we
 		// are deterministic for different hierarchies.
@@ -933,7 +993,7 @@ void URigHierarchy::CopyHierarchy(URigHierarchy* InHierarchy)
 	EnsureCacheValidity();
 
 	bIsCopyingHierarchy = false;
-	Notify(ERigHierarchyNotification::HierarchyCopied, nullptr);
+	Notify(ERigHierarchyNotification::HierarchyCopied, {});
 }
 
 uint32 URigHierarchy::GetNameHash() const
@@ -1001,6 +1061,17 @@ uint32 URigHierarchy::GetTopologyHash(bool bIncludeTopologyVersion, bool bInclud
 		}
 	}
 
+	for(int32 ComponentIndex = 0; ComponentIndex < ElementComponents.Num(); ComponentIndex++)
+	{
+		if(const FRigBaseComponent* Component = GetComponent(ComponentIndex))
+		{
+			Hash = HashCombine(Hash, GetTypeHash(Component->GetScriptStruct()));
+			Hash = HashCombine(Hash, GetTypeHash(Component->GetIndexInHierarchy()));
+			Hash = HashCombine(Hash, GetTypeHash(Component->GetIndexInElement()));
+			Hash = HashCombine(Hash, GetTypeHash(Component->GetElement()->GetIndex()));
+		}
+	}
+
 	return Hash;
 }
 
@@ -1360,6 +1431,77 @@ bool URigHierarchy::IsProcedural(const FRigBaseElement* InElement) const
 	return InElement->IsProcedural();
 }
 
+const FRigBaseComponent* URigHierarchy::GetComponent(int32 InIndex) const
+{
+	if(ElementComponents.IsValidIndex(InIndex))
+	{
+		if(ElementComponents[InIndex].IsValid())
+		{
+			return ElementComponents[InIndex].GetPtr<FRigBaseComponent>();
+		}
+	}
+	return nullptr;
+}
+
+FRigBaseComponent* URigHierarchy::GetComponent(int32 InIndex)
+{
+	if(ElementComponents.IsValidIndex(InIndex))
+	{
+		if(ElementComponents[InIndex].IsValid())
+		{
+			return ElementComponents[InIndex].GetMutablePtr<FRigBaseComponent>();
+		}
+	}
+	return nullptr;
+}
+
+int32 URigHierarchy::NumComponents(FRigElementKey InElement) const
+{
+	if(const FRigBaseElement* Element = Find(InElement))
+	{
+		return Element->NumComponents();
+	}
+	return 0;
+}
+
+FName URigHierarchy::GetComponentName(FRigElementKey InElement, int32 InComponentIndex) const
+{
+	if(const FRigBaseElement* Element = Find(InElement))
+	{
+		if(const FRigBaseComponent* Component = Element->GetComponent(InComponentIndex))
+		{
+			return Component->GetFName();
+		}
+	}
+	return NAME_None;
+}
+
+UScriptStruct* URigHierarchy::GetComponentType(FRigElementKey InElement, int32 InComponentIndex) const
+{
+	if(const FRigBaseElement* Element = Find(InElement))
+	{
+		if(const FRigBaseComponent* Component = Element->GetComponent(InComponentIndex))
+		{
+			return Component->GetScriptStruct();
+		}
+	}
+	return nullptr;
+}
+
+FString URigHierarchy::GetComponentContent(FRigElementKey InElement, int32 InComponentIndex) const
+{
+	if(const FRigBaseElement* Element = Find(InElement))
+	{
+		if(const FRigBaseComponent* Component = Element->GetComponent(InComponentIndex))
+		{
+			FString Content;
+			Component->GetScriptStruct()->ExportText(Content, Component, Component, nullptr, PPF_None, nullptr);
+			return Content;
+		}
+	}
+	return FString();
+}
+
 TArray<FRigSocketState> URigHierarchy::GetSocketStates() const
 {
 	const TArray<FRigElementKey> Keys = GetSocketKeys(true);
@@ -3564,7 +3706,7 @@ void URigHierarchy::LinkPoseAdapter(TSharedPtr<FRigHierarchyPoseAdapter> InPoseA
 	}
 }
 
-void URigHierarchy::Notify(ERigHierarchyNotification InNotifType, const FRigBaseElement* InElement)
+void URigHierarchy::Notify(ERigHierarchyNotification InNotifType, const FRigNotificationSubject& InSubject)
 {
 	if(bSuspendNotifications)
 	{
@@ -3581,7 +3723,7 @@ void URigHierarchy::Notify(ERigHierarchyNotification InNotifType, const FRigBase
 		FScopeLock Lock(&ExecuteContextLock);
 		if(ExecuteContext != nullptr)
 		{
-			QueueNotification(InNotifType, InElement);
+			QueueNotification(InNotifType, InSubject);
 			return;
 		}
 	}
@@ -3589,23 +3731,26 @@ void URigHierarchy::Notify(ERigHierarchyNotification InNotifType, const FRigBase
 
 	if(QueuedNotifications.IsEmpty())
 	{
-		ModifiedEvent.Broadcast(InNotifType, this, InElement);
+		ModifiedEvent.Broadcast(InNotifType, this, InSubject);
 		if(ModifiedEventDynamic.IsBound())
 		{
 			FRigElementKey Key;
-			if(InElement)
+			if(InSubject.Element)
 			{
-				Key = InElement->GetKey();
+				Key = InSubject.Element->GetKey();
 			}
 			ModifiedEventDynamic.Broadcast(InNotifType, this, Key);
 		}
 	}
 	else
 	{
-		QueueNotification(InNotifType, InElement);
+		QueueNotification(InNotifType, InSubject);
 		SendQueuedNotifications();
 	}
 
+	const FRigBaseElement* InElement = InSubject.Element;
+	const FRigBaseComponent* InComponent = InSubject.Component;
+
 #if WITH_EDITOR
 
 	// certain events needs to be forwarded to the listening hierarchies.
@@ -3633,6 +3778,28 @@ void URigHierarchy::Notify(ERigHierarchyNotification InNotifType, const FRigBase
 			}
 			break;
 		}
+		case ERigHierarchyNotification::ComponentAdded:
+		case ERigHierarchyNotification::ComponentRemoved:
+		{
+			if (ensure(InComponent != nullptr))
+			{
+				const FRigBaseElement* Element = InComponent->GetElement();
+				ForEachListeningHierarchy([this, InNotifType, InComponent, Element](const FRigHierarchyListener& Listener)
+				{
+					if (URigHierarchy* ListeningHierarchy = Listener.Hierarchy.Get())
+					{			
+						if(const FRigBaseElement* ListeningElement = ListeningHierarchy->Find( Element->GetKey()))
+						{
+							if(const FRigBaseComponent* ListeningComponent = ListeningElement->GetComponentByName(InComponent->GetFName()))
+							{
+								ListeningHierarchy->Notify(InNotifType, ListeningComponent);
+							}
+						}
+					}
+				});
+			}
+			break;
+		}
 		default:
 		{
 			break;
@@ -3642,11 +3809,24 @@ void URigHierarchy::Notify(ERigHierarchyNotification InNotifType, const FRigBase
 #endif
 }
 
-void URigHierarchy::QueueNotification(ERigHierarchyNotification InNotification, const FRigBaseElement* InElement)
+void URigHierarchy::QueueNotification(ERigHierarchyNotification InNotification, const FRigNotificationSubject& InSubject)
 {
 	FQueuedNotification Entry;
 	Entry.Type = InNotification;
-	Entry.Key = InElement ? InElement->GetKey() : FRigElementKey();
+	if(InSubject.Element)
+	{
+		Entry.Key = InSubject.Element->GetKey();
+		Entry.ComponentName = NAME_None;
+	}
+	else if(InSubject.Component)
+	{
+		Entry.Key = InSubject.Component->GetElementKey();
+		Entry.ComponentName = InSubject.Component->GetFName();
+	}
+	else
+	{
+		checkNoEntry();
+	}
 	QueuedNotifications.Enqueue(Entry);
 }
 
@@ -3699,6 +3879,7 @@ void URigHierarchy::SendQueuedNotifications()
 			case ERigHierarchyNotification::HierarchyReset:
 			case ERigHierarchyNotification::ElementRemoved:
 			case ERigHierarchyNotification::ElementRenamed:
+			case ERigHierarchyNotification::ComponentRemoved:
 			{
 				// we don't allow these to happen during the run of a VM
 				if(const URigHierarchyController* Controller = GetController())
@@ -3745,6 +3926,8 @@ void URigHierarchy::SendQueuedNotifications()
 				break;
 			}
 			case ERigHierarchyNotification::ElementAdded:
+			case ERigHierarchyNotification::ComponentAdded:
+			case ERigHierarchyNotification::ComponentContentChanged:
 			case ERigHierarchyNotification::InteractionBracketOpened:
 			case ERigHierarchyNotification::InteractionBracketClosed:
 			default:
@@ -3772,7 +3955,7 @@ void URigHierarchy::SendQueuedNotifications()
 		return;
 	}
 
-	ModifiedEvent.Broadcast(ERigHierarchyNotification::InteractionBracketOpened, this, nullptr);
+	ModifiedEvent.Broadcast(ERigHierarchyNotification::InteractionBracketOpened, this, {});
 	if(ModifiedEventDynamic.IsBound())
 	{
 		ModifiedEventDynamic.Broadcast(ERigHierarchyNotification::InteractionBracketOpened, this, FRigElementKey());
@@ -3783,18 +3966,27 @@ void URigHierarchy::SendQueuedNotifications()
 	for(int32 Index = FilteredNotifications.Num() - 1; Index >= 0; Index--)
 	{
 		const FQueuedNotification& Entry = FilteredNotifications[Index];
-		const FRigBaseElement* Element = Find(Entry.Key);
-		if(Element)
+		FRigNotificationSubject Subject;
+		Subject.Element = Find(Entry.Key);
+		if(Subject.Element)
 		{
-			ModifiedEvent.Broadcast(Entry.Type, this, Element);
-			if(ModifiedEventDynamic.IsBound())
+			Subject.Component = Subject.Element->GetComponentByName(Entry.ComponentName);
+			if(Subject.Component)
+			{
+				Subject.Element = nullptr;
+			}
+
+			ModifiedEvent.Broadcast(Entry.Type, this, Subject);
+
+			// the dynamic multicast for now only supports element notifs
+			if(Subject.Element && ModifiedEventDynamic.IsBound())
 			{
-				ModifiedEventDynamic.Broadcast(Entry.Type, this, Element->GetKey());
+				ModifiedEventDynamic.Broadcast(Entry.Type, this, Subject.Element->GetKey());
 			}
 		}
 	}
 
-	ModifiedEvent.Broadcast(ERigHierarchyNotification::InteractionBracketClosed, this, nullptr);
+	ModifiedEvent.Broadcast(ERigHierarchyNotification::InteractionBracketClosed, this, {});
 	if(ModifiedEventDynamic.IsBound())
 	{
 		ModifiedEventDynamic.Broadcast(ERigHierarchyNotification::InteractionBracketClosed, this, FRigElementKey());
@@ -5632,6 +5824,7 @@ void URigHierarchy::DestroyElement(FRigBaseElement*& InElement)
 			for(int32 Index=0;Index<Count;Index++)
 			{
 				TGuardValue<const FRigBaseElement*> DestroyGuard(ElementBeingDestroyed, &ExistingElements[Index]);
+				DestroyComponents(&ExistingElements[Index]);
 				DeallocateElementStorage(&ExistingElements[Index]);
 				ExistingElements[Index].~FRigBoneElement(); 
 			}
@@ -5643,6 +5836,7 @@ void URigHierarchy::DestroyElement(FRigBaseElement*& InElement)
 			for(int32 Index=0;Index<Count;Index++)
 			{
 				TGuardValue<const FRigBaseElement*> DestroyGuard(ElementBeingDestroyed, &ExistingElements[Index]);
+				DestroyComponents(&ExistingElements[Index]);
 				DeallocateElementStorage(&ExistingElements[Index]);
 				ExistingElements[Index].~FRigNullElement(); 
 			}
@@ -5654,6 +5848,7 @@ void URigHierarchy::DestroyElement(FRigBaseElement*& InElement)
 			for(int32 Index=0;Index<Count;Index++)
 			{
 				TGuardValue<const FRigBaseElement*> DestroyGuard(ElementBeingDestroyed, &ExistingElements[Index]);
+				DestroyComponents(&ExistingElements[Index]);
 				DeallocateElementStorage(&ExistingElements[Index]);
 				ExistingElements[Index].~FRigControlElement(); 
 			}
@@ -5665,6 +5860,7 @@ void URigHierarchy::DestroyElement(FRigBaseElement*& InElement)
 			for(int32 Index=0;Index<Count;Index++)
 			{
 				TGuardValue<const FRigBaseElement*> DestroyGuard(ElementBeingDestroyed, &ExistingElements[Index]);
+				DestroyComponents(&ExistingElements[Index]);
 				DeallocateElementStorage(&ExistingElements[Index]);
 				ExistingElements[Index].~FRigCurveElement(); 
 			}
@@ -5676,6 +5872,7 @@ void URigHierarchy::DestroyElement(FRigBaseElement*& InElement)
 			for(int32 Index=0;Index<Count;Index++)
 			{
 				TGuardValue<const FRigBaseElement*> DestroyGuard(ElementBeingDestroyed, &ExistingElements[Index]);
+				DestroyComponents(&ExistingElements[Index]);
 				DeallocateElementStorage(&ExistingElements[Index]);
 				ExistingElements[Index].~FRigPhysicsElement(); 
 			}
@@ -5687,6 +5884,7 @@ void URigHierarchy::DestroyElement(FRigBaseElement*& InElement)
 			for(int32 Index=0;Index<Count;Index++)
 			{
 				TGuardValue<const FRigBaseElement*> DestroyGuard(ElementBeingDestroyed, &ExistingElements[Index]);
+				DestroyComponents(&ExistingElements[Index]);
 				DeallocateElementStorage(&ExistingElements[Index]);
 				ExistingElements[Index].~FRigReferenceElement(); 
 			}
@@ -5698,6 +5896,7 @@ void URigHierarchy::DestroyElement(FRigBaseElement*& InElement)
 			for(int32 Index=0;Index<Count;Index++)
 			{
 				TGuardValue<const FRigBaseElement*> DestroyGuard(ElementBeingDestroyed, &ExistingElements[Index]);
+				DestroyComponents(&ExistingElements[Index]);
 				DeallocateElementStorage(&ExistingElements[Index]);
 				ExistingElements[Index].~FRigConnectorElement(); 
 			}
@@ -5709,6 +5908,7 @@ void URigHierarchy::DestroyElement(FRigBaseElement*& InElement)
 			for(int32 Index=0;Index<Count;Index++)
 			{
 				TGuardValue<const FRigBaseElement*> DestroyGuard(ElementBeingDestroyed, &ExistingElements[Index]);
+				DestroyComponents(&ExistingElements[Index]);
 				DeallocateElementStorage(&ExistingElements[Index]);
 				ExistingElements[Index].~FRigSocketElement(); 
 			}
@@ -5727,6 +5927,121 @@ void URigHierarchy::DestroyElement(FRigBaseElement*& InElement)
 	ElementTransformRanges.Reset();
 }
 
+FRigBaseComponent* URigHierarchy::MakeComponent(const UScriptStruct* InComponentStruct, const FName& InName, FRigBaseElement* InElement)
+{
+	check(InComponentStruct);
+	check(InComponentStruct->IsChildOf(FRigBaseComponent::StaticStruct()));
+	check(InElement);
+	check(InElement->GetOwner() == this);
+	
+	int32 ComponentIndex = INDEX_NONE;
+	for(int32 ExistingIndex = 0; ExistingIndex < ElementComponents.Num(); ExistingIndex++)
+	{
+		if(!ElementComponents[ExistingIndex].IsValid())
+		{
+			ComponentIndex = ExistingIndex;
+			ElementComponents[ComponentIndex] = FInstancedStruct(InComponentStruct);
+			break;
+		}
+	}
+
+	if(ComponentIndex == INDEX_NONE)
+	{
+		ComponentIndex = ElementComponents.Emplace(InComponentStruct);
+	}
+
+	FRigName SanitizedName = InName;
+	SanitizeName(SanitizedName, false);
+
+	FRigName UniqueName = SanitizedName;
+
+	int32 Suffix = 1;
+	while (InElement->GetComponentByName(UniqueName) != nullptr)
+	{
+		FString BaseString = SanitizedName.GetName();
+		if (BaseString.Len() > GetMaxNameLength() - 4)
+		{
+			BaseString.LeftChopInline(BaseString.Len() - (GetMaxNameLength() - 4));
+		}
+		UniqueName.SetName(FString::Printf(TEXT("%s_%d"), *BaseString, ++Suffix));
+	}
+	
+	FRigBaseComponent* Component = ElementComponents[ComponentIndex].GetMutablePtr<FRigBaseComponent>();
+	Component->Name = UniqueName;
+	Component->Element = InElement;
+	Component->ElementKey = InElement->GetKey();
+	Component->IndexInHierarchy = ComponentIndex;
+	Component->IndexInElement = InElement->ComponentIndices.Add(Component->IndexInHierarchy);
+
+	return Component;
+}
+
+void URigHierarchy::DestroyComponent(FRigBaseComponent*& InComponent)
+{
+	for(int32 ComponentIndex = 0; ComponentIndex < ElementComponents.Num(); ComponentIndex ++)
+	{
+		FInstancedStruct& InstancedStruct = ElementComponents[ComponentIndex];
+		const FRigBaseComponent* Component = InstancedStruct.GetMutablePtr<FRigBaseComponent>();
+		if(Component == InComponent)
+		{
+			if(FRigBaseElement* Element = Find(Component->GetElementKey()))
+			{
+				Element->ComponentIndices.Remove(ComponentIndex);
+				
+				for(int32 IndexInElement = 0; IndexInElement < Element->ComponentIndices.Num(); IndexInElement++)
+				{
+					if(FRigBaseComponent* RemainingComponent = GetComponent(Element->ComponentIndices[IndexInElement]))
+					{
+						RemainingComponent->IndexInElement = IndexInElement;
+					}
+				}
+			}
+			InstancedStruct.Reset();
+			InComponent = nullptr;
+			break;
+		}
+	}
+}
+
+void URigHierarchy::DestroyComponents(FRigBaseElement* InElement)
+{
+	for(const int32& ComponentIndex : InElement->ComponentIndices)
+	{
+		ElementComponents[ComponentIndex].Reset();
+	}
+	InElement->ComponentIndices.Reset();
+}
+
+void URigHierarchy::ShrinkComponentStorage()
+{
+	// only shrink the memory if any of the components is invalid
+	if(!ElementComponents.ContainsByPredicate([](const FInstancedStruct& InstancedStruct) -> bool
+	{
+		return !InstancedStruct.IsValid();
+	}))
+	{
+		return;
+	}
+
+	for(FRigBaseElement* Element : Elements)
+	{
+		Element->ComponentIndices.Reset();
+	}
+
+	ElementComponents.RemoveAll([](const FInstancedStruct& InstancedStruct) -> bool
+	{
+		return !InstancedStruct.IsValid();
+	});
+
+	for(int32 ComponentIndex = 0; ComponentIndex < ElementComponents.Num(); ComponentIndex ++)
+	{
+		FInstancedStruct& InstancedStruct = ElementComponents[ComponentIndex];
+		FRigBaseComponent* Component = InstancedStruct.GetMutablePtr<FRigBaseComponent>();
+		Component->Element = Find(Component->GetElementKey());
+		Component->Element->ComponentIndices.Add(ComponentIndex);
+	}
+}
+
 void URigHierarchy::PropagateDirtyFlags(FRigTransformElement* InTransformElement, bool bInitial, bool bAffectChildren, bool bComputeOpposed, bool bMarkDirty) const
 {
 	if(!bEnableDirtyPropagation)
diff --git a/Engine/Plugins/Animation/ControlRig/Source/ControlRig/Private/Rigs/RigHierarchyCache.cpp b/Engine/Plugins/Animation/ControlRig/Source/ControlRig/Private/Rigs/RigHierarchyCache.cpp
--- a/Engine/Plugins/Animation/ControlRig/Source/ControlRig/Private/Rigs/RigHierarchyCache.cpp
+++ b/Engine/Plugins/Animation/ControlRig/Source/ControlRig/Private/Rigs/RigHierarchyCache.cpp
@@ -6,6 +6,10 @@
 
 #include UE_INLINE_GENERATED_CPP_BY_NAME(RigHierarchyCache)
 
+////////////////////////////////////////////////////////////////////////////////
+// FCachedRigElement
+////////////////////////////////////////////////////////////////////////////////
+
 const FRigElementKey& FCachedRigElement::GetResolvedKey() const
 {
 	if(Element)
@@ -65,7 +69,7 @@ bool FCachedRigElement::UpdateCache(const FRigElementKey& InKey, const URigHiera
 			}
 			else
 			{
-				Reset();
+				Invalidate();
 				Key = KeyToResolve;
 			}
 
@@ -81,6 +85,72 @@ bool FCachedRigElement::IsIdentical(const FRigElementKey& InKey, const URigHiera
 	return InKey == Key && InHierarchy->GetTopologyVersionHash() == ContainerVersion;
 }
 
+////////////////////////////////////////////////////////////////////////////////
+// FCachedRigComponent
+////////////////////////////////////////////////////////////////////////////////
+
+const FRigElementKey& FCachedRigComponent::GetResolvedElementKey() const
+{
+	if(Component)
+	{
+		if(const FRigBaseElement* Element = Component->GetElement())
+		{
+			return Element->GetKey();
+		}
+	}
+	static FRigElementKey InvalidKey;
+	return InvalidKey;
+}
+
+bool FCachedRigComponent::UpdateCache(const URigHierarchy* InHierarchy)
+{
+	if(!IsValid() || InHierarchy->GetTopologyVersionHash() != CachedElement.ContainerVersion || Component != InHierarchy->GetComponent(Index))
+	{
+		return UpdateCache(GetElementKey(), GetName(), InHierarchy);
+	}
+	return IsValid();
+}
+
+bool FCachedRigComponent::UpdateCache(const FRigElementKey& InKey, const FName& InName, const URigHierarchy* InHierarchy)
+{
+	if(InHierarchy)
+	{
+		if(!IsValid() || !IsIdentical(InKey, InName, InHierarchy) || Component != InHierarchy->GetComponent(Index))
+		{
+			if(CachedElement.UpdateCache(InKey, InHierarchy))
+			{
+				if(const FRigBaseElement* Element = CachedElement.GetElement())
+				{
+					Component = Element->GetComponentByName(InName);
+					if(Component)
+					{
+						Index = Component->GetIndexInHierarchy();
+					}
+					else
+					{
+						Invalidate();
+					}
+				}
+				else
+				{
+					Invalidate();
+				}
+			}
+		}
+		return IsValid();
+	}
+	return false;
+}
+
+bool FCachedRigComponent::IsIdentical(const FRigElementKey& InKey, const FName& InName, const URigHierarchy* InHierarchy)
+{
+	return InName == Name && CachedElement.IsIdentical(InKey, InHierarchy);
+}
+
+////////////////////////////////////////////////////////////////////////////////
+// FRigElementKeyRedirector
+////////////////////////////////////////////////////////////////////////////////
+
 FRigElementKeyRedirector::FRigElementKeyRedirector(const TMap<FRigElementKey, FRigElementKey>& InMap, const URigHierarchy* InHierarchy)
 {
 	check(InHierarchy);
diff --git a/Engine/Plugins/Animation/ControlRig/Source/ControlRig/Private/Rigs/RigHierarchyComponents.cpp b/Engine/Plugins/Animation/ControlRig/Source/ControlRig/Private/Rigs/RigHierarchyComponents.cpp
new file mode 100644
--- /dev/null
+++ b/Engine/Plugins/Animation/ControlRig/Source/ControlRig/Private/Rigs/RigHierarchyComponents.cpp
@@ -0,0 +1,47 @@
+// Copyright Epic Games, Inc. All Rights Reserved.
+
+#include "Rigs/RigHierarchyComponents.h"
+#include "Rigs/RigHierarchyElements.h"
+#include "ControlRigObjectVersion.h"
+
+#include UE_INLINE_GENERATED_CPP_BY_NAME(RigHierarchyComponents)
+
+////////////////////////////////////////////////////////////////////////////////
+// FRigBaseComponent
+////////////////////////////////////////////////////////////////////////////////
+
+FRigBaseComponent::FRigBaseComponent(FRigBaseElement* InElement)
+	: Element(InElement)
+	, ElementKey(InElement->GetKey())
+{
+}
+
+FRigBaseComponent::~FRigBaseComponent()
+{
+}
+
+bool FRigBaseComponent::Serialize(FArchive& Ar)
+{
+	if(Ar.IsLoading())
+	{
+		Load(Ar);
+	}
+	else if(Ar.IsSaving())
+	{
+		Save(Ar);
+	}
+	return true;
+}
+
+void FRigBaseComponent::Save(FArchive& Ar)
+{
+	Ar << Name;
+	Ar << ElementKey;
+}
+
+void FRigBaseComponent::Load(FArchive& Ar)
+{
+	Ar << Name;
+	CachedNameString.Reset();
+	Ar << ElementKey;
+}
diff --git a/Engine/Plugins/Animation/ControlRig/Source/ControlRig/Private/Rigs/RigHierarchyController.cpp b/Engine/Plugins/Animation/ControlRig/Source/ControlRig/Private/Rigs/RigHierarchyController.cpp
--- a/Engine/Plugins/Animation/ControlRig/Source/ControlRig/Private/Rigs/RigHierarchyController.cpp
+++ b/Engine/Plugins/Animation/ControlRig/Source/ControlRig/Private/Rigs/RigHierarchyController.cpp
@@ -18,6 +18,29 @@
 
 #include UE_INLINE_GENERATED_CPP_BY_NAME(RigHierarchyController)
 
+////////////////////////////////////////////////////////////////////////////////
+// FRigHierarchyImportErrorContext
+////////////////////////////////////////////////////////////////////////////////
+
+class FRigHierarchyImportErrorContext : public FOutputDevice
+{
+public:
+
+	int32 NumErrors;
+
+	FRigHierarchyImportErrorContext()
+		: FOutputDevice()
+		, NumErrors(0)
+	{
+	}
+
+	virtual void Serialize(const TCHAR* V, ELogVerbosity::Type Verbosity, const class FName& Category) override
+	{
+		UE_LOG(LogControlRig, Error, TEXT("Error Importing To Hierarchy: %s"), V);
+		NumErrors++;
+	}
+};
+
 ////////////////////////////////////////////////////////////////////////////////
 // URigHierarchyController
 ////////////////////////////////////////////////////////////////////////////////
@@ -858,6 +881,198 @@ bool URigHierarchyController::SetControlSettings(FRigElementKey InKey, FRigContr
 	return true;
 }
 
+FName URigHierarchyController::AddComponent(UScriptStruct* InComponentStruct, FName InName, FRigElementKey InElement, FString InContent, bool bSetupUndo, bool bPrintPythonCommand)
+{
+	if(!IsValid())
+	{
+		return NAME_None;
+	}
+	
+	if(InComponentStruct == nullptr)
+	{
+		ReportError(TEXT("The passed component struct is nullptr."));
+		return NAME_None;
+	}
+
+	if(InName.IsNone())
+	{
+		ReportError(TEXT("The passed name is None."));
+		return NAME_None;
+	}
+
+	URigHierarchy* Hierarchy = GetHierarchy();
+	check(Hierarchy);
+	
+	FRigBaseElement* Element = Hierarchy->Find(InElement);
+	if(Element == nullptr)
+	{
+		ReportErrorf(TEXT("The element '%s' could not be found."), *InElement.ToString());
+		return NAME_None;
+	}
+
+#if WITH_EDITOR
+	TSharedPtr<FScopedTransaction> TransactionPtr;
+	if(bSetupUndo)
+	{
+		TransactionPtr = MakeShared<FScopedTransaction>(NSLOCTEXT("RigHierarchyController", "Add Component", "Add Component"));
+		Hierarchy->Modify();
+	}
+#endif
+
+	const FRigBaseComponent* NewComponent = Hierarchy->MakeComponent(InComponentStruct, InName, Element);
+	(void)SetComponentContent(NewComponent->GetFName(), NewComponent->GetElementKey(), InContent, false, false);
+	Hierarchy->IncrementTopologyVersion();
+	Notify(ERigHierarchyNotification::ComponentAdded, NewComponent);
+
+#if WITH_EDITOR
+	TransactionPtr.Reset();
+
+	if (bPrintPythonCommand && !bSuspendPythonPrinting)
+	{
+		if (const UBlueprint* Blueprint = GetTypedOuter<UBlueprint>())
+		{
+			TArray<FString> Commands = GetAddComponentPythonCommands(NewComponent);
+			for (const FString& Command : Commands)
+			{			
+				RigVMPythonUtils::Print(Blueprint->GetFName().ToString(),
+					FString::Printf(TEXT("%s"), *Command));
+			}
+		}
+	}
+#endif
+
+	return NewComponent->GetFName();
+}
+
+bool URigHierarchyController::RemoveComponent(FName InName, FRigElementKey InElement, bool bSetupUndo, bool bPrintPythonCommand)
+{
+	if(!IsValid())
+	{
+		return false;
+	}
+	
+	if(InName.IsNone())
+	{
+		ReportError(TEXT("The passed name is None."));
+		return false;
+	}
+
+	URigHierarchy* Hierarchy = GetHierarchy();
+	check(Hierarchy);
+	
+	FRigBaseElement* Element = Hierarchy->Find(InElement);
+	if(Element == nullptr)
+	{
+		ReportErrorf(TEXT("The element '%s' could not be found."), *InElement.ToString());
+		return false;
+	}
+
+	FRigBaseComponent* Component = Element->GetComponentByName(InName);
+	if(Component == nullptr)
+	{
+		ReportErrorf(TEXT("The component '%s' cannot be found on element '%s'."), *InName.ToString(), *InElement.ToString());
+		return false;
+	}
+
+#if WITH_EDITOR
+	TSharedPtr<FScopedTransaction> TransactionPtr;
+	if(bSetupUndo)
+	{
+		TransactionPtr = MakeShared<FScopedTransaction>(NSLOCTEXT("RigHierarchyController", "Remove Component", "Remove Component"));
+		Hierarchy->Modify();
+	}
+#endif
+
+	Notify(ERigHierarchyNotification::ComponentRemoved, Component);
+	Hierarchy->DestroyComponent(Component);
+	Hierarchy->IncrementTopologyVersion();
+
+#if WITH_EDITOR
+	TransactionPtr.Reset();
+
+	if (bPrintPythonCommand && !bSuspendPythonPrinting)
+	{
+		if (const UBlueprint* Blueprint = GetTypedOuter<UBlueprint>())
+		{
+			RigVMPythonUtils::Print(Blueprint->GetFName().ToString(),
+				FString::Printf(TEXT("hierarchy_controller.remove_component('%s', %s"), *InName.ToString(), *InElement.ToPythonString()));
+		}
+	}
+#endif
+
+	return true;
+}
+
+bool URigHierarchyController::SetComponentContent(FName InName, FRigElementKey InElement, const FString& InContent, bool bSetupUndo, bool bPrintPythonCommand)
+{
+	if(!IsValid())
+	{
+		return false;
+	}
+	
+	if(InName.IsNone())
+	{
+		ReportError(TEXT("The passed name is None."));
+		return false;
+	}
+
+	URigHierarchy* Hierarchy = GetHierarchy();
+	check(Hierarchy);
+	
+	FRigBaseElement* Element = Hierarchy->Find(InElement);
+	if(Element == nullptr)
+	{
+		ReportErrorf(TEXT("The element '%s' could not be found."), *InElement.ToString());
+		return false;
+	}
+
+	FRigBaseComponent* Component = Element->GetComponentByName(InName);
+	if(Component == nullptr)
+	{
+		ReportErrorf(TEXT("The component '%s' cannot be found on element '%s'."), *InName.ToString(), *InElement.ToString());
+		return false;
+	}
+
+#if WITH_EDITOR
+	TSharedPtr<FScopedTransaction> TransactionPtr;
+	if(bSetupUndo)
+	{
+		TransactionPtr = MakeShared<FScopedTransaction>(NSLOCTEXT("RigHierarchyController", "Set Component Content", "Set Component Content"));
+		Hierarchy->Modify();
+	}
+#endif
+
+	// create a dummy component first to import onto
+	FStructOnScope ScopedStruct(Component->GetScriptStruct());
+	FRigHierarchyImportErrorContext ErrorPipe;
+	Component->GetScriptStruct()->ImportText(*InContent, ScopedStruct.GetStructMemory(), nullptr, EPropertyPortFlags::PPF_None, &ErrorPipe, Component->GetScriptStruct()->GetName(), true);
+	if (ErrorPipe.NumErrors > 0)
+	{
+		return false;;
+	}
+
+	// now import onto the actual component
+	ErrorPipe.NumErrors = 0;
+	Component->GetScriptStruct()->ImportText(*InContent, Component, nullptr, EPropertyPortFlags::PPF_None, &ErrorPipe, Component->GetScriptStruct()->GetName(), true);
+
+	Notify(ERigHierarchyNotification::ComponentContentChanged, Component);
+
+#if WITH_EDITOR
+	TransactionPtr.Reset();
+
+	if (bPrintPythonCommand && !bSuspendPythonPrinting)
+	{
+		if (const UBlueprint* Blueprint = GetTypedOuter<UBlueprint>())
+		{
+			RigVMPythonUtils::Print(Blueprint->GetFName().ToString(),
+				FString::Printf(TEXT("hierarchy_controller.set_component_content('%s', %s, '%s'"), *InName.ToString(), *InElement.ToPythonString(), *InContent));
+		}
+	}
+#endif
+
+	return true;
+}
+
 TArray<FRigElementKey> URigHierarchyController::ImportBones(const FReferenceSkeleton& InSkeleton,
                                                             const FName& InNameSpace, bool bReplaceExistingBones, bool bRemoveObsoleteBones, bool bSelectBones,
                                                             bool bSetupUndo)
@@ -1406,25 +1621,6 @@ FString URigHierarchyController::ExportToText(TArray<FRigElementKey> InKeys) con
 	return ExportedText;
 }
 
-class FRigHierarchyImportErrorContext : public FOutputDevice
-{
-public:
-
-	int32 NumErrors;
-
-	FRigHierarchyImportErrorContext()
-        : FOutputDevice()
-        , NumErrors(0)
-	{
-	}
-
-	virtual void Serialize(const TCHAR* V, ELogVerbosity::Type Verbosity, const class FName& Category) override
-	{
-		UE_LOG(LogControlRig, Error, TEXT("Error Importing To Hierarchy: %s"), V);
-		NumErrors++;
-	}
-};
-
 TArray<FRigElementKey> URigHierarchyController::ImportFromText(FString InContent, bool bReplaceExistingElements, bool bSelectNewElements, bool bSetupUndo, bool bPrintPythonCommands)
 {
 	return ImportFromText(InContent, ERigElementType::All, bReplaceExistingElements, bSelectNewElements, bSetupUndo, bPrintPythonCommands);
@@ -2113,9 +2309,29 @@ TArray<FString> URigHierarchyController::GetSetControlShapeTransformPythonComman
 		*RigVMPythonUtils::TransformToPythonString(InTransform),
 		bInitial ? TEXT("True") : TEXT("False"))};
 }
+
+TArray<FString> URigHierarchyController::GetAddComponentPythonCommands(const FRigBaseComponent* Component) const
+{
+	TArray<FString> Commands;
+
+	const FString ElementKeyStr = Component->GetElement()->GetKey().ToPythonString();
+
+	FString ContentStr;
+	Component->GetScriptStruct()->ExportText(ContentStr, Component, Component, nullptr, PPF_None, nullptr);
+	
+	// AddComponent(UScriptStruct* InComponentStruct, FName InName, FRigElementKey InElement, FString InContent);
+	Commands.Add(FString::Printf(TEXT("hierarchy_controller.add_component(unreal.%s, '%s', %s, '%s')"),
+		*Component->GetScriptStruct()->GetName(),
+		*Component->GetName(),
+		*ElementKeyStr,
+		*ContentStr
+	));
+
+	return Commands;
+}
 #endif
 
-void URigHierarchyController::Notify(ERigHierarchyNotification InNotifType, const FRigBaseElement* InElement)
+void URigHierarchyController::Notify(ERigHierarchyNotification InNotifType, const FRigNotificationSubject& InSubject)
 {
 	if(!IsValid())
 	{
@@ -2133,10 +2349,10 @@ void URigHierarchyController::Notify(ERigHierarchyNotification InNotifType, cons
 			return;
 		}
 	}	
-	GetHierarchy()->Notify(InNotifType, InElement);
+	GetHierarchy()->Notify(InNotifType, InSubject);
 }
 
-void URigHierarchyController::HandleHierarchyModified(ERigHierarchyNotification InNotifType, URigHierarchy* InHierarchy, const FRigBaseElement* InElement) const
+void URigHierarchyController::HandleHierarchyModified(ERigHierarchyNotification InNotifType, URigHierarchy* InHierarchy, const FRigNotificationSubject& InSubject) const
 {
 	if(bSuspendAllNotifications)
 	{
@@ -2144,7 +2360,7 @@ void URigHierarchyController::HandleHierarchyModified(ERigHierarchyNotification
 	}
 	ensure(IsValid());
 	ensure(InHierarchy == GetHierarchy());
-	ModifiedEvent.Broadcast(InNotifType, InHierarchy, InElement);
+	ModifiedEvent.Broadcast(InNotifType, InHierarchy, InSubject);
 }
 
 bool URigHierarchyController::IsValid() const
@@ -2410,7 +2626,7 @@ bool URigHierarchyController::RemoveElement(FRigBaseElement* InElement)
 	Notify(ERigHierarchyNotification::ElementRemoved, InElement);
 	if(Hierarchy->Num() == 0)
 	{
-		Notify(ERigHierarchyNotification::HierarchyReset, nullptr);
+		Notify(ERigHierarchyNotification::HierarchyReset, {});
 	}
 
 	if (InElement->OwnedInstances == 1)
diff --git a/Engine/Plugins/Animation/ControlRig/Source/ControlRig/Private/Rigs/RigHierarchyElements.cpp b/Engine/Plugins/Animation/ControlRig/Source/ControlRig/Private/Rigs/RigHierarchyElements.cpp
--- a/Engine/Plugins/Animation/ControlRig/Source/ControlRig/Private/Rigs/RigHierarchyElements.cpp
+++ b/Engine/Plugins/Animation/ControlRig/Source/ControlRig/Private/Rigs/RigHierarchyElements.cpp
@@ -159,6 +159,80 @@ void FRigBaseElement::NotifyMetadataTagChanged(const FName& InTag, bool bAdded)
 	}
 }
 
+int32 FRigBaseElement::NumComponents() const
+{
+	return ComponentIndices.Num();
+}
+
+const FRigBaseComponent* FRigBaseElement::GetComponent(int32 InIndex) const
+{
+	return const_cast<FRigBaseElement*>(this)->GetComponent(InIndex);
+}
+
+FRigBaseComponent* FRigBaseElement::GetComponent(int32 InIndex)
+{
+	if(ComponentIndices.IsValidIndex(InIndex))
+	{
+		if(URigHierarchy* Hierarchy = GetOwner())
+		{
+			return Hierarchy->GetComponent(ComponentIndices[InIndex]);
+		}
+	}
+	return nullptr;
+}
+
+const FRigBaseComponent* FRigBaseElement::GetComponentByName(const FName& InName) const
+{
+	return const_cast<FRigBaseElement*>(this)->GetComponentByName(InName);
+}
+
+FRigBaseComponent* FRigBaseElement::GetComponentByName(const FName& InName)
+{
+	if(!ComponentIndices.IsEmpty())
+	{
+		if(URigHierarchy* Hierarchy = GetOwner())
+		{
+			for(int32 ComponentIndex : ComponentIndices)
+			{
+				if(FRigBaseComponent* Component = Hierarchy->GetComponent(ComponentIndex))
+				{
+					if(Component->GetName() == InName)
+					{
+						return Component;
+					}
+				}
+			}
+		}
+	}
+	return nullptr;
+}
+
+const FRigBaseComponent* FRigBaseElement::GetComponentOfType(const UScriptStruct* InComponentStruct) const
+{
+	return const_cast<FRigBaseElement*>(this)->GetComponentOfType(InComponentStruct);
+}
+
+FRigBaseComponent* FRigBaseElement::GetComponentOfType(const UScriptStruct* InComponentStruct)
+{
+	if(!ComponentIndices.IsEmpty())
+	{
+		if(URigHierarchy* Hierarchy = GetOwner())
+		{
+			for(const int32& ComponentIndex : ComponentIndices)
+			{
+				if(FRigBaseComponent* Component = Hierarchy->GetComponent(ComponentIndex))
+				{
+					if(Component->IsA(InComponentStruct))
+					{
+						return Component;
+					}
+				}
+			}
+		}
+	}
+	return nullptr;
+}
+
 
 void FRigBaseElement::InitializeFrom(const FRigBaseElement* InOther)
 {
@@ -548,6 +622,40 @@ FRigBaseElement* FRigElementHandle::Get()
 	return nullptr;
 }
 
+////////////////////////////////////////////////////////////////////////////////
+// FRigComponentHandle
+////////////////////////////////////////////////////////////////////////////////
+
+FRigComponentHandle::FRigComponentHandle(URigHierarchy* InHierarchy, const FRigElementKey& InKey, const FName& InComponentName)
+: ElementHandle(InHierarchy, InKey)
+, ComponentName(InComponentName)
+{
+}
+
+FRigComponentHandle::FRigComponentHandle(URigHierarchy* InHierarchy, const FRigBaseComponent* InComponent)
+: ElementHandle(InHierarchy, InComponent->GetElement())
+, ComponentName(InComponent->GetFName())
+{
+}
+
+const FRigBaseComponent* FRigComponentHandle::Get() const
+{
+	if(const FRigBaseElement* Element = ElementHandle.Get())
+	{
+		return Element->GetComponentByName(ComponentName);
+	}
+	return nullptr;
+}
+
+FRigBaseComponent* FRigComponentHandle::Get()
+{
+	if(FRigBaseElement* Element = ElementHandle.Get())
+	{
+		return Element->GetComponentByName(ComponentName);
+	}
+	return nullptr;
+}
+
 ////////////////////////////////////////////////////////////////////////////////
 // FRigTransformElement
 ////////////////////////////////////////////////////////////////////////////////
diff --git a/Engine/Plugins/Animation/ControlRig/Source/ControlRig/Public/Constraints/ControlRigTransformableHandle.h b/Engine/Plugins/Animation/ControlRig/Source/ControlRig/Public/Constraints/ControlRigTransformableHandle.h
--- a/Engine/Plugins/Animation/ControlRig/Source/ControlRig/Public/Constraints/ControlRigTransformableHandle.h
+++ b/Engine/Plugins/Animation/ControlRig/Source/ControlRig/Public/Constraints/ControlRigTransformableHandle.h
@@ -134,7 +134,7 @@ private:
 	void OnHierarchyModified(
 		ERigHierarchyNotification InNotif,
 		URigHierarchy* InHierarchy,
-		const FRigBaseElement* InElement);
+		const FRigNotificationSubject& InSubject);
 
 	void OnControlRigBound(UControlRig* InControlRig);
 	void OnObjectBoundToControlRig(UObject* InObject);
diff --git a/Engine/Plugins/Animation/ControlRig/Source/ControlRig/Public/ControlRig.h b/Engine/Plugins/Animation/ControlRig/Source/ControlRig/Public/ControlRig.h
--- a/Engine/Plugins/Animation/ControlRig/Source/ControlRig/Public/ControlRig.h
+++ b/Engine/Plugins/Animation/ControlRig/Source/ControlRig/Public/ControlRig.h
@@ -497,7 +497,7 @@ private:
 	FControlRigExecuteEvent PreAdditiveValuesApplicationEvent;
 
 	/** Handle changes within the hierarchy */
-	void HandleHierarchyModified(ERigHierarchyNotification InNotification, URigHierarchy* InHierarchy, const FRigBaseElement* InElement);
+	void HandleHierarchyModified(ERigHierarchyNotification InNotification, URigHierarchy* InHierarchy, const FRigNotificationSubject& InSubject);
 
 protected:
 	
diff --git a/Engine/Plugins/Animation/ControlRig/Source/ControlRig/Public/ControlRigObjectVersion.h b/Engine/Plugins/Animation/ControlRig/Source/ControlRig/Public/ControlRigObjectVersion.h
--- a/Engine/Plugins/Animation/ControlRig/Source/ControlRig/Public/ControlRigObjectVersion.h
+++ b/Engine/Plugins/Animation/ControlRig/Source/ControlRig/Public/ControlRigObjectVersion.h
@@ -110,6 +110,9 @@ struct CONTROLRIG_API FControlRigObjectVersion
 		// Compress the rig hierarchy when storing to disk
 		RigHierarchyCompressElements,
 
+		// Added the notion of components to the rig hierarchy
+		RigHierarchyStoresComponents,
+
 		// -----<new versions can be added above this line>-------------------------------------------------
 		VersionPlusOne,
 		LatestVersion = VersionPlusOne - 1
diff --git a/Engine/Plugins/Animation/ControlRig/Source/ControlRig/Public/Rigs/RigHierarchy.h b/Engine/Plugins/Animation/ControlRig/Source/ControlRig/Public/Rigs/RigHierarchy.h
--- a/Engine/Plugins/Animation/ControlRig/Source/ControlRig/Public/Rigs/RigHierarchy.h
+++ b/Engine/Plugins/Animation/ControlRig/Source/ControlRig/Public/Rigs/RigHierarchy.h
@@ -26,7 +26,7 @@ class URigHierarchy;
 class URigHierarchyController;
 class UModularRigRuleManager; 
 
-DECLARE_MULTICAST_DELEGATE_ThreeParams(FRigHierarchyModifiedEvent, ERigHierarchyNotification /* type */, URigHierarchy* /* hierarchy */, const FRigBaseElement* /* element */);
+DECLARE_MULTICAST_DELEGATE_ThreeParams(FRigHierarchyModifiedEvent, ERigHierarchyNotification /* type */, URigHierarchy* /* hierarchy */, const FRigNotificationSubject& /* element or component */);
 DECLARE_DYNAMIC_MULTICAST_DELEGATE_ThreeParams(FRigHierarchyModifiedDynamicEvent, ERigHierarchyNotification, NotifType, URigHierarchy*, Hierarchy, FRigElementKey, Subject);
 DECLARE_EVENT_FiveParams(URigHierarchy, FRigHierarchyUndoRedoTransformEvent, URigHierarchy*, const FRigElementKey&, ERigTransformType::Type, const FTransform&, bool /* bUndo */);
 DECLARE_MULTICAST_DELEGATE_TwoParams(FRigHierarchyMetadataChangedDelegate, const FRigElementKey& /* Key */, const FName& /* Name */);
@@ -705,6 +705,84 @@ public:
 		return Get<T>(GetIndex(InKey));
 	}
 
+	/**
+	 * Returns a component at a given index or nullptr.
+	 * @param InIndex The index of the component to retrieve.
+	 * @return The retrieved component or nullptr.
+	 */
+	const FRigBaseComponent* GetComponent(int32 InIndex) const;
+
+	/**
+	 * Returns a component at a given index or nullptr.
+	 * @param InIndex The index of the component to retrieve.
+	 * @return The retrieved component or nullptr.
+	 */
+	FRigBaseComponent* GetComponent(int32 InIndex);
+
+	/**
+	 * Returns a component at a given index or nullptr.
+	 * This templated method also casts to the chosen
+	 * component type but does not guarantee a valid result.
+	 * @param InIndex The index of the component to retrieve.
+	 * @return The retrieved component or nullptr.
+	 */
+	template<typename T>
+	const T* GetComponent(int32 InIndex) const
+	{
+		return Cast<T>(GetComponent(InIndex));
+	}
+
+	/**
+	 * Returns a component at a given index or nullptr.
+	 * This templated method also casts to the chosen
+	 * component type but does not guarantee a valid result.
+	 * @param InIndex The index of the component to retrieve.
+	 * @return The retrieved component or nullptr.
+	 */
+	template<typename T>
+	T* GetComponent(int32 InIndex)
+	{
+		return Cast<T>(GetComponent(InIndex));
+	}
+
+	/**
+	 * Returns a component at a given index.
+	 * This templated method also casts to the chosen
+	 * component type and checks for a the valid result.
+	 * @param InIndex The index of the component to retrieve.
+	 * @return The retrieved component or nullptr.
+	 */
+	template<typename T>
+	const T* GetComponentChecked(int32 InIndex) const
+	{
+		return CastChecked<T>(GetComponent(InIndex));
+	}
+
+	/**
+	 * Returns a component at a given index.
+	 * This templated method also casts to the chosen
+	 * component type and checks for a the valid result.
+	 * @param InIndex The index of the component to retrieve.
+	 * @return The retrieved component or nullptr.
+	 */
+	template<typename T>
+    T* GetComponentChecked(int32 InIndex)
+	{
+		return CastChecked<T>(GetComponent(InIndex));
+	}
+	
+	UFUNCTION(BlueprintPure, Category = URigHierarchy)
+	int32 NumComponents(FRigElementKey InElement) const;
+
+	UFUNCTION(BlueprintPure, Category = URigHierarchy)
+	FName GetComponentName(FRigElementKey InElement, int32 InComponentIndex) const;
+	
+	UFUNCTION(BlueprintPure, Category = URigHierarchy)
+	UScriptStruct* GetComponentType(FRigElementKey InElement, int32 InComponentIndex) const;
+
+	UFUNCTION(BlueprintPure, Category = URigHierarchy)
+	FString GetComponentContent(FRigElementKey InElement, int32 InComponentIndex) const;
+
 private:
 	/**
 	* Returns bone element for a given key, for scripting purpose only, for cpp usage, use Find<FRigBoneElement>()
@@ -2890,7 +2968,6 @@ public:
 	/**
 	 * Sets a curve's value given its key
 	 * @param InKey The key of the element to set the value for
-	 * @param InValue The value to set on the curve
 	 * @param bSetupUndo If true the transform stack will be setup for undo / redo
 	 */
 	UFUNCTION(BlueprintCallable, Category = URigHierarchy)
@@ -3965,7 +4042,7 @@ private:
 
 public:
 
-	void Notify(ERigHierarchyNotification InNotifType, const FRigBaseElement* InElement);
+	void Notify(ERigHierarchyNotification InNotifType, const FRigNotificationSubject& InSubject);
 
 	/**
 	 * Returns a transform based on a given transform type
@@ -4251,7 +4328,7 @@ private:
 	FRigBaseElement* MakeElement(ERigElementType InElementType, int32 InCount = 1, int32* OutStructureSize = nullptr); 
 
 	/*
-	* Helper function to create an element for a given type
+	* Helper function to destroy an element
 	*/
 	void DestroyElement(FRigBaseElement*& InElement);
 
@@ -4277,6 +4354,32 @@ private:
 		return NewElements;
 	}
 
+	FRigBaseComponent* MakeComponent(const UScriptStruct* InComponentStruct, const FName& InName, FRigBaseElement* InElement); 
+
+	/*
+	 * Helper function to destroy a component
+	 */
+	void DestroyComponent(FRigBaseComponent*& InComponent);
+
+	/*
+	 * Helper function to destroy all components of an element
+	 */
+	void DestroyComponents(FRigBaseElement* InElement);
+
+	/*
+	 * Helper function to shrink component memory
+	 */
+	void ShrinkComponentStorage();
+
+	/*
+	 * Templated helper function to create an element
+	 */
+	template<typename ComponentType = FRigBaseComponent>
+	ComponentType* NewComponent(const FName& InName, FRigBaseElement* InElement)
+	{
+		return MakeComponent(ComponentType::StaticStruct(), InName, InElement);
+	}
+
 	/**
 	 * Marks all affected elements of a given element as dirty
 	 * @param InTransformElement The element that has changed
@@ -4424,6 +4527,11 @@ private:
 	// do not have URigHierarchy as an owner and therefore do not carry metadata with them.
 	FRigReusableElementStorage<FMetadataStorage> ElementMetadata;
 
+	// Element component storage. Storage is defined here rather than on the elements
+	// to reduce memory consumption. Only elements created by MakeComponent point to
+	// the component storage.
+	TArray<FInstancedStruct> ElementComponents;
+
 	// A quick-lookup cache for elements' children. Each element that has a ChildCacheIndex
 	// not equal to INDEX_NONE is an index into the offset and count cache below, which in
 	// turn contains an offset into the ChildElementCache, which stores consecutive runs of
@@ -4975,15 +5083,16 @@ private:
 	{
 		ERigHierarchyNotification Type;
 		FRigElementKey Key;
+		FName ComponentName;
 		
 		bool operator == (const FQueuedNotification& InOther) const
 		{
-			return Type == InOther.Type && Key == InOther.Key;
+			return Type == InOther.Type && Key == InOther.Key && ComponentName == InOther.ComponentName;
 		}
 	};
 	TQueue<FQueuedNotification> QueuedNotifications;
 
-	void QueueNotification(ERigHierarchyNotification InNotification, const FRigBaseElement* InElement);
+	void QueueNotification(ERigHierarchyNotification InNotification, const FRigNotificationSubject& InSubject);
 	void SendQueuedNotifications();
 	void Reset_Impl(bool bResetElements);
 #if WITH_EDITOR
@@ -5084,12 +5193,12 @@ public:
 		: Hierarchy(InHierarchy)
 	{
 		check(Hierarchy);
-		Hierarchy->Notify(ERigHierarchyNotification::InteractionBracketOpened, nullptr);
+		Hierarchy->Notify(ERigHierarchyNotification::InteractionBracketOpened, {});
 	}
 
 	~FRigHierarchyInteractionBracket()
 	{
-		Hierarchy->Notify(ERigHierarchyNotification::InteractionBracketClosed, nullptr);
+		Hierarchy->Notify(ERigHierarchyNotification::InteractionBracketClosed, {});
 	}
 
 private:
diff --git a/Engine/Plugins/Animation/ControlRig/Source/ControlRig/Public/Rigs/RigHierarchyCache.h b/Engine/Plugins/Animation/ControlRig/Source/ControlRig/Public/Rigs/RigHierarchyCache.h
--- a/Engine/Plugins/Animation/ControlRig/Source/ControlRig/Public/Rigs/RigHierarchyCache.h
+++ b/Engine/Plugins/Animation/ControlRig/Source/ControlRig/Public/Rigs/RigHierarchyCache.h
@@ -22,12 +22,14 @@ public:
 		: Key()
 		, Index(UINT16_MAX)
 		, ContainerVersion(INDEX_NONE)
+		, Element(nullptr)
 	{}
 
 	FCachedRigElement(const FRigElementKey& InKey, const URigHierarchy* InHierarchy, bool bForceStoreKey = false)
 		: Key()
 		, Index(UINT16_MAX)
 		, ContainerVersion(INDEX_NONE)
+		, Element(nullptr)
 	{
 		UpdateCache(InKey, InHierarchy);
 		if(bForceStoreKey)
@@ -41,6 +43,13 @@ public:
 		return GetIndex() != INDEX_NONE && Key.IsValid();
 	}
 
+	void Invalidate()
+	{
+		Index = UINT16_MAX;
+		ContainerVersion = INDEX_NONE;
+		Element = nullptr;
+	}
+
 	void Reset()
 	{
 		Key = FRigElementKey();
@@ -85,6 +94,15 @@ public:
 		return Element;
 	}
 
+	const FRigBaseElement* GetElement(const URigHierarchy* InHierarchy)
+	{
+		if(UpdateCache(InHierarchy))
+		{
+			return Element;
+		}
+		return nullptr;
+	}
+
 	bool UpdateCache(const URigHierarchy* InHierarchy);
 
 	bool UpdateCache(const FRigElementKey& InKey, const URigHierarchy* InHierarchy);
@@ -156,6 +174,172 @@ private:
 	int32 ContainerVersion;
 
 	const FRigBaseElement* Element;
+	friend struct FCachedRigComponent;
+};
+
+USTRUCT(BlueprintType)
+struct CONTROLRIG_API FCachedRigComponent
+{
+	GENERATED_BODY()
+
+public:
+
+	FCachedRigComponent()
+		: CachedElement()
+		, Name(NAME_None)
+		, Index(UINT16_MAX)
+		, Component(nullptr)
+	{}
+
+	FCachedRigComponent(const FRigElementKey& InKey, const FName& InName, const URigHierarchy* InHierarchy, bool bForceStoreName = false)
+		: CachedElement(InKey, InHierarchy, bForceStoreName)
+		, Name(NAME_None)
+		, Index(UINT16_MAX)
+		, Component(nullptr)
+	{
+		UpdateCache(InKey, InName, InHierarchy);
+		if(bForceStoreName)
+		{
+			Name = InName;
+		}
+	}
+
+	bool IsValid() const
+	{
+		return GetIndex() != INDEX_NONE && CachedElement.IsValid();
+	}
+
+	void Invalidate()
+	{
+		CachedElement.Invalidate();
+		Index = UINT16_MAX;
+		Component = nullptr;
+	}
+
+	void Reset()
+	{
+		CachedElement.Reset();
+		Index = UINT16_MAX;
+		Name = NAME_None;
+		Component = nullptr;
+	}
+
+	explicit operator bool() const
+	{
+		return IsValid();
+	}
+
+	operator int32() const
+	{
+		return GetIndex();
+	}
+
+	int32 GetIndex() const
+	{
+		if(Index == UINT16_MAX)
+		{
+			return INDEX_NONE;
+		}
+		return (int32)Index;
+	}
+
+	const FRigElementKey& GetElementKey() const
+	{
+		return CachedElement.GetKey();
+	}
+
+	const FRigElementKey& GetResolvedElementKey() const;
+
+	const FName& GetName() const
+	{
+		return Name;
+	}
+
+	const FRigBaseElement* GetElement() const
+	{
+		return CachedElement.GetElement();
+	}
+
+	const FRigBaseElement* GetElement(const URigHierarchy* InHierarchy)
+	{
+		return CachedElement.GetElement(InHierarchy);
+	}
+
+	const FRigBaseComponent* GetComponent() const
+	{
+		return Component;
+	}
+
+	const FRigBaseComponent* GetComponent(const URigHierarchy* InHierarchy)
+	{
+		if(UpdateCache(InHierarchy))
+		{
+			return Component;
+		}
+		return nullptr;
+	}
+
+	bool UpdateCache(const URigHierarchy* InHierarchy);
+
+	bool UpdateCache(const FRigElementKey& InKey, const FName& InName, const URigHierarchy* InHierarchy);
+
+	friend uint32 GetTypeHash(const FCachedRigComponent& CachedComponent)
+	{
+		return GetTypeHash(CachedComponent.CachedElement) * 17 + (uint32)CachedComponent.Index;
+	}
+
+	bool IsIdentical(const FRigElementKey& InKey, const FName& InName, const URigHierarchy* InHierarchy);
+
+	bool operator ==(const FCachedRigComponent& Other) const
+	{
+		return Index == Other.Index && CachedElement == Other.CachedElement;
+	}
+
+	bool operator !=(const FCachedRigComponent& Other) const
+	{
+		return Index != Other.Index || CachedElement != Other.CachedElement;
+	}
+
+	bool operator ==(const int32& Other) const
+	{
+		return GetIndex() == Other;
+	}
+
+	bool operator !=(const int32& Other) const
+	{
+		return GetIndex() != Other;
+	}
+
+	bool operator <(const FCachedRigComponent& Other) const
+	{
+		if (CachedElement < Other.CachedElement)
+		{
+			return true;
+		}
+		return Index < Other.Index;
+	}
+
+	bool operator >(const FCachedRigComponent& Other) const
+	{
+		if (CachedElement > Other.CachedElement)
+		{
+			return true;
+		}
+		return Index > Other.Index;
+	}
+
+private:
+
+	UPROPERTY()
+	FCachedRigElement CachedElement;
+
+	UPROPERTY()
+	FName Name;
+
+	UPROPERTY()
+	uint16 Index;
+
+	const FRigBaseComponent* Component;
 };
 
 class CONTROLRIG_API FRigElementKeyRedirector
diff --git a/Engine/Plugins/Animation/ControlRig/Source/ControlRig/Public/Rigs/RigHierarchyComponents.h b/Engine/Plugins/Animation/ControlRig/Source/ControlRig/Public/Rigs/RigHierarchyComponents.h
new file mode 100644
--- /dev/null
+++ b/Engine/Plugins/Animation/ControlRig/Source/ControlRig/Public/Rigs/RigHierarchyComponents.h
@@ -0,0 +1,189 @@
+// Copyright Epic Games, Inc. All Rights Reserved.
+
+#pragma once
+
+#include "CoreMinimal.h"
+#include "RigHierarchyDefines.h"
+#include "RigHierarchyComponents.generated.h"
+
+#define DECLARE_RIG_COMPONENT_METHODS(ComponentType) \
+virtual UScriptStruct* GetScriptStruct() const override { return ComponentType::StaticStruct(); } \
+template<typename T> \
+friend const T* Cast(const ComponentType* InComponent) \
+{ \
+   return Cast<T>((const FRigBaseComponent*) InComponent); \
+} \
+template<typename T> \
+friend T* Cast(ComponentType* InComponent) \
+{ \
+   return Cast<T>((FRigBaseComponent*) InComponent); \
+} \
+template<typename T> \
+friend const T* CastChecked(const ComponentType* InComponent) \
+{ \
+	return CastChecked<T>((const FRigBaseComponent*) InComponent); \
+} \
+template<typename T> \
+friend T* CastChecked(ComponentType* InComponent) \
+{ \
+	return CastChecked<T>((FRigBaseComponent*) InComponent); \
+} \
+static bool IsClassOf(const FRigBaseComponent* InComponent) \
+{ \
+	return InComponent->GetScriptStruct()->IsChildOf(StaticStruct()); \
+}
+
+USTRUCT(BlueprintType)
+struct CONTROLRIG_API FRigBaseComponent
+{
+	GENERATED_BODY()
+
+public:
+
+	FRigBaseComponent() = default;
+	virtual ~FRigBaseComponent();
+
+	FRigBaseComponent(const FRigBaseComponent& InOther)
+	{
+		*this = InOther;
+	}
+	
+	FRigBaseComponent& operator=(const FRigBaseComponent& InOther)
+	{
+		// We purposefully do not copy any non-UPROPERTY entries, including Element
+		Name = InOther.Name;
+		CachedNameString.Reset();
+		IndexInElement = InOther.IndexInElement;
+		IndexInHierarchy = InOther.IndexInHierarchy;
+		ElementKey = InOther.ElementKey;
+		CreatedAtInstructionIndex = InOther.CreatedAtInstructionIndex;
+		bSelected = InOther.bSelected;
+		return *this;
+	}
+	
+protected:
+
+	// Only derived types should be able to construct this one.
+	explicit FRigBaseComponent(FRigBaseElement* InElement);
+
+	static bool IsClassOf(const FRigBaseComponent* InComponent)
+	{
+		return true;
+	}
+
+	// ReSharper disable once CppUE4ProbableMemoryIssuesWithUObject
+	FRigBaseElement* Element = nullptr;
+	
+	UPROPERTY(BlueprintReadOnly, VisibleAnywhere, Category = RigElement, meta = (AllowPrivateAccess = "true"))
+	FName Name = FName(NAME_None);
+
+	UPROPERTY(BlueprintReadOnly, Transient, VisibleAnywhere, Category = RigElement, meta = (AllowPrivateAccess = "true"))
+	int32 IndexInHierarchy = INDEX_NONE;
+
+	UPROPERTY(BlueprintReadOnly, Transient, VisibleAnywhere, Category = RigElement, meta = (AllowPrivateAccess = "true"))
+	int32 IndexInElement = INDEX_NONE;
+
+	UPROPERTY(BlueprintReadOnly, VisibleAnywhere, Category = RigElement, meta = (AllowPrivateAccess = "true"))
+	FRigElementKey ElementKey;
+
+	UPROPERTY(BlueprintReadOnly, Transient, Category = RigElement, meta = (AllowPrivateAccess = "true"))
+	int32 CreatedAtInstructionIndex = INDEX_NONE;
+
+	UPROPERTY(BlueprintReadOnly, Transient, Category = RigElement, meta = (AllowPrivateAccess = "true"))
+	bool bSelected = false;
+
+	mutable FString CachedNameString;
+
+public:
+
+	virtual UScriptStruct* GetScriptStruct() const { return FRigBaseComponent::StaticStruct(); }
+
+	const FName& GetFName() const { return Name; }
+	const FString& GetName() const
+	{
+		if(CachedNameString.IsEmpty() && !Name.IsNone())
+		{
+			CachedNameString = Name.ToString();
+		}
+		return CachedNameString;
+	}
+
+	virtual const FName& GetDisplayName() const { return GetFName(); }
+	int32 GetIndexInElement() const { return IndexInElement; }
+	int32 GetIndexInHierarchy() const { return IndexInHierarchy; }
+	bool IsSelected() const { return bSelected; }
+	int32 GetCreatedAtInstructionIndex() const { return CreatedAtInstructionIndex; }
+	bool IsProcedural() const { return CreatedAtInstructionIndex != INDEX_NONE; }
+	const FRigElementKey& GetElementKey() const { return ElementKey; }
+	const FRigBaseElement* GetElement() const { return Element; }
+	FRigBaseElement* GetElement() { return Element; }
+
+	bool Serialize(FArchive& Ar);
+	virtual void Save(FArchive& Ar);
+	virtual void Load(FArchive& Ar);
+
+	bool IsA(const UScriptStruct* InScriptStruct) const
+	{
+		return GetScriptStruct()->IsChildOf(InScriptStruct);
+	}
+
+	template<typename T>
+	bool IsA() const { return T::IsClassOf(this); }
+
+	template<typename T>
+    friend const T* Cast(const FRigBaseComponent* InComponent)
+	{
+		if(InComponent)
+		{
+			if(InComponent->IsA<T>())
+			{
+				return static_cast<const T*>(InComponent);
+			}
+		}
+		return nullptr;
+	}
+
+	template<typename T>
+    friend T* Cast(FRigBaseComponent* InComponent)
+	{
+		if(InComponent)
+		{
+			if(InComponent->IsA<T>())
+			{
+				return static_cast<T*>(InComponent);
+			}
+		}
+		return nullptr;
+	}
+
+	template<typename T>
+    friend const T* CastChecked(const FRigBaseComponent* InComponent)
+	{
+		const T* Component = Cast<T>(InComponent);
+		check(Component);
+		return Component;
+	}
+
+	template<typename T>
+    friend T* CastChecked(FRigBaseComponent* InComponent)
+	{
+		T* Component = Cast<T>(InComponent);
+		check(Component);
+		return Component;
+	}
+
+protected:
+
+	friend class FControlRigEditor;
+	friend class URigHierarchy;
+	friend class URigHierarchyController;
+};
+
+template<>
+struct TStructOpsTypeTraits<FRigBaseComponent> : public TStructOpsTypeTraitsBase2<FRigBaseComponent>
+{
+	enum 
+	{
+		WithSerializer = true,
+	};
+};
diff --git a/Engine/Plugins/Animation/ControlRig/Source/ControlRig/Public/Rigs/RigHierarchyController.h b/Engine/Plugins/Animation/ControlRig/Source/ControlRig/Public/Rigs/RigHierarchyController.h
--- a/Engine/Plugins/Animation/ControlRig/Source/ControlRig/Public/Rigs/RigHierarchyController.h
+++ b/Engine/Plugins/Animation/ControlRig/Source/ControlRig/Public/Rigs/RigHierarchyController.h
@@ -283,6 +283,42 @@ public:
 	UFUNCTION(BlueprintCallable, Category = URigHierarchyController)
     bool SetControlSettings(FRigElementKey InKey, FRigControlSettings InSettings, bool bSetupUndo = false) const;
 
+	/**
+	 * Adds a component to the hierarchy
+	 * @param InComponentStruct The script struct of the component to add
+	 * @param InName The suggested name of the new component.
+	 * @param InElement The element the component will be added to.
+	 * @param InContent The (optional) serialized text default for the component
+	 * @param bSetupUndo If set to true the stack will record the change for undo / redo
+	 * @param bPrintPythonCommand If set to true a python command equivalent to this call will be printed out
+	 * @return The name of the newly created component
+	 */
+	UFUNCTION(BlueprintCallable, Category = URigHierarchyController)
+	FName AddComponent(UScriptStruct* InComponentStruct, FName InName, FRigElementKey InElement, FString InContent = TEXT(""), bool bSetupUndo = false, bool bPrintPythonCommand = false);
+
+	/**
+     * Removes a component from the hierarchy
+     * @param InName The name of the component to remove
+     * @param InElement The element the component belongs to
+     * @param bSetupUndo If set to true the stack will record the change for undo / redo
+     * @param bPrintPythonCommand If set to true a python command equivalent to this call will be printed out
+     * @return True if the component was removed successfully
+     */
+    UFUNCTION(BlueprintCallable, Category = URigHierarchyController)
+    bool RemoveComponent(FName InName, FRigElementKey InElement, bool bSetupUndo = false, bool bPrintPythonCommand = false);
+
+	/**
+	 * Updates the content of a component in the hierarchy
+	 * @param InName The name of the component
+	 * @param InElement The element the component belongs to
+	 * @param InContent The serialized text default for the component
+	 * @param bSetupUndo If set to true the stack will record the change for undo / redo
+	 * @param bPrintPythonCommand If set to true a python command equivalent to this call will be printed out
+	 * @return True if the component was updated correctly
+	 */
+	UFUNCTION(BlueprintCallable, Category = URigHierarchyController)
+	bool SetComponentContent(FName InName, FRigElementKey InElement, const FString& InContent, bool bSetupUndo = false, bool bPrintPythonCommand = false);
+
 	/**
 	 * Imports an existing skeleton to the hierarchy
 	 * @param InSkeleton The skeleton to import
@@ -710,13 +746,15 @@ public:
 	TArray<FString> GetSetControlOffsetTransformPythonCommands(const FRigControlElement* Control, const FTransform& Offset, bool bInitial = false, bool bAffectChildren = true) const;
 	
 	TArray<FString> GetSetControlShapeTransformPythonCommands(const FRigControlElement* Control, const FTransform& Transform, bool bInitial = false) const;
+
+	TArray<FString> GetAddComponentPythonCommands(const FRigBaseComponent* Component) const;
 #endif
 	
 private:
 
 	FRigHierarchyModifiedEvent ModifiedEvent;
-	void Notify(ERigHierarchyNotification InNotifType, const FRigBaseElement* InElement);
-	void HandleHierarchyModified(ERigHierarchyNotification InNotifType, URigHierarchy* InHierarchy, const FRigBaseElement* InElement) const;
+	void Notify(ERigHierarchyNotification InNotifType, const FRigNotificationSubject& InSubject);
+	void HandleHierarchyModified(ERigHierarchyNotification InNotifType, URigHierarchy* InHierarchy, const FRigNotificationSubject& InSubject) const;
 
 	/**
 	 * Returns true if this controller is valid / linked to a valid hierarchy.
diff --git a/Engine/Plugins/Animation/ControlRig/Source/ControlRig/Public/Rigs/RigHierarchyDefines.h b/Engine/Plugins/Animation/ControlRig/Source/ControlRig/Public/Rigs/RigHierarchyDefines.h
--- a/Engine/Plugins/Animation/ControlRig/Source/ControlRig/Public/Rigs/RigHierarchyDefines.h
+++ b/Engine/Plugins/Animation/ControlRig/Source/ControlRig/Public/Rigs/RigHierarchyDefines.h
@@ -9,6 +9,7 @@
 
 class URigHierarchy;
 struct FRigBaseElement;
+struct FRigBaseComponent;
 
 // Debug define which performs a full check on the cache validity for all elements of the hierarchy.
 // This can be useful for debugging cache validity bugs.
@@ -99,11 +100,35 @@ enum class ERigHierarchyNotification : uint8
 	SocketDescriptionChanged,
 	SocketDesiredParentChanged,
 	HierarchyCopied,
+	ComponentAdded,
+	ComponentRemoved,
+	ComponentContentChanged,
 
 	/** MAX - invalid */
 	Max UMETA(Hidden),
 };
 
+/*
+ * Used for notifications - the subject can be a variety of things
+ */
+struct CONTROLRIG_API FRigNotificationSubject
+{
+	FRigNotificationSubject() = default;
+
+	FRigNotificationSubject(const FRigBaseElement* InElement)
+	: Element(InElement)
+	, Component(nullptr)
+	{}
+
+	FRigNotificationSubject(const FRigBaseComponent* InComponent)
+	: Element(nullptr)
+	, Component(InComponent)
+	{}
+	
+	const FRigBaseElement* Element = nullptr;
+	const FRigBaseComponent* Component = nullptr;
+};
+
 UENUM(meta = (RigVMTypeAllowed))
 enum class ERigEvent : uint8
 {
diff --git a/Engine/Plugins/Animation/ControlRig/Source/ControlRig/Public/Rigs/RigHierarchyElements.h b/Engine/Plugins/Animation/ControlRig/Source/ControlRig/Public/Rigs/RigHierarchyElements.h
--- a/Engine/Plugins/Animation/ControlRig/Source/ControlRig/Public/Rigs/RigHierarchyElements.h
+++ b/Engine/Plugins/Animation/ControlRig/Source/ControlRig/Public/Rigs/RigHierarchyElements.h
@@ -2,8 +2,7 @@
 
 #pragma once
 
-#include "CoreMinimal.h"
-#include "RigHierarchyDefines.h"
+#include "RigHierarchyComponents.h"
 #include "RigHierarchyMetadata.h"
 #include "RigConnectionRules.h"
 #include "RigReusableElementStorage.h"
@@ -21,16 +20,16 @@ DECLARE_DELEGATE_ThreeParams(FRigElementMetadataTagChangedDelegate, const FRigEl
 DECLARE_DELEGATE_RetVal_ThreeParams(FTransform, FRigReferenceGetWorldTransformDelegate, const FRigVMExecuteContext*, const FRigElementKey& /* Key */, bool /* bInitial */);
 
 #define DECLARE_RIG_ELEMENT_METHODS(ElementType) \
-virtual UScriptStruct* GetElementStruct() const override { return ElementType::StaticStruct(); } \
+virtual UScriptStruct* GetScriptStruct() const override { return ElementType::StaticStruct(); } \
 template<typename T> \
 friend const T* Cast(const ElementType* InElement) \
 { \
-   return Cast<T>((const FRigBaseElement*) InElement); \
+	return Cast<T>((const FRigBaseElement*) InElement); \
 } \
 template<typename T> \
 friend T* Cast(ElementType* InElement) \
 { \
-   return Cast<T>((FRigBaseElement*) InElement); \
+	return Cast<T>((FRigBaseElement*) InElement); \
 } \
 template<typename T> \
 friend const T* CastChecked(const ElementType* InElement) \
@@ -42,6 +41,10 @@ friend T* CastChecked(ElementType* InElement) \
 { \
 	return CastChecked<T>((FRigBaseElement*) InElement); \
 } \
+static bool IsClassOf(const FRigBaseElement* InElement) \
+{ \
+	return InElement->GetScriptStruct()->IsChildOf(StaticStruct()); \
+} \
 virtual int32 GetElementTypeIndex() const override { return (int32)ElementType::ElementTypeIndex; }
 
 UENUM()
@@ -706,6 +709,58 @@ private:
 	FRigElementKey Key;
 };
 
+struct CONTROLRIG_API FRigComponentHandle
+{
+public:
+
+	FRigComponentHandle()
+		: ElementHandle()
+		, ComponentName(NAME_None)
+	{}
+
+	FRigComponentHandle(URigHierarchy* InHierarchy, const FRigElementKey& InKey, const FName& InComponentName);
+	FRigComponentHandle(URigHierarchy* InHierarchy, const FRigBaseComponent* InComponent);
+
+	bool IsValid() const { return Get() != nullptr; }
+	operator bool() const { return IsValid(); }
+	
+	const URigHierarchy* GetHierarchy() const { return ElementHandle.GetHierarchy(); }
+	URigHierarchy* GetHierarchy() { return ElementHandle.GetHierarchy(); }
+	const FRigElementKey& GetElementKey() const { return ElementHandle.GetKey(); }
+
+	const FRigBaseComponent* Get() const;
+	FRigBaseComponent* Get();
+
+	template<typename T>
+	T* Get()
+	{
+		return Cast<T>(Get());
+	}
+
+	template<typename T>
+	const T* Get() const
+	{
+		return Cast<T>(Get());
+	}
+
+	template<typename T>
+	T* GetChecked()
+	{
+		return CastChecked<T>(Get());
+	}
+
+	template<typename T>
+	const T* GetChecked() const
+	{
+		return CastChecked<T>(Get());
+	}
+
+private:
+
+	FRigElementHandle ElementHandle;
+	FName ComponentName;
+};
+
 USTRUCT(BlueprintType)
 struct CONTROLRIG_API FRigBaseElement
 {
@@ -805,12 +860,15 @@ protected:
 
 	// Index into the metadata storage for this element.
 	int32 MetadataStorageIndex = INDEX_NONE;
+
+	// The indices of the components on this element
+	TArray<int32> ComponentIndices;
 	
 	mutable FString CachedNameString;
 
 public:
 
-	virtual UScriptStruct* GetElementStruct() const { return FRigBaseElement::StaticStruct(); }
+	virtual UScriptStruct* GetScriptStruct() const { return FRigBaseElement::StaticStruct(); }
 	void Serialize(FArchive& Ar, ESerializationPhase SerializationPhase);
 	virtual void Save(FArchive& Ar, ESerializationPhase SerializationPhase);
 	virtual void Load(FArchive& Ar, ESerializationPhase SerializationPhase);
@@ -846,9 +904,35 @@ public:
 	
 	void NotifyMetadataTagChanged(const FName& InTag, bool bAdded);
 
+	// Components
+	int32 NumComponents() const;
+	const FRigBaseComponent* GetComponent(int32 InIndex) const;
+	FRigBaseComponent* GetComponent(int32 InIndex);
+	const FRigBaseComponent* GetComponentByName(const FName& InName) const;
+	FRigBaseComponent* GetComponentByName(const FName& InName);
+	const FRigBaseComponent* GetComponentOfType(const UScriptStruct* InComponentStruct) const;
+	FRigBaseComponent* GetComponentOfType(const UScriptStruct* InComponentStruct);
+
+	template<typename T>
+	const T* GetComponent() const
+	{
+		return Cast<T>(GetComponentOfType(T::StaticStruct()));
+	}
+
+	template<typename T>
+	T* GetComponent()
+	{
+		return Cast<T>(GetComponentOfType(T::StaticStruct()));
+	}
+
 	virtual int32 GetNumTransforms() const { return 0; }
 	virtual int32 GetNumCurves() const { return 0; }
-	
+
+	bool IsA(const UScriptStruct* InScriptStruct) const
+	{
+		return GetScriptStruct()->IsChildOf(InScriptStruct);
+	}
+
 	template<typename T>
 	bool IsA() const { return T::IsClassOf(this); }
 
@@ -995,16 +1079,6 @@ protected:
 	typedef TArray<FElementToDirty, TInlineAllocator<3>> FElementsToDirtyArray;  
 	FElementsToDirtyArray ElementsToDirty;
 
-	static bool IsClassOf(const FRigBaseElement* InElement)
-	{
-		return InElement->GetType() == ERigElementType::Bone ||
-			InElement->GetType() == ERigElementType::Null ||
-			InElement->GetType() == ERigElementType::Control ||
-			InElement->GetType() == ERigElementType::Physics ||
-			InElement->GetType() == ERigElementType::Reference ||
-			InElement->GetType() == ERigElementType::Socket;
-	}
-
 	virtual void CopyFrom(const FRigBaseElement* InOther) override;
 
 	virtual void LinkStorage(const TArrayView<FTransform>& InTransforms, const TArrayView<bool>& InDirtyStates, const TArrayView<float>& InCurves) override;
@@ -1039,14 +1113,6 @@ protected:
 
 	virtual void CopyFrom(const FRigBaseElement* InOther) override;
 
-	static bool IsClassOf(const FRigBaseElement* InElement)
-	{
-		return InElement->GetType() == ERigElementType::Bone ||
-			InElement->GetType() == ERigElementType::Physics ||
-			InElement->GetType() == ERigElementType::Reference ||
-			InElement->GetType() == ERigElementType::Socket;
-	}
-	
 	friend class URigHierarchy;
 	friend struct FRigBaseElement;
 };
@@ -1193,11 +1259,6 @@ protected:
 	virtual void CopyFrom(const FRigBaseElement* InOther) override;
 	
 private:
-	static bool IsClassOf(const FRigBaseElement* InElement)
-	{
-		return InElement->GetType() == ERigElementType::Null ||
-			InElement->GetType() == ERigElementType::Control;
-	}
 
 	friend class URigHierarchy;
 	friend struct FRigBaseElement;
@@ -1242,11 +1303,6 @@ private:
 
 	virtual void CopyFrom(const FRigBaseElement* InOther) override;
 
-	static bool IsClassOf(const FRigBaseElement* InElement)
-	{
-		return InElement->GetType() == ERigElementType::Bone;
-	}
-
 	friend class URigHierarchy;
 	friend struct FRigBaseElement;
 };
@@ -1272,11 +1328,6 @@ private:
 		: FRigMultiParentElement(InOwner, ERigElementType::Null)
 	{}
 
-	static bool IsClassOf(const FRigBaseElement* InElement)
-	{
-		return InElement->GetType() == ERigElementType::Null;
-	}
-
 	friend class URigHierarchy;
 	friend struct FRigBaseElement;
 };
@@ -1661,11 +1712,6 @@ private:
 
 	virtual void CopyFrom(const FRigBaseElement* InOther) override;
 
-	static bool IsClassOf(const FRigBaseElement* InElement)
-	{
-		return InElement->GetType() == ERigElementType::Control;
-	}
-
 	// Offset storage for this element.
 	FRigCurrentAndInitialTransform OffsetStorage;
 	FRigCurrentAndInitialDirtyState OffsetDirtyState;
@@ -1723,11 +1769,6 @@ private:
 	
 	virtual void CopyFrom(const FRigBaseElement* InOther) override;
 
-	static bool IsClassOf(const FRigBaseElement* InElement)
-	{
-		return InElement->GetType() == ERigElementType::Curve;
-	}
-
 	// Set to true if the value was actually set. Used to carry back and forth blend curve
 	// value validity state.
 	bool bIsValueSet = true;
@@ -1901,11 +1942,6 @@ private:
 	
 	virtual void CopyFrom(const FRigBaseElement* InOther) override;
 
-	static bool IsClassOf(const FRigBaseElement* InElement)
-	{
-		return InElement->GetType() == ERigElementType::Physics;
-	}
-
 	friend class URigHierarchy;
 	friend struct FRigBaseElement;
 };
@@ -1942,11 +1978,6 @@ private:
 
 	virtual void CopyFrom(const FRigBaseElement* InOther) override;
 
-	static bool IsClassOf(const FRigBaseElement* InElement)
-	{
-		return InElement->GetType() == ERigElementType::Reference;
-	}
-
 	friend class URigHierarchy;
 	friend class URigHierarchyController;
 	friend struct FRigBaseElement;
@@ -2063,11 +2094,6 @@ private:
 	
 	virtual void CopyFrom(const FRigBaseElement* InOther) override;
 
-	static bool IsClassOf(const FRigBaseElement* InElement)
-	{
-		return InElement->GetType() == ERigElementType::Connector;
-	}
-
 	friend class URigHierarchy;
 	friend struct FRigBaseElement;
 };
@@ -2131,11 +2157,6 @@ private:
 	
 	virtual void CopyFrom(const FRigBaseElement* InOther) override;
 
-	static bool IsClassOf(const FRigBaseElement* InElement)
-	{
-		return InElement->GetType() == ERigElementType::Socket;
-	}
-
 	friend class URigHierarchy;
 	friend struct FRigBaseElement;
 };
diff --git a/Engine/Plugins/Animation/ControlRig/Source/ControlRigDeveloper/Private/ControlRigBlueprint.cpp b/Engine/Plugins/Animation/ControlRig/Source/ControlRigDeveloper/Private/ControlRigBlueprint.cpp
--- a/Engine/Plugins/Animation/ControlRig/Source/ControlRigDeveloper/Private/ControlRigBlueprint.cpp
+++ b/Engine/Plugins/Animation/ControlRig/Source/ControlRigDeveloper/Private/ControlRigBlueprint.cpp
@@ -2672,7 +2672,7 @@ void UControlRigBlueprint::RefreshModuleConnectors(const FRigModuleReference* In
 	}
 }
 
-void UControlRigBlueprint::HandleHierarchyModified(ERigHierarchyNotification InNotification, URigHierarchy* InHierarchy, const FRigBaseElement* InElement)
+void UControlRigBlueprint::HandleHierarchyModified(ERigHierarchyNotification InNotification, URigHierarchy* InHierarchy, const FRigNotificationSubject& InSubject)
 {
 #if WITH_EDITOR
 
@@ -2681,6 +2681,9 @@ void UControlRigBlueprint::HandleHierarchyModified(ERigHierarchyNotification InN
 		return;
 	}
 
+	const FRigBaseElement* InElement = InSubject.Element;
+	const FRigBaseComponent* InComponent = InSubject.Component;
+
 	switch(InNotification)
 	{
 		case ERigHierarchyNotification::ElementRemoved:
@@ -2878,7 +2881,7 @@ void UControlRigBlueprint::HandleRigModulesModified(EModularRigNotification InNo
 			Hierarchy->Modify();
 			
 			UpdateConnectionMapFromModel();
-			HierarchyModifiedEvent.Broadcast(ERigHierarchyNotification::HierarchyReset, Hierarchy, nullptr);
+			HierarchyModifiedEvent.Broadcast(ERigHierarchyNotification::HierarchyReset, Hierarchy, {});
 			break;
 		}
 		case EModularRigNotification::ModuleClassChanged:
diff --git a/Engine/Plugins/Animation/ControlRig/Source/ControlRigDeveloper/Private/ControlRigBlueprintCompiler.cpp b/Engine/Plugins/Animation/ControlRig/Source/ControlRigDeveloper/Private/ControlRigBlueprintCompiler.cpp
--- a/Engine/Plugins/Animation/ControlRig/Source/ControlRigDeveloper/Private/ControlRigBlueprintCompiler.cpp
+++ b/Engine/Plugins/Animation/ControlRig/Source/ControlRigDeveloper/Private/ControlRigBlueprintCompiler.cpp
@@ -73,7 +73,7 @@ void FControlRigBlueprintCompilerContext::CopyTermDefaultsToDefaultObject(UObjec
 		}
 
 		// notify clients that the hierarchy has changed
-		ControlRig->GetHierarchy()->Notify(ERigHierarchyNotification::HierarchyReset, nullptr);
+		ControlRig->GetHierarchy()->Notify(ERigHierarchyNotification::HierarchyReset, {});
 
 		ControlRig->DrawContainer = ControlRigBlueprint->DrawContainer;
 		ControlRig->Influences = ControlRigBlueprint->Influences;
diff --git a/Engine/Plugins/Animation/ControlRig/Source/ControlRigDeveloper/Public/ControlRigBlueprint.h b/Engine/Plugins/Animation/ControlRig/Source/ControlRigDeveloper/Public/ControlRigBlueprint.h
--- a/Engine/Plugins/Animation/ControlRig/Source/ControlRigDeveloper/Public/ControlRigBlueprint.h
+++ b/Engine/Plugins/Animation/ControlRig/Source/ControlRigDeveloper/Public/ControlRigBlueprint.h
@@ -389,7 +389,7 @@ private:
 	int32 ModulesRecompilationBracket = 0;
 
 
-	void HandleHierarchyModified(ERigHierarchyNotification InNotification, URigHierarchy* InHierarchy, const FRigBaseElement* InElement);
+	void HandleHierarchyModified(ERigHierarchyNotification InNotification, URigHierarchy* InHierarchy, const FRigNotificationSubject& InSubject);
 
 	void HandleRigModulesModified(EModularRigNotification InNotification, const FRigModuleReference* InModule);
 
diff --git a/Engine/Plugins/Animation/ControlRig/Source/ControlRigEditor/Private/EditMode/ControlRigEditMode.cpp b/Engine/Plugins/Animation/ControlRig/Source/ControlRigEditor/Private/EditMode/ControlRigEditMode.cpp
--- a/Engine/Plugins/Animation/ControlRig/Source/ControlRigEditor/Private/EditMode/ControlRigEditMode.cpp
+++ b/Engine/Plugins/Animation/ControlRig/Source/ControlRigEditor/Private/EditMode/ControlRigEditMode.cpp
@@ -4366,8 +4366,11 @@ void FControlRigEditMode::ReplaceControlRig(UControlRig* OldControlRig, UControl
 	RequestToRecreateControlShapeActors(NewControlRig);
 
 }
-void FControlRigEditMode::OnHierarchyModified(ERigHierarchyNotification InNotif, URigHierarchy* InHierarchy, const FRigBaseElement* InElement)
+void FControlRigEditMode::OnHierarchyModified(ERigHierarchyNotification InNotif, URigHierarchy* InHierarchy, const FRigNotificationSubject& InSubject)
 {
+	const FRigBaseElement* InElement = InSubject.Element;
+	const FRigBaseComponent* InComponent = InSubject.Component;
+	
 	if(bSuspendHierarchyNotifs || InElement == nullptr)
 	{
 		return;
@@ -4581,7 +4584,7 @@ void FControlRigEditMode::OnHierarchyModified(ERigHierarchyNotification InNotif,
 	}
 }
 
-void FControlRigEditMode::OnHierarchyModified_AnyThread(ERigHierarchyNotification InNotif, URigHierarchy* InHierarchy, const FRigBaseElement* InElement)
+void FControlRigEditMode::OnHierarchyModified_AnyThread(ERigHierarchyNotification InNotif, URigHierarchy* InHierarchy, const FRigNotificationSubject& InSubject)
 {
 	if(bSuspendHierarchyNotifs)
 	{
@@ -4595,7 +4598,7 @@ void FControlRigEditMode::OnHierarchyModified_AnyThread(ERigHierarchyNotificatio
 
 	if(IsInGameThread())
 	{
-		OnHierarchyModified(InNotif, InHierarchy, InElement);
+		OnHierarchyModified(InNotif, InHierarchy, InSubject);
 		return;
 	}
 
@@ -4606,19 +4609,25 @@ void FControlRigEditMode::OnHierarchyModified_AnyThread(ERigHierarchyNotificatio
 		InNotif != ERigHierarchyNotification::ElementSelected &&
 		InNotif != ERigHierarchyNotification::ElementDeselected)
 	{
-		OnHierarchyModified(InNotif, InHierarchy, InElement);
+		OnHierarchyModified(InNotif, InHierarchy, InSubject);
 		return;
 	}
 	
 	FRigElementKey Key;
-	if(InElement)
+	FName ComponentName = NAME_None;
+	if(InSubject.Element)
+	{
+		Key = InSubject.Element->GetKey();
+	}
+	else if(InSubject.Component)
 	{
-		Key = InElement->GetKey();
+		Key = InSubject.Component->GetElementKey();
+		ComponentName = InSubject.Component->GetFName();
 	}
 
 	TWeakObjectPtr<URigHierarchy> WeakHierarchy = InHierarchy;
 	
-	FFunctionGraphTask::CreateAndDispatchWhenReady([this, InNotif, WeakHierarchy, Key]()
+	FFunctionGraphTask::CreateAndDispatchWhenReady([this, InNotif, WeakHierarchy, Key, ComponentName]()
 	{
 		if(!WeakHierarchy.IsValid())
 		{
@@ -4626,7 +4635,14 @@ void FControlRigEditMode::OnHierarchyModified_AnyThread(ERigHierarchyNotificatio
 		}
 		if (const FRigBaseElement* Element = WeakHierarchy.Get()->Find(Key))
 		{
-			OnHierarchyModified(InNotif, WeakHierarchy.Get(), Element);
+			if(ComponentName.IsNone())
+			{
+				OnHierarchyModified(InNotif, WeakHierarchy.Get(), Element);
+			}
+			else if(const FRigBaseComponent* Component = Element->GetComponentByName(ComponentName))
+			{
+				OnHierarchyModified(InNotif, WeakHierarchy.Get(), Component);
+			}
 		}
 		
 	}, TStatId(), NULL, ENamedThreads::GameThread);
diff --git a/Engine/Plugins/Animation/ControlRig/Source/ControlRigEditor/Private/Editor/ControlRigEditor.cpp b/Engine/Plugins/Animation/ControlRig/Source/ControlRigEditor/Private/Editor/ControlRigEditor.cpp
--- a/Engine/Plugins/Animation/ControlRig/Source/ControlRigEditor/Private/Editor/ControlRigEditor.cpp
+++ b/Engine/Plugins/Animation/ControlRig/Source/ControlRigEditor/Private/Editor/ControlRigEditor.cpp
@@ -1214,7 +1214,7 @@ void FControlRigEditor::SetDetailViewForRigElements(const TArray<FRigElementKey>
 			continue;
 		}
 
-		URigVMDetailsViewWrapperObject* WrapperObject = URigVMDetailsViewWrapperObject::MakeInstance(GetDetailWrapperClass(), GetBlueprintObj(), Element->GetElementStruct(), (uint8*)Element, HierarchyBeingDebugged);
+		URigVMDetailsViewWrapperObject* WrapperObject = URigVMDetailsViewWrapperObject::MakeInstance(GetDetailWrapperClass(), GetBlueprintObj(), Element->GetScriptStruct(), (uint8*)Element, HierarchyBeingDebugged);
 		WrapperObject->GetWrappedPropertyChangedChainEvent().AddSP(this, &FControlRigEditor::OnWrappedPropertyChangedChainEvent);
 		WrapperObject->AddToRoot();
 
@@ -2964,7 +2964,7 @@ void FControlRigEditor::HandlePreviewMeshChanged(USkeletalMesh* InOldSkeletalMes
 			
 			if(UControlRig* DebuggedControlRig = Cast<UControlRig>(ControlRigBP->GetObjectBeingDebugged()))
 			{
-				DebuggedControlRig->GetHierarchy()->Notify(ERigHierarchyNotification::HierarchyReset, nullptr);
+				DebuggedControlRig->GetHierarchy()->Notify(ERigHierarchyNotification::HierarchyReset, {});
 				DebuggedControlRig->Initialize(true);
 			}
 
@@ -3859,7 +3859,7 @@ void FControlRigEditor::OnHierarchyChanged()
 }
 
 
-void FControlRigEditor::OnHierarchyModified(ERigHierarchyNotification InNotif, URigHierarchy* InHierarchy, const FRigBaseElement* InElement)
+void FControlRigEditor::OnHierarchyModified(ERigHierarchyNotification InNotif, URigHierarchy* InHierarchy, const FRigNotificationSubject& InSubject)
 {
 	UControlRigBlueprint* RigBlueprint = Cast<UControlRigBlueprint>(GetBlueprintObj());
 	if(RigBlueprint == nullptr)
@@ -3877,6 +3877,9 @@ void FControlRigEditor::OnHierarchyModified(ERigHierarchyNotification InNotif, U
 		return;
 	}
 
+	const FRigBaseElement* InElement = InSubject.Element;
+	const FRigBaseComponent* InComponent = InSubject.Component;
+
 	switch(InNotif)
 	{
 		case ERigHierarchyNotification::ElementAdded:
@@ -4069,7 +4072,7 @@ void FControlRigEditor::OnHierarchyModified(ERigHierarchyNotification InNotif, U
 	}
 }
 
-void FControlRigEditor::OnHierarchyModified_AnyThread(ERigHierarchyNotification InNotif, URigHierarchy* InHierarchy, const FRigBaseElement* InElement)
+void FControlRigEditor::OnHierarchyModified_AnyThread(ERigHierarchyNotification InNotif, URigHierarchy* InHierarchy, const FRigNotificationSubject& InSubject)
 {
 	if(bIsConstructionEventRunning)
 	{
@@ -4078,13 +4081,19 @@ void FControlRigEditor::OnHierarchyModified_AnyThread(ERigHierarchyNotification
 
 	if(SchematicViewport)
 	{
-		SchematicModel.OnHierarchyModified(InNotif, InHierarchy, InElement);
+		SchematicModel.OnHierarchyModified(InNotif, InHierarchy, InSubject);
 	}
 	
 	FRigElementKey Key;
-	if(InElement)
+	FName ComponentName = NAME_None;
+	if(InSubject.Element)
+	{
+		Key = InSubject.Element->GetKey();
+	}
+	else if(InSubject.Component)
 	{
-		Key = InElement->GetKey();
+		Key = InSubject.Component->GetElementKey();
+		ComponentName = InSubject.Component->GetFName();
 	}
 
 	if(IsInGameThread())
@@ -4099,14 +4108,19 @@ void FControlRigEditor::OnHierarchyModified_AnyThread(ERigHierarchyNotification
 	}
 
 	TWeakObjectPtr<URigHierarchy> WeakHierarchy = InHierarchy;
-	auto Task = [this, InNotif, WeakHierarchy, Key]()
+	auto Task = [this, InNotif, WeakHierarchy, Key, ComponentName]()
     {
 		if(!WeakHierarchy.IsValid())
     	{
     		return;
     	}
 
+		FRigBaseComponent* Component = nullptr;
         FRigBaseElement* Element = WeakHierarchy.Get()->Find(Key);
+		if(Element)
+		{
+			Component = Element->GetComponentByName(ComponentName);
+		}
 
 		switch(InNotif)
 		{
diff --git a/Engine/Plugins/Animation/ControlRig/Source/ControlRigEditor/Private/Editor/ControlRigEditor.h b/Engine/Plugins/Animation/ControlRig/Source/ControlRigEditor/Private/Editor/ControlRigEditor.h
--- a/Engine/Plugins/Animation/ControlRig/Source/ControlRigEditor/Private/Editor/ControlRigEditor.h
+++ b/Engine/Plugins/Animation/ControlRig/Source/ControlRigEditor/Private/Editor/ControlRigEditor.h
@@ -158,8 +158,8 @@ public:
 
 	void OnCurveContainerChanged();
 
-	void OnHierarchyModified(ERigHierarchyNotification InNotif, URigHierarchy* InHierarchy, const FRigBaseElement* InElement);
-	void OnHierarchyModified_AnyThread(ERigHierarchyNotification InNotif, URigHierarchy* InHierarchy, const FRigBaseElement* InElement);
+	void OnHierarchyModified(ERigHierarchyNotification InNotif, URigHierarchy* InHierarchy, const FRigNotificationSubject& InSubject);
+	void OnHierarchyModified_AnyThread(ERigHierarchyNotification InNotif, URigHierarchy* InHierarchy, const FRigNotificationSubject& InSubject);
 
 	void HandleRigTypeChanged(UControlRigBlueprint* InBlueprint);
 
diff --git a/Engine/Plugins/Animation/ControlRig/Source/ControlRigEditor/Private/Editor/ControlRigSchematicModel.cpp b/Engine/Plugins/Animation/ControlRig/Source/ControlRigEditor/Private/Editor/ControlRigSchematicModel.cpp
--- a/Engine/Plugins/Animation/ControlRig/Source/ControlRigEditor/Private/Editor/ControlRigSchematicModel.cpp
+++ b/Engine/Plugins/Animation/ControlRig/Source/ControlRigEditor/Private/Editor/ControlRigSchematicModel.cpp
@@ -1838,8 +1838,11 @@ void FControlRigSchematicModel::OnHideCandidatesForConnector()
 	UpdateElementKeyLinks();
 }
 
-void FControlRigSchematicModel::OnHierarchyModified(ERigHierarchyNotification InNotif, URigHierarchy* InHierarchy, const FRigBaseElement* InElement)
+void FControlRigSchematicModel::OnHierarchyModified(ERigHierarchyNotification InNotif, URigHierarchy* InHierarchy, const FRigNotificationSubject& InSubject)
 {
+	const FRigBaseElement* InElement = InSubject.Element;
+	const FRigBaseComponent* InComponent = InSubject.Component;
+	
 	switch(InNotif)
 	{
 		case ERigHierarchyNotification::ElementSelected:
diff --git a/Engine/Plugins/Animation/ControlRig/Source/ControlRigEditor/Private/Editor/ControlRigSchematicModel.h b/Engine/Plugins/Animation/ControlRig/Source/ControlRigEditor/Private/Editor/ControlRigSchematicModel.h
--- a/Engine/Plugins/Animation/ControlRig/Source/ControlRigEditor/Private/Editor/ControlRigSchematicModel.h
+++ b/Engine/Plugins/Animation/ControlRig/Source/ControlRigEditor/Private/Editor/ControlRigSchematicModel.h
@@ -130,7 +130,7 @@ private:
 	void OnShowCandidatesForConnector(const FRigModuleConnector* InModuleConnector);
 	void OnShowCandidatesForMatches(const FModularRigResolveResult& InMatches);
 	void OnHideCandidatesForConnector();
-	void OnHierarchyModified(ERigHierarchyNotification InNotif, URigHierarchy* InHierarchy, const FRigBaseElement* InElement);
+	void OnHierarchyModified(ERigHierarchyNotification InNotif, URigHierarchy* InHierarchy, const FRigNotificationSubject& InSubject);
 
 	TWeakPtr<FControlRigEditor> ControlRigEditor;
 	TWeakObjectPtr<UControlRigBlueprint> ControlRigBlueprint;
diff --git a/Engine/Plugins/Animation/ControlRig/Source/ControlRigEditor/Private/Editor/ControlRigSkeletalMeshComponent.cpp b/Engine/Plugins/Animation/ControlRig/Source/ControlRigEditor/Private/Editor/ControlRigSkeletalMeshComponent.cpp
--- a/Engine/Plugins/Animation/ControlRig/Source/ControlRigEditor/Private/Editor/ControlRigSkeletalMeshComponent.cpp
+++ b/Engine/Plugins/Animation/ControlRig/Source/ControlRigEditor/Private/Editor/ControlRigSkeletalMeshComponent.cpp
@@ -233,7 +233,7 @@ void UControlRigSkeletalMeshComponent::SetControlRigBeingDebugged(UControlRig* I
 }
 
 void UControlRigSkeletalMeshComponent::OnHierarchyModified(ERigHierarchyNotification InNotif,
-    URigHierarchy* InHierarchy, const FRigBaseElement* InElement)
+    URigHierarchy* InHierarchy, const FRigNotificationSubject& InSubject)
 {
 	switch(InNotif)
 	{
@@ -273,7 +273,7 @@ void UControlRigSkeletalMeshComponent::OnHierarchyModified(ERigHierarchyNotifica
 }
 
 void UControlRigSkeletalMeshComponent::OnHierarchyModified_AnyThread(ERigHierarchyNotification InNotif,
-    URigHierarchy* InHierarchy, const FRigBaseElement* InElement)
+    URigHierarchy* InHierarchy, const FRigNotificationSubject& InSubject)
 {
 	if (bIsConstructionEventRunning)
 	{
@@ -281,9 +281,9 @@ void UControlRigSkeletalMeshComponent::OnHierarchyModified_AnyThread(ERigHierarc
 	}
 	
 	FRigElementKey Key;
-	if(InElement)
+	if(InSubject.Element)
 	{
-		Key = InElement->GetKey();
+		Key = InSubject.Element->GetKey();
 	}
 
 	TWeakObjectPtr<URigHierarchy> WeakHierarchy = InHierarchy;
diff --git a/Engine/Plugins/Animation/ControlRig/Source/ControlRigEditor/Private/Editor/ControlRigSkeletalMeshComponent.h b/Engine/Plugins/Animation/ControlRig/Source/ControlRigEditor/Private/Editor/ControlRigSkeletalMeshComponent.h
--- a/Engine/Plugins/Animation/ControlRig/Source/ControlRigEditor/Private/Editor/ControlRigSkeletalMeshComponent.h
+++ b/Engine/Plugins/Animation/ControlRig/Source/ControlRigEditor/Private/Editor/ControlRigSkeletalMeshComponent.h
@@ -51,8 +51,8 @@ public:
 	 */
 	void RebuildDebugDrawSkeleton();
 
-	void OnHierarchyModified(ERigHierarchyNotification InNotif, URigHierarchy* InHierarchy, const FRigBaseElement* InElement);
-	void OnHierarchyModified_AnyThread(ERigHierarchyNotification InNotif, URigHierarchy* InHierarchy, const FRigBaseElement* InElement);
+	void OnHierarchyModified(ERigHierarchyNotification InNotif, URigHierarchy* InHierarchy, const FRigNotificationSubject& InSubject);
+	void OnHierarchyModified_AnyThread(ERigHierarchyNotification InNotif, URigHierarchy* InHierarchy, const FRigNotificationSubject& InSubject);
 
 	void OnPreConstruction_AnyThread(UControlRig* InControlRig, const FName& InEventName);
 	void OnPostConstruction_AnyThread(UControlRig* InControlRig, const FName& InEventName);
diff --git a/Engine/Plugins/Animation/ControlRig/Source/ControlRigEditor/Private/Editor/SModularRigModel.cpp b/Engine/Plugins/Animation/ControlRig/Source/ControlRigEditor/Private/Editor/SModularRigModel.cpp
--- a/Engine/Plugins/Animation/ControlRig/Source/ControlRigEditor/Private/Editor/SModularRigModel.cpp
+++ b/Engine/Plugins/Animation/ControlRig/Source/ControlRigEditor/Private/Editor/SModularRigModel.cpp
@@ -1284,13 +1284,16 @@ void SModularRigModel::OnModularRigModified(EModularRigNotification InNotif, con
 	}
 }
 
-void SModularRigModel::OnHierarchyModified(ERigHierarchyNotification InNotif, URigHierarchy* InHierarchy, const FRigBaseElement* InElement)
+void SModularRigModel::OnHierarchyModified(ERigHierarchyNotification InNotif, URigHierarchy* InHierarchy, const FRigNotificationSubject& InSubject)
 {
 	if(!ControlRigBlueprint.IsValid())
 	{
 		return;
 	}
 	
+	const FRigBaseElement* InElement = InSubject.Element;
+	const FRigBaseComponent* InComponent = InSubject.Component;
+	
 	switch(InNotif)
 	{
 		case ERigHierarchyNotification::ElementSelected:
diff --git a/Engine/Plugins/Animation/ControlRig/Source/ControlRigEditor/Private/Editor/SModularRigModel.h b/Engine/Plugins/Animation/ControlRig/Source/ControlRigEditor/Private/Editor/SModularRigModel.h
--- a/Engine/Plugins/Animation/ControlRig/Source/ControlRigEditor/Private/Editor/SModularRigModel.h
+++ b/Engine/Plugins/Animation/ControlRig/Source/ControlRigEditor/Private/Editor/SModularRigModel.h
@@ -145,7 +145,7 @@ private:
 	void HandlePreCompileModularRigs(URigVMBlueprint* InBlueprint);
 	void HandlePostCompileModularRigs(URigVMBlueprint* InBlueprint);
 	void OnModularRigModified(EModularRigNotification InNotif, const FRigModuleReference* InModule);
-	void OnHierarchyModified(ERigHierarchyNotification InNotif, URigHierarchy* InHierarchy, const FRigBaseElement* InElement);
+	void OnHierarchyModified(ERigHierarchyNotification InNotif, URigHierarchy* InHierarchy, const FRigNotificationSubject& InSubject);
 
 	void HandleRefreshEditorFromBlueprint(URigVMBlueprint* InBlueprint);
 	void HandleSetObjectBeingDebugged(UObject* InObject);
diff --git a/Engine/Plugins/Animation/ControlRig/Source/ControlRigEditor/Private/Editor/SRigCurveContainer.cpp b/Engine/Plugins/Animation/ControlRig/Source/ControlRigEditor/Private/Editor/SRigCurveContainer.cpp
--- a/Engine/Plugins/Animation/ControlRig/Source/ControlRigEditor/Private/Editor/SRigCurveContainer.cpp
+++ b/Engine/Plugins/Animation/ControlRig/Source/ControlRigEditor/Private/Editor/SRigCurveContainer.cpp
@@ -581,13 +581,14 @@ void SRigCurveContainer::OnSelectionChanged(FDisplayedRigCurveInfoPtr Selection,
 	}
 }
 
-void SRigCurveContainer::OnHierarchyModified(ERigHierarchyNotification InNotif, URigHierarchy* InHierarchy, const FRigBaseElement* InElement)
+void SRigCurveContainer::OnHierarchyModified(ERigHierarchyNotification InNotif, URigHierarchy* InHierarchy, const FRigNotificationSubject& InSubject)
 {
 	if (bIsChangingRigHierarchy)
 	{
 		return;
 	}
 
+	const FRigBaseElement* InElement = InSubject.Element;
 	if(InElement)
 	{
 		if(!InElement->IsTypeOf(ERigElementType::Curve))
diff --git a/Engine/Plugins/Animation/ControlRig/Source/ControlRigEditor/Private/Editor/SRigCurveContainer.h b/Engine/Plugins/Animation/ControlRig/Source/ControlRigEditor/Private/Editor/SRigCurveContainer.h
--- a/Engine/Plugins/Animation/ControlRig/Source/ControlRigEditor/Private/Editor/SRigCurveContainer.h
+++ b/Engine/Plugins/Animation/ControlRig/Source/ControlRigEditor/Private/Editor/SRigCurveContainer.h
@@ -242,7 +242,7 @@ private:
 	void OnSelectionChanged(FDisplayedRigCurveInfoPtr Selection, ESelectInfo::Type SelectInfo);
 
 	bool bIsChangingRigHierarchy;
-	void OnHierarchyModified(ERigHierarchyNotification InNotif, URigHierarchy* InHierarchy, const FRigBaseElement* InElement);
+	void OnHierarchyModified(ERigHierarchyNotification InNotif, URigHierarchy* InHierarchy, const FRigNotificationSubject& InSubject);
 	void HandleRefreshEditorFromBlueprint(URigVMBlueprint* InBlueprint);
 
 	// import curve part
diff --git a/Engine/Plugins/Animation/ControlRig/Source/ControlRigEditor/Private/Editor/SRigHierarchy.cpp b/Engine/Plugins/Animation/ControlRig/Source/ControlRigEditor/Private/Editor/SRigHierarchy.cpp
--- a/Engine/Plugins/Animation/ControlRig/Source/ControlRigEditor/Private/Editor/SRigHierarchy.cpp
+++ b/Engine/Plugins/Animation/ControlRig/Source/ControlRigEditor/Private/Editor/SRigHierarchy.cpp
@@ -688,7 +688,7 @@ void SRigHierarchy::OnSelectionChanged(TSharedPtr<FRigTreeElement> Selection, ES
 	}
 }
 
-void SRigHierarchy::OnHierarchyModified(ERigHierarchyNotification InNotif, URigHierarchy* InHierarchy, const FRigBaseElement* InElement)
+void SRigHierarchy::OnHierarchyModified(ERigHierarchyNotification InNotif, URigHierarchy* InHierarchy, const FRigNotificationSubject& InSubject)
 {
 	if(!ControlRigBlueprint.IsValid())
 	{
@@ -705,6 +705,9 @@ void SRigHierarchy::OnHierarchyModified(ERigHierarchyNotification InNotif, URigH
 		return;
 	}
 
+	const FRigBaseElement* InElement = InSubject.Element;
+	const FRigBaseComponent* InComponent = InSubject.Component;
+
 	if(InElement)
 	{
 		if(InElement->IsTypeOf(ERigElementType::Curve))
@@ -847,7 +850,7 @@ void SRigHierarchy::OnHierarchyModified(ERigHierarchyNotification InNotif, URigH
 	}
 }
 
-void SRigHierarchy::OnHierarchyModified_AnyThread(ERigHierarchyNotification InNotif, URigHierarchy* InHierarchy, const FRigBaseElement* InElement)
+void SRigHierarchy::OnHierarchyModified_AnyThread(ERigHierarchyNotification InNotif, URigHierarchy* InHierarchy, const FRigNotificationSubject& InSubject)
 {
 	if(bIsChangingRigHierarchy)
 	{
@@ -869,9 +872,12 @@ void SRigHierarchy::OnHierarchyModified_AnyThread(ERigHierarchyNotification InNo
 		return;
 	}
 
+	const FRigBaseElement* InElement = InSubject.Element;
+	const FRigBaseComponent* InComponent = InSubject.Component;
+
 	if(IsInGameThread())
 	{
-		OnHierarchyModified(InNotif, InHierarchy, InElement);
+		OnHierarchyModified(InNotif, InHierarchy, InSubject);
 	}
 	else
 	{
diff --git a/Engine/Plugins/Animation/ControlRig/Source/ControlRigEditor/Private/Editor/SRigHierarchy.h b/Engine/Plugins/Animation/ControlRig/Source/ControlRigEditor/Private/Editor/SRigHierarchy.h
--- a/Engine/Plugins/Animation/ControlRig/Source/ControlRigEditor/Private/Editor/SRigHierarchy.h
+++ b/Engine/Plugins/Animation/ControlRig/Source/ControlRigEditor/Private/Editor/SRigHierarchy.h
@@ -213,8 +213,8 @@ private:
 	void ClearDetailPanel() const;
 
 	bool bIsChangingRigHierarchy;
-	void OnHierarchyModified(ERigHierarchyNotification InNotif, URigHierarchy* InHierarchy, const FRigBaseElement* InElement);
-	void OnHierarchyModified_AnyThread(ERigHierarchyNotification InNotif, URigHierarchy* InHierarchy, const FRigBaseElement* InElement);
+	void OnHierarchyModified(ERigHierarchyNotification InNotif, URigHierarchy* InHierarchy, const FRigNotificationSubject& InSubject);
+	void OnHierarchyModified_AnyThread(ERigHierarchyNotification InNotif, URigHierarchy* InHierarchy, const FRigNotificationSubject& InSubject);
 	void OnModularRigModified(EModularRigNotification InNotif, const FRigModuleReference* InModule);
 	void HandleRefreshEditorFromBlueprint(URigVMBlueprint* InBlueprint);
 	void HandleSetObjectBeingDebugged(UObject* InObject);
diff --git a/Engine/Plugins/Animation/ControlRig/Source/ControlRigEditor/Private/Editor/SRigSpacePickerWidget.cpp b/Engine/Plugins/Animation/ControlRig/Source/ControlRigEditor/Private/Editor/SRigSpacePickerWidget.cpp
--- a/Engine/Plugins/Animation/ControlRig/Source/ControlRigEditor/Private/Editor/SRigSpacePickerWidget.cpp
+++ b/Engine/Plugins/Animation/ControlRig/Source/ControlRigEditor/Private/Editor/SRigSpacePickerWidget.cpp
@@ -854,8 +854,11 @@ bool SRigSpacePickerWidget::IsSpaceMoveDownEnabled(FRigElementKey InKey) const
 }
 
 void SRigSpacePickerWidget::OnHierarchyModified(ERigHierarchyNotification InNotif, URigHierarchy* InHierarchy,
-                                                const FRigBaseElement* InElement)
+                                                const FRigNotificationSubject& InSubject)
 {
+	const FRigBaseElement* InElement = InSubject.Element;
+	const FRigBaseComponent* InComponent = InSubject.Component;
+
 	if(InElement == nullptr)
 	{
 		return;
diff --git a/Engine/Plugins/Animation/ControlRig/Source/ControlRigEditor/Private/Editor/SRigSpacePickerWidget.h b/Engine/Plugins/Animation/ControlRig/Source/ControlRigEditor/Private/Editor/SRigSpacePickerWidget.h
--- a/Engine/Plugins/Animation/ControlRig/Source/ControlRigEditor/Private/Editor/SRigSpacePickerWidget.h
+++ b/Engine/Plugins/Animation/ControlRig/Source/ControlRigEditor/Private/Editor/SRigSpacePickerWidget.h
@@ -139,7 +139,7 @@ private:
 	bool IsSpaceMoveUpEnabled(FRigElementKey InKey) const;
 	bool IsSpaceMoveDownEnabled(FRigElementKey InKey) const;
 
-	void OnHierarchyModified(ERigHierarchyNotification InNotif, URigHierarchy* InHierarchy, const FRigBaseElement* InElement);
+	void OnHierarchyModified(ERigHierarchyNotification InNotif, URigHierarchy* InHierarchy, const FRigNotificationSubject& InSubject);
 
 	FSlateColor GetButtonColor(ESpacePickerType InType, FRigElementKey InKey) const;
 	FRigElementKey GetActiveSpace_Private(URigHierarchy* InHierarchy, const FRigElementKey& InControlKey) const;
diff --git a/Engine/Plugins/Animation/ControlRig/Source/ControlRigEditor/Public/EditMode/ControlRigEditMode.h b/Engine/Plugins/Animation/ControlRig/Source/ControlRigEditor/Public/EditMode/ControlRigEditMode.h
--- a/Engine/Plugins/Animation/ControlRig/Source/ControlRigEditor/Public/EditMode/ControlRigEditMode.h
+++ b/Engine/Plugins/Animation/ControlRig/Source/ControlRigEditor/Public/EditMode/ControlRigEditMode.h
@@ -210,8 +210,8 @@ public:
 	FControlRigSelected& OnControlRigSelected() { return OnControlRigSelectedDelegate; }
 
 	// callback that gets called when rig element is selected in other view
-	void OnHierarchyModified(ERigHierarchyNotification InNotif, URigHierarchy* InHierarchy, const FRigBaseElement* InElement);
-	void OnHierarchyModified_AnyThread(ERigHierarchyNotification InNotif, URigHierarchy* InHierarchy, const FRigBaseElement* InElement);
+	void OnHierarchyModified(ERigHierarchyNotification InNotif, URigHierarchy* InHierarchy, const FRigNotificationSubject& InSubject);
+	void OnHierarchyModified_AnyThread(ERigHierarchyNotification InNotif, URigHierarchy* InHierarchy, const FRigNotificationSubject& InSubject);
 	void OnControlModified(UControlRig* Subject, FRigControlElement* InControlElement, const FRigControlModifiedContext& Context);
 	void OnPreConstruction_AnyThread(UControlRig* InRig, const FName& InEventName);
 	void OnPostConstruction_AnyThread(UControlRig* InRig, const FName& InEventName);
diff --git a/Engine/Plugins/Editor/SequencerAnimTools/Source/SequencerAnimTools/Private/Sequencer/SequencerTrailHierarchy.cpp b/Engine/Plugins/Editor/SequencerAnimTools/Source/SequencerAnimTools/Private/Sequencer/SequencerTrailHierarchy.cpp
--- a/Engine/Plugins/Editor/SequencerAnimTools/Source/SequencerAnimTools/Private/Sequencer/SequencerTrailHierarchy.cpp
+++ b/Engine/Plugins/Editor/SequencerAnimTools/Source/SequencerAnimTools/Private/Sequencer/SequencerTrailHierarchy.cpp
@@ -791,16 +791,15 @@ void FSequencerTrailHierarchy::RegisterControlRigDelegates(USkeletalMeshComponen
 	);
 	
 	DelegateHandles.OnHierarchyModified = RigHierarchy->OnModified().AddLambda(
-		[this, RigHierarchy, ControlRig](ERigHierarchyNotification InNotif, URigHierarchy* InHierarchy, const FRigBaseElement* InElement) {
+		[this, RigHierarchy, ControlRig](ERigHierarchyNotification InNotif, URigHierarchy* InHierarchy, const FRigNotificationSubject& InSubject) {
 			TSharedPtr<ISequencer> Sequencer = WeakSequencer.Pin();
 			check(Sequencer);
 
-			const FRigControlElement* ControlElement = Cast<FRigControlElement>(InElement); 
+			const FRigControlElement* ControlElement = Cast<FRigControlElement>(InSubject.Element); 
 			if(ControlElement == nullptr)
 			{
 				return;
 			}
-
 			
 			if(InNotif == ERigHierarchyNotification::ElementRemoved)
 			{
