AnimLayer: Don't create anim layer when checking to see if they exist, added new function for this

diff --git a/Engine/Plugins/Animation/ControlRig/Source/ControlRigEditor/Private/EditMode/ControlRigEditModeToolkit.cpp b/Engine/Plugins/Animation/ControlRig/Source/ControlRigEditor/Private/EditMode/ControlRigEditModeToolkit.cpp
--- a/Engine/Plugins/Animation/ControlRig/Source/ControlRigEditor/Private/EditMode/ControlRigEditModeToolkit.cpp
+++ b/Engine/Plugins/Animation/ControlRig/Source/ControlRigEditor/Private/EditMode/ControlRigEditModeToolkit.cpp
@@ -462,7 +462,7 @@ void FControlRigEditModeToolkit::InvokeUI()
 		else if (TSharedPtr<ISequencer> SequencerPtr = UAnimLayers::GetSequencerFromAsset())
 		{
 			//wasn't open but there are layers in open level sequence we try to open it
-			if (UAnimLayers* AnimLayers = UAnimLayers::GetAnimLayers(SequencerPtr.Get()))
+			if (UAnimLayers::HasAnimLayers(SequencerPtr.Get()))
 			{
 				TryInvokeToolkitUI(AnimLayerTabName);
 			}
diff --git a/Engine/Plugins/Animation/ControlRig/Source/ControlRigEditor/Private/Sequencer/AnimLayers/AnimLayers.cpp b/Engine/Plugins/Animation/ControlRig/Source/ControlRigEditor/Private/Sequencer/AnimLayers/AnimLayers.cpp
--- a/Engine/Plugins/Animation/ControlRig/Source/ControlRigEditor/Private/Sequencer/AnimLayers/AnimLayers.cpp
+++ b/Engine/Plugins/Animation/ControlRig/Source/ControlRigEditor/Private/Sequencer/AnimLayers/AnimLayers.cpp
@@ -1423,6 +1423,28 @@ void UAnimLayers::AddBaseLayer()
 	AnimLayers.Add(AnimLayer);
 }
 
+bool UAnimLayers::HasAnimLayers(ISequencer* InSequencer)
+{
+	if (!InSequencer)
+	{
+		return false;
+	}
+	if(ULevelSequence* LevelSequence = Cast<ULevelSequence>(InSequencer->GetFocusedMovieSceneSequence()))
+	{ 
+		if (LevelSequence->GetClass()->ImplementsInterface(UInterface_AssetUserData::StaticClass()))
+		{
+			if (IInterface_AssetUserData* AssetUserDataInterface = Cast< IInterface_AssetUserData >(LevelSequence))
+			{
+				if (UAnimLayers* AnimLayers = AssetUserDataInterface->GetAssetUserData< UAnimLayers >())
+				{
+					return true;
+				}
+			}
+		}
+	}
+	return false;
+}
+
 UAnimLayers* UAnimLayers::GetAnimLayers(ISequencer* SequencerPtr)
 {
 	if (!SequencerPtr)
diff --git a/Engine/Plugins/Animation/ControlRig/Source/ControlRigEditor/Private/Sequencer/AnimLayers/AnimLayers.h b/Engine/Plugins/Animation/ControlRig/Source/ControlRigEditor/Private/Sequencer/AnimLayers/AnimLayers.h
--- a/Engine/Plugins/Animation/ControlRig/Source/ControlRigEditor/Private/Sequencer/AnimLayers/AnimLayers.h
+++ b/Engine/Plugins/Animation/ControlRig/Source/ControlRigEditor/Private/Sequencer/AnimLayers/AnimLayers.h
@@ -290,6 +290,8 @@ public:
 	//Get
 	static UAnimLayers* GetAnimLayers(ISequencer* InSequencer);
 	static UAnimLayers* GetAnimLayers(ULevelSequence* InLevelSequence);
+	static bool HasAnimLayers(ISequencer* InSequencer);
+
 	static 	TSharedPtr<ISequencer> GetSequencerFromAsset();
 public:
 	DECLARE_EVENT_OneParam(UAnimLayers, FAnimLayerListChanged, UAnimLayers*);
