Control Rig: fixed typo

#rnx

diff --git a/Engine/Plugins/Animation/ControlRig/Source/ControlRig/Public/Rigs/RigHierarchy.h b/Engine/Plugins/Animation/ControlRig/Source/ControlRig/Public/Rigs/RigHierarchy.h
--- a/Engine/Plugins/Animation/ControlRig/Source/ControlRig/Public/Rigs/RigHierarchy.h
+++ b/Engine/Plugins/Animation/ControlRig/Source/ControlRig/Public/Rigs/RigHierarchy.h
@@ -2611,9 +2611,7 @@ public:
 	{
 		if (InControlElement && bUsePreferredEulerAngles)
 		{
-			{
-				return InControlElement->Settings.bUsePreferredRotationOrder;
-			}
+			return InControlElement->Settings.bUsePreferredRotationOrder;
 		}
 		return false;
 	}
@@ -2622,9 +2620,7 @@ public:
 	{
 		if (InControlElement)
 		{
-			{
-				InControlElement->Settings.bUsePreferredRotationOrder = bVal;
-			}
+			InControlElement->Settings.bUsePreferredRotationOrder = bVal;
 		}
 	}
 
