Find Closest Item RigUnit

diff --git a/Engine/Plugins/Animation/ControlRig/Source/ControlRig/Private/Units/Hierarchy/RigUnit_FindClosestItem.cpp b/Engine/Plugins/Animation/ControlRig/Source/ControlRig/Private/Units/Hierarchy/RigUnit_FindClosestItem.cpp
new file mode 100644
--- /dev/null
+++ b/Engine/Plugins/Animation/ControlRig/Source/ControlRig/Private/Units/Hierarchy/RigUnit_FindClosestItem.cpp
@@ -0,0 +1,59 @@
+// Copyright Epic Games, Inc. All Rights Reserved.
+
+#include "Units/Hierarchy/RigUnit_FindClosestItem.h"
+#include "Units/RigUnitContext.h"
+
+#include UE_INLINE_GENERATED_CPP_BY_NAME(RigUnit_FindClosestItem)
+
+FRigUnit_FindClosestItem_Execute()
+{
+	DECLARE_SCOPE_HIERARCHICAL_COUNTER_RIGUNIT()
+
+	if (Items.Num() == 0)
+	{
+		static FRigElementKey InvalidKey;
+		Item = InvalidKey;
+		return;
+	}
+
+	URigHierarchy* Hierarchy = ExecuteContext.Hierarchy;
+	if (Hierarchy == nullptr)
+	{
+		return;
+	}
+
+	CachedItems.SetNum(Items.Num());
+	for (int32 ItemIndex = 0; ItemIndex < Items.Num(); ItemIndex++)
+	{
+		if (!CachedItems[ItemIndex].UpdateCache(Items[ItemIndex], Hierarchy))
+		{
+			CachedItems.Reset();
+			break;
+		}
+	}
+
+	if (CachedItems.Num() == 0 && Items.Num() > 0)
+	{
+		for (FRigElementKey CacheItem : Items)
+		{
+			CachedItems.Add(FCachedRigElement(CacheItem, Hierarchy));
+		}
+
+	}
+
+	float NearestDistance = BIG_NUMBER;
+	Item = CachedItems[0].GetKey();
+	for (int32 ItemIndex = 0; ItemIndex < CachedItems.Num(); ItemIndex++)
+	{
+		if (CachedItems[ItemIndex].IsValid())
+		{
+			const FVector Location = Hierarchy->GetGlobalTransform(CachedItems[ItemIndex]).GetLocation();
+			const float Distance = (Location - Point).Size();
+			if (Distance < NearestDistance)
+			{
+				NearestDistance = Distance;
+				Item = CachedItems[ItemIndex].GetKey();
+			}
+		}
+	}
+}
\ No newline at end of file
diff --git a/Engine/Plugins/Animation/ControlRig/Source/ControlRig/Public/Units/Hierarchy/RigUnit_FindClosestItem.h b/Engine/Plugins/Animation/ControlRig/Source/ControlRig/Public/Units/Hierarchy/RigUnit_FindClosestItem.h
new file mode 100644
--- /dev/null
+++ b/Engine/Plugins/Animation/ControlRig/Source/ControlRig/Public/Units/Hierarchy/RigUnit_FindClosestItem.h
@@ -0,0 +1,34 @@
+// Copyright Epic Games, Inc. All Rights Reserved.
+
+#pragma once
+
+#include "Units/RigUnit.h"
+#include "RigUnit_FindClosestItem.generated.h"
+
+USTRUCT(meta = (DisplayName = "Find Closest Item", Category = "Hierarchy", Keywords = "Find,Closest,Item,Transform,Bone,Joint", NodeColor = "0.3 0.1 0.1"))
+struct CONTROLRIG_API FRigUnit_FindClosestItem : public FRigUnitMutable
+{
+	GENERATED_BODY()
+
+	FRigUnit_FindClosestItem()
+	{
+		Point = FVector::ZeroVector;
+	}
+
+	RIGVM_METHOD()
+		virtual void Execute() override;
+
+	UPROPERTY(meta = (Input))
+	TArray<FRigElementKey> Items;
+
+	UPROPERTY(meta = (Input))
+	FVector Point;
+
+	UPROPERTY(meta = (Output, ExpandByDefault))
+	FRigElementKey Item;
+
+	// Used to cache the internally used bone index
+	UPROPERTY(transient)
+	TArray<FCachedRigElement> CachedItems;
+
+};
\ No newline at end of file
