Cotnrol Rig: Details panel support for components

#rnx

diff --git a/Engine/Plugins/Animation/ControlRig/Source/ControlRig/Private/Rigs/RigHierarchy.cpp b/Engine/Plugins/Animation/ControlRig/Source/ControlRig/Private/Rigs/RigHierarchy.cpp
--- a/Engine/Plugins/Animation/ControlRig/Source/ControlRig/Private/Rigs/RigHierarchy.cpp
+++ b/Engine/Plugins/Animation/ControlRig/Source/ControlRig/Private/Rigs/RigHierarchy.cpp
@@ -967,7 +967,7 @@ void URigHierarchy::CopyHierarchy(URigHierarchy* InHierarchy)
 			}
 			const int32 TargetComponentIndex = ElementComponents.Add(InHierarchy->ElementComponents[SourceComponentIndex]);
 			FRigBaseComponent* TargetComponent = ElementComponents[TargetComponentIndex].GetMutablePtr<FRigBaseComponent>();
-			TargetComponent->Element = Find(TargetComponent->ElementKey);
+			TargetComponent->Element = Find(TargetComponent->GetElementKey());
 			TargetComponent->IndexInHierarchy = TargetComponentIndex;
 			TargetComponent->IndexInElement = TargetComponent->Element->ComponentIndices.Add(TargetComponent->IndexInHierarchy);
 		}
@@ -1431,6 +1431,24 @@ bool URigHierarchy::IsProcedural(const FRigBaseElement* InElement) const
 	return InElement->IsProcedural();
 }
 
+const FRigBaseComponent* URigHierarchy::FindComponent(const FRigComponentKey& InKey) const
+{
+	if(const FRigBaseElement* Element = Find(InKey.ElementKey))
+	{
+		return Element->FindComponent(InKey.Name);
+	}
+	return nullptr;
+}
+
+FRigBaseComponent* URigHierarchy::FindComponent(const FRigComponentKey& InKey)
+{
+	if(FRigBaseElement* Element = Find(InKey.ElementKey))
+	{
+		return Element->FindComponent(InKey.Name);
+	}
+	return nullptr;
+}
+
 const FRigBaseComponent* URigHierarchy::GetComponent(int32 InIndex) const
 {
 	if(ElementComponents.IsValidIndex(InIndex))
@@ -3790,7 +3808,7 @@ void URigHierarchy::Notify(ERigHierarchyNotification InNotifType, const FRigNoti
 					{			
 						if(const FRigBaseElement* ListeningElement = ListeningHierarchy->Find( Element->GetKey()))
 						{
-							if(const FRigBaseComponent* ListeningComponent = ListeningElement->GetComponentByName(InComponent->GetFName()))
+							if(const FRigBaseComponent* ListeningComponent = ListeningElement->FindComponent(InComponent->GetFName()))
 							{
 								ListeningHierarchy->Notify(InNotifType, ListeningComponent);
 							}
@@ -3970,7 +3988,7 @@ void URigHierarchy::SendQueuedNotifications()
 		Subject.Element = Find(Entry.Key);
 		if(Subject.Element)
 		{
-			Subject.Component = Subject.Element->GetComponentByName(Entry.ComponentName);
+			Subject.Component = Subject.Element->FindComponent(Entry.ComponentName);
 			if(Subject.Component)
 			{
 				Subject.Element = nullptr;
@@ -5956,7 +5974,7 @@ FRigBaseComponent* URigHierarchy::MakeComponent(const UScriptStruct* InComponent
 	FRigName UniqueName = SanitizedName;
 
 	int32 Suffix = 1;
-	while (InElement->GetComponentByName(UniqueName) != nullptr)
+	while (InElement->FindComponent(UniqueName) != nullptr)
 	{
 		FString BaseString = SanitizedName.GetName();
 		if (BaseString.Len() > GetMaxNameLength() - 4)
@@ -5967,9 +5985,8 @@ FRigBaseComponent* URigHierarchy::MakeComponent(const UScriptStruct* InComponent
 	}
 	
 	FRigBaseComponent* Component = ElementComponents[ComponentIndex].GetMutablePtr<FRigBaseComponent>();
-	Component->Name = UniqueName;
+	Component->Key = FRigComponentKey(InElement->GetKey(), UniqueName);
 	Component->Element = InElement;
-	Component->ElementKey = InElement->GetKey();
 	Component->IndexInHierarchy = ComponentIndex;
 	Component->IndexInElement = InElement->ComponentIndices.Add(Component->IndexInHierarchy);
 
diff --git a/Engine/Plugins/Animation/ControlRig/Source/ControlRig/Private/Rigs/RigHierarchyCache.cpp b/Engine/Plugins/Animation/ControlRig/Source/ControlRig/Private/Rigs/RigHierarchyCache.cpp
--- a/Engine/Plugins/Animation/ControlRig/Source/ControlRig/Private/Rigs/RigHierarchyCache.cpp
+++ b/Engine/Plugins/Animation/ControlRig/Source/ControlRig/Private/Rigs/RigHierarchyCache.cpp
@@ -121,7 +121,7 @@ bool FCachedRigComponent::UpdateCache(const FRigElementKey& InKey, const FName&
 			{
 				if(const FRigBaseElement* Element = CachedElement.GetElement())
 				{
-					Component = Element->GetComponentByName(InName);
+					Component = Element->FindComponent(InName);
 					if(Component)
 					{
 						Index = Component->GetIndexInHierarchy();
@@ -142,6 +142,11 @@ bool FCachedRigComponent::UpdateCache(const FRigElementKey& InKey, const FName&
 	return false;
 }
 
+bool FCachedRigComponent::UpdateCache(const FRigComponentKey& InKey, const URigHierarchy* InHierarchy)
+{
+	return UpdateCache(InKey.ElementKey, InKey.Name, InHierarchy);
+}
+
 bool FCachedRigComponent::IsIdentical(const FRigElementKey& InKey, const FName& InName, const URigHierarchy* InHierarchy)
 {
 	return InName == Name && CachedElement.IsIdentical(InKey, InHierarchy);
diff --git a/Engine/Plugins/Animation/ControlRig/Source/ControlRig/Private/Rigs/RigHierarchyComponents.cpp b/Engine/Plugins/Animation/ControlRig/Source/ControlRig/Private/Rigs/RigHierarchyComponents.cpp
--- a/Engine/Plugins/Animation/ControlRig/Source/ControlRig/Private/Rigs/RigHierarchyComponents.cpp
+++ b/Engine/Plugins/Animation/ControlRig/Source/ControlRig/Private/Rigs/RigHierarchyComponents.cpp
@@ -10,12 +10,6 @@
 // FRigBaseComponent
 ////////////////////////////////////////////////////////////////////////////////
 
-FRigBaseComponent::FRigBaseComponent(FRigBaseElement* InElement)
-	: Element(InElement)
-	, ElementKey(InElement->GetKey())
-{
-}
-
 FRigBaseComponent::~FRigBaseComponent()
 {
 }
@@ -35,13 +29,11 @@ bool FRigBaseComponent::Serialize(FArchive& Ar)
 
 void FRigBaseComponent::Save(FArchive& Ar)
 {
-	Ar << Name;
-	Ar << ElementKey;
+	Ar << Key;
 }
 
 void FRigBaseComponent::Load(FArchive& Ar)
 {
-	Ar << Name;
+	Ar << Key;
 	CachedNameString.Reset();
-	Ar << ElementKey;
 }
diff --git a/Engine/Plugins/Animation/ControlRig/Source/ControlRig/Private/Rigs/RigHierarchyController.cpp b/Engine/Plugins/Animation/ControlRig/Source/ControlRig/Private/Rigs/RigHierarchyController.cpp
--- a/Engine/Plugins/Animation/ControlRig/Source/ControlRig/Private/Rigs/RigHierarchyController.cpp
+++ b/Engine/Plugins/Animation/ControlRig/Source/ControlRig/Private/Rigs/RigHierarchyController.cpp
@@ -881,23 +881,23 @@ bool URigHierarchyController::SetControlSettings(FRigElementKey InKey, FRigContr
 	return true;
 }
 
-FName URigHierarchyController::AddComponent(UScriptStruct* InComponentStruct, FName InName, FRigElementKey InElement, FString InContent, bool bSetupUndo, bool bPrintPythonCommand)
+FRigComponentKey URigHierarchyController::AddComponent(UScriptStruct* InComponentStruct, FName InName, FRigElementKey InElement, FString InContent, bool bSetupUndo, bool bPrintPythonCommand)
 {
 	if(!IsValid())
 	{
-		return NAME_None;
+		return FRigComponentKey();
 	}
 	
 	if(InComponentStruct == nullptr)
 	{
 		ReportError(TEXT("The passed component struct is nullptr."));
-		return NAME_None;
+		return FRigComponentKey();
 	}
 
 	if(InName.IsNone())
 	{
 		ReportError(TEXT("The passed name is None."));
-		return NAME_None;
+		return FRigComponentKey();
 	}
 
 	URigHierarchy* Hierarchy = GetHierarchy();
@@ -907,7 +907,7 @@ FName URigHierarchyController::AddComponent(UScriptStruct* InComponentStruct, FN
 	if(Element == nullptr)
 	{
 		ReportErrorf(TEXT("The element '%s' could not be found."), *InElement.ToString());
-		return NAME_None;
+		return FRigComponentKey();
 	}
 
 #if WITH_EDITOR
@@ -919,8 +919,12 @@ FName URigHierarchyController::AddComponent(UScriptStruct* InComponentStruct, FN
 	}
 #endif
 
-	const FRigBaseComponent* NewComponent = Hierarchy->MakeComponent(InComponentStruct, InName, Element);
-	(void)SetComponentContent(NewComponent->GetFName(), NewComponent->GetElementKey(), InContent, false, false);
+	FRigBaseComponent* NewComponent = Hierarchy->MakeComponent(InComponentStruct, InName, Element);
+	NewComponent->CreatedAtInstructionIndex = CurrentInstructionIndex;
+	if(!InContent.IsEmpty())
+	{
+		(void)SetComponentContent(NewComponent->GetKey(), InContent, false, false);
+	}
 	Hierarchy->IncrementTopologyVersion();
 	Notify(ERigHierarchyNotification::ComponentAdded, NewComponent);
 
@@ -941,36 +945,29 @@ FName URigHierarchyController::AddComponent(UScriptStruct* InComponentStruct, FN
 	}
 #endif
 
-	return NewComponent->GetFName();
+	return NewComponent->GetKey();
 }
 
-bool URigHierarchyController::RemoveComponent(FName InName, FRigElementKey InElement, bool bSetupUndo, bool bPrintPythonCommand)
+bool URigHierarchyController::RemoveComponent(FRigComponentKey InComponent, bool bSetupUndo, bool bPrintPythonCommand)
 {
 	if(!IsValid())
 	{
 		return false;
 	}
 	
-	if(InName.IsNone())
+	if(!InComponent.IsValid())
 	{
-		ReportError(TEXT("The passed name is None."));
+		ReportError(TEXT("The passed component key is invalid."));
 		return false;
 	}
 
 	URigHierarchy* Hierarchy = GetHierarchy();
 	check(Hierarchy);
 	
-	FRigBaseElement* Element = Hierarchy->Find(InElement);
-	if(Element == nullptr)
-	{
-		ReportErrorf(TEXT("The element '%s' could not be found."), *InElement.ToString());
-		return false;
-	}
-
-	FRigBaseComponent* Component = Element->GetComponentByName(InName);
+	FRigBaseComponent* Component = Hierarchy->FindComponent(InComponent);
 	if(Component == nullptr)
 	{
-		ReportErrorf(TEXT("The component '%s' cannot be found on element '%s'."), *InName.ToString(), *InElement.ToString());
+		ReportErrorf(TEXT("The component '%s' cannot be found."), *InComponent.ToString());
 		return false;
 	}
 
@@ -995,7 +992,7 @@ bool URigHierarchyController::RemoveComponent(FName InName, FRigElementKey InEle
 		if (const UBlueprint* Blueprint = GetTypedOuter<UBlueprint>())
 		{
 			RigVMPythonUtils::Print(Blueprint->GetFName().ToString(),
-				FString::Printf(TEXT("hierarchy_controller.remove_component('%s', %s"), *InName.ToString(), *InElement.ToPythonString()));
+				FString::Printf(TEXT("hierarchy_controller.remove_component(%s"), *InComponent.ToPythonString()));
 		}
 	}
 #endif
@@ -1003,33 +1000,32 @@ bool URigHierarchyController::RemoveComponent(FName InName, FRigElementKey InEle
 	return true;
 }
 
-bool URigHierarchyController::SetComponentContent(FName InName, FRigElementKey InElement, const FString& InContent, bool bSetupUndo, bool bPrintPythonCommand)
+bool URigHierarchyController::SetComponentContent(FRigComponentKey InComponent, const FString& InContent, bool bSetupUndo, bool bPrintPythonCommand)
 {
 	if(!IsValid())
 	{
 		return false;
 	}
 	
-	if(InName.IsNone())
+	if(!InComponent.IsValid())
 	{
-		ReportError(TEXT("The passed name is None."));
+		ReportError(TEXT("The passed component key is valid."));
 		return false;
 	}
 
-	URigHierarchy* Hierarchy = GetHierarchy();
-	check(Hierarchy);
-	
-	FRigBaseElement* Element = Hierarchy->Find(InElement);
-	if(Element == nullptr)
+	if(InContent.IsEmpty())
 	{
-		ReportErrorf(TEXT("The element '%s' could not be found."), *InElement.ToString());
+		ReportError(TEXT("The passed content is empty."));
 		return false;
 	}
 
-	FRigBaseComponent* Component = Element->GetComponentByName(InName);
+	URigHierarchy* Hierarchy = GetHierarchy();
+	check(Hierarchy);
+	
+	FRigBaseComponent* Component = Hierarchy->FindComponent(InComponent);
 	if(Component == nullptr)
 	{
-		ReportErrorf(TEXT("The component '%s' cannot be found on element '%s'."), *InName.ToString(), *InElement.ToString());
+		ReportErrorf(TEXT("The component '%s' cannot be found."), *InComponent.ToString());
 		return false;
 	}
 
@@ -1065,7 +1061,7 @@ bool URigHierarchyController::SetComponentContent(FName InName, FRigElementKey I
 		if (const UBlueprint* Blueprint = GetTypedOuter<UBlueprint>())
 		{
 			RigVMPythonUtils::Print(Blueprint->GetFName().ToString(),
-				FString::Printf(TEXT("hierarchy_controller.set_component_content('%s', %s, '%s'"), *InName.ToString(), *InElement.ToPythonString(), *InContent));
+				FString::Printf(TEXT("hierarchy_controller.set_component_content( %s, '%s'"), *InComponent.ToPythonString(), *InContent));
 		}
 	}
 #endif
diff --git a/Engine/Plugins/Animation/ControlRig/Source/ControlRig/Private/Rigs/RigHierarchyDefines.cpp b/Engine/Plugins/Animation/ControlRig/Source/ControlRig/Private/Rigs/RigHierarchyDefines.cpp
--- a/Engine/Plugins/Animation/ControlRig/Source/ControlRig/Private/Rigs/RigHierarchyDefines.cpp
+++ b/Engine/Plugins/Animation/ControlRig/Source/ControlRig/Private/Rigs/RigHierarchyDefines.cpp
@@ -94,6 +94,49 @@ FString FRigElementKey::ToPythonString() const
 #endif
 }
 
+////////////////////////////////////////////////////////////////////////////////
+// FRigComponentKey
+////////////////////////////////////////////////////////////////////////////////
+
+void FRigComponentKey::Serialize(FArchive& Ar)
+{
+	if (Ar.IsSaving() || Ar.IsObjectReferenceCollector() || Ar.IsCountingMemory())
+	{
+		Save(Ar);
+	}
+	else if (Ar.IsLoading())
+	{
+		Load(Ar);
+	}
+	else
+	{
+		// remove due to FPIEFixupSerializer hitting this checkNoEntry();
+	}
+}
+
+void FRigComponentKey::Save(FArchive& Ar)
+{
+	Ar << Name;
+	Ar << ElementKey;
+}
+
+void FRigComponentKey::Load(FArchive& Ar)
+{
+	Ar << Name;
+	Ar << ElementKey;
+}
+
+FString FRigComponentKey::ToPythonString() const
+{
+#if WITH_EDITOR
+	return FString::Printf(TEXT("unreal.RigComponentKey(element_key=%s, name='%s')"),
+		*ElementKey.ToPythonString(),
+		*Name.ToString());
+#else
+	return FString();
+#endif
+}
+
 ////////////////////////////////////////////////////////////////////////////////
 // FRigElementKeyAndIndex
 ////////////////////////////////////////////////////////////////////////////////
diff --git a/Engine/Plugins/Animation/ControlRig/Source/ControlRig/Private/Rigs/RigHierarchyElements.cpp b/Engine/Plugins/Animation/ControlRig/Source/ControlRig/Private/Rigs/RigHierarchyElements.cpp
--- a/Engine/Plugins/Animation/ControlRig/Source/ControlRig/Private/Rigs/RigHierarchyElements.cpp
+++ b/Engine/Plugins/Animation/ControlRig/Source/ControlRig/Private/Rigs/RigHierarchyElements.cpp
@@ -181,12 +181,12 @@ FRigBaseComponent* FRigBaseElement::GetComponent(int32 InIndex)
 	return nullptr;
 }
 
-const FRigBaseComponent* FRigBaseElement::GetComponentByName(const FName& InName) const
+const FRigBaseComponent* FRigBaseElement::FindComponent(const FName& InName) const
 {
-	return const_cast<FRigBaseElement*>(this)->GetComponentByName(InName);
+	return const_cast<FRigBaseElement*>(this)->FindComponent(InName);
 }
 
-FRigBaseComponent* FRigBaseElement::GetComponentByName(const FName& InName)
+FRigBaseComponent* FRigBaseElement::FindComponent(const FName& InName)
 {
 	if(!ComponentIndices.IsEmpty())
 	{
@@ -207,12 +207,12 @@ FRigBaseComponent* FRigBaseElement::GetComponentByName(const FName& InName)
 	return nullptr;
 }
 
-const FRigBaseComponent* FRigBaseElement::GetComponentOfType(const UScriptStruct* InComponentStruct) const
+const FRigBaseComponent* FRigBaseElement::GetFirstComponent(const UScriptStruct* InComponentStruct) const
 {
-	return const_cast<FRigBaseElement*>(this)->GetComponentOfType(InComponentStruct);
+	return const_cast<FRigBaseElement*>(this)->GetFirstComponent(InComponentStruct);
 }
 
-FRigBaseComponent* FRigBaseElement::GetComponentOfType(const UScriptStruct* InComponentStruct)
+FRigBaseComponent* FRigBaseElement::GetFirstComponent(const UScriptStruct* InComponentStruct)
 {
 	if(!ComponentIndices.IsEmpty())
 	{
@@ -626,32 +626,32 @@ FRigBaseElement* FRigElementHandle::Get()
 // FRigComponentHandle
 ////////////////////////////////////////////////////////////////////////////////
 
-FRigComponentHandle::FRigComponentHandle(URigHierarchy* InHierarchy, const FRigElementKey& InKey, const FName& InComponentName)
-: ElementHandle(InHierarchy, InKey)
-, ComponentName(InComponentName)
+FRigComponentHandle::FRigComponentHandle(URigHierarchy* InHierarchy, const FRigComponentKey& InKey)
+: Hierarchy(InHierarchy)
+, Key(InKey)
 {
 }
 
 FRigComponentHandle::FRigComponentHandle(URigHierarchy* InHierarchy, const FRigBaseComponent* InComponent)
-: ElementHandle(InHierarchy, InComponent->GetElement())
-, ComponentName(InComponent->GetFName())
+: Hierarchy(InHierarchy)
+, Key(InComponent->GetKey())
 {
 }
 
 const FRigBaseComponent* FRigComponentHandle::Get() const
 {
-	if(const FRigBaseElement* Element = ElementHandle.Get())
+	if(const URigHierarchy* CurrentHierarchy = Hierarchy.Get())
 	{
-		return Element->GetComponentByName(ComponentName);
+		return CurrentHierarchy->FindComponent(Key);
 	}
 	return nullptr;
 }
 
 FRigBaseComponent* FRigComponentHandle::Get()
 {
-	if(FRigBaseElement* Element = ElementHandle.Get())
+	if(URigHierarchy* CurrentHierarchy = Hierarchy.Get())
 	{
-		return Element->GetComponentByName(ComponentName);
+		return CurrentHierarchy->FindComponent(Key);
 	}
 	return nullptr;
 }
diff --git a/Engine/Plugins/Animation/ControlRig/Source/ControlRig/Public/Rigs/RigHierarchy.h b/Engine/Plugins/Animation/ControlRig/Source/ControlRig/Public/Rigs/RigHierarchy.h
--- a/Engine/Plugins/Animation/ControlRig/Source/ControlRig/Public/Rigs/RigHierarchy.h
+++ b/Engine/Plugins/Animation/ControlRig/Source/ControlRig/Public/Rigs/RigHierarchy.h
@@ -705,6 +705,20 @@ public:
 		return Get<T>(GetIndex(InKey));
 	}
 
+	/**
+	 * Returns a component for a given key or nullptr.
+	 * @param InKey The key of the component to retrieve.
+	 * @return The retrieved component or nullptr.
+	 */
+	const FRigBaseComponent* FindComponent(const FRigComponentKey& InKey) const;
+
+	/**
+	 * Returns a component for a given key or nullptr.
+	 * @param InKey The key of the component to retrieve.
+	 * @return The retrieved component or nullptr.
+	 */
+	FRigBaseComponent* FindComponent(const FRigComponentKey& InKey);
+
 	/**
 	 * Returns a component at a given index or nullptr.
 	 * @param InIndex The index of the component to retrieve.
diff --git a/Engine/Plugins/Animation/ControlRig/Source/ControlRig/Public/Rigs/RigHierarchyCache.h b/Engine/Plugins/Animation/ControlRig/Source/ControlRig/Public/Rigs/RigHierarchyCache.h
--- a/Engine/Plugins/Animation/ControlRig/Source/ControlRig/Public/Rigs/RigHierarchyCache.h
+++ b/Engine/Plugins/Animation/ControlRig/Source/ControlRig/Public/Rigs/RigHierarchyCache.h
@@ -204,6 +204,19 @@ public:
 		}
 	}
 
+	FCachedRigComponent(const FRigComponentKey& InKey, const URigHierarchy* InHierarchy, bool bForceStoreName = false)
+	: CachedElement(InKey.ElementKey, InHierarchy, bForceStoreName)
+	, Name(NAME_None)
+	, Index(UINT16_MAX)
+	, Component(nullptr)
+	{
+		UpdateCache(InKey, InHierarchy);
+		if(bForceStoreName)
+		{
+			Name = InKey.Name;
+		}
+	}
+
 	bool IsValid() const
 	{
 		return GetIndex() != INDEX_NONE && CachedElement.IsValid();
@@ -250,6 +263,16 @@ public:
 
 	const FRigElementKey& GetResolvedElementKey() const;
 
+	FRigComponentKey GetComponentKey() const
+	{
+		return FRigComponentKey(GetElementKey(), GetName());
+	}
+
+	FRigComponentKey GetResolvedComponentKey() const
+	{
+		return FRigComponentKey(GetResolvedElementKey(), GetName());
+	}
+
 	const FName& GetName() const
 	{
 		return Name;
@@ -283,6 +306,8 @@ public:
 
 	bool UpdateCache(const FRigElementKey& InKey, const FName& InName, const URigHierarchy* InHierarchy);
 
+	bool UpdateCache(const FRigComponentKey& InKey, const URigHierarchy* InHierarchy);
+
 	friend uint32 GetTypeHash(const FCachedRigComponent& CachedComponent)
 	{
 		return GetTypeHash(CachedComponent.CachedElement) * 17 + (uint32)CachedComponent.Index;
diff --git a/Engine/Plugins/Animation/ControlRig/Source/ControlRig/Public/Rigs/RigHierarchyComponents.h b/Engine/Plugins/Animation/ControlRig/Source/ControlRig/Public/Rigs/RigHierarchyComponents.h
--- a/Engine/Plugins/Animation/ControlRig/Source/ControlRig/Public/Rigs/RigHierarchyComponents.h
+++ b/Engine/Plugins/Animation/ControlRig/Source/ControlRig/Public/Rigs/RigHierarchyComponents.h
@@ -51,11 +51,10 @@ public:
 	FRigBaseComponent& operator=(const FRigBaseComponent& InOther)
 	{
 		// We purposefully do not copy any non-UPROPERTY entries, including Element
-		Name = InOther.Name;
+		Key = InOther.Key;
 		CachedNameString.Reset();
 		IndexInElement = InOther.IndexInElement;
 		IndexInHierarchy = InOther.IndexInHierarchy;
-		ElementKey = InOther.ElementKey;
 		CreatedAtInstructionIndex = InOther.CreatedAtInstructionIndex;
 		bSelected = InOther.bSelected;
 		return *this;
@@ -63,9 +62,6 @@ public:
 	
 protected:
 
-	// Only derived types should be able to construct this one.
-	explicit FRigBaseComponent(FRigBaseElement* InElement);
-
 	static bool IsClassOf(const FRigBaseComponent* InComponent)
 	{
 		return true;
@@ -74,18 +70,15 @@ protected:
 	// ReSharper disable once CppUE4ProbableMemoryIssuesWithUObject
 	FRigBaseElement* Element = nullptr;
 	
-	UPROPERTY(BlueprintReadOnly, VisibleAnywhere, Category = RigElement, meta = (AllowPrivateAccess = "true"))
-	FName Name = FName(NAME_None);
-
-	UPROPERTY(BlueprintReadOnly, Transient, VisibleAnywhere, Category = RigElement, meta = (AllowPrivateAccess = "true"))
+	UPROPERTY(BlueprintReadOnly, Category = RigElement, meta = (AllowPrivateAccess = "true"))
+	FRigComponentKey Key;
+	
+	UPROPERTY(BlueprintReadOnly, Transient, Category = RigElement, meta = (AllowPrivateAccess = "true"))
 	int32 IndexInHierarchy = INDEX_NONE;
 
-	UPROPERTY(BlueprintReadOnly, Transient, VisibleAnywhere, Category = RigElement, meta = (AllowPrivateAccess = "true"))
+	UPROPERTY(BlueprintReadOnly, Transient, Category = RigElement, meta = (AllowPrivateAccess = "true"))
 	int32 IndexInElement = INDEX_NONE;
 
-	UPROPERTY(BlueprintReadOnly, VisibleAnywhere, Category = RigElement, meta = (AllowPrivateAccess = "true"))
-	FRigElementKey ElementKey;
-
 	UPROPERTY(BlueprintReadOnly, Transient, Category = RigElement, meta = (AllowPrivateAccess = "true"))
 	int32 CreatedAtInstructionIndex = INDEX_NONE;
 
@@ -98,12 +91,13 @@ public:
 
 	virtual UScriptStruct* GetScriptStruct() const { return FRigBaseComponent::StaticStruct(); }
 
-	const FName& GetFName() const { return Name; }
+	const FRigComponentKey& GetKey() const { return Key; }
+	const FName& GetFName() const { return Key.Name; }
 	const FString& GetName() const
 	{
-		if(CachedNameString.IsEmpty() && !Name.IsNone())
+		if(CachedNameString.IsEmpty() && !Key.Name.IsNone())
 		{
-			CachedNameString = Name.ToString();
+			CachedNameString = Key.Name.ToString();
 		}
 		return CachedNameString;
 	}
@@ -114,7 +108,7 @@ public:
 	bool IsSelected() const { return bSelected; }
 	int32 GetCreatedAtInstructionIndex() const { return CreatedAtInstructionIndex; }
 	bool IsProcedural() const { return CreatedAtInstructionIndex != INDEX_NONE; }
-	const FRigElementKey& GetElementKey() const { return ElementKey; }
+	const FRigElementKey& GetElementKey() const { return Key.ElementKey; }
 	const FRigBaseElement* GetElement() const { return Element; }
 	FRigBaseElement* GetElement() { return Element; }
 
diff --git a/Engine/Plugins/Animation/ControlRig/Source/ControlRig/Public/Rigs/RigHierarchyController.h b/Engine/Plugins/Animation/ControlRig/Source/ControlRig/Public/Rigs/RigHierarchyController.h
--- a/Engine/Plugins/Animation/ControlRig/Source/ControlRig/Public/Rigs/RigHierarchyController.h
+++ b/Engine/Plugins/Animation/ControlRig/Source/ControlRig/Public/Rigs/RigHierarchyController.h
@@ -294,30 +294,28 @@ public:
 	 * @return The name of the newly created component
 	 */
 	UFUNCTION(BlueprintCallable, Category = URigHierarchyController)
-	FName AddComponent(UScriptStruct* InComponentStruct, FName InName, FRigElementKey InElement, FString InContent = TEXT(""), bool bSetupUndo = false, bool bPrintPythonCommand = false);
+	FRigComponentKey AddComponent(UScriptStruct* InComponentStruct, FName InName, FRigElementKey InElement, FString InContent = TEXT(""), bool bSetupUndo = false, bool bPrintPythonCommand = false);
 
 	/**
      * Removes a component from the hierarchy
-     * @param InName The name of the component to remove
-     * @param InElement The element the component belongs to
+     * @param InComponent The component to remove
      * @param bSetupUndo If set to true the stack will record the change for undo / redo
      * @param bPrintPythonCommand If set to true a python command equivalent to this call will be printed out
      * @return True if the component was removed successfully
      */
     UFUNCTION(BlueprintCallable, Category = URigHierarchyController)
-    bool RemoveComponent(FName InName, FRigElementKey InElement, bool bSetupUndo = false, bool bPrintPythonCommand = false);
+    bool RemoveComponent(FRigComponentKey InComponent, bool bSetupUndo = false, bool bPrintPythonCommand = false);
 
 	/**
 	 * Updates the content of a component in the hierarchy
-	 * @param InName The name of the component
-	 * @param InElement The element the component belongs to
+	 * @param InComponent The component to change the content for
 	 * @param InContent The serialized text default for the component
 	 * @param bSetupUndo If set to true the stack will record the change for undo / redo
 	 * @param bPrintPythonCommand If set to true a python command equivalent to this call will be printed out
 	 * @return True if the component was updated correctly
 	 */
 	UFUNCTION(BlueprintCallable, Category = URigHierarchyController)
-	bool SetComponentContent(FName InName, FRigElementKey InElement, const FString& InContent, bool bSetupUndo = false, bool bPrintPythonCommand = false);
+	bool SetComponentContent(FRigComponentKey InComponent, const FString& InContent, bool bSetupUndo = false, bool bPrintPythonCommand = false);
 
 	/**
 	 * Imports an existing skeleton to the hierarchy
diff --git a/Engine/Plugins/Animation/ControlRig/Source/ControlRig/Public/Rigs/RigHierarchyDefines.h b/Engine/Plugins/Animation/ControlRig/Source/ControlRig/Public/Rigs/RigHierarchyDefines.h
--- a/Engine/Plugins/Animation/ControlRig/Source/ControlRig/Public/Rigs/RigHierarchyDefines.h
+++ b/Engine/Plugins/Animation/ControlRig/Source/ControlRig/Public/Rigs/RigHierarchyDefines.h
@@ -1498,14 +1498,13 @@ public:
 	FRigElementKey(const FString& InKeyString)
 	{
 		FString TypeStr, NameStr;
-		check(InKeyString.Split(TEXT("("), &TypeStr, &NameStr));
+		verify(InKeyString.Split(TEXT("("), &TypeStr, &NameStr));
 		NameStr.RemoveFromEnd(TEXT(")"));
 		Name = *NameStr;
 
 		const UEnum* ElementTypeEnum = StaticEnum<ERigElementType>();
 		Type = (ERigElementType)ElementTypeEnum->GetValueByName(*TypeStr);
 		check(Type > ERigElementType::None && Type < ERigElementType::All);
-		
 	}
 
 	void Serialize(FArchive& Ar);
@@ -1630,6 +1629,122 @@ public:
 	FString ToPythonString() const;
 };
 
+USTRUCT(BlueprintType)
+struct CONTROLRIG_API FRigComponentKey
+{
+public:
+	
+	GENERATED_BODY()
+
+	UPROPERTY(BlueprintReadWrite, EditAnywhere, Category = "Hierarchy", meta = (DisplayName="Item"))
+	FRigElementKey ElementKey;
+
+	UPROPERTY(BlueprintReadWrite, EditAnywhere, Category = "Hierarchy", meta = (CustomWidget = "ComponentName"))
+	FName Name;
+
+	FRigComponentKey()
+		: ElementKey()
+		, Name(NAME_None)
+	{
+	}
+
+	FRigComponentKey(const FRigElementKey& InElementKey, const FName& InName)
+		: ElementKey(InElementKey)
+		, Name(InName)
+	{}
+
+	FRigComponentKey(const FString& InKeyString)
+	{
+		FString KeyString = InKeyString;
+		KeyString.RemoveFromStart(TEXT("Component("));
+		KeyString.RemoveFromEnd(TEXT(")"));
+		FString NameStr, ElementKeyStr;
+		verify(KeyString.Split(TEXT(","), &NameStr, &ElementKeyStr));
+		Name = *NameStr;
+		ElementKey = FRigElementKey(ElementKeyStr);
+	}
+	
+	void Serialize(FArchive& Ar);
+	void Save(FArchive& Ar);
+	void Load(FArchive& Ar);
+	friend FArchive& operator<<(FArchive& Ar, FRigComponentKey& P)
+	{
+		P.Serialize(Ar);
+		return Ar;
+	}
+
+	bool IsValid() const
+	{
+		return Name.IsValid() && Name != NAME_None && ElementKey.IsValid();
+	}
+
+	explicit operator bool() const
+	{
+		return IsValid();
+	}
+
+	void Reset()
+	{
+		ElementKey.Reset();
+		Name = NAME_None;
+	}
+
+	friend uint32 GetTypeHash(const FRigComponentKey& Key)
+	{
+		return GetTypeHash(Key.Name) * 15 + GetTypeHash(Key.ElementKey);
+	}
+
+	friend uint32 GetTypeHash(const TArrayView<const FRigComponentKey>& Keys)
+	{
+		uint32 Hash = (uint32)(Keys.Num() * 17 + 3);
+		for (const FRigComponentKey& Key : Keys)
+		{
+			Hash += GetTypeHash(Key);
+		}
+		return Hash;
+	}
+
+	friend uint32 GetTypeHash(const TArray<FRigComponentKey>& Keys)
+	{
+		return GetTypeHash(TArrayView<const FRigComponentKey>(Keys.GetData(), Keys.Num()));
+	}
+
+	bool operator ==(const FRigComponentKey& Other) const
+	{
+		return Name == Other.Name && ElementKey == Other.ElementKey;
+	}
+
+	bool operator !=(const FRigComponentKey& Other) const
+	{
+		return Name != Other.Name || ElementKey != Other.ElementKey;
+	}
+
+	bool operator <(const FRigComponentKey& Other) const
+	{
+		if (ElementKey < Other.ElementKey)
+		{
+			return true;
+		}
+		return Name.LexicalLess(Other.Name);
+	}
+
+	bool operator >(const FRigComponentKey& Other) const
+	{
+		if (ElementKey > Other.ElementKey)
+		{
+			return true;
+		}
+		return Other.Name.LexicalLess(Name);
+	}
+
+	FString ToString() const
+	{
+		return FString::Printf(TEXT("Component(%s,%s)"), *Name.ToString(), *ElementKey.ToString());
+	}
+
+	FString ToPythonString() const;
+};
+
 struct CONTROLRIG_API FRigElementKeyAndIndex
 {
 public:
diff --git a/Engine/Plugins/Animation/ControlRig/Source/ControlRig/Public/Rigs/RigHierarchyElements.h b/Engine/Plugins/Animation/ControlRig/Source/ControlRig/Public/Rigs/RigHierarchyElements.h
--- a/Engine/Plugins/Animation/ControlRig/Source/ControlRig/Public/Rigs/RigHierarchyElements.h
+++ b/Engine/Plugins/Animation/ControlRig/Source/ControlRig/Public/Rigs/RigHierarchyElements.h
@@ -714,19 +714,21 @@ struct CONTROLRIG_API FRigComponentHandle
 public:
 
 	FRigComponentHandle()
-		: ElementHandle()
-		, ComponentName(NAME_None)
-	{}
+		: Key()
+	{
+	}
 
-	FRigComponentHandle(URigHierarchy* InHierarchy, const FRigElementKey& InKey, const FName& InComponentName);
+	FRigComponentHandle(URigHierarchy* InHierarchy, const FRigComponentKey& InKey);
 	FRigComponentHandle(URigHierarchy* InHierarchy, const FRigBaseComponent* InComponent);
 
 	bool IsValid() const { return Get() != nullptr; }
 	operator bool() const { return IsValid(); }
 	
-	const URigHierarchy* GetHierarchy() const { return ElementHandle.GetHierarchy(); }
-	URigHierarchy* GetHierarchy() { return ElementHandle.GetHierarchy(); }
-	const FRigElementKey& GetElementKey() const { return ElementHandle.GetKey(); }
+	const URigHierarchy* GetHierarchy() const { return Hierarchy.Get(); }
+	URigHierarchy* GetHierarchy() { return Hierarchy.Get(); }
+	const FRigElementKey& GetElementKey() const { return Key.ElementKey; }
+	const FRigComponentKey& GetComponentKey() const { return Key; }
+	const FName& GetComponentName() const { return Key.Name; }
 
 	const FRigBaseComponent* Get() const;
 	FRigBaseComponent* Get();
@@ -757,8 +759,8 @@ public:
 
 private:
 
-	FRigElementHandle ElementHandle;
-	FName ComponentName;
+	TWeakObjectPtr<URigHierarchy> Hierarchy;
+	FRigComponentKey Key;
 };
 
 USTRUCT(BlueprintType)
@@ -908,21 +910,21 @@ public:
 	int32 NumComponents() const;
 	const FRigBaseComponent* GetComponent(int32 InIndex) const;
 	FRigBaseComponent* GetComponent(int32 InIndex);
-	const FRigBaseComponent* GetComponentByName(const FName& InName) const;
-	FRigBaseComponent* GetComponentByName(const FName& InName);
-	const FRigBaseComponent* GetComponentOfType(const UScriptStruct* InComponentStruct) const;
-	FRigBaseComponent* GetComponentOfType(const UScriptStruct* InComponentStruct);
+	const FRigBaseComponent* FindComponent(const FName& InName) const;
+	FRigBaseComponent* FindComponent(const FName& InName);
+	const FRigBaseComponent* GetFirstComponent(const UScriptStruct* InComponentStruct) const;
+	FRigBaseComponent* GetFirstComponent(const UScriptStruct* InComponentStruct);
 
 	template<typename T>
-	const T* GetComponent() const
+	const T* GetFirstComponent() const
 	{
-		return Cast<T>(GetComponentOfType(T::StaticStruct()));
+		return Cast<T>(GetFirstComponent(T::StaticStruct()));
 	}
 
 	template<typename T>
-	T* GetComponent()
+	T* GetFirstComponent()
 	{
-		return Cast<T>(GetComponentOfType(T::StaticStruct()));
+		return Cast<T>(GetFirstComponent(T::StaticStruct()));
 	}
 
 	virtual int32 GetNumTransforms() const { return 0; }
diff --git a/Engine/Plugins/Animation/ControlRig/Source/ControlRigDeveloper/Private/Graph/ControlRigGraphSchema.cpp b/Engine/Plugins/Animation/ControlRig/Source/ControlRigDeveloper/Private/Graph/ControlRigGraphSchema.cpp
--- a/Engine/Plugins/Animation/ControlRig/Source/ControlRigDeveloper/Private/Graph/ControlRigGraphSchema.cpp
+++ b/Engine/Plugins/Animation/ControlRig/Source/ControlRigDeveloper/Private/Graph/ControlRigGraphSchema.cpp
@@ -31,7 +31,12 @@ FLinearColor UControlRigGraphSchema::GetPinTypeColor(const FEdGraphPinType& PinT
 				return FLinearColor(0.0, 0.6588, 0.9490);
 			}
 
-			if (Struct == FRigElementKey::StaticStruct() || Struct == FRigPose::StaticStruct())
+			if (Struct == FRigComponentKey::StaticStruct())
+			{
+				return FLinearColor(0.0, 0.6588 * 0.7, 0.9490 * 0.7);
+			}
+
+			if (Struct == FRigPose::StaticStruct())
 			{
 				return FLinearColor(0.0, 0.3588, 0.5490);
 			}
diff --git a/Engine/Plugins/Animation/ControlRig/Source/ControlRigEditor/Private/ControlRigElementDetails.cpp b/Engine/Plugins/Animation/ControlRig/Source/ControlRigEditor/Private/ControlRigElementDetails.cpp
--- a/Engine/Plugins/Animation/ControlRig/Source/ControlRigEditor/Private/ControlRigElementDetails.cpp
+++ b/Engine/Plugins/Animation/ControlRig/Source/ControlRigEditor/Private/ControlRigElementDetails.cpp
@@ -26,6 +26,7 @@
 #include "StructViewerFilter.h"
 #include "StructViewerModule.h"
 #include "Widgets/SRigVMGraphPinEnumPicker.h"
+#include "IStructureDataProvider.h"
 
 #define LOCTEXT_NAMESPACE "ControlRigElementDetails"
 
@@ -893,6 +894,7 @@ void FRigBaseElementDetails::CustomizeDetails(IDetailLayoutBuilder& DetailBuilde
 		!IsAnyElementOfType(ERigElementType::Connector) &&
 		!IsAnyElementOfType(ERigElementType::Socket))
 	{
+		CustomizeComponents(DetailBuilder);
 		CustomizeMetadata(DetailBuilder);
 	}
 }
@@ -1293,6 +1295,211 @@ FReply FRigBaseElementDetails::OnSelectElementClicked(const FRigElementKey& InKe
 	return FReply::Handled();
 }
 
+class FRigComponentStructProvider : public IStructureDataProvider
+{
+public:
+	FRigComponentStructProvider(URigHierarchy* InHierarchy)
+		: HierarchyPtr(InHierarchy)
+	{
+	}
+	
+	virtual ~FRigComponentStructProvider() override {}
+
+	int32 Num() const
+	{
+		return ComponentIndices.Num();
+	}
+
+	const FRigBaseComponent* GetComponent(int32 InIndex) const
+	{
+		if(URigHierarchy* Hierarchy = GetHierarchy())
+		{
+			return Hierarchy->GetComponent(ComponentIndices[InIndex]);
+		}
+		return nullptr;
+	}
+
+	URigHierarchy* GetHierarchy() const
+	{
+		if(HierarchyPtr.IsValid())
+		{
+			return HierarchyPtr.Get();
+		}
+		return nullptr;
+	}
+
+	void Reset()
+	{
+		HierarchyPtr.Reset();
+		ComponentIndices.Reset();
+	}
+
+	void AddComponent(const FRigBaseComponent* InComponent)
+	{
+		ComponentIndices.AddUnique(InComponent->GetIndexInHierarchy());
+	}
+	
+	virtual bool IsValid() const override
+	{
+		return GetHierarchy() && Num() != 0;
+	}
+	
+	virtual const UStruct* GetBaseStructure() const override
+	{
+		if(Num() > 0)
+		{
+			if(const FRigBaseComponent* Component = GetComponent(0))
+			{
+				return Component->GetScriptStruct();
+			}
+		}
+		return nullptr;
+	}
+
+	virtual void GetInstances(TArray<TSharedPtr<FStructOnScope>>& OutInstances, const UStruct* ExpectedBaseStructure) const override
+	{
+		for(int32 Index = 0; Index < Num(); Index++)
+		{
+			if(const FRigBaseComponent* Component = GetComponent(Index))
+			{
+				check(Component->GetScriptStruct() == ExpectedBaseStructure);
+				uint8* Memory = reinterpret_cast<uint8*>(const_cast<FRigBaseComponent*>(Component));
+				OutInstances.Add(MakeShareable(new FStructOnScope(ExpectedBaseStructure, Memory)));
+			}
+		}
+	}
+
+	virtual bool IsPropertyIndirection() const override
+	{
+		return false;
+	}
+
+protected:
+
+	mutable TWeakObjectPtr<URigHierarchy> HierarchyPtr;
+	TArray<int32> ComponentIndices;
+};
+
+void FRigBaseElementDetails::CustomizeComponents(IDetailLayoutBuilder& DetailBuilder)
+{
+	if(PerElementInfos.IsEmpty())
+	{
+		return;
+	}
+
+	URigHierarchy* Hierarchy = PerElementInfos[0].GetHierarchy();
+	if(Hierarchy == nullptr)
+	{
+		return;
+	}
+
+	if(const FRigBaseElement* Element = PerElementInfos[0].GetElement<>())
+	{
+		for(int32 ComponentIndex = 0; ComponentIndex < Element->NumComponents(); ComponentIndex++)
+		{
+			if(const FRigBaseComponent* Component = Element->GetComponent(ComponentIndex))
+			{
+				FSharedComponent SharedComponent;
+				SharedComponent.Handle = FRigComponentHandle(Hierarchy, Component);
+				SharedComponent.ScriptStruct = Component->GetScriptStruct();
+				// todo: we need to discuss if procedural components should be disabled for editing?
+				//SharedComponent.bIsProcedural = Component->IsProcedural();
+				SharedComponents.Add(SharedComponent);
+			}
+		}
+	}
+	
+	for(int32 ElementIndex = 1; ElementIndex < PerElementInfos.Num(); ElementIndex++)
+	{
+		if(const FRigBaseElement* Element = PerElementInfos[ElementIndex].GetElement<>())
+		{
+			// remove any missing or type-mismatching component from the list to display
+			SharedComponents.RemoveAll([Element](const FSharedComponent& SharedComponent) -> bool {
+				if(const FRigBaseComponent* Component = Element->FindComponent(SharedComponent.Handle.GetComponentName()))
+				{
+					return Component->GetScriptStruct() != SharedComponent.ScriptStruct;
+				}
+				return true;
+			});
+
+			// update the procedural flag in case any component is procedural
+			// todo: we need to discuss if procedural components should be disabled for editing?
+			/*
+			for(FSharedComponent& SharedComponent : SharedComponents)
+			{
+				if(SharedComponent.bIsProcedural)
+				{
+					continue;
+				}
+				if(const FRigBaseComponent* Component = Element->FindComponent(SharedComponent.Handle.GetComponentName()))
+				{
+					if(Component->IsProcedural())
+					{
+						SharedComponent.bIsProcedural = true;
+					}
+				}
+			}
+			*/
+		}
+	}
+
+	if(SharedComponents.IsEmpty())
+	{
+		return;
+	}
+
+	for(FSharedComponent& SharedComponent : SharedComponents)
+	{
+		TSharedPtr<FRigComponentStructProvider> StructProvider = MakeShareable(new FRigComponentStructProvider(Hierarchy));
+		for(int32 ElementIndex = 0; ElementIndex < PerElementInfos.Num(); ElementIndex++)
+		{
+			if(const FRigBaseElement* Element = PerElementInfos[ElementIndex].GetElement<>())
+			{
+				if(const FRigBaseComponent* Component = Element->FindComponent(SharedComponent.Handle.GetComponentName()))
+				{
+					StructProvider->AddComponent(Component);
+				}
+			}
+		}
+
+		if(StructProvider->Num() == 0)
+		{
+			continue;
+		}
+
+		const FName ComponentName = SharedComponent.Handle.GetComponentName();
+		const FText DisplayName = FText::Format(LOCTEXT("ComponentCategoryTitleFormat", "{0} Component"), FText::FromName(ComponentName));
+		IDetailCategoryBuilder& Category = DetailBuilder.EditCategory(ComponentName, DisplayName);
+
+		TArray<IDetailPropertyRow*> DetailPropertyRows;
+		Category.AddAllExternalStructureProperties(StructProvider, EPropertyLocation::Default, &DetailPropertyRows);
+
+		for(IDetailPropertyRow* DetailPropertyRow : DetailPropertyRows)
+		{
+			if(SharedComponent.bIsProcedural)
+			{
+				DetailPropertyRow->IsEnabled(false);
+			}
+			DetailPropertyRow->GetPropertyHandle()->SetOnPropertyValueChanged(FSimpleDelegate::CreateLambda([this, ComponentName]()
+			{
+				for(int32 ElementIndex = 0; ElementIndex < PerElementInfos.Num(); ElementIndex++)
+				{
+					if(URigHierarchy* Hierarchy = PerElementInfos[ElementIndex].GetHierarchy())
+					{
+						if(const FRigBaseElement* Element = PerElementInfos[ElementIndex].GetElement<>())
+						{
+							if(const FRigBaseComponent* Component = Element->FindComponent(ComponentName))
+							{
+								Hierarchy->Notify(ERigHierarchyNotification::ComponentContentChanged, Component);
+							}
+						}
+					}
+				}
+			}));
+		}
+	}
+}
+
 void FRigBaseElementDetails::CustomizeMetadata(IDetailLayoutBuilder& DetailBuilder)
 {
 	if(PerElementInfos.Num() != 1)
@@ -2801,6 +3008,7 @@ void FRigBoneElementDetails::CustomizeDetails(IDetailLayoutBuilder& DetailBuilde
 {
 	FRigTransformElementDetails::CustomizeDetails(DetailBuilder);
 	CustomizeTransform(DetailBuilder);
+	CustomizeComponents(DetailBuilder);
 	CustomizeMetadata(DetailBuilder);
 }
 
@@ -2814,6 +3022,7 @@ void FRigControlElementDetails::CustomizeDetails(IDetailLayoutBuilder& DetailBui
 	CustomizeShape(DetailBuilder);
 	CustomizeAvailableSpaces(DetailBuilder);
 	CustomizeAnimationChannels(DetailBuilder);
+	CustomizeComponents(DetailBuilder);
 	CustomizeMetadata(DetailBuilder);
 }
 
@@ -5654,6 +5863,7 @@ void FRigNullElementDetails::CustomizeDetails(IDetailLayoutBuilder& DetailBuilde
 {
 	FRigTransformElementDetails::CustomizeDetails(DetailBuilder);
 	CustomizeTransform(DetailBuilder);
+	CustomizeComponents(DetailBuilder);
 	CustomizeMetadata(DetailBuilder);
 }
 
@@ -5662,6 +5872,7 @@ void FRigConnectorElementDetails::CustomizeDetails(IDetailLayoutBuilder& DetailB
 	FRigTransformElementDetails::CustomizeDetails(DetailBuilder);
 	CustomizeSettings(DetailBuilder);
 	//CustomizeTransform(DetailBuilder);
+	CustomizeComponents(DetailBuilder);
 	//CustomizeMetadata(DetailBuilder);
 }
 
@@ -5730,6 +5941,7 @@ void FRigSocketElementDetails::CustomizeDetails(IDetailLayoutBuilder& DetailBuil
 	FRigTransformElementDetails::CustomizeDetails(DetailBuilder);
 	CustomizeSettings(DetailBuilder);
 	CustomizeTransform(DetailBuilder);
+	CustomizeComponents(DetailBuilder);
 	CustomizeMetadata(DetailBuilder);
 }
 
@@ -6052,6 +6264,7 @@ void FRigPhysicsElementDetails::CustomizeDetails(IDetailLayoutBuilder& DetailBui
 	FRigTransformElementDetails::CustomizeDetails(DetailBuilder);
 	CustomizeSettings(DetailBuilder);
 	CustomizeTransform(DetailBuilder);
+	CustomizeComponents(DetailBuilder);
 	CustomizeMetadata(DetailBuilder);
 }
 
diff --git a/Engine/Plugins/Animation/ControlRig/Source/ControlRigEditor/Private/ControlRigElementDetails.h b/Engine/Plugins/Animation/ControlRig/Source/ControlRigEditor/Private/ControlRigElementDetails.h
--- a/Engine/Plugins/Animation/ControlRig/Source/ControlRigEditor/Private/ControlRigElementDetails.h
+++ b/Engine/Plugins/Animation/ControlRig/Source/ControlRigEditor/Private/ControlRigElementDetails.h
@@ -308,12 +308,21 @@ public:
 
 protected:
 
+	void CustomizeComponents(IDetailLayoutBuilder& DetailBuilder);
 	void CustomizeMetadata(IDetailLayoutBuilder& DetailBuilder);
 
 	TArray<FPerElementInfo> PerElementInfos;
 	FDelegateHandle MetadataHandle;
 	TArray<FControlRigInteractionScope*> InteractionScopes;
-	
+
+	struct FSharedComponent
+	{
+		FRigComponentHandle Handle;
+		const UScriptStruct* ScriptStruct = nullptr;
+		bool bIsProcedural = false;
+	};
+	TArray<FSharedComponent> SharedComponents;
+
 	TSharedPtr<SButton> SelectParentElementButton;
 };
 
diff --git a/Engine/Plugins/Animation/ControlRig/Source/ControlRigEditor/Private/EditMode/ControlRigEditMode.cpp b/Engine/Plugins/Animation/ControlRig/Source/ControlRigEditor/Private/EditMode/ControlRigEditMode.cpp
--- a/Engine/Plugins/Animation/ControlRig/Source/ControlRigEditor/Private/EditMode/ControlRigEditMode.cpp
+++ b/Engine/Plugins/Animation/ControlRig/Source/ControlRigEditor/Private/EditMode/ControlRigEditMode.cpp
@@ -4639,7 +4639,7 @@ void FControlRigEditMode::OnHierarchyModified_AnyThread(ERigHierarchyNotificatio
 			{
 				OnHierarchyModified(InNotif, WeakHierarchy.Get(), Element);
 			}
-			else if(const FRigBaseComponent* Component = Element->GetComponentByName(ComponentName))
+			else if(const FRigBaseComponent* Component = Element->FindComponent(ComponentName))
 			{
 				OnHierarchyModified(InNotif, WeakHierarchy.Get(), Component);
 			}
diff --git a/Engine/Plugins/Animation/ControlRig/Source/ControlRigEditor/Private/Editor/ControlRigEditor.cpp b/Engine/Plugins/Animation/ControlRig/Source/ControlRigEditor/Private/Editor/ControlRigEditor.cpp
--- a/Engine/Plugins/Animation/ControlRig/Source/ControlRigEditor/Private/Editor/ControlRigEditor.cpp
+++ b/Engine/Plugins/Animation/ControlRig/Source/ControlRigEditor/Private/Editor/ControlRigEditor.cpp
@@ -4065,6 +4065,12 @@ void FControlRigEditor::OnHierarchyModified(ERigHierarchyNotification InNotif, U
 
 			break;
 		}
+		case ERigHierarchyNotification::ComponentAdded:
+		case ERigHierarchyNotification::ComponentRemoved:
+		{
+			OnHierarchyChanged();
+			break;
+		}
 		default:
 		{
 			break;
@@ -4119,7 +4125,7 @@ void FControlRigEditor::OnHierarchyModified_AnyThread(ERigHierarchyNotification
         FRigBaseElement* Element = WeakHierarchy.Get()->Find(Key);
 		if(Element)
 		{
-			Component = Element->GetComponentByName(ComponentName);
+			Component = Element->FindComponent(ComponentName);
 		}
 
 		switch(InNotif)
