Control Rig: Speed up hierarchy component lookup with dedicated map

#rnx

diff --git a/Engine/Plugins/Animation/ControlRig/Source/ControlRig/Private/Rigs/RigHierarchy.cpp b/Engine/Plugins/Animation/ControlRig/Source/ControlRig/Private/Rigs/RigHierarchy.cpp
--- a/Engine/Plugins/Animation/ControlRig/Source/ControlRig/Private/Rigs/RigHierarchy.cpp
+++ b/Engine/Plugins/Animation/ControlRig/Source/ControlRig/Private/Rigs/RigHierarchy.cpp
@@ -107,7 +107,7 @@ URigHierarchy::URigHierarchy()
 , bEnableDirtyPropagation(true)
 , Elements()
 , bRecordCurveChanges(true)
-, IndexLookup()
+, ElementIndexLookup()
 , TransformStackIndex(0)
 , bTransactingForTransformChange(false)
 , bIsInteracting(false)
@@ -429,7 +429,7 @@ void URigHierarchy::Load(FArchive& Ar)
 		Element->SubIndex = Num(Key.Type);
 		Element->Index = Elements.Add(Element);
 		ElementsPerType[RigElementTypeToFlatIndex(Key.Type)].Add(Element);
-		IndexLookup.Add(Key, Element->Index);
+		ElementIndexLookup.Add(Key, Element->Index);
 
 		if(bAllocateStoragePerElement)
 		{
@@ -572,11 +572,13 @@ void URigHierarchy::Load(FArchive& Ar)
 	if (ArchiveForElements->CustomVer(FControlRigObjectVersion::GUID) >= FControlRigObjectVersion::RigHierarchyStoresComponents)
 	{
 		ElementComponents.Reset();
+		ComponentIndexLookup.Reset();
 
 		int32 NumComponents = 0;
 		*ArchiveForElements << NumComponents;
 
 		ElementComponents.AddZeroed(NumComponents);
+		ComponentIndexLookup.Reserve(NumComponents);
 		for (int32 ComponentIndex = 0; ComponentIndex < ElementComponents.Num(); ComponentIndex++)
 		{
 			FInstancedStruct& InstancedStruct = ElementComponents[ComponentIndex];
@@ -588,6 +590,7 @@ void URigHierarchy::Load(FArchive& Ar)
 			Component->IndexInHierarchy = ComponentIndex;
 			Component->Element = Element;
 			Component->IndexInElement = Element->ComponentIndices.Add(Component->IndexInHierarchy);
+			ComponentIndexLookup.Add(Component->GetKey(), ComponentIndex);
 		}
 
 		IncrementTopologyVersion();
@@ -648,7 +651,8 @@ void URigHierarchy::Reset_Impl(bool bResetElements)
 			}
 			ElementsPerType.Add(TArray<FRigBaseElement*>());
 		}
-		IndexLookup.Reset();
+		ElementIndexLookup.Reset();
+		ComponentIndexLookup.Reset();
 
 		ElementTransforms.Reset();
 		ElementDirtyStates.Reset();
@@ -813,7 +817,7 @@ void URigHierarchy::CopyHierarchy(URigHierarchy* InHierarchy)
 			}
 
 			Elements.Reserve(InHierarchy->Elements.Num());
-			IndexLookup.Reserve(InHierarchy->IndexLookup.Num());
+			ElementIndexLookup.Reserve(InHierarchy->ElementIndexLookup.Num());
 
 			for(int32 Index = 0; Index < InHierarchy->Num(); Index++)
 			{
@@ -839,7 +843,7 @@ void URigHierarchy::CopyHierarchy(URigHierarchy* InHierarchy)
 				Target->ComponentIndices.Reset();
 
 				ElementsPerType[ElementTypeIndex].Add(Target);
-				IndexLookup.Add(Key, Target->Index);
+				ElementIndexLookup.Add(Key, Target->Index);
 			
 				IncrementPoseVersion(Index);
 
@@ -883,7 +887,7 @@ void URigHierarchy::CopyHierarchy(URigHierarchy* InHierarchy)
 				IncrementPoseVersion(Index);
 			}
 
-			IndexLookup = InHierarchy->IndexLookup;
+			ElementIndexLookup = InHierarchy->ElementIndexLookup;
 		}
 
 		ElementTransforms.Reset();
@@ -959,6 +963,9 @@ void URigHierarchy::CopyHierarchy(URigHierarchy* InHierarchy)
 		// Components
 		ElementComponents.Reset();
 		ElementComponents.Reserve(InHierarchy->ElementComponents.Num());
+		ComponentIndexLookup.Reset();
+		ComponentIndexLookup.Reserve(InHierarchy->ElementComponents.Num());
+
 		for(int32 SourceComponentIndex = 0; SourceComponentIndex < ElementComponents.Num(); SourceComponentIndex++)
 		{
 			if(!InHierarchy->ElementComponents[SourceComponentIndex].IsValid())
@@ -970,6 +977,7 @@ void URigHierarchy::CopyHierarchy(URigHierarchy* InHierarchy)
 			TargetComponent->Element = Find(TargetComponent->GetElementKey());
 			TargetComponent->IndexInHierarchy = TargetComponentIndex;
 			TargetComponent->IndexInElement = TargetComponent->Element->ComponentIndices.Add(TargetComponent->IndexInHierarchy);
+			ComponentIndexLookup.Add(TargetComponent->GetKey(), TargetComponentIndex);
 		}
 
 		// copy the topology version from the hierarchy.
@@ -1433,22 +1441,49 @@ bool URigHierarchy::IsProcedural(const FRigBaseElement* InElement) const
 
 const FRigBaseComponent* URigHierarchy::FindComponent(const FRigComponentKey& InKey) const
 {
-	if(const FRigBaseElement* Element = Find(InKey.ElementKey))
+	const int32 ComponentIndex = GetComponentIndex(InKey);
+	if(ComponentIndex != INDEX_NONE)
 	{
-		return Element->FindComponent(InKey.Name);
+		return GetComponent(ComponentIndex);
 	}
 	return nullptr;
 }
 
 FRigBaseComponent* URigHierarchy::FindComponent(const FRigComponentKey& InKey)
 {
-	if(FRigBaseElement* Element = Find(InKey.ElementKey))
+	const int32 ComponentIndex = GetComponentIndex(InKey);
+	if(ComponentIndex != INDEX_NONE)
 	{
-		return Element->FindComponent(InKey.Name);
+		return GetComponent(ComponentIndex);
 	}
 	return nullptr;
 }
 
+int32 URigHierarchy::GetComponentIndex(const FRigComponentKey& InComponentKey, bool bFollowRedirector) const
+{
+	if(bFollowRedirector && ElementKeyRedirector)
+	{
+		const int32 ElementIndex = GetIndex(InComponentKey.ElementKey);
+		if(ElementIndex != INDEX_NONE)
+		{
+			const FRigElementKey& RedirectedKey = Elements[ElementIndex]->GetKey(); 
+			if(RedirectedKey != InComponentKey.ElementKey)
+			{
+				const FRigComponentKey RedirectedComponentKey(RedirectedKey, InComponentKey.Name);
+				return GetComponentIndex(RedirectedComponentKey, false);
+			}
+		}
+	}
+	
+	if(const int32* IndexPtr = ComponentIndexLookup.Find(InComponentKey))
+	{
+		check(ElementComponents.IsValidIndex(*IndexPtr));
+		check(ElementComponents[*IndexPtr].IsValid());
+		return (*IndexPtr);
+	}
+	return INDEX_NONE;
+}
+
 const FRigBaseComponent* URigHierarchy::GetComponent(int32 InIndex) const
 {
 	if(ElementComponents.IsValidIndex(InIndex))
@@ -3806,12 +3841,9 @@ void URigHierarchy::Notify(ERigHierarchyNotification InNotifType, const FRigNoti
 				{
 					if (URigHierarchy* ListeningHierarchy = Listener.Hierarchy.Get())
 					{			
-						if(const FRigBaseElement* ListeningElement = ListeningHierarchy->Find( Element->GetKey()))
+						if(const FRigBaseComponent* ListeningComponent = ListeningHierarchy->FindComponent(InComponent->GetKey()))
 						{
-							if(const FRigBaseComponent* ListeningComponent = ListeningElement->FindComponent(InComponent->GetFName()))
-							{
-								ListeningHierarchy->Notify(InNotifType, ListeningComponent);
-							}
+							ListeningHierarchy->Notify(InNotifType, ListeningComponent);
 						}
 					}
 				});
@@ -3985,22 +4017,18 @@ void URigHierarchy::SendQueuedNotifications()
 	{
 		const FQueuedNotification& Entry = FilteredNotifications[Index];
 		FRigNotificationSubject Subject;
-		Subject.Element = Find(Entry.Key);
-		if(Subject.Element)
+		Subject.Component = FindComponent({Entry.Key, Entry.ComponentName});
+		if(Subject.Component == nullptr)
 		{
-			Subject.Component = Subject.Element->FindComponent(Entry.ComponentName);
-			if(Subject.Component)
-			{
-				Subject.Element = nullptr;
-			}
+			Subject.Element = Find(Entry.Key);
+		}
 
-			ModifiedEvent.Broadcast(Entry.Type, this, Subject);
+		ModifiedEvent.Broadcast(Entry.Type, this, Subject);
 
-			// the dynamic multicast for now only supports element notifs
-			if(Subject.Element && ModifiedEventDynamic.IsBound())
-			{
-				ModifiedEventDynamic.Broadcast(Entry.Type, this, Subject.Element->GetKey());
-			}
+		// the dynamic multicast for now only supports element notifs
+		if(Subject.Element && ModifiedEventDynamic.IsBound())
+		{
+			ModifiedEventDynamic.Broadcast(Entry.Type, this, Subject.Element->GetKey());
 		}
 	}
 
@@ -5973,8 +6001,15 @@ FRigBaseComponent* URigHierarchy::MakeComponent(const UScriptStruct* InComponent
 
 	FRigName UniqueName = SanitizedName;
 
+	TArray<FName> ExistingNames;
+	ExistingNames.Reserve(InElement->NumComponents());
+	for(int32 ExistingComponentIndex = 0; ExistingComponentIndex < InElement->NumComponents(); ExistingComponentIndex++)
+	{
+		ExistingNames.Add(InElement->GetComponent(ExistingComponentIndex)->GetFName());
+	}
+
 	int32 Suffix = 1;
-	while (InElement->FindComponent(UniqueName) != nullptr)
+	while (ExistingNames.Contains(UniqueName))
 	{
 		FString BaseString = SanitizedName.GetName();
 		if (BaseString.Len() > GetMaxNameLength() - 4)
@@ -5990,6 +6025,8 @@ FRigBaseComponent* URigHierarchy::MakeComponent(const UScriptStruct* InComponent
 	Component->IndexInHierarchy = ComponentIndex;
 	Component->IndexInElement = InElement->ComponentIndices.Add(Component->IndexInHierarchy);
 
+	ComponentIndexLookup.Add(Component->GetKey(), Component->GetIndexInHierarchy());
+	
 	return Component;
 }
 
@@ -6013,6 +6050,7 @@ void URigHierarchy::DestroyComponent(FRigBaseComponent*& InComponent)
 					}
 				}
 			}
+			ComponentIndexLookup.Remove(InComponent->GetKey());
 			InstancedStruct.Reset();
 			InComponent = nullptr;
 			break;
@@ -6024,6 +6062,10 @@ void URigHierarchy::DestroyComponents(FRigBaseElement* InElement)
 {
 	for(const int32& ComponentIndex : InElement->ComponentIndices)
 	{
+		if(const FRigBaseComponent* Component = GetComponent(ComponentIndex))
+		{
+			ComponentIndexLookup.Remove(Component->GetKey());
+		}
 		ElementComponents[ComponentIndex].Reset();
 	}
 	InElement->ComponentIndices.Reset();
@@ -6050,12 +6092,16 @@ void URigHierarchy::ShrinkComponentStorage()
 		return !InstancedStruct.IsValid();
 	});
 
+	ComponentIndexLookup.Reset();;
+
 	for(int32 ComponentIndex = 0; ComponentIndex < ElementComponents.Num(); ComponentIndex ++)
 	{
 		FInstancedStruct& InstancedStruct = ElementComponents[ComponentIndex];
 		FRigBaseComponent* Component = InstancedStruct.GetMutablePtr<FRigBaseComponent>();
+		Component->IndexInHierarchy = ComponentIndex;
 		Component->Element = Find(Component->GetElementKey());
 		Component->Element->ComponentIndices.Add(ComponentIndex);
+		ComponentIndexLookup.Add(Component->GetKey(), ComponentIndex);
 	}
 }
 
diff --git a/Engine/Plugins/Animation/ControlRig/Source/ControlRig/Private/Rigs/RigHierarchyCache.cpp b/Engine/Plugins/Animation/ControlRig/Source/ControlRig/Private/Rigs/RigHierarchyCache.cpp
--- a/Engine/Plugins/Animation/ControlRig/Source/ControlRig/Private/Rigs/RigHierarchyCache.cpp
+++ b/Engine/Plugins/Animation/ControlRig/Source/ControlRig/Private/Rigs/RigHierarchyCache.cpp
@@ -119,17 +119,10 @@ bool FCachedRigComponent::UpdateCache(const FRigElementKey& InKey, const FName&
 		{
 			if(CachedElement.UpdateCache(InKey, InHierarchy))
 			{
-				if(const FRigBaseElement* Element = CachedElement.GetElement())
+				Component = InHierarchy->FindComponent({InKey, InName});
+				if(Component)
 				{
-					Component = Element->FindComponent(InName);
-					if(Component)
-					{
-						Index = Component->GetIndexInHierarchy();
-					}
-					else
-					{
-						Invalidate();
-					}
+					Index = Component->GetIndexInHierarchy();
 				}
 				else
 				{
diff --git a/Engine/Plugins/Animation/ControlRig/Source/ControlRig/Private/Rigs/RigHierarchyController.cpp b/Engine/Plugins/Animation/ControlRig/Source/ControlRig/Private/Rigs/RigHierarchyController.cpp
--- a/Engine/Plugins/Animation/ControlRig/Source/ControlRig/Private/Rigs/RigHierarchyController.cpp
+++ b/Engine/Plugins/Animation/ControlRig/Source/ControlRig/Private/Rigs/RigHierarchyController.cpp
@@ -2395,7 +2395,7 @@ int32 URigHierarchyController::AddElement(FRigBaseElement* InElementToAdd, FRigB
 	InElementToAdd->SubIndex = Hierarchy->Num(InElementToAdd->Key.Type);
 	InElementToAdd->Index = Hierarchy->Elements.Add(InElementToAdd);
 	Hierarchy->ElementsPerType[URigHierarchy::RigElementTypeToFlatIndex(InElementToAdd->GetKey().Type)].Add(InElementToAdd);
-	Hierarchy->IndexLookup.Add(InElementToAdd->Key, InElementToAdd->Index);
+	Hierarchy->ElementIndexLookup.Add(InElementToAdd->Key, InElementToAdd->Index);
 	Hierarchy->AllocateDefaultElementStorage(InElementToAdd, true);
 	Hierarchy->IncrementTopologyVersion();
 
@@ -2564,9 +2564,9 @@ bool URigHierarchyController::RemoveElement(FRigBaseElement* InElement)
 	const int32 NumTypeElementsRemoved = Hierarchy->ElementsPerType[URigHierarchy::RigElementTypeToFlatIndex(InElement->GetKey().Type)].Remove(InElement);
 	ensure(NumTypeElementsRemoved == 1);
 
-	const int32 NumLookupsRemoved = Hierarchy->IndexLookup.Remove(InElement->Key);
+	const int32 NumLookupsRemoved = Hierarchy->ElementIndexLookup.Remove(InElement->Key);
 	ensure(NumLookupsRemoved == 1);
-	for(TPair<FRigElementKey, int32>& Pair : Hierarchy->IndexLookup)
+	for(TPair<FRigElementKey, int32>& Pair : Hierarchy->ElementIndexLookup)
 	{
 		if(Pair.Value > InElement->Index)
 		{
@@ -2807,6 +2807,7 @@ bool URigHierarchyController::RenameElement(FRigBaseElement* InElement, const FN
 	}
 
 	const FRigElementKey OldKey = InElement->GetKey();
+	TArray<FRigComponentKey> OldComponentKeys = InElement->GetComponentKeys();
 
 	URigHierarchy* Hierarchy = GetHierarchy();
 
@@ -2820,18 +2821,32 @@ bool URigHierarchyController::RenameElement(FRigBaseElement* InElement, const FN
 
 	{
 		// create a temp copy of the map and remove the current item's key
-		TMap<FRigElementKey, int32> TemporaryMap = Hierarchy->IndexLookup;
+		TMap<FRigElementKey, int32> TemporaryMap = Hierarchy->ElementIndexLookup;
 		TemporaryMap.Remove(OldKey);
    
-		TGuardValue<TMap<FRigElementKey, int32>> MapGuard(Hierarchy->IndexLookup, TemporaryMap);
+		TGuardValue<TMap<FRigElementKey, int32>> MapGuard(Hierarchy->ElementIndexLookup, TemporaryMap);
 		InElement->Key.Name = GetSafeNewName(InName, InElement->GetType());
 		InElement->CachedNameString.Reset();
 	}
 	
 	const FRigElementKey NewKey = InElement->GetKey();
 
-	Hierarchy->IndexLookup.Remove(OldKey);
-	Hierarchy->IndexLookup.Add(NewKey, InElement->Index);
+	Hierarchy->ElementIndexLookup.Remove(OldKey);
+	Hierarchy->ElementIndexLookup.Add(NewKey, InElement->Index);
+
+	for(const FRigComponentKey& ComponentKey : OldComponentKeys)
+	{
+		Hierarchy->ComponentIndexLookup.Remove(ComponentKey);
+	}
+	for(int32 ComponentIndex = 0; ComponentIndex < InElement->NumComponents(); ComponentIndex++)
+	{
+		FRigBaseComponent* Component = InElement->GetComponent(ComponentIndex);
+		if(ensure(Component))
+		{
+			Component->Key.ElementKey = NewKey;
+			Hierarchy->ComponentIndexLookup.Add(Component->GetKey(), Component->GetIndexInHierarchy());
+		}
+	}
 
 	// update all multi parent elements' index lookups
 	for (FRigBaseElement* Element : Hierarchy->Elements)
@@ -2935,7 +2950,7 @@ bool URigHierarchyController::ReorderElement(FRigBaseElement* InElement, int32 I
 		FRigBaseElement* Element = LocalElements[LocalIndex];
 		Hierarchy->Elements[GlobalIndex] = Element;
 		Element->Index = GlobalIndex;
-		Hierarchy->IndexLookup.FindOrAdd(Element->Key) = GlobalIndex;
+		Hierarchy->ElementIndexLookup.FindOrAdd(Element->Key) = GlobalIndex;
 	}
 
 	Notify(ERigHierarchyNotification::ElementReordered, InElement);
diff --git a/Engine/Plugins/Animation/ControlRig/Source/ControlRig/Private/Rigs/RigHierarchyElements.cpp b/Engine/Plugins/Animation/ControlRig/Source/ControlRig/Private/Rigs/RigHierarchyElements.cpp
--- a/Engine/Plugins/Animation/ControlRig/Source/ControlRig/Private/Rigs/RigHierarchyElements.cpp
+++ b/Engine/Plugins/Animation/ControlRig/Source/ControlRig/Private/Rigs/RigHierarchyElements.cpp
@@ -233,6 +233,21 @@ FRigBaseComponent* FRigBaseElement::GetFirstComponent(const UScriptStruct* InCom
 	return nullptr;
 }
 
+TArray<FRigComponentKey> FRigBaseElement::GetComponentKeys() const
+{
+	TArray<FRigComponentKey> Keys;
+	Keys.Reserve(NumComponents());
+	for(int32 ComponentIndex = 0; ComponentIndex < NumComponents(); ComponentIndex++)
+	{
+		const FRigBaseComponent* Component = GetComponent(ComponentIndex);
+		if(ensure(Component))
+		{
+			Keys.Add(Component->GetKey());
+		}
+	}
+	return Keys;
+}
+
 
 void FRigBaseElement::InitializeFrom(const FRigBaseElement* InOther)
 {
diff --git a/Engine/Plugins/Animation/ControlRig/Source/ControlRig/Public/Rigs/RigHierarchy.h b/Engine/Plugins/Animation/ControlRig/Source/ControlRig/Public/Rigs/RigHierarchy.h
--- a/Engine/Plugins/Animation/ControlRig/Source/ControlRig/Public/Rigs/RigHierarchy.h
+++ b/Engine/Plugins/Animation/ControlRig/Source/ControlRig/Public/Rigs/RigHierarchy.h
@@ -443,7 +443,7 @@ public:
 			}
 		}
 		
-		if(const int32* Index = IndexLookup.Find(InKey))
+		if(const int32* Index = ElementIndexLookup.Find(InKey))
 		{
 			return *Index;
 		}
@@ -719,6 +719,14 @@ public:
 	 */
 	FRigBaseComponent* FindComponent(const FRigComponentKey& InKey);
 
+	/**
+	 * Returns the index of a component or INDEX_NONE
+	 * @param InComponentKey The key of the component to look up
+	 * @param bFollowRedirector If True the element key will be resolved using the element redirector
+	 * @return The index of the component
+	 */
+	int32 GetComponentIndex(const FRigComponentKey& InComponentKey, bool bFollowRedirector = true) const;
+
 	/**
 	 * Returns a component at a given index or nullptr.
 	 * @param InIndex The index of the component to retrieve.
@@ -4475,8 +4483,11 @@ private:
 		}
 	};
 
-	// Managed lookup from Key to Index
-	TMap<FRigElementKey, int32> IndexLookup;
+	// Managed lookup from Element Key to Index
+	TMap<FRigElementKey, int32> ElementIndexLookup;
+
+	// Managed lookup from Component Key to Index
+	TMap<FRigComponentKey, int32> ComponentIndexLookup;
 
 	TMap<FRigElementKey, FString> UserDefinedElementName;
 
diff --git a/Engine/Plugins/Animation/ControlRig/Source/ControlRig/Public/Rigs/RigHierarchyElements.h b/Engine/Plugins/Animation/ControlRig/Source/ControlRig/Public/Rigs/RigHierarchyElements.h
--- a/Engine/Plugins/Animation/ControlRig/Source/ControlRig/Public/Rigs/RigHierarchyElements.h
+++ b/Engine/Plugins/Animation/ControlRig/Source/ControlRig/Public/Rigs/RigHierarchyElements.h
@@ -914,6 +914,7 @@ public:
 	FRigBaseComponent* FindComponent(const FName& InName);
 	const FRigBaseComponent* GetFirstComponent(const UScriptStruct* InComponentStruct) const;
 	FRigBaseComponent* GetFirstComponent(const UScriptStruct* InComponentStruct);
+	TArray<FRigComponentKey> GetComponentKeys() const;
 
 	template<typename T>
 	const T* GetFirstComponent() const
diff --git a/Engine/Plugins/Animation/ControlRig/Source/ControlRigEditor/Private/ControlRigElementDetails.cpp b/Engine/Plugins/Animation/ControlRig/Source/ControlRigEditor/Private/ControlRigElementDetails.cpp
--- a/Engine/Plugins/Animation/ControlRig/Source/ControlRigEditor/Private/ControlRigElementDetails.cpp
+++ b/Engine/Plugins/Animation/ControlRig/Source/ControlRigEditor/Private/ControlRigElementDetails.cpp
@@ -1414,8 +1414,9 @@ void FRigBaseElementDetails::CustomizeComponents(IDetailLayoutBuilder& DetailBui
 		if(const FRigBaseElement* Element = PerElementInfos[ElementIndex].GetElement<>())
 		{
 			// remove any missing or type-mismatching component from the list to display
-			SharedComponents.RemoveAll([Element](const FSharedComponent& SharedComponent) -> bool {
-				if(const FRigBaseComponent* Component = Element->FindComponent(SharedComponent.Handle.GetComponentName()))
+			SharedComponents.RemoveAll([Element, Hierarchy](const FSharedComponent& SharedComponent) -> bool {
+				const FRigComponentKey ComponentKey(Element->GetKey(), SharedComponent.Handle.GetComponentName());
+				if(const FRigBaseComponent* Component = Hierarchy->FindComponent(ComponentKey))
 				{
 					return Component->GetScriptStruct() != SharedComponent.ScriptStruct;
 				}
@@ -1453,12 +1454,10 @@ void FRigBaseElementDetails::CustomizeComponents(IDetailLayoutBuilder& DetailBui
 		TSharedPtr<FRigComponentStructProvider> StructProvider = MakeShareable(new FRigComponentStructProvider(Hierarchy));
 		for(int32 ElementIndex = 0; ElementIndex < PerElementInfos.Num(); ElementIndex++)
 		{
-			if(const FRigBaseElement* Element = PerElementInfos[ElementIndex].GetElement<>())
+			const FRigComponentKey ComponentKey(PerElementInfos[ElementIndex].Element.GetKey(), SharedComponent.Handle.GetComponentName());
+			if(const FRigBaseComponent* Component = Hierarchy->FindComponent(ComponentKey))
 			{
-				if(const FRigBaseComponent* Component = Element->FindComponent(SharedComponent.Handle.GetComponentName()))
-				{
-					StructProvider->AddComponent(Component);
-				}
+				StructProvider->AddComponent(Component);
 			}
 		}
 
@@ -1486,12 +1485,10 @@ void FRigBaseElementDetails::CustomizeComponents(IDetailLayoutBuilder& DetailBui
 				{
 					if(URigHierarchy* Hierarchy = PerElementInfos[ElementIndex].GetHierarchy())
 					{
-						if(const FRigBaseElement* Element = PerElementInfos[ElementIndex].GetElement<>())
+						const FRigComponentKey ComponentKey(PerElementInfos[ElementIndex].Element.GetKey(), ComponentName);
+						if(const FRigBaseComponent* Component = Hierarchy->FindComponent(ComponentKey))
 						{
-							if(const FRigBaseComponent* Component = Element->FindComponent(ComponentName))
-							{
-								Hierarchy->Notify(ERigHierarchyNotification::ComponentContentChanged, Component);
-							}
+							Hierarchy->Notify(ERigHierarchyNotification::ComponentContentChanged, Component);
 						}
 					}
 				}
diff --git a/Engine/Plugins/Animation/ControlRig/Source/ControlRigEditor/Private/EditMode/ControlRigEditMode.cpp b/Engine/Plugins/Animation/ControlRig/Source/ControlRigEditor/Private/EditMode/ControlRigEditMode.cpp
--- a/Engine/Plugins/Animation/ControlRig/Source/ControlRigEditor/Private/EditMode/ControlRigEditMode.cpp
+++ b/Engine/Plugins/Animation/ControlRig/Source/ControlRigEditor/Private/EditMode/ControlRigEditMode.cpp
@@ -4613,36 +4613,33 @@ void FControlRigEditMode::OnHierarchyModified_AnyThread(ERigHierarchyNotificatio
 		return;
 	}
 	
-	FRigElementKey Key;
-	FName ComponentName = NAME_None;
+	FRigElementKey ElementKey;
+	FRigComponentKey ComponentKey;
 	if(InSubject.Element)
 	{
-		Key = InSubject.Element->GetKey();
+		ElementKey = InSubject.Element->GetKey();
 	}
 	else if(InSubject.Component)
 	{
-		Key = InSubject.Component->GetElementKey();
-		ComponentName = InSubject.Component->GetFName();
+		ElementKey = InSubject.Component->GetElementKey();
+		ComponentKey = InSubject.Component->GetKey();
 	}
 
 	TWeakObjectPtr<URigHierarchy> WeakHierarchy = InHierarchy;
 	
-	FFunctionGraphTask::CreateAndDispatchWhenReady([this, InNotif, WeakHierarchy, Key, ComponentName]()
+	FFunctionGraphTask::CreateAndDispatchWhenReady([this, InNotif, WeakHierarchy, ElementKey, ComponentKey]()
 	{
 		if(!WeakHierarchy.IsValid())
 		{
 			return;
 		}
-		if (const FRigBaseElement* Element = WeakHierarchy.Get()->Find(Key))
+		if(const FRigBaseComponent* Component = WeakHierarchy.Get()->FindComponent(ComponentKey))
 		{
-			if(ComponentName.IsNone())
-			{
-				OnHierarchyModified(InNotif, WeakHierarchy.Get(), Element);
-			}
-			else if(const FRigBaseComponent* Component = Element->FindComponent(ComponentName))
-			{
-				OnHierarchyModified(InNotif, WeakHierarchy.Get(), Component);
-			}
+			OnHierarchyModified(InNotif, WeakHierarchy.Get(), Component);
+		}
+		else if (const FRigBaseElement* Element = WeakHierarchy.Get()->Find(ElementKey))
+		{
+			OnHierarchyModified(InNotif, WeakHierarchy.Get(), Element);
 		}
 		
 	}, TStatId(), NULL, ENamedThreads::GameThread);
