Control Rig: Fix CIS breakage in RigHierachy.cpp

#rnx

diff --git a/Engine/Plugins/Animation/ControlRig/Source/ControlRig/Private/Rigs/RigHierarchy.cpp b/Engine/Plugins/Animation/ControlRig/Source/ControlRig/Private/Rigs/RigHierarchy.cpp
--- a/Engine/Plugins/Animation/ControlRig/Source/ControlRig/Private/Rigs/RigHierarchy.cpp
+++ b/Engine/Plugins/Animation/ControlRig/Source/ControlRig/Private/Rigs/RigHierarchy.cpp
@@ -6009,7 +6009,10 @@ FRigBaseComponent* URigHierarchy::MakeComponent(const UScriptStruct* InComponent
 	}
 
 	int32 Suffix = 1;
-	while (ExistingNames.Contains(UniqueName))
+	while (ExistingNames.ContainsByPredicate([UniqueName](const FName& ExistingName) -> bool
+	{
+		return ExistingName == UniqueName.GetFName();
+	}))
 	{
 		FString BaseString = SanitizedName.GetName();
 		if (BaseString.Len() > GetMaxNameLength() - 4)
@@ -6020,7 +6023,7 @@ FRigBaseComponent* URigHierarchy::MakeComponent(const UScriptStruct* InComponent
 	}
 	
 	FRigBaseComponent* Component = ElementComponents[ComponentIndex].GetMutablePtr<FRigBaseComponent>();
-	Component->Key = FRigComponentKey(InElement->GetKey(), UniqueName);
+	Component->Key = FRigComponentKey(InElement->GetKey(), UniqueName.GetFName());
 	Component->Element = InElement;
 	Component->IndexInHierarchy = ComponentIndex;
 	Component->IndexInElement = InElement->ComponentIndices.Add(Component->IndexInHierarchy);
