Control Rig: Fix for ensures messing up Dev-Cine due to SkeletalMeshAsset not existing when trying to access it in SetBoneInitialTransformsFromSkeletalMeshComponent

diff --git a/Engine/Plugins/Animation/ControlRig/Source/ControlRig/Private/ControlRig.cpp b/Engine/Plugins/Animation/ControlRig/Source/ControlRig/Private/ControlRig.cpp
--- a/Engine/Plugins/Animation/ControlRig/Source/ControlRig/Private/ControlRig.cpp
+++ b/Engine/Plugins/Animation/ControlRig/Source/ControlRig/Private/ControlRig.cpp
@@ -3273,7 +3273,7 @@ void UControlRig::SetBoneInitialTransformsFromSkeletalMeshComponent(USkeletalMes
 	{
 		SetBoneInitialTransformsFromAnimInstance(InSkelMeshComp->GetAnimInstance());
 	}
-	else
+	else if(InSkelMeshComp->GetSkeletalMeshAsset())
 	{
 		SetBoneInitialTransformsFromSkeletalMesh(InSkelMeshComp->GetSkeletalMeshAsset());
 	}
