Control Rig and IK Rig import hierarchy dialog made non-modal.

diff --git a/Engine/Plugins/Animation/ControlRig/Source/ControlRigEditor/Private/Editor/SRigHierarchy.cpp b/Engine/Plugins/Animation/ControlRig/Source/ControlRigEditor/Private/Editor/SRigHierarchy.cpp
--- a/Engine/Plugins/Animation/ControlRig/Source/ControlRigEditor/Private/Editor/SRigHierarchy.cpp
+++ b/Engine/Plugins/Animation/ControlRig/Source/ControlRigEditor/Private/Editor/SRigHierarchy.cpp
@@ -532,30 +532,29 @@ EVisibility SRigHierarchy::IsSearchbarVisible() const
 
 FReply SRigHierarchy::OnImportSkeletonClicked()
 {
-	FRigHierarchyImportSettings Settings;
-	TSharedPtr<FStructOnScope> StructToDisplay = MakeShareable(new FStructOnScope(FRigHierarchyImportSettings::StaticStruct(), (uint8*)&Settings));
+	TSharedPtr<FStructOnScope> StructToDisplay = MakeShareable(new FStructOnScope(FRigHierarchyImportSettings::StaticStruct(), (uint8*)&ImportSettings));
 
 	TSharedRef<SKismetInspector> KismetInspector = SNew(SKismetInspector);
 	KismetInspector->ShowSingleStruct(StructToDisplay);
 
 	SGenericDialogWidget::FArguments DialogArguments;
-	DialogArguments.OnOkPressed_Lambda([&Settings, this] ()
+	DialogArguments.OnOkPressed_Lambda([this] ()
 	{
-		if (Settings.Mesh != nullptr)
+		if (ImportSettings.Mesh != nullptr)
 		{
 			if(ControlRigBlueprint->IsControlRigModule())
 			{
-				//ControlRigBlueprint->SetPreviewMesh(Settings.Mesh);
-				UpdateMesh(Settings.Mesh, true);
+				UpdateMesh(ImportSettings.Mesh, true);
 			}
 			else
 			{
-				ImportHierarchy(FAssetData(Settings.Mesh));
+				ImportHierarchy(FAssetData(ImportSettings.Mesh));
 			}
 		}
 	});
-	
-	SGenericDialogWidget::OpenDialog(LOCTEXT("ControlRigHierarchyImport", "Import Hierarchy"), KismetInspector, DialogArguments, true);
+
+	constexpr bool bAsModalDialog = false;
+	SGenericDialogWidget::OpenDialog(LOCTEXT("ControlRigHierarchyImport", "Import Hierarchy"), KismetInspector, DialogArguments, bAsModalDialog);
 
 	return FReply::Handled();
 }
diff --git a/Engine/Plugins/Animation/ControlRig/Source/ControlRigEditor/Private/Editor/SRigHierarchy.h b/Engine/Plugins/Animation/ControlRig/Source/ControlRigEditor/Private/Editor/SRigHierarchy.h
--- a/Engine/Plugins/Animation/ControlRig/Source/ControlRigEditor/Private/Editor/SRigHierarchy.h
+++ b/Engine/Plugins/Animation/ControlRig/Source/ControlRigEditor/Private/Editor/SRigHierarchy.h
@@ -146,7 +146,10 @@ private:
 	TWeakPtr<FControlRigEditor> ControlRigEditor;
 
 	FRigTreeDisplaySettings DisplaySettings;
-	const FRigTreeDisplaySettings& GetDisplaySettings() const { return DisplaySettings; } 
+	const FRigTreeDisplaySettings& GetDisplaySettings() const { return DisplaySettings; }
+
+	// struct for import hierarchy dialog (not on stack to allow non-modal dialog)
+	FRigHierarchyImportSettings ImportSettings;
 
 	/** Search box widget */
 	TSharedPtr<SSearchBox> FilterBox;
diff --git a/Engine/Plugins/Animation/IKRig/Source/IKRigEditor/Private/RigEditor/SIKRigHierarchy.cpp b/Engine/Plugins/Animation/IKRig/Source/IKRigEditor/Private/RigEditor/SIKRigHierarchy.cpp
--- a/Engine/Plugins/Animation/IKRig/Source/IKRigEditor/Private/RigEditor/SIKRigHierarchy.cpp
+++ b/Engine/Plugins/Animation/IKRig/Source/IKRigEditor/Private/RigEditor/SIKRigHierarchy.cpp
@@ -731,16 +731,15 @@ bool SIKRigHierarchy::IsSkeletonLoaded() const
 
 FReply SIKRigHierarchy::OnImportSkeletonClicked() const
 {
-	FIKRigHierarchyImportSettings Settings;
-	TSharedPtr<FStructOnScope> StructToDisplay = MakeShareable(new FStructOnScope(FIKRigHierarchyImportSettings::StaticStruct(), (uint8*)&Settings));
+	TSharedPtr<FStructOnScope> StructToDisplay = MakeShareable(new FStructOnScope(FIKRigHierarchyImportSettings::StaticStruct(), (uint8*)&ImportSettings));
 
 	TSharedRef<SKismetInspector> KismetInspector = SNew(SKismetInspector);
 	KismetInspector->ShowSingleStruct(StructToDisplay);
 
 	SGenericDialogWidget::FArguments DialogArguments;
-	DialogArguments.OnOkPressed_Lambda([&Settings, this] ()
+	DialogArguments.OnOkPressed_Lambda([this] ()
 	{
-		if (Settings.Mesh == nullptr)
+		if (ImportSettings.Mesh == nullptr)
 		{
 			return;
 		}
@@ -751,10 +750,11 @@ FReply SIKRigHierarchy::OnImportSkeletonClicked() const
 			return;
 		}
 
-		Controller->AssetController->SetSkeletalMesh(Settings.Mesh);
+		Controller->AssetController->SetSkeletalMesh(ImportSettings.Mesh);
 	});
 
-	SGenericDialogWidget::OpenDialog(LOCTEXT("IKRigHierarchyImport", "Import Hierarchy"), KismetInspector, DialogArguments, true);
+	constexpr bool bAsModalDialog = false;
+	SGenericDialogWidget::OpenDialog(LOCTEXT("IKRigHierarchyImport", "Import Hierarchy"), KismetInspector, DialogArguments, bAsModalDialog);
 
 	return FReply::Handled();
 }
diff --git a/Engine/Plugins/Animation/IKRig/Source/IKRigEditor/Public/RigEditor/SIKRigHierarchy.h b/Engine/Plugins/Animation/IKRig/Source/IKRigEditor/Public/RigEditor/SIKRigHierarchy.h
--- a/Engine/Plugins/Animation/IKRig/Source/IKRigEditor/Public/RigEditor/SIKRigHierarchy.h
+++ b/Engine/Plugins/Animation/IKRig/Source/IKRigEditor/Public/RigEditor/SIKRigHierarchy.h
@@ -170,12 +170,8 @@ struct FIKRigHierarchyImportSettings
 {
 	GENERATED_BODY()
 
-	FIKRigHierarchyImportSettings()
-	: Mesh(nullptr)
-	{}
-
 	UPROPERTY(EditAnywhere, Category = "Hierachy Import")
-	TObjectPtr<USkeletalMesh> Mesh;
+	TObjectPtr<USkeletalMesh> Mesh = nullptr;
 };
 
 typedef SBaseHierarchyTreeView<FIKRigTreeElement> SIKRigSkeletonTreeView;
@@ -316,6 +312,8 @@ private:
 	void FillContextMenu(FMenuBuilder& MenuBuilder);
 	/** END tree view callbacks */
 
+	FIKRigHierarchyImportSettings ImportSettings;
+
 	friend SIKRigHierarchyItem;
 	friend FIKRigEditorController;
 	friend SIKRigSolverStack;
