Control Rig: Detail customization + pin widget for rig component key

#rnx

diff --git a/Engine/Plugins/Animation/ControlRig/Source/ControlRig/Private/Rigs/RigHierarchy.cpp b/Engine/Plugins/Animation/ControlRig/Source/ControlRig/Private/Rigs/RigHierarchy.cpp
--- a/Engine/Plugins/Animation/ControlRig/Source/ControlRig/Private/Rigs/RigHierarchy.cpp
+++ b/Engine/Plugins/Animation/ControlRig/Source/ControlRig/Private/Rigs/RigHierarchy.cpp
@@ -1517,6 +1517,27 @@ int32 URigHierarchy::NumComponents(FRigElementKey InElement) const
 	return 0;
 }
 
+TArray<FRigComponentKey> URigHierarchy::GetComponentKeys(FRigElementKey InElement) const
+{
+	if(const FRigBaseElement* Element = Find(InElement))
+	{
+		return Element->GetComponentKeys();
+	}
+	return {};
+}
+
+FRigComponentKey URigHierarchy::GetComponentKey(FRigElementKey InElement, int32 InComponentIndex) const
+{
+	if(const FRigBaseElement* Element = Find(InElement))
+	{
+		if(const FRigBaseComponent* Component = Element->GetComponent(InComponentIndex))
+		{
+			return Component->GetKey();
+		}
+	}
+	return FRigComponentKey();
+}
+
 FName URigHierarchy::GetComponentName(FRigElementKey InElement, int32 InComponentIndex) const
 {
 	if(const FRigBaseElement* Element = Find(InElement))
@@ -2097,6 +2118,10 @@ FRigName URigHierarchy::GetSafeNewDisplayName(const FRigElementKey& InParentElem
 	{
 		return FRigName();
 	}
+	
+	FRigName SanitizedName = InPotentialNewDisplayName;
+	SanitizeName(SanitizedName);
+	FRigName Name = SanitizedName;
 
 	TArray<FRigElementKey> KeysToCheck;
 	if(InParentElement.IsValid())
@@ -2120,10 +2145,6 @@ FRigName URigHierarchy::GetSafeNewDisplayName(const FRigElementKey& InParentElem
 		}
 	}
 
-	FRigName SanitizedName = InPotentialNewDisplayName;
-	SanitizeName(SanitizedName);
-	FRigName Name = SanitizedName;
-
 	TArray<FString> DisplayNames;
 	Algo::Transform(KeysToCheck, DisplayNames, [this](const FRigElementKey& InKey) -> FString
 	{
diff --git a/Engine/Plugins/Animation/ControlRig/Source/ControlRig/Public/Rigs/RigHierarchy.h b/Engine/Plugins/Animation/ControlRig/Source/ControlRig/Public/Rigs/RigHierarchy.h
--- a/Engine/Plugins/Animation/ControlRig/Source/ControlRig/Public/Rigs/RigHierarchy.h
+++ b/Engine/Plugins/Animation/ControlRig/Source/ControlRig/Public/Rigs/RigHierarchy.h
@@ -796,6 +796,12 @@ public:
 	UFUNCTION(BlueprintPure, Category = URigHierarchy)
 	int32 NumComponents(FRigElementKey InElement) const;
 
+	UFUNCTION(BlueprintPure, Category = URigHierarchy)
+	TArray<FRigComponentKey> GetComponentKeys(FRigElementKey InElement) const;
+
+	UFUNCTION(BlueprintPure, Category = URigHierarchy)
+	FRigComponentKey GetComponentKey(FRigElementKey InElement, int32 InComponentIndex) const;
+
 	UFUNCTION(BlueprintPure, Category = URigHierarchy)
 	FName GetComponentName(FRigElementKey InElement, int32 InComponentIndex) const;
 	
diff --git a/Engine/Plugins/Animation/ControlRig/Source/ControlRigEditor/Private/ControlRigEditorModule.cpp b/Engine/Plugins/Animation/ControlRig/Source/ControlRigEditor/Private/ControlRigEditorModule.cpp
--- a/Engine/Plugins/Animation/ControlRig/Source/ControlRigEditor/Private/ControlRigEditorModule.cpp
+++ b/Engine/Plugins/Animation/ControlRig/Source/ControlRigEditor/Private/ControlRigEditorModule.cpp
@@ -261,6 +261,9 @@ void FControlRigEditorModule::StartupModule()
 	PropertiesToUnregisterOnShutdown.Add(FRigElementKey::StaticStruct()->GetFName());
 	PropertyEditorModule.RegisterCustomPropertyTypeLayout(PropertiesToUnregisterOnShutdown.Last(), FOnGetPropertyTypeCustomizationInstance::CreateStatic(&FRigElementKeyDetails::MakeInstance));
 
+	PropertiesToUnregisterOnShutdown.Add(FRigComponentKey::StaticStruct()->GetFName());
+	PropertyEditorModule.RegisterCustomPropertyTypeLayout(PropertiesToUnregisterOnShutdown.Last(), FOnGetPropertyTypeCustomizationInstance::CreateStatic(&FRigComponentKeyDetails::MakeInstance));
+
 	PropertiesToUnregisterOnShutdown.Add(FRigComputedTransform::StaticStruct()->GetFName());
 	PropertyEditorModule.RegisterCustomPropertyTypeLayout(PropertiesToUnregisterOnShutdown.Last(), FOnGetPropertyTypeCustomizationInstance::CreateStatic(&FRigComputedTransformDetails::MakeInstance));
 
diff --git a/Engine/Plugins/Animation/ControlRig/Source/ControlRigEditor/Private/ControlRigElementDetails.cpp b/Engine/Plugins/Animation/ControlRig/Source/ControlRigEditor/Private/ControlRigElementDetails.cpp
--- a/Engine/Plugins/Animation/ControlRig/Source/ControlRigEditor/Private/ControlRigElementDetails.cpp
+++ b/Engine/Plugins/Animation/ControlRig/Source/ControlRigEditor/Private/ControlRigElementDetails.cpp
@@ -129,6 +129,26 @@ void RigElementKeyDetails_GetCustomizedInfo(TSharedRef<IPropertyHandle> InStruct
 			}
 		}
 	}
+
+	UControlRigGraph* RigGraph = nullptr;
+	if(OutBlueprint)
+	{
+		for (UEdGraph* Graph : OutBlueprint->UbergraphPages)
+		{
+			RigGraph = Cast<UControlRigGraph>(Graph);
+			if (RigGraph)
+			{
+				break;
+			}
+		}
+	}
+
+	// only allow blueprints with at least one rig graph
+	if (RigGraph == nullptr)
+	{
+		OutBlueprint = nullptr;
+	}
+
 }
 
 UControlRigBlueprint* RigElementDetails_GetBlueprintFromHierarchy(URigHierarchy* InHierarchy)
@@ -317,25 +337,6 @@ void FRigElementKeyDetails::CustomizeHeader(TSharedRef<IPropertyHandle> InStruct
 	BlueprintBeingCustomized = nullptr;
 	RigElementKeyDetails_GetCustomizedInfo(InStructPropertyHandle, BlueprintBeingCustomized);
 
-	UControlRigGraph* RigGraph = nullptr;
-	if(BlueprintBeingCustomized)
-	{
-		for (UEdGraph* Graph : BlueprintBeingCustomized->UbergraphPages)
-		{
-			RigGraph = Cast<UControlRigGraph>(Graph);
-			if (RigGraph)
-			{
-				break;
-			}
-		}
-	}
-
-	// only allow blueprints with at least one rig graph
-	if (RigGraph == nullptr)
-	{
-		BlueprintBeingCustomized = nullptr;
-	}
-
 	if (BlueprintBeingCustomized == nullptr)
 	{
 		HeaderRow
@@ -550,6 +551,194 @@ FReply FRigElementKeyDetails::OnSelectInHierarchyClicked()
 	return FReply::Handled();
 }
 
+void FRigComponentKeyDetails::CustomizeHeader(TSharedRef<IPropertyHandle> InStructPropertyHandle, FDetailWidgetRow& HeaderRow,
+	IPropertyTypeCustomizationUtils& StructCustomizationUtils)
+{
+	BlueprintBeingCustomized = nullptr;
+	RigElementKeyDetails_GetCustomizedInfo(InStructPropertyHandle, BlueprintBeingCustomized);
+
+	HeaderRow
+	.NameContent()
+	[
+		InStructPropertyHandle->CreatePropertyNameWidget()
+	];
+
+	if (BlueprintBeingCustomized == nullptr)
+	{
+		HeaderRow
+		.ValueContent()
+		[
+			InStructPropertyHandle->CreatePropertyValueWidget(false)
+		];
+	}
+}
+
+void FRigComponentKeyDetails::CustomizeChildren(TSharedRef<IPropertyHandle> InStructPropertyHandle, IDetailChildrenBuilder& StructBuilder,
+	IPropertyTypeCustomizationUtils& StructCustomizationUtils)
+{
+	if (InStructPropertyHandle->IsValidHandle())
+	{
+		BlueprintBeingCustomized = nullptr;
+		RigElementKeyDetails_GetCustomizedInfo(InStructPropertyHandle, BlueprintBeingCustomized);
+
+		ElementKeyHandle = InStructPropertyHandle->GetChildHandle(TEXT("ElementKey"));
+		NameHandle = InStructPropertyHandle->GetChildHandle(TEXT("Name"));
+
+		// only fill the children if the blueprint cannot be found
+		if (BlueprintBeingCustomized == nullptr || !ElementKeyHandle.IsValid() || !NameHandle.IsValid())
+		{
+			uint32 NumChildren = 0;
+			InStructPropertyHandle->GetNumChildren(NumChildren);
+			for (uint32 ChildIndex = 0; ChildIndex < NumChildren; ChildIndex++)
+			{
+				StructBuilder.AddProperty(InStructPropertyHandle->GetChildHandle(ChildIndex).ToSharedRef());
+			}
+		}
+		else
+		{
+			ElementKeyHandle->SetOnChildPropertyValueChanged(FSimpleDelegate::CreateLambda([this]()
+			{
+				UpdateComponentNameList();
+			}));
+			
+			StructBuilder
+			.AddProperty(ElementKeyHandle.ToSharedRef());
+
+			StructBuilder
+			.AddProperty(NameHandle.ToSharedRef())
+			.CustomWidget()
+			.NameContent()
+			[
+				NameHandle->CreatePropertyNameWidget()
+			]
+			.ValueContent()
+			[
+				SNew(SHorizontalBox)
+				
+				+ SHorizontalBox::Slot()
+				.AutoWidth()
+				.Padding(4.f, 0.f, 0.f, 0.f)
+				[
+					SAssignNew(SearchableComboBox, SSearchableComboBox)
+					.OptionsSource(&ComponentNameList)
+					.OnSelectionChanged(this, &FRigComponentKeyDetails::OnComponentNameChanged)
+					.OnGenerateWidget_Lambda([](TSharedPtr<FString> InItem)
+					{
+						return SNew(STextBlock)
+						.Text(FText::FromString(InItem.IsValid() ? *InItem : FString()))
+						.Font(IDetailLayoutBuilder::GetDetailFont());
+					})
+					.Content()
+					[
+						SNew(STextBlock)
+						.Text(this, &FRigComponentKeyDetails::GetComponentNameAsText)
+						.Font(IDetailLayoutBuilder::GetDetailFont())
+					]
+				]
+			];
+		}
+	}
+
+	UpdateComponentNameList();
+}
+
+FRigElementKey FRigComponentKeyDetails::GetElementKey() const
+{
+	FRigElementKey ElementKey;
+	if (ElementKeyHandle.IsValid())
+	{
+		void* ElementKeyData = nullptr;
+		if(ElementKeyHandle->GetValueData(ElementKeyData) == FPropertyAccess::Success)
+		{
+			ElementKey = *static_cast<const FRigElementKey*>(ElementKeyData);
+		}
+	}
+	return ElementKey;
+}
+
+FString FRigComponentKeyDetails::GetComponentName() const
+{
+	FString ElementNameStr;
+	if (NameHandle.IsValid())
+	{
+		for(int32 ObjectIndex = 0; ObjectIndex < NameHandle->GetNumPerObjectValues(); ObjectIndex++)
+		{
+			FString PerObjectValue;
+			NameHandle->GetPerObjectValue(ObjectIndex, PerObjectValue);
+
+			if(ObjectIndex == 0)
+			{
+				ElementNameStr = PerObjectValue;
+			}
+			else if(ElementNameStr != PerObjectValue)
+			{
+				return ControlRigDetailsMultipleValues.ToString();
+			}
+		}
+	}
+	return ElementNameStr;
+}
+
+void FRigComponentKeyDetails::SetComponentName(FString InName)
+{
+	if (NameHandle.IsValid())
+	{
+		NameHandle->SetValue(InName);
+	}
+}
+
+void FRigComponentKeyDetails::OnComponentNameChanged(TSharedPtr<FString> InItem, ESelectInfo::Type InSelectionInfo)
+{
+	if (InItem.IsValid())
+	{
+		SetComponentName(*InItem);
+	}
+	else
+	{
+		SetComponentName(FString());
+	}
+}
+
+FText FRigComponentKeyDetails::GetComponentNameAsText() const
+{
+	return FText::FromString(GetComponentName());
+}
+
+void FRigComponentKeyDetails::UpdateComponentNameList()
+{
+	if(BlueprintBeingCustomized == nullptr)
+	{
+		return;
+	}
+	
+	const FRigElementKey ElementKey = GetElementKey();
+	if(!ElementKey.IsValid())
+	{
+		return;
+	}
+
+	URigHierarchy* Hierarchy = BlueprintBeingCustomized->Hierarchy;
+	if(UControlRig* RigBeingDebugged = Cast<UControlRig>(BlueprintBeingCustomized->GetObjectBeingDebugged()))
+	{
+		Hierarchy = RigBeingDebugged->GetHierarchy();
+	}
+
+	TArray<FRigComponentKey> ComponentKeys = Hierarchy->GetComponentKeys(ElementKey);
+	ComponentKeys.Sort();
+	
+	ComponentNameList.Reset();
+	ComponentNameList.Emplace(new FString(FName(NAME_None).ToString()));
+	for(const FRigComponentKey& ComponentKey : ComponentKeys)
+	{
+		ComponentNameList.Emplace(new FString(ComponentKey.Name.ToString()));
+	}
+
+	if(SearchableComboBox.IsValid())
+	{
+		SearchableComboBox->RefreshOptions();
+	}
+}
+
 void FRigComputedTransformDetails::CustomizeHeader(TSharedRef<IPropertyHandle> InStructPropertyHandle, FDetailWidgetRow& HeaderRow, IPropertyTypeCustomizationUtils& StructCustomizationUtils)
 {
 	BlueprintBeingCustomized = nullptr;
diff --git a/Engine/Plugins/Animation/ControlRig/Source/ControlRigEditor/Private/ControlRigElementDetails.h b/Engine/Plugins/Animation/ControlRig/Source/ControlRigEditor/Private/ControlRigElementDetails.h
--- a/Engine/Plugins/Animation/ControlRig/Source/ControlRigEditor/Private/ControlRigElementDetails.h
+++ b/Engine/Plugins/Animation/ControlRig/Source/ControlRigEditor/Private/ControlRigElementDetails.h
@@ -134,6 +134,37 @@ protected:
 	TSharedPtr<SRigElementKeyWidget> RigElementKeyWidget;
 };
 
+class FRigComponentKeyDetails : public IPropertyTypeCustomization
+{
+public:
+
+	static TSharedRef<IPropertyTypeCustomization> MakeInstance()
+	{
+		return MakeShareable(new FRigComponentKeyDetails);
+	}
+
+	/** IPropertyTypeCustomization interface */
+	virtual void CustomizeHeader(TSharedRef<class IPropertyHandle> InStructPropertyHandle, class FDetailWidgetRow& HeaderRow, IPropertyTypeCustomizationUtils& StructCustomizationUtils) override;
+	virtual void CustomizeChildren(TSharedRef<class IPropertyHandle> InStructPropertyHandle, class IDetailChildrenBuilder& StructBuilder, IPropertyTypeCustomizationUtils& StructCustomizationUtils) override;
+
+protected:
+
+	FRigElementKey GetElementKey() const;
+	FString GetComponentName() const;
+	void SetComponentName(FString InName);
+	void OnComponentNameChanged(TSharedPtr<FString> InItem, ESelectInfo::Type InSelectionInfo);
+	FText GetComponentNameAsText() const;
+	void UpdateComponentNameList();
+	
+	TSharedPtr<IPropertyHandle> ElementKeyHandle;
+	TSharedPtr<IPropertyHandle> NameHandle;
+	UControlRigBlueprint* BlueprintBeingCustomized;
+	
+	TArray<TSharedPtr<FString>> ComponentNameList;
+	TSharedPtr<SSearchableComboBox> SearchableComboBox;
+
+};
+
 UENUM()
 enum class ERigElementDetailsTransformComponent : uint8
 {
diff --git a/Engine/Plugins/Animation/ControlRig/Source/ControlRigEditor/Private/Graph/ControlRigGraphPanelPinFactory.cpp b/Engine/Plugins/Animation/ControlRig/Source/ControlRigEditor/Private/Graph/ControlRigGraphPanelPinFactory.cpp
--- a/Engine/Plugins/Animation/ControlRig/Source/ControlRigEditor/Private/Graph/ControlRigGraphPanelPinFactory.cpp
+++ b/Engine/Plugins/Animation/ControlRig/Source/ControlRigEditor/Private/Graph/ControlRigGraphPanelPinFactory.cpp
@@ -157,6 +157,13 @@ TSharedPtr<SGraphPin> FControlRigGraphPanelPinFactory::CreatePin_Internal(UEdGra
 
 			if (URigVMPin* ModelPin = RigNode->GetModelPinFromPinPath(InPin->GetName()))
 			{
+				if(ModelPin->GetDirection() != ERigVMPinDirection::Input && 
+					ModelPin->GetDirection() != ERigVMPinDirection::IO && 
+					ModelPin->GetDirection() != ERigVMPinDirection::Visible)
+				{
+					return nullptr;
+				}
+				
 				FName CustomWidgetName = ModelPin->GetCustomWidgetName();
 				if (CustomWidgetName == TEXT("BoneName"))
 				{
@@ -224,6 +231,100 @@ TSharedPtr<SGraphPin> FControlRigGraphPanelPinFactory::CreatePin_Internal(UEdGra
 						.ModelPin(ModelPin)
 						.OnGetNameListContent_UObject(RigGraph, &UControlRigGraph::GetShapeNameList);
 				}
+				else if (CustomWidgetName == TEXT("ComponentName"))
+				{
+					struct FCachedComponentNames
+					{
+						int32 TopologyVersion;
+						TSharedPtr<TArray<TSharedPtr<FRigVMStringWithTag>>> Names;
+						
+						FCachedComponentNames()
+						: TopologyVersion(INDEX_NONE)
+						{}
+					};
+					
+					return SNew(SRigVMGraphPinNameList, InPin)
+						.ModelPin(ModelPin)
+						.OnGetNameListContent_Lambda([RigGraph](const URigVMPin* InPin)
+						{
+							if (const UControlRigBlueprint* Blueprint = RigGraph->GetTypedOuter<UControlRigBlueprint>())
+							{
+								FRigElementKey ElementKey;
+
+								// find the pin that holds the element
+								if(const URigVMPin* ParentPin = InPin->GetParentPin())
+								{
+									if(const URigVMPin* ElementPin = ParentPin->FindSubPin(GET_MEMBER_NAME_STRING_CHECKED(FRigComponentKey, ElementKey)))
+									{
+										if(ElementPin->GetCPPType() == RigVMTypeUtils::GetUniqueStructTypeName(FRigElementKey::StaticStruct()))
+										{
+											const FString DefaultValue = ElementPin->GetDefaultValue();
+											if(!DefaultValue.IsEmpty())
+											{
+												FRigElementKey::StaticStruct()->ImportText(*DefaultValue, &ElementKey, nullptr, EPropertyPortFlags::PPF_None, nullptr, FRigElementKey::StaticStruct()->GetName(), true);
+											}
+										}
+									}
+								}
+
+								const URigHierarchy* Hierarchy = Blueprint->Hierarchy;
+								if (UControlRig* ControlRig = Cast<UControlRig>(Blueprint->GetObjectBeingDebugged()))
+								{
+									Hierarchy = ControlRig->GetHierarchy();
+								}
+								
+								if(Hierarchy->Find(ElementKey) == nullptr)
+								{
+									ElementKey.Reset();
+								}
+
+								const FString MapHash = ElementKey.IsValid() ?
+									Blueprint->GetPathName() + TEXT("|") + ElementKey.ToString() :
+									FName(NAME_None).ToString();
+
+								static TMap<FString, FCachedComponentNames> ComponentNameLists;
+								FCachedComponentNames& ComponentNames = ComponentNameLists.FindOrAdd(MapHash);
+
+								bool bRefreshList = !ComponentNames.Names.IsValid();
+
+								if(!bRefreshList)
+								{
+									const int32 TopologyVersion = Hierarchy->GetTopologyVersion();
+									if(ComponentNames.TopologyVersion != TopologyVersion)
+									{
+										bRefreshList = true;
+										ComponentNames.TopologyVersion = TopologyVersion;
+									}
+								}
+
+								if(bRefreshList)
+								{
+									if(!ComponentNames.Names.IsValid())
+									{
+										ComponentNames.Names = MakeShareable(new TArray<TSharedPtr<FRigVMStringWithTag>>());
+									}
+									ComponentNames.Names->Reset();
+									ComponentNames.Names->Add(MakeShareable(new FRigVMStringWithTag(FName(NAME_None).ToString())));
+
+									if(const FRigBaseElement* Element = Hierarchy->Find<FRigBaseElement>(ElementKey))
+									{
+										TArray<FRigComponentKey> ComponentKeys = Element->GetComponentKeys();
+										ComponentKeys.Sort();
+										for(const FRigComponentKey& ComponentKey : ComponentKeys)
+										{
+											ComponentNames.Names->Add(MakeShareable(new FRigVMStringWithTag(ComponentKey.Name.ToString())));
+										}
+									}
+								}
+								return ComponentNames.Names.Get();
+							}
+
+							static TArray<TSharedPtr<FRigVMStringWithTag>> EmptyNameList;
+							return &EmptyNameList;
+						})
+						.OnGetSelectedClicked_UObject(RigGraph, &UControlRigGraph::HandleGetSelectedClicked)
+						.OnBrowseClicked_UObject(RigGraph, &UControlRigGraph::HandleBrowseClicked);
+				}
 				else if (CustomWidgetName == TEXT("AnimationChannelName"))
 				{
 					struct FCachedAnimationChannelNames
diff --git a/Engine/Plugins/Runtime/RigVM/Source/RigVMDeveloper/Private/RigVMModel/RigVMPin.cpp b/Engine/Plugins/Runtime/RigVM/Source/RigVMDeveloper/Private/RigVMModel/RigVMPin.cpp
--- a/Engine/Plugins/Runtime/RigVM/Source/RigVMDeveloper/Private/RigVMModel/RigVMPin.cpp
+++ b/Engine/Plugins/Runtime/RigVM/Source/RigVMDeveloper/Private/RigVMModel/RigVMPin.cpp
@@ -1308,7 +1308,7 @@ FString URigVMPin::GetMetaData(FName InKey) const
 			}
 		}
 	}
-	else if(const URigVMTemplateNode* TemplateNode = Cast<URigVMTemplateNode>(GetNode()))
+	if(const URigVMTemplateNode* TemplateNode = Cast<URigVMTemplateNode>(GetNode()))
 	{
 		if(const FRigVMTemplate* Template = TemplateNode->GetTemplate())
 		{
