Control Rig: Get Component Dispatch node

#rnx

diff --git a/Engine/Plugins/Animation/ControlRig/Source/ControlRig/Private/Units/Hierarchy/RigUnit_Component.cpp b/Engine/Plugins/Animation/ControlRig/Source/ControlRig/Private/Units/Hierarchy/RigUnit_Component.cpp
new file mode 100644
--- /dev/null
+++ b/Engine/Plugins/Animation/ControlRig/Source/ControlRig/Private/Units/Hierarchy/RigUnit_Component.cpp
@@ -0,0 +1,301 @@
+// Copyright Epic Games, Inc. All Rights Reserved.
+
+#include "Units/Hierarchy/RigUnit_Component.h"
+
+#include "Animation/AttributeTypes.h"
+#include "Units/RigUnitContext.h"
+#include "Rigs/RigHierarchyController.h"
+#include "UObject/UObjectIterator.h"
+
+#include UE_INLINE_GENERATED_CPP_BY_NAME(RigUnit_Component)
+
+//////////////////////////////////////////////////////////////////////////////////////////////////////
+// FRigDispatch_ComponentBase
+//////////////////////////////////////////////////////////////////////////////////////////////////////
+
+bool FRigDispatch_ComponentBase::IsTypeSupported(const TRigVMTypeIndex& InTypeIndex)
+{
+	const FRigVMRegistry_NoLock& Registry = FRigVMRegistry_NoLock::GetForRead();
+	const FRigVMTemplateArgumentType& InType = Registry.GetType_NoLock(InTypeIndex);
+
+	if (IsValid(InType.CPPTypeObject))
+	{
+		if (const UScriptStruct* ScriptStruct = Cast<UScriptStruct>(InType.CPPTypeObject))
+		{
+			return ScriptStruct != FRigBaseComponent::StaticStruct() &&
+				ScriptStruct->IsChildOf(FRigBaseComponent::StaticStruct());
+		}
+	}
+	return false;
+}
+
+const TRigVMTypeIndex& FRigDispatch_ComponentBase::GetElementKeyType()
+{
+	static FRigVMRegistry_NoLock& Registry = FRigVMRegistry_NoLock::GetForWrite();
+	static TRigVMTypeIndex RigElementKeyType = Registry.FindOrAddType_NoLock(FRigVMTemplateArgumentType(FRigElementKey::StaticStruct()), false);
+	return RigElementKeyType;
+}
+
+const TRigVMTypeIndex& FRigDispatch_ComponentBase::GetComponentKeyType()
+{
+	static FRigVMRegistry_NoLock& Registry = FRigVMRegistry_NoLock::GetForWrite();
+	static TRigVMTypeIndex RigComponentKeyType = Registry.FindOrAddType_NoLock(FRigVMTemplateArgumentType(FRigComponentKey::StaticStruct()), false);
+	return RigComponentKeyType;
+}
+
+void FRigDispatch_ComponentBase::RegisterDependencyTypes_NoLock(FRigVMRegistry_NoLock& InRegistry) const
+{
+	Super::RegisterDependencyTypes_NoLock(InRegistry);
+
+	FRigVMRegistry_NoLock& Registry = FRigVMRegistry_NoLock::GetForWrite();
+	(void)Registry.FindOrAddType_NoLock(FRigVMTemplateArgumentType(FRigComponentKey::StaticStruct()), true);
+	
+	for (TObjectIterator<UScriptStruct> ScriptStructIt; ScriptStructIt; ++ScriptStructIt)
+	{
+		UScriptStruct* ScriptStruct = *ScriptStructIt;
+		if(ScriptStruct != FRigBaseComponent::StaticStruct() && ScriptStruct->IsChildOf(FRigBaseComponent::StaticStruct()))
+		{
+			(void)Registry.FindOrAddType_NoLock(FRigVMTemplateArgumentType(ScriptStruct), true);
+		}
+	}
+}
+
+const TArray<FRigVMExecuteArgument>& FRigDispatch_ComponentBase::GetExecuteArguments_Impl(const FRigVMDispatchContext& InContext) const
+{
+	if(IsMutable() && ExecuteInfos.IsEmpty())
+	{
+		ExecuteInfos.Emplace(ExecuteArgName, ERigVMPinDirection::IO);
+	}
+	return ExecuteInfos;
+}
+
+#if WITH_EDITOR
+
+FText FRigDispatch_ComponentBase::GetArgumentTooltip(const FName& InArgumentName, TRigVMTypeIndex InTypeIndex) const
+{
+	if(InArgumentName == NameArgName)
+	{
+		return NSLOCTEXT("FRigDispatch_ComponentBase", "NameArgTooltip", "The name of the component");
+	}
+	if(InArgumentName == ItemArgName)
+	{
+		return NSLOCTEXT("FRigDispatch_ComponentBase", "ItemArgTooltip", "The item for this component");
+	}
+	if(InArgumentName == KeyArgName)
+	{
+		return NSLOCTEXT("FRigDispatch_ComponentBase", "KeyArgTooltip", "The key of the component");
+	}
+	if(InArgumentName == ComponentArgName)
+	{
+		return NSLOCTEXT("FRigDispatch_ComponentBase", "ComponentArgTooltip", "The actual component");
+	}
+	if(InArgumentName == SuccessArgName)
+	{
+		return NSLOCTEXT("FRigDispatch_ComponentBase", "SuccessArgTooltip", "Returns true if the operation was successful.");
+	}
+	return FRigDispatchFactory::GetArgumentTooltip(InArgumentName, InTypeIndex);	
+}
+
+FString FRigDispatch_ComponentBase::GetArgumentDefaultValue(const FName& InArgumentName, TRigVMTypeIndex InTypeIndex) const
+{
+	if(InArgumentName == NameArgName)
+	{
+		return TEXT("Component");
+	}
+	if(InArgumentName == ItemArgName)
+	{
+		const FRigElementKey DefaultItem(NAME_None, ERigElementType::Bone);
+		FString DefaultValue;
+		FRigElementKey::StaticStruct()->ExportText(DefaultValue, &DefaultItem, &DefaultItem, nullptr, PPF_None, nullptr);
+		return DefaultValue;
+	}
+	return FRigDispatchFactory::GetArgumentDefaultValue(InArgumentName, InTypeIndex);
+}
+
+#endif
+
+//////////////////////////////////////////////////////////////////////////////////////////////////////
+// FRigDispatch_SpawnComponent
+//////////////////////////////////////////////////////////////////////////////////////////////////////
+
+const TArray<FRigVMTemplateArgumentInfo>& FRigDispatch_SpawnComponent::GetArgumentInfos() const
+{
+	if (Infos.IsEmpty())
+	{
+		ItemArgIndex = Infos.Emplace(ItemArgName, ERigVMPinDirection::Input, GetElementKeyType());
+		NameArgIndex = Infos.Emplace(NameArgName, ERigVMPinDirection::Input, RigVMTypeUtils::TypeIndex::FName);
+		ComponentArgIndex = Infos.Add(FRigVMTemplateArgumentInfo(ComponentArgName, ERigVMPinDirection::Input, {FRigVMTemplateArgument::ETypeCategory_SingleScriptStructValue}, &FRigDispatch_SpawnComponent::IsTypeSupported));
+		KeyArgIndex = Infos.Emplace(KeyArgName, ERigVMPinDirection::Output, GetComponentKeyType());
+		SuccessArgIndex = Infos.Emplace(SuccessArgName, ERigVMPinDirection::Output, RigVMTypeUtils::TypeIndex::Bool);
+	}
+	
+	return Infos;
+}
+
+
+void FRigDispatch_SpawnComponent::Execute(FRigVMExtendedExecuteContext& InContext, FRigVMMemoryHandleArray Handles, FRigVMPredicateBranchArray /* not needed */)
+{
+	const FRigDispatch_SpawnComponent* Factory = static_cast<const FRigDispatch_SpawnComponent*>(InContext.Factory);
+		
+	// unpack the memory
+	const FRigElementKey& ElementKey = *(const FRigElementKey*)Handles[Factory->ItemArgIndex].GetData();
+	const FName& Name = *(const FName*)Handles[Factory->NameArgIndex].GetData();
+	const uint8* SourceComponent = Handles[Factory->ComponentArgIndex].GetData();
+	FRigComponentKey& ComponentKey = *(FRigComponentKey*)Handles[Factory->KeyArgIndex].GetData();
+	bool& bSuccess = *(bool*)Handles[Factory->SuccessArgIndex].GetData();
+
+	const FStructProperty* StructProperty = CastField<FStructProperty>(Handles[Factory->ComponentArgIndex].GetResolvedProperty());
+	check(StructProperty);
+	UScriptStruct* ComponentStruct = StructProperty->Struct;
+	check(ComponentStruct && ComponentStruct->IsChildOf(FRigBaseComponent::StaticStruct()));
+
+	ComponentKey = FRigComponentKey();
+	bSuccess = false;
+
+	// extract the component from the hierarchy
+	const FControlRigExecuteContext& Context = InContext.GetPublicDataSafe<FControlRigExecuteContext>();
+	if(!Context.Hierarchy->Contains(ElementKey))
+	{
+#if WITH_EDITOR
+		Context.Report(FRigVMLogSettings(EMessageSeverity::Warning), Context.GetFunctionName(), Context.GetInstructionIndex(), FString::Printf(TEXT("Item %s not found."), *ElementKey.ToString()));
+#endif
+		return;
+	}
+
+	if(URigHierarchyController* Controller = Context.Hierarchy->GetController())
+	{
+		FRigHierarchyControllerInstructionBracket InstructionBracket(Controller, Context.GetInstructionIndex());
+		ComponentKey = Controller->AddComponent(ComponentStruct, Name, ElementKey);
+		if(ComponentKey.IsValid())
+		{
+			FRigBaseComponent* TargetComponent = Context.Hierarchy->FindComponent(ComponentKey);
+			check(TargetComponent);
+			
+			// copy the public content over (members which are not UProperty will be skipped) 
+			for (TFieldIterator<FProperty> It(ComponentStruct); It; ++It)
+			{
+				It->CopyCompleteValue_InContainer(TargetComponent, SourceComponent);
+			}
+
+			bSuccess = true;
+		}
+	}
+}
+
+//////////////////////////////////////////////////////////////////////////////////////////////////////
+// FRigDispatch_GetComponentContent
+//////////////////////////////////////////////////////////////////////////////////////////////////////
+
+const TArray<FRigVMTemplateArgumentInfo>& FRigDispatch_GetComponentContent::GetArgumentInfos() const
+{
+	if (Infos.IsEmpty())
+	{
+		KeyArgIndex = Infos.Emplace(KeyArgName, ERigVMPinDirection::Input, GetComponentKeyType());
+		ComponentArgIndex = Infos.Add(FRigVMTemplateArgumentInfo(ComponentArgName, ERigVMPinDirection::Output, {FRigVMTemplateArgument::ETypeCategory_SingleScriptStructValue}, &FRigDispatch_GetComponentContent::IsTypeSupported));
+		SuccessArgIndex = Infos.Emplace(SuccessArgName, ERigVMPinDirection::Output, RigVMTypeUtils::TypeIndex::Bool);
+	}
+	
+	return Infos;
+}
+
+void FRigDispatch_GetComponentContent::Execute(FRigVMExtendedExecuteContext& InContext, FRigVMMemoryHandleArray Handles, FRigVMPredicateBranchArray)
+{
+	const FRigDispatch_GetComponentContent* Factory = static_cast<const FRigDispatch_GetComponentContent*>(InContext.Factory);
+			
+	// unpack the memory
+	const FRigComponentKey& ComponentKey = *(const FRigComponentKey*)Handles[Factory->KeyArgIndex].GetData();
+	uint8* TargetComponent = Handles[Factory->ComponentArgIndex].GetData();
+	bool& bSuccess = *(bool*)Handles[Factory->SuccessArgIndex].GetData();
+
+	const FStructProperty* StructProperty = CastField<FStructProperty>(Handles[Factory->ComponentArgIndex].GetResolvedProperty());
+	check(StructProperty);
+	const UScriptStruct* ComponentStruct = StructProperty->Struct;
+	check(ComponentStruct && ComponentStruct->IsChildOf(FRigBaseComponent::StaticStruct()));
+
+	// extract the component from the hierarchy
+	const FControlRigExecuteContext& Context = InContext.GetPublicDataSafe<FControlRigExecuteContext>();
+	const FRigBaseComponent* SourceComponent = Context.Hierarchy->FindComponent(ComponentKey);
+	if(SourceComponent == nullptr)
+	{
+		ComponentStruct->InitializeDefaultValue(TargetComponent);
+#if WITH_EDITOR
+		Context.Report(FRigVMLogSettings(EMessageSeverity::Warning), Context.GetFunctionName(), Context.GetInstructionIndex(), FString::Printf(TEXT("%s not found."), *ComponentKey.ToString()));
+#endif
+		bSuccess = false;
+		return;
+	}
+
+	if(SourceComponent->GetScriptStruct() != ComponentStruct)
+	{
+		ComponentStruct->InitializeDefaultValue(TargetComponent);
+#if WITH_EDITOR
+		const FString Message = FString::Printf(TEXT("Component pin type (%s) doesn't match component in hierarchy (%s)."), *ComponentStruct->GetName(), *SourceComponent->GetScriptStruct()->GetName());
+		Context.Report(FRigVMLogSettings(EMessageSeverity::Warning), Context.GetFunctionName(), Context.GetInstructionIndex(), Message);
+#endif
+		bSuccess = false;
+		return;
+	}
+
+	ComponentStruct->CopyScriptStruct(TargetComponent, SourceComponent, 1);
+	bSuccess = true;
+}
+
+//////////////////////////////////////////////////////////////////////////////////////////////////////
+// FRigDispatch_SetComponentContent
+//////////////////////////////////////////////////////////////////////////////////////////////////////
+
+const TArray<FRigVMTemplateArgumentInfo>& FRigDispatch_SetComponentContent::GetArgumentInfos() const
+{
+	if (Infos.IsEmpty())
+	{
+		KeyArgIndex = Infos.Emplace(KeyArgName, ERigVMPinDirection::Input, GetComponentKeyType());
+		ComponentArgIndex = Infos.Add(FRigVMTemplateArgumentInfo(ComponentArgName, ERigVMPinDirection::Input, {FRigVMTemplateArgument::ETypeCategory_SingleScriptStructValue}, &FRigDispatch_SetComponentContent::IsTypeSupported));
+		SuccessArgIndex = Infos.Emplace(SuccessArgName, ERigVMPinDirection::Output, RigVMTypeUtils::TypeIndex::Bool);
+	}
+	
+	return Infos;
+}
+
+void FRigDispatch_SetComponentContent::Execute(FRigVMExtendedExecuteContext& InContext, FRigVMMemoryHandleArray Handles, FRigVMPredicateBranchArray)
+{
+	const FRigDispatch_SetComponentContent* Factory = static_cast<const FRigDispatch_SetComponentContent*>(InContext.Factory);
+			
+	// unpack the memory
+	const FRigComponentKey& ComponentKey = *(const FRigComponentKey*)Handles[Factory->KeyArgIndex].GetData();
+	const  uint8* SourceComponent = Handles[Factory->ComponentArgIndex].GetData();
+	bool& bSuccess = *(bool*)Handles[Factory->SuccessArgIndex].GetData();
+
+	const FStructProperty* StructProperty = CastField<FStructProperty>(Handles[Factory->ComponentArgIndex].GetResolvedProperty());
+	check(StructProperty);
+	const UScriptStruct* ComponentStruct = StructProperty->Struct;
+	check(ComponentStruct && ComponentStruct->IsChildOf(FRigBaseComponent::StaticStruct()));
+
+	// extract the component from the hierarchy
+	const FControlRigExecuteContext& Context = InContext.GetPublicDataSafe<FControlRigExecuteContext>();
+	FRigBaseComponent* TargetComponent = Context.Hierarchy->FindComponent(ComponentKey);
+	if(TargetComponent == nullptr)
+	{
+#if WITH_EDITOR
+		Context.Report(FRigVMLogSettings(EMessageSeverity::Warning), Context.GetFunctionName(), Context.GetInstructionIndex(), FString::Printf(TEXT("%s not found."), *ComponentKey.ToString()));
+#endif
+		bSuccess = false;
+		return;
+	}
+
+	if(TargetComponent->GetScriptStruct() != ComponentStruct)
+	{
+#if WITH_EDITOR
+		const FString Message = FString::Printf(TEXT("Component pin type (%s) doesn't match component in hierarchy (%s)."), *ComponentStruct->GetName(), *TargetComponent->GetScriptStruct()->GetName());
+		Context.Report(FRigVMLogSettings(EMessageSeverity::Warning), Context.GetFunctionName(), Context.GetInstructionIndex(), Message);
+#endif
+		bSuccess = false;
+		return;
+	}
+
+	// copy the public content over (members which are not UProperty will be skipped) 
+	for (TFieldIterator<FProperty> It(ComponentStruct); It; ++It)
+	{
+		It->CopyCompleteValue_InContainer(TargetComponent, SourceComponent);
+	}
+	bSuccess = true;
+}
\ No newline at end of file
diff --git a/Engine/Plugins/Animation/ControlRig/Source/ControlRig/Public/Rigs/RigHierarchyComponents.h b/Engine/Plugins/Animation/ControlRig/Source/ControlRig/Public/Rigs/RigHierarchyComponents.h
--- a/Engine/Plugins/Animation/ControlRig/Source/ControlRig/Public/Rigs/RigHierarchyComponents.h
+++ b/Engine/Plugins/Animation/ControlRig/Source/ControlRig/Public/Rigs/RigHierarchyComponents.h
@@ -70,19 +70,10 @@ protected:
 	// ReSharper disable once CppUE4ProbableMemoryIssuesWithUObject
 	FRigBaseElement* Element = nullptr;
 	
-	UPROPERTY(BlueprintReadOnly, Category = RigElement, meta = (AllowPrivateAccess = "true"))
 	FRigComponentKey Key;
-	
-	UPROPERTY(BlueprintReadOnly, Transient, Category = RigElement, meta = (AllowPrivateAccess = "true"))
 	int32 IndexInHierarchy = INDEX_NONE;
-
-	UPROPERTY(BlueprintReadOnly, Transient, Category = RigElement, meta = (AllowPrivateAccess = "true"))
 	int32 IndexInElement = INDEX_NONE;
-
-	UPROPERTY(BlueprintReadOnly, Transient, Category = RigElement, meta = (AllowPrivateAccess = "true"))
 	int32 CreatedAtInstructionIndex = INDEX_NONE;
-
-	UPROPERTY(BlueprintReadOnly, Transient, Category = RigElement, meta = (AllowPrivateAccess = "true"))
 	bool bSelected = false;
 
 	mutable FString CachedNameString;
diff --git a/Engine/Plugins/Animation/ControlRig/Source/ControlRig/Public/Units/Hierarchy/RigUnit_Component.h b/Engine/Plugins/Animation/ControlRig/Source/ControlRig/Public/Units/Hierarchy/RigUnit_Component.h
new file mode 100644
--- /dev/null
+++ b/Engine/Plugins/Animation/ControlRig/Source/ControlRig/Public/Units/Hierarchy/RigUnit_Component.h
@@ -0,0 +1,105 @@
+// Copyright Epic Games, Inc. All Rights Reserved.
+
+#pragma once
+#include "Units/RigUnit.h"
+#include "Units/RigDispatchFactory.h"
+#include "RigUnit_Component.generated.h"
+
+USTRUCT(meta=(Abstract, Category="Hierarchy", NodeColor="0.462745, 1,0, 0.329412", Varying))
+struct CONTROLRIG_API FRigDispatch_ComponentBase : public FRigDispatchFactory
+{
+	GENERATED_BODY()
+
+	virtual void RegisterDependencyTypes_NoLock(FRigVMRegistry_NoLock& InRegistry) const override;
+	virtual UScriptStruct* GetExecuteContextStruct() const override { return FControlRigExecuteContext::StaticStruct(); }
+	virtual const TArray<FRigVMExecuteArgument>& GetExecuteArguments_Impl(const FRigVMDispatchContext& InContext) const override;
+
+#if WITH_EDITOR
+	virtual FText GetArgumentTooltip(const FName& InArgumentName, TRigVMTypeIndex InTypeIndex) const override;
+	virtual FString GetArgumentDefaultValue(const FName& InArgumentName, TRigVMTypeIndex InTypeIndex) const override;
+#endif
+
+protected:
+
+	virtual bool IsMutable() const { return false; }
+	static bool IsTypeSupported(const TRigVMTypeIndex& InTypeIndex);
+	static const TRigVMTypeIndex& GetElementKeyType();
+	static const TRigVMTypeIndex& GetComponentKeyType();
+
+	mutable TArray<FRigVMTemplateArgumentInfo> Infos;
+	mutable TArray<FRigVMExecuteArgument> ExecuteInfos;
+
+	mutable int32 NameArgIndex = INDEX_NONE;
+	mutable int32 ItemArgIndex = INDEX_NONE;
+	mutable int32 KeyArgIndex = INDEX_NONE;
+	mutable int32 ComponentArgIndex = INDEX_NONE;
+	mutable int32 SuccessArgIndex = INDEX_NONE;
+
+	static inline const FLazyName ExecuteArgName = FLazyName(TEXT("Execute"));
+	static inline const FLazyName NameArgName = FLazyName(TEXT("Name"));
+	static inline const FLazyName ItemArgName = FLazyName(TEXT("Item"));
+	static inline const FLazyName KeyArgName = FLazyName(TEXT("Key"));
+	static inline const FLazyName ComponentArgName = FLazyName(TEXT("Component"));
+	static inline const FLazyName SuccessArgName = FLazyName(TEXT("Success"));
+};
+
+/**
+ * Break / Decompose node for a hierarchy component
+ */
+USTRUCT(meta=(DisplayName="Spawn Component"))
+struct CONTROLRIG_API FRigDispatch_SpawnComponent : public FRigDispatch_ComponentBase
+{
+	GENERATED_BODY()
+
+	virtual const TArray<FRigVMTemplateArgumentInfo>& GetArgumentInfos() const override;
+
+protected:
+
+	virtual bool IsMutable() const override { return true; }
+
+	virtual FRigVMFunctionPtr GetDispatchFunctionImpl(const FRigVMTemplateTypeMap& InTypes) const override
+	{
+		return &Execute;
+	}
+	static void Execute(FRigVMExtendedExecuteContext& InContext, FRigVMMemoryHandleArray Handles, FRigVMPredicateBranchArray /* not needed */);
+};
+
+/**
+ * Break / Decompose node for a hierarchy component
+ */
+USTRUCT(meta=(DisplayName="Get Component"))
+struct CONTROLRIG_API FRigDispatch_GetComponentContent : public FRigDispatch_ComponentBase
+{
+	GENERATED_BODY()
+
+	virtual const TArray<FRigVMTemplateArgumentInfo>& GetArgumentInfos() const override;
+
+protected:
+
+	virtual FRigVMFunctionPtr GetDispatchFunctionImpl(const FRigVMTemplateTypeMap& InTypes) const override
+	{
+		return &Execute;
+	}
+	static void Execute(FRigVMExtendedExecuteContext& InContext, FRigVMMemoryHandleArray Handles, FRigVMPredicateBranchArray /* not needed */);
+};
+
+/**
+ * Set the content of a component
+ */
+USTRUCT(meta=(DisplayName="Set Component"))
+struct CONTROLRIG_API FRigDispatch_SetComponentContent : public FRigDispatch_ComponentBase
+{
+	GENERATED_BODY()
+
+	virtual const TArray<FRigVMTemplateArgumentInfo>& GetArgumentInfos() const override;
+
+protected:
+
+	virtual bool IsMutable() const override { return true; }
+
+	virtual FRigVMFunctionPtr GetDispatchFunctionImpl(const FRigVMTemplateTypeMap& InTypes) const override
+	{
+		return &Execute;
+	}
+	static void Execute(FRigVMExtendedExecuteContext& InContext, FRigVMMemoryHandleArray Handles, FRigVMPredicateBranchArray /* not needed */);
+};
\ No newline at end of file
