Control Rig: Fix unit test failure in hierarchy component tests

#rnx

diff --git a/Engine/Plugins/Animation/ControlRig/Source/ControlRig/Private/Rigs/RigHierarchy.cpp b/Engine/Plugins/Animation/ControlRig/Source/ControlRig/Private/Rigs/RigHierarchy.cpp
--- a/Engine/Plugins/Animation/ControlRig/Source/ControlRig/Private/Rigs/RigHierarchy.cpp
+++ b/Engine/Plugins/Animation/ControlRig/Source/ControlRig/Private/Rigs/RigHierarchy.cpp
@@ -289,7 +289,10 @@ void URigHierarchy::Save(FArchive& Ar)
 	{
 		if(InstancedStruct.IsValid())
 		{
-			InstancedStruct.Serialize(*ArchiveForElements);
+			FRigBaseComponent* Component = InstancedStruct.GetMutablePtr<FRigBaseComponent>();
+			UScriptStruct* ScriptStruct = Component->GetScriptStruct();
+			*ArchiveForElements << ScriptStruct;
+			Component->Serialize(*ArchiveForElements);
 		}
 	}
 
@@ -581,12 +584,22 @@ void URigHierarchy::Load(FArchive& Ar)
 		ComponentIndexLookup.Reserve(NumComponents);
 		for (int32 ComponentIndex = 0; ComponentIndex < ElementComponents.Num(); ComponentIndex++)
 		{
+			UScriptStruct* ScriptStruct = nullptr;
+			*ArchiveForElements << ScriptStruct;
+			if (ScriptStruct)
+			{
+				ArchiveForElements->Preload(ScriptStruct);
+			}
+
 			FInstancedStruct& InstancedStruct = ElementComponents[ComponentIndex];
-			InstancedStruct.Serialize(*ArchiveForElements);
+			InstancedStruct.InitializeAs(ScriptStruct);
 			
 			FRigBaseComponent* Component = InstancedStruct.GetMutablePtr<FRigBaseComponent>();
+			Component->Serialize(*ArchiveForElements);
+
 			FRigBaseElement* Element = Find(Component->GetElementKey());
 			check(Element);
+			
 			Component->IndexInHierarchy = ComponentIndex;
 			Component->Element = Element;
 			Component->IndexInElement = Element->ComponentIndices.Add(Component->IndexInHierarchy);
