[Backout] - CL37742880
Original CL Desc
-----------------------------------------------------------------
Control Rig: Get Component Dispatch node

#rnx

diff --git a/Engine/Plugins/Animation/ControlRig/Source/ControlRig/Private/Units/Hierarchy/RigUnit_Component.cpp b/Engine/Plugins/Animation/ControlRig/Source/ControlRig/Private/Units/Hierarchy/RigUnit_Component.cpp
deleted file mode 100644
diff --git a/Engine/Plugins/Animation/ControlRig/Source/ControlRig/Public/Rigs/RigHierarchyComponents.h b/Engine/Plugins/Animation/ControlRig/Source/ControlRig/Public/Rigs/RigHierarchyComponents.h
--- a/Engine/Plugins/Animation/ControlRig/Source/ControlRig/Public/Rigs/RigHierarchyComponents.h
+++ b/Engine/Plugins/Animation/ControlRig/Source/ControlRig/Public/Rigs/RigHierarchyComponents.h
@@ -70,10 +70,19 @@ protected:
 	// ReSharper disable once CppUE4ProbableMemoryIssuesWithUObject
 	FRigBaseElement* Element = nullptr;
 	
+	UPROPERTY(BlueprintReadOnly, Category = RigElement, meta = (AllowPrivateAccess = "true"))
 	FRigComponentKey Key;
+	
+	UPROPERTY(BlueprintReadOnly, Transient, Category = RigElement, meta = (AllowPrivateAccess = "true"))
 	int32 IndexInHierarchy = INDEX_NONE;
+
+	UPROPERTY(BlueprintReadOnly, Transient, Category = RigElement, meta = (AllowPrivateAccess = "true"))
 	int32 IndexInElement = INDEX_NONE;
+
+	UPROPERTY(BlueprintReadOnly, Transient, Category = RigElement, meta = (AllowPrivateAccess = "true"))
 	int32 CreatedAtInstructionIndex = INDEX_NONE;
+
+	UPROPERTY(BlueprintReadOnly, Transient, Category = RigElement, meta = (AllowPrivateAccess = "true"))
 	bool bSelected = false;
 
 	mutable FString CachedNameString;
diff --git a/Engine/Plugins/Animation/ControlRig/Source/ControlRig/Public/Units/Hierarchy/RigUnit_Component.h b/Engine/Plugins/Animation/ControlRig/Source/ControlRig/Public/Units/Hierarchy/RigUnit_Component.h
deleted file mode 100644
