Control Rig: Fix typo in sequencer filters

#rnx

diff --git a/Engine/Plugins/Animation/ControlRig/Source/ControlRigEditor/Private/Sequencer/ControlRigParameterTrackEditor.cpp b/Engine/Plugins/Animation/ControlRig/Source/ControlRigEditor/Private/Sequencer/ControlRigParameterTrackEditor.cpp
--- a/Engine/Plugins/Animation/ControlRig/Source/ControlRigEditor/Private/Sequencer/ControlRigParameterTrackEditor.cpp
+++ b/Engine/Plugins/Animation/ControlRig/Source/ControlRigEditor/Private/Sequencer/ControlRigParameterTrackEditor.cpp
@@ -257,8 +257,8 @@ public:
 		Tag(InTag)
 	{}
 
-	FString GetName() const override { return Tag.bMarksSubjectAsInvalid ? TEXT("Exlude ") + Tag.Name.ToString() : Tag.Name.ToString(); }
-	FText GetDisplayName() const override { return FText::FromString(Tag.bMarksSubjectAsInvalid ? TEXT("Exlude ") + Tag.GetLabel() : Tag.GetLabel()); }
+	FString GetName() const override { return Tag.bMarksSubjectAsInvalid ? TEXT("Exclude ") + Tag.Name.ToString() : Tag.Name.ToString(); }
+	FText GetDisplayName() const override { return FText::FromString(Tag.bMarksSubjectAsInvalid ? TEXT("Exclude ") + Tag.GetLabel() : Tag.GetLabel()); }
 	FText GetToolTipText() const override { return Tag.ToolTip; }
 	FLinearColor GetColor() const override { return Tag.Color; };
 	
