Control Rig: Performance improvement for unique name & unique display name

#rnx

diff --git a/Engine/Plugins/Animation/ControlRig/Source/ControlRig/Private/Rigs/RigHierarchy.cpp b/Engine/Plugins/Animation/ControlRig/Source/ControlRig/Private/Rigs/RigHierarchy.cpp
--- a/Engine/Plugins/Animation/ControlRig/Source/ControlRig/Private/Rigs/RigHierarchy.cpp
+++ b/Engine/Plugins/Animation/ControlRig/Source/ControlRig/Private/Rigs/RigHierarchy.cpp
@@ -2108,21 +2108,18 @@ FRigName URigHierarchy::GetSafeNewName(const FRigName& InPotentialNewName, ERigE
 		}
 	}
 
-	FRigName Name = SanitizedName;
-
-	int32 Suffix = 1;
-	while (!IsNameAvailable(Name, InType))
+	if (SanitizedName.Len() > GetMaxNameLength() - 4)
 	{
-		FString BaseString = SanitizedName.GetName();
-		if (BaseString.Len() > GetMaxNameLength() - 4)
-		{
-			BaseString.LeftChopInline(BaseString.Len() - (GetMaxNameLength() - 4));
-		}
-		Name.SetName(FString::Printf(TEXT("%s_%d"), *BaseString, ++Suffix));
+		SanitizedName.SetName(SanitizedName.GetName().LeftChop(SanitizedName.Len() - (GetMaxNameLength() - 4)));
 	}
 
+	const FName UniqueName = FRigName(GetUniqueName(*SanitizedName.GetName(), [this, InType](const FName& InName) -> bool
+	{
+		return IsNameAvailable(InName, InType);
+	}));
+
 	bAllowNameSpaceWhenSanitizingName = false;
-	return Name;
+	return FRigName(UniqueName);
 }
 
 FRigName URigHierarchy::GetSafeNewDisplayName(const FRigElementKey& InParentElement, const FRigName& InPotentialNewDisplayName) const
@@ -2131,55 +2128,91 @@ FRigName URigHierarchy::GetSafeNewDisplayName(const FRigElementKey& InParentElem
 	{
 		return FRigName();
 	}
-	
+
+	static const int32 ControlElementIndex = RigElementTypeToFlatIndex(ERigElementType::Control); 
+
 	FRigName SanitizedName = InPotentialNewDisplayName;
 	SanitizeName(SanitizedName);
-	FRigName Name = SanitizedName;
+
+	// this can be nullptr - and that's ok.
+	const FRigBaseElement* ParentElement = Find(InParentElement);
 
 	TArray<FRigElementKey> KeysToCheck;
-	if(InParentElement.IsValid())
+
+	// two paths: if the children are up2date, rely on the previously cached children.
+	if(ChildElementCacheTopologyVersion == TopologyVersion)
 	{
-		KeysToCheck = GetChildren(InParentElement);
+		if(InParentElement.IsValid())
+		{
+			const TConstArrayView<FRigBaseElement*> Children = GetChildren(ParentElement);
+			for(const FRigBaseElement* ChildElement : Children)
+			{
+				if(ChildElement->IsA<FRigControlElement>())
+				{
+					KeysToCheck.Add(ChildElement->GetKey());
+				}
+			}
+		}
+		else
+		{
+			// get all of the root control elements
+			for(const FRigBaseElement* Element : ElementsPerType[ControlElementIndex])
+			{
+				if(GetNumberOfParents(Element) == 0)
+				{
+					KeysToCheck.Add(Element->GetKey());
+				}
+			}
+		}
 	}
 	else
 	{
-		// get all of the root elements
-		for(const FRigBaseElement* Element : Elements)
+		// if we enter here we need to find the children for us. first we'll visit all elements
+		// and will check for a name collision there. if there is a collision we can rely on the code
+		// path based on the KeysToCheck
+		bool bHasNameCollision = false;
+		KeysToCheck.Reserve(10);
+		
+		for(const FRigBaseElement* Element : ElementsPerType[ControlElementIndex])
 		{
-			if(!Element->IsA<FRigTransformElement>())
-			{
-				continue;
-			}
-			
-			if(GetNumberOfParents(Element) == 0)
+			if(GetFirstParent(Element) == ParentElement)
 			{
+				if(CastChecked<FRigControlElement>(Element)->Settings.DisplayName == SanitizedName.GetFName())
+				{
+					bHasNameCollision = true;
+				}
 				KeysToCheck.Add(Element->GetKey());
 			}
 		}
+
+		// early exit - we haven't hit any name collision anywhere
+		if(!bHasNameCollision)
+		{
+			return SanitizedName;
+		}
+	}
+
+	if (SanitizedName.Len() > GetMaxNameLength() - 4)
+	{
+		SanitizedName.SetName(SanitizedName.GetName().LeftChop(SanitizedName.Len() - (GetMaxNameLength() - 4)));
 	}
 
 	TArray<FString> DisplayNames;
 	Algo::Transform(KeysToCheck, DisplayNames, [this](const FRigElementKey& InKey) -> FString
 	{
-		if(const FRigBaseElement* BaseElement = Find(InKey))
+		if(const FRigControlElement* ControlElement = Find<FRigControlElement>(InKey))
 		{
-			return BaseElement->GetDisplayName().ToString();
+			return ControlElement->GetDisplayName().ToString();
 		}
 		return FString();
 	});
 
-	int32 Suffix = 1;
-	while (DisplayNames.Contains(Name.GetName()))
+	const FName UniqueName = FRigName(GetUniqueName(*SanitizedName.GetName(), [this, &DisplayNames](const FName& InName) -> bool
 	{
-		FString BaseString = SanitizedName.GetName();
-		if (BaseString.Len() > GetMaxNameLength() - 4)
-		{
-			BaseString.LeftChopInline(BaseString.Len() - (GetMaxNameLength() - 4));
-		}
-		Name.SetName(FString::Printf(TEXT("%s_%d"), *BaseString, ++Suffix));
-	}
+		return !DisplayNames.Contains(InName);
+	}));
 
-	return Name;
+	return UniqueName;
 }
 
 FText URigHierarchy::GetDisplayNameForUI(const FRigBaseElement* InElement, bool bIncludeNameSpace) const
@@ -7355,6 +7388,27 @@ void URigHierarchy::EnsureCacheValidityImpl()
 	});
 }
 
+FName URigHierarchy::GetUniqueName(const FName& InName, TFunction<bool(const FName&)> IsNameAvailablePredicate)
+{
+	static constexpr int32 InitialSuffix = 3; // offset since FName uses 0 as the indicator that there's no number suffix, plus we want to start with index 2. 
+	int32 NameSuffix = InitialSuffix; 
+	FName Name = InName;
+
+	while (!IsNameAvailablePredicate(Name))
+	{
+		if(NameSuffix == InitialSuffix)
+		{
+			Name = FName(Name, NameSuffix);
+		}
+		else
+		{
+			Name.SetNumber(NameSuffix);
+		}
+		NameSuffix++;
+	}
+	return Name;
+}
+
 #if WITH_EDITOR
 
 URigHierarchy::TElementDependencyMap URigHierarchy::GetDependenciesForVM(const URigVM* InVM, FName InEventName) const
diff --git a/Engine/Plugins/Animation/ControlRig/Source/ControlRig/Private/Rigs/RigHierarchyController.cpp b/Engine/Plugins/Animation/ControlRig/Source/ControlRig/Private/Rigs/RigHierarchyController.cpp
--- a/Engine/Plugins/Animation/ControlRig/Source/ControlRig/Private/Rigs/RigHierarchyController.cpp
+++ b/Engine/Plugins/Animation/ControlRig/Source/ControlRig/Private/Rigs/RigHierarchyController.cpp
@@ -386,13 +386,19 @@ FRigElementKey URigHierarchyController::AddControl(
 
 		if(!NewElement->Settings.DisplayName.IsNone())
 		{
-			NewElement->Settings.DisplayName = Hierarchy->GetSafeNewDisplayName(InParent, NewElement->Settings.DisplayName); 
+			// avoid self name collision
+			FName DesiredDisplayName = NAME_None;
+			Swap(DesiredDisplayName, NewElement->Settings.DisplayName);
+			
+			NewElement->Settings.DisplayName = Hierarchy->GetSafeNewDisplayName(InParent, DesiredDisplayName); 
 		}
 		else if(Hierarchy->HasExecuteContext())
 		{
 			const FControlRigExecuteContext& CRContext = Hierarchy->ExecuteContext->GetPublicData<FControlRigExecuteContext>();
 			if(!CRContext.GetRigModuleNameSpace().IsEmpty())
 			{
+				// avoid self name collision
+				NewElement->Settings.DisplayName = NAME_None;
 				NewElement->Settings.DisplayName = Hierarchy->GetSafeNewDisplayName(InParent, InName);
 			}
 		}
@@ -2978,6 +2984,8 @@ FName URigHierarchyController::SetDisplayName(FRigControlElement* InControlEleme
 		ParentElementKey = ParentElement->GetKey();
 	}
 
+	// avoid self name collision
+	InControlElement->Settings.DisplayName = NAME_None;
 	const FName DisplayName = Hierarchy->GetSafeNewDisplayName(ParentElementKey, InDisplayName);
 	InControlElement->Settings.DisplayName = DisplayName;
 
@@ -3196,10 +3204,14 @@ bool URigHierarchyController::AddParent(FRigBaseElement* InChild, FRigBaseElemen
 		{
 			if(!ControlElement->Settings.DisplayName.IsNone())
 			{
+				// avoid self name collision
+				FName DesiredDisplayName = NAME_None;
+				Swap(DesiredDisplayName, ControlElement->Settings.DisplayName);
+				
 				ControlElement->Settings.DisplayName =
 					Hierarchy->GetSafeNewDisplayName(
 						InParent->GetKey(),
-						ControlElement->Settings.DisplayName);
+						DesiredDisplayName);
 			}
 		}
 		
diff --git a/Engine/Plugins/Animation/ControlRig/Source/ControlRig/Public/Rigs/RigHierarchy.h b/Engine/Plugins/Animation/ControlRig/Source/ControlRig/Public/Rigs/RigHierarchy.h
--- a/Engine/Plugins/Animation/ControlRig/Source/ControlRig/Public/Rigs/RigHierarchy.h
+++ b/Engine/Plugins/Animation/ControlRig/Source/ControlRig/Public/Rigs/RigHierarchy.h
@@ -5063,6 +5063,8 @@ private:
 	
 	void EnsureCacheValidityImpl();
 
+	static FName GetUniqueName(const FName& InName, TFunction<bool(const FName&)> IsNameAvailablePredicate);
+
 	// NOTE: it is not safe to read or write the execute context without locking it first as the
 	// FRigHierarchyExecuteContextBracket can change the context from the main thread (from sequencer for example)
 	// at the same time that a control rig is being evaluated on an animation thread.
diff --git a/Engine/Plugins/Runtime/RigVM/Source/RigVMDeveloper/Private/RigVMModel/RigVMSchema.cpp b/Engine/Plugins/Runtime/RigVM/Source/RigVMDeveloper/Private/RigVMModel/RigVMSchema.cpp
--- a/Engine/Plugins/Runtime/RigVM/Source/RigVMDeveloper/Private/RigVMModel/RigVMSchema.cpp
+++ b/Engine/Plugins/Runtime/RigVM/Source/RigVMDeveloper/Private/RigVMModel/RigVMSchema.cpp
@@ -785,14 +785,23 @@ FName URigVMSchema::GetUniqueName(const FName& InName, TFunction<bool(const FNam
 	FString SanitizedPrefix = InName.ToString();
 	SanitizeName(SanitizedPrefix, bAllowPeriod, bAllowSpace);
 
-	int32 NameSuffix = 0;
-	FString Name = SanitizedPrefix;
-	while (!IsNameAvailablePredicate(*Name))
+	static constexpr int32 InitialSuffix = 2; // offset since FName uses 0 as the indicator that there's no number suffix, plus we want to start with index 1. 
+	int32 NameSuffix = InitialSuffix; 
+	FName Name = *SanitizedPrefix;
+
+	while (!IsNameAvailablePredicate(Name))
 	{
+		if(NameSuffix == InitialSuffix)
+		{
+			Name = FName(Name, NameSuffix);
+		}
+		else
+		{
+			Name.SetNumber(NameSuffix);
+		}
 		NameSuffix++;
-		Name = FString::Printf(TEXT("%s_%d"), *SanitizedPrefix, NameSuffix);
 	}
-	return *Name;
+	return Name;
 }
 
 void URigVMSchema::SetExecuteContextStruct(UScriptStruct* InExecuteContextStruct)
