Control Rig: Relink pose adapter after recompiling a rig

#rnx

diff --git a/Engine/Plugins/Animation/ControlRig/Source/ControlRig/Private/AnimNode_ControlRigBase.cpp b/Engine/Plugins/Animation/ControlRig/Source/ControlRig/Private/AnimNode_ControlRigBase.cpp
--- a/Engine/Plugins/Animation/ControlRig/Source/ControlRig/Private/AnimNode_ControlRigBase.cpp
+++ b/Engine/Plugins/Animation/ControlRig/Source/ControlRig/Private/AnimNode_ControlRigBase.cpp
@@ -135,8 +135,6 @@ void FAnimNode_ControlRig_PoseAdapter::PostLinked(URigHierarchy* InHierarchy)
 
 	// 5. Shrink the storage on the hierarchy now that we've relinked it
 	ShrinkHierarchyStorage();
-	
-	LastTopologyVersion = Hierarchy->GetTopologyVersion();
 }
 
 void FAnimNode_ControlRig_PoseAdapter::PreUnlinked(URigHierarchy* InHierarchy)
@@ -144,7 +142,6 @@ void FAnimNode_ControlRig_PoseAdapter::PreUnlinked(URigHierarchy* InHierarchy)
 	UnlinkTransformStorage();
 	UpdateHierarchyStorage();
 	
-	LastTopologyVersion = UINT32_MAX;
 	Hierarchy = nullptr;
 	FRigHierarchyPoseAdapter::PreUnlinked(InHierarchy);
 }
@@ -967,8 +964,7 @@ void FAnimNode_ControlRigBase::ExecuteControlRig(FPoseContext& InOutput)
 
 		if(bEnablePoseAdapter)
 		{
-			const uint32 TopologyVersion = Hierarchy->GetTopologyVersion();
-			if(PoseAdapter && PoseAdapter->LastTopologyVersion != TopologyVersion)
+			if(!PoseAdapter->IsUpdateToDate(Hierarchy))
 			{
 				Hierarchy->UnlinkPoseAdapter();
 				UpdateInputOutputMappingIfRequired(ControlRig, InOutput.Pose.GetBoneContainer());
diff --git a/Engine/Plugins/Animation/ControlRig/Source/ControlRig/Private/Rigs/RigHierarchy.cpp b/Engine/Plugins/Animation/ControlRig/Source/ControlRig/Private/Rigs/RigHierarchy.cpp
--- a/Engine/Plugins/Animation/ControlRig/Source/ControlRig/Private/Rigs/RigHierarchy.cpp
+++ b/Engine/Plugins/Animation/ControlRig/Source/ControlRig/Private/Rigs/RigHierarchy.cpp
@@ -3816,6 +3816,7 @@ void URigHierarchy::LinkPoseAdapter(TSharedPtr<FRigHierarchyPoseAdapter> InPoseA
 	if(PoseAdapter)
 	{
 		PoseAdapter->PreUnlinked(this);
+		PoseAdapter->bLinked = false;
 		PoseAdapter.Reset();
 	}
 
@@ -3823,6 +3824,7 @@ void URigHierarchy::LinkPoseAdapter(TSharedPtr<FRigHierarchyPoseAdapter> InPoseA
 	{
 		PoseAdapter = InPoseAdapter;
 		PoseAdapter->PostLinked(this);
+		PoseAdapter->bLinked = true;
 	}
 }
 
diff --git a/Engine/Plugins/Animation/ControlRig/Source/ControlRig/Private/Rigs/RigHierarchyPoseAdapter.cpp b/Engine/Plugins/Animation/ControlRig/Source/ControlRig/Private/Rigs/RigHierarchyPoseAdapter.cpp
--- a/Engine/Plugins/Animation/ControlRig/Source/ControlRig/Private/Rigs/RigHierarchyPoseAdapter.cpp
+++ b/Engine/Plugins/Animation/ControlRig/Source/ControlRig/Private/Rigs/RigHierarchyPoseAdapter.cpp
@@ -14,11 +14,13 @@ URigHierarchy* FRigHierarchyPoseAdapter::GetHierarchy() const
 
 void FRigHierarchyPoseAdapter::PostLinked(URigHierarchy* InHierarchy)
 {
-	WeakHierarchy = InHierarchy; 
+	WeakHierarchy = InHierarchy;
+	LastTopologyVersion = InHierarchy->GetTopologyVersion();
 }
 
 void FRigHierarchyPoseAdapter::PreUnlinked(URigHierarchy* InHierarchy)
 {
+	LastTopologyVersion = UINT32_MAX;
 }
 
 void FRigHierarchyPoseAdapter::PostUnlinked(URigHierarchy* InHierarchy)
@@ -26,6 +28,22 @@ void FRigHierarchyPoseAdapter::PostUnlinked(URigHierarchy* InHierarchy)
 	WeakHierarchy.Reset();
 }
 
+bool FRigHierarchyPoseAdapter::IsLinkedTo(const URigHierarchy* InHierarchy) const
+{
+	check(InHierarchy);
+	return IsLinked() && (GetHierarchy() == InHierarchy);
+}
+
+bool FRigHierarchyPoseAdapter::IsUpdateToDate(const URigHierarchy* InHierarchy) const
+{
+	check(InHierarchy);
+	if(!IsLinkedTo(InHierarchy))
+	{
+		return false;
+	}
+	return LastTopologyVersion != InHierarchy->GetTopologyVersion();
+}
+
 TTuple<FRigComputedTransform*, FRigTransformDirtyState*> FRigHierarchyPoseAdapter::GetElementTransformStorage(
 	const FRigElementKeyAndIndex& InKeyAndIndex, ERigTransformType::Type InTransformType, ERigTransformStorageType::Type InStorageType) const
 {
diff --git a/Engine/Plugins/Animation/ControlRig/Source/ControlRig/Public/AnimNode_ControlRigBase.h b/Engine/Plugins/Animation/ControlRig/Source/ControlRig/Public/AnimNode_ControlRigBase.h
--- a/Engine/Plugins/Animation/ControlRig/Source/ControlRig/Public/AnimNode_ControlRigBase.h
+++ b/Engine/Plugins/Animation/ControlRig/Source/ControlRig/Public/AnimNode_ControlRigBase.h
@@ -60,7 +60,6 @@ public:
 
 	FAnimNode_ControlRig_PoseAdapter()
 	: bTransferInLocalSpace(true)
-	, LastTopologyVersion(UINT32_MAX)
 	, bRequiresResetPoseToInitial(true)
 	, Hierarchy(nullptr)
 	, bEnabled(true)
@@ -94,7 +93,6 @@ protected:
 	TArray<int32> PoseCurveToHierarchyCurve;
 	TArray<bool> HierarchyCurveCopied;
 	TMap<FName, int32> HierarchyCurveLookup;
-	uint32 LastTopologyVersion;
 	bool bRequiresResetPoseToInitial;
 	TArray<int32> BonesToResetToInitial;
 
diff --git a/Engine/Plugins/Animation/ControlRig/Source/ControlRig/Public/Rigs/RigHierarchyPoseAdapter.h b/Engine/Plugins/Animation/ControlRig/Source/ControlRig/Public/Rigs/RigHierarchyPoseAdapter.h
--- a/Engine/Plugins/Animation/ControlRig/Source/ControlRig/Public/Rigs/RigHierarchyPoseAdapter.h
+++ b/Engine/Plugins/Animation/ControlRig/Source/ControlRig/Public/Rigs/RigHierarchyPoseAdapter.h
@@ -22,6 +22,9 @@ protected:
 	virtual void PostLinked(URigHierarchy* InHierarchy);
 	virtual void PreUnlinked(URigHierarchy* InHierarchy);
 	virtual void PostUnlinked(URigHierarchy* InHierarchy);
+	bool IsLinked() const { return bLinked; }
+	bool IsLinkedTo(const URigHierarchy* InHierarchy) const;
+	bool IsUpdateToDate(const URigHierarchy* InHierarchy) const;
 
 	TTuple<FRigComputedTransform*, FRigTransformDirtyState*> GetElementTransformStorage(const FRigElementKeyAndIndex& InKeyAndIndex, ERigTransformType::Type InTransformType, ERigTransformStorageType::Type InStorageType) const;
 
@@ -40,6 +43,8 @@ protected:
 	bool UpdateHierarchyStorage();
 
 	TWeakObjectPtr<URigHierarchy> WeakHierarchy;
+	bool bLinked = false;
+	uint32 LastTopologyVersion = UINT32_MAX;
 
 	friend class URigHierarchy;
 	friend class URigHierarchyController;
