Control Rig: Increase max element name length to 200

#rnx

diff --git a/Engine/Plugins/Animation/ControlRig/Source/ControlRig/Public/Rigs/RigHierarchy.h b/Engine/Plugins/Animation/ControlRig/Source/ControlRig/Public/Rigs/RigHierarchy.h
--- a/Engine/Plugins/Animation/ControlRig/Source/ControlRig/Public/Rigs/RigHierarchy.h
+++ b/Engine/Plugins/Animation/ControlRig/Source/ControlRig/Public/Rigs/RigHierarchy.h
@@ -1886,7 +1886,7 @@ public:
 	 * Returns the max allowed length for a name within the hierarchy.
 	 * @return Returns the max allowed length for a name within the hierarchy.
 	 */
-	static int32 GetMaxNameLength() { return 100; }
+	static int32 GetMaxNameLength() { return 200; }
 
 	/**
 	 * Sanitizes a name by removing invalid characters.
