Anim Layers: Don't refresh on ActivateSequence callback form Sequencer if Transacting, since Sequencer annoying always broadcasts this on Undo.

diff --git a/Engine/Plugins/Animation/ControlRig/Source/ControlRigEditor/Private/Sequencer/AnimLayers/SAnimLayers.cpp b/Engine/Plugins/Animation/ControlRig/Source/ControlRigEditor/Private/Sequencer/AnimLayers/SAnimLayers.cpp
--- a/Engine/Plugins/Animation/ControlRig/Source/ControlRigEditor/Private/Sequencer/AnimLayers/SAnimLayers.cpp
+++ b/Engine/Plugins/Animation/ControlRig/Source/ControlRigEditor/Private/Sequencer/AnimLayers/SAnimLayers.cpp
@@ -1639,7 +1639,7 @@ void SAnimLayers::OnActorSelectionChanged(const TArray<UObject*>& NewSelection,
 
 void SAnimLayers::OnActivateSequence(FMovieSceneSequenceIDRef ID)
 {
-	if (AnimLayerController)
+	if (!GIsTransacting && AnimLayerController) //don't do this if in UNDO, Sequencer unfortunately calls this delegate then
 	{
 		AnimLayerController->RefreshSourceData(true);
 		AnimLayerController->RefreshSelectionData();
