Control Rig: Compact Transform Serialization

#rnx

diff --git a/Engine/Plugins/Animation/ControlRig/Source/ControlRig/Private/Rigs/RigHierarchy.cpp b/Engine/Plugins/Animation/ControlRig/Source/ControlRig/Private/Rigs/RigHierarchy.cpp
--- a/Engine/Plugins/Animation/ControlRig/Source/ControlRig/Private/Rigs/RigHierarchy.cpp
+++ b/Engine/Plugins/Animation/ControlRig/Source/ControlRig/Private/Rigs/RigHierarchy.cpp
@@ -188,6 +188,9 @@ void URigHierarchy::Save(FArchive& Ar)
 	FScopeLock Lock(&ElementsLock);
 	UE_RIGVM_ARCHIVETRACE_SCOPE(Ar, FString::Printf(TEXT("URigHierarchy(%s)"), *GetName()));
 
+	Ar.UsingCustomVersion(FAnimObjectVersion::GUID);
+	Ar.UsingCustomVersion(FControlRigObjectVersion::GUID);
+
 	if(Ar.IsTransacting())
 	{
 		Ar << TransformStackIndex;
@@ -209,24 +212,26 @@ void URigHierarchy::Save(FArchive& Ar)
 	// this ensures cache validity.
 	EnsureCacheValidity();
 
+	// by computing all transforms we can be sure that initial & current,
+	// and local & global each are correct.
 	ComputeAllTransforms();
 
-	FArchive* ArchiveForElements = &Ar;
-
-	const bool bIsDuplicating = (Ar.GetPortFlags() & PPF_Duplicate) != 0;
-	bool bUseCompressedArchive =
-		!bIsDuplicating &&
-		Ar.IsPersistent() &&
-		!Ar.IsObjectReferenceCollector() &&
-		!Ar.ShouldSkipBulkData() &&
-		!Ar.IsTransacting();
+	FRigHierarchySerializationSettings SerializationSettings(Ar);
+	if(SerializationSettings.bIsSerializingToPackage)
+	{
+		SerializationSettings.bSerializeGlobalTransform = false;
+		SerializationSettings.bSerializeCurrentTransform = false;
+		SerializationSettings.bUseCompressedArchive = true;
+	}
+	
+	SerializationSettings.Save(Ar);
 
-	Ar << bUseCompressedArchive;
+	FArchive* ArchiveForElements = &Ar;
 
 	TArray<uint8> UncompressedBytes;
 	TArray<FName> UniqueNames;
 	FRigHierarchyMemoryWriter MemoryWriter(UncompressedBytes, UniqueNames, Ar.IsPersistent());
-	if(bUseCompressedArchive)
+	if(SerializationSettings.bUseCompressedArchive)
 	{
 		MemoryWriter.UsingCustomVersion(FAnimObjectVersion::GUID);
 		MemoryWriter.UsingCustomVersion(FControlRigObjectVersion::GUID);
@@ -246,16 +251,18 @@ void URigHierarchy::Save(FArchive& Ar)
 		*ArchiveForElements << Key;
 	}
 
+	SerializationSettings.SerializationPhase = FRigHierarchySerializationSettings::StaticData;
 	for(int32 ElementIndex = 0; ElementIndex < ElementCount; ElementIndex++)
 	{
 		FRigBaseElement* Element = Elements[ElementIndex];
-		Element->Serialize(*ArchiveForElements, FRigBaseElement::StaticData);
+		Element->Serialize(*ArchiveForElements, SerializationSettings);
 	}
 
+	SerializationSettings.SerializationPhase = FRigHierarchySerializationSettings::InterElementData;
 	for(int32 ElementIndex = 0; ElementIndex < ElementCount; ElementIndex++)
 	{
 		FRigBaseElement* Element = Elements[ElementIndex];
-		Element->Serialize(*ArchiveForElements, FRigBaseElement::InterElementData);
+		Element->Serialize(*ArchiveForElements, SerializationSettings);
 	}
 
 	*ArchiveForElements << PreviousNameMap;
@@ -296,7 +303,7 @@ void URigHierarchy::Save(FArchive& Ar)
 		}
 	}
 
-	if(bUseCompressedArchive)
+	if(SerializationSettings.bUseCompressedArchive)
 	{
 		// It is possible for compression to actually increase the size of the data, so we over allocate here to handle that.
 		int32 UncompressedSize = UncompressedBytes.Num();
@@ -372,11 +379,16 @@ void URigHierarchy::Load(FArchive& Ar)
 	
 	Reset();
 
-	bool bUseCompressedArchive = false;
-	if(Ar.CustomVer(FControlRigObjectVersion::GUID) >= FControlRigObjectVersion::RigHierarchyCompressElements)
+	FRigHierarchySerializationSettings SerializationSettings(Ar);
+	if(Ar.CustomVer(FControlRigObjectVersion::GUID) >= FControlRigObjectVersion::RigHierarchyCompactTransformSerialization)
+	{
+		SerializationSettings.Load(Ar);
+	}
+	else if(Ar.CustomVer(FControlRigObjectVersion::GUID) >= FControlRigObjectVersion::RigHierarchyCompressElements)
 	{
-		Ar << bUseCompressedArchive;
+		Ar << SerializationSettings.bUseCompressedArchive;
 	}
+	SerializationSettings.SerializationPhase = FRigHierarchySerializationSettings::StaticData;
 
 	FArchive* ArchiveForElements = &Ar;
 
@@ -384,7 +396,7 @@ void URigHierarchy::Load(FArchive& Ar)
 	TArray<uint8> UncompressedBytes;
 	TArray<uint8> CompressedBytes;
 	TArray<FName> UniqueNames;
-	if(bUseCompressedArchive)
+	if(SerializationSettings.bUseCompressedArchive)
 	{
 		// this block is not going to be entered if we are in an older asset
 		// ( < FControlRigObjectVersion::RigHierarchyCompressElements ) since
@@ -419,7 +431,7 @@ void URigHierarchy::Load(FArchive& Ar)
 	int32 NumTransforms = 0;
 	int32 NumDirtyStates = 0;
 	int32 NumCurves = 0;
-
+	
 	const bool bAllocateStoragePerElement = ArchiveForElements->CustomVer(FControlRigObjectVersion::GUID) < FControlRigObjectVersion::RigHierarchyIndirectElementStorage;
 	for(int32 ElementIndex = 0; ElementIndex < ElementCount; ElementIndex++)
 	{
@@ -437,7 +449,7 @@ void URigHierarchy::Load(FArchive& Ar)
 		if(bAllocateStoragePerElement)
 		{
 			AllocateDefaultElementStorage(Element, false);
-			Element->Load(*ArchiveForElements, FRigBaseElement::StaticData);
+			Element->Load(*ArchiveForElements, SerializationSettings);
 		}
 		else
 		{
@@ -507,15 +519,16 @@ void URigHierarchy::Load(FArchive& Ar)
 			}
 
 			Element->LinkStorage(ElementTransforms.GetStorage(), ElementDirtyStates.GetStorage(), ElementCurves.GetStorage());
-			Element->Load(*ArchiveForElements, FRigBaseElement::StaticData);
+			Element->Load(*ArchiveForElements, SerializationSettings);
 		}
 	}
 	IncrementTopologyVersion();
 
+	SerializationSettings.SerializationPhase = FRigHierarchySerializationSettings::InterElementData;
 	for(int32 ElementIndex = 0; ElementIndex < ElementCount; ElementIndex++)
 	{
 		FRigBaseElement* Element = Elements[ElementIndex];
-		Element->Load(*ArchiveForElements, FRigBaseElement::InterElementData);
+		Element->Load(*ArchiveForElements, SerializationSettings);
 	}
 
 	IncrementTopologyVersion();
diff --git a/Engine/Plugins/Animation/ControlRig/Source/ControlRig/Private/Rigs/RigHierarchyDefines.cpp b/Engine/Plugins/Animation/ControlRig/Source/ControlRig/Private/Rigs/RigHierarchyDefines.cpp
--- a/Engine/Plugins/Animation/ControlRig/Source/ControlRig/Private/Rigs/RigHierarchyDefines.cpp
+++ b/Engine/Plugins/Animation/ControlRig/Source/ControlRig/Private/Rigs/RigHierarchyDefines.cpp
@@ -10,6 +10,54 @@
 #include "RigVMPythonUtils.h"
 #endif
 
+FRigHierarchySerializationSettings::FRigHierarchySerializationSettings(const FArchive& InArchive)
+{
+	const bool bIsDuplicating = (InArchive.GetPortFlags() & PPF_Duplicate) != 0;
+	bIsSerializingToPackage =
+		!bIsDuplicating &&
+		InArchive.IsPersistent() &&
+		!InArchive.IsObjectReferenceCollector() &&
+		!InArchive.ShouldSkipBulkData() &&
+		!InArchive.IsTransacting();
+
+	ControlRigVersion = (FControlRigObjectVersion::Type)InArchive.CustomVer(FControlRigObjectVersion::GUID);
+}
+
+void FRigHierarchySerializationSettings::Save(FArchive& InArchive)
+{
+	int32 ControlRigVersionInt = (int32)ControlRigVersion;
+	int32 SerializationPhaseInt = (int32)SerializationPhase;
+
+	InArchive << ControlRigVersionInt;
+	InArchive << bIsSerializingToPackage;
+	InArchive << bUseCompressedArchive;
+	InArchive << bStoreCompactTransforms;
+	InArchive << bSerializeLocalTransform;
+	InArchive << bSerializeGlobalTransform;
+	InArchive << bSerializeInitialTransform;
+	InArchive << bSerializeCurrentTransform;
+	InArchive << SerializationPhaseInt;
+}
+
+void FRigHierarchySerializationSettings::Load(FArchive& InArchive)
+{
+	int32 ControlRigVersionInt = INDEX_NONE;
+	int32 SerializationPhaseInt = INDEX_NONE;
+
+	InArchive << ControlRigVersionInt;
+	InArchive << bIsSerializingToPackage;
+	InArchive << bUseCompressedArchive;
+	InArchive << bStoreCompactTransforms;
+	InArchive << bSerializeLocalTransform;
+	InArchive << bSerializeGlobalTransform;
+	InArchive << bSerializeInitialTransform;
+	InArchive << bSerializeCurrentTransform;
+	InArchive << SerializationPhaseInt;
+
+	ControlRigVersion = (FControlRigObjectVersion::Type)ControlRigVersionInt;
+	SerializationPhase = (ESerializationPhase)SerializationPhaseInt;
+}
+
 ////////////////////////////////////////////////////////////////////////////////
 // FRigControlLimitEnabled
 ////////////////////////////////////////////////////////////////////////////////
diff --git a/Engine/Plugins/Animation/ControlRig/Source/ControlRig/Private/Rigs/RigHierarchyElements.cpp b/Engine/Plugins/Animation/ControlRig/Source/ControlRig/Private/Rigs/RigHierarchyElements.cpp
--- a/Engine/Plugins/Animation/ControlRig/Source/ControlRig/Private/Rigs/RigHierarchyElements.cpp
+++ b/Engine/Plugins/Animation/ControlRig/Source/ControlRig/Private/Rigs/RigHierarchyElements.cpp
@@ -25,33 +25,33 @@ FRigBaseElement::~FRigBaseElement()
 	}
 }
 
-void FRigBaseElement::Serialize(FArchive& Ar, ESerializationPhase SerializationPhase)
+void FRigBaseElement::Serialize(FArchive& Ar, const FRigHierarchySerializationSettings& InSettings)
 {
 	Ar.UsingCustomVersion(FControlRigObjectVersion::GUID);
 
 	if (Ar.IsLoading())
 	{
-		Load(Ar, SerializationPhase);
+		Load(Ar, InSettings);
 	}
 	else
 	{
-		Save(Ar, SerializationPhase);
+		Save(Ar, InSettings);
 	}
 }
 
-void FRigBaseElement::Save(FArchive& Ar, ESerializationPhase SerializationPhase)
+void FRigBaseElement::Save(FArchive& Ar, const FRigHierarchySerializationSettings& InSettings)
 {
-	if(SerializationPhase == ESerializationPhase::StaticData)
+	if(InSettings.SerializationPhase == FRigHierarchySerializationSettings::StaticData)
 	{
 		Ar << Key;
 	}
 }
 
-void FRigBaseElement::Load(FArchive& Ar, ESerializationPhase SerializationPhase)
+void FRigBaseElement::Load(FArchive& Ar, const FRigHierarchySerializationSettings& InSettings)
 {
 	checkf(Owner != nullptr, TEXT("Loading should not happen on a rig element without an owner"));
 	
-	if(SerializationPhase == ESerializationPhase::StaticData)
+	if(InSettings.SerializationPhase == FRigHierarchySerializationSettings::StaticData)
 	{
 		FRigElementKey LoadedKey;
 	
@@ -263,6 +263,278 @@ void FRigBaseElement::CopyFrom(const FRigBaseElement* InOther)
 {
 }
 
+////////////////////////////////////////////////////////////////////////////////
+// FRigCompactTransform
+////////////////////////////////////////////////////////////////////////////////
+
+FRigCompactTransform::FRigCompactTransform(FTransform& InTransform)
+: Transform(InTransform)
+{
+}
+
+void FRigCompactTransform::Serialize(FArchive& Ar, const FRigHierarchySerializationSettings& InSettings, ERepresentation* OutRepresentation)
+{
+	if(Ar.IsLoading())
+	{
+		Load(Ar, InSettings, OutRepresentation);
+	}
+	else if(Ar.IsSaving())
+	{
+		Save(Ar, InSettings, OutRepresentation);
+	}
+}
+
+void FRigCompactTransform::Save(FArchive& Ar, const FRigHierarchySerializationSettings& InSettings, ERepresentation* OutRepresentation)
+{
+	uint8 State = Float_Zero_Identity_One;
+	if(!InSettings.bStoreCompactTransforms)
+	{
+		State = Double_Complete;
+		if(OutRepresentation)
+		{
+			*OutRepresentation = (ERepresentation)State;
+		}
+		Ar << State;
+		Ar << Transform;
+		return;
+	}
+
+	FVector3f Position(Transform.GetLocation());
+	if(!Position.IsNearlyZero(UE_SMALL_NUMBER))
+	{
+		State += 6;
+	}
+	FQuat4f Rotation(Transform.GetRotation());
+	if(!Rotation.IsIdentity(UE_SMALL_NUMBER))
+	{
+		State += 3;
+	}
+	FVector3f Scale(Transform.GetScale3D());
+	float ScaleX = Scale.X;
+
+	if(!Scale.IsUniform(UE_SMALL_NUMBER))
+	{
+		State += 2;
+	}
+	else if(!FMath::IsNearlyEqual(Scale.X, 1.f, UE_SMALL_NUMBER))
+	{
+		State += 1;
+	}
+
+	if(OutRepresentation)
+	{
+		*OutRepresentation = (ERepresentation)State;
+	}
+
+	Ar << State;
+	switch(State)
+	{
+		case Float_Zero_Identity_One:
+		{
+			// don't store anything
+			break;
+		}
+		case Float_Zero_Identity_Uniform:
+		{
+			Ar << ScaleX;
+			break;
+		}
+		case Float_Zero_Identity_NonUniform:
+		{
+			Ar << Scale;
+			break;
+		}
+		case Float_Zero_Quat_One:
+		{
+			Ar << Rotation;
+			break;
+		}
+		case Float_Zero_Quat_Uniform:
+		{
+			Ar << Rotation;
+			Ar << ScaleX;
+			break;
+		}
+		case Float_Zero_Quat_NonUniform:
+		{
+			Ar << Rotation;
+			Ar << Scale;
+			break;
+		}
+		case Float_Position_Identity_One:
+		{
+			Ar << Position;
+			break;
+		}
+		case Float_Position_Identity_Uniform:
+		{
+			Ar << Position;
+			Ar << ScaleX;
+			break;
+		}
+		case Float_Position_Identity_NonUniform:
+		{
+			Ar << Position;
+			Ar << Scale;
+			break;
+		}
+		case Float_Position_Quat_One:
+		{
+			Ar << Position;
+			Ar << Rotation;
+			break;
+		}
+		case Float_Position_Quat_Uniform:
+		{
+			Ar << Position;
+			Ar << Rotation;
+			Ar << ScaleX;
+			break;
+		}
+		case Float_Position_Quat_NonUniform:
+		{
+			Ar << Position;
+			Ar << Rotation;
+			Ar << Scale;
+			break;
+		}
+		case Double_Complete:
+		default:
+		{
+			Ar << Transform;
+			break;
+		}
+	}
+}
+
+void FRigCompactTransform::Load(FArchive& Ar, const FRigHierarchySerializationSettings& InSettings, ERepresentation* OutRepresentation)
+{
+	if(InSettings.ControlRigVersion < FControlRigObjectVersion::RigHierarchyCompactTransformSerialization)
+	{
+		Ar << Transform;
+		return;
+	}
+
+	uint8 State = Float_Zero_Identity_One;
+	Ar << State;
+
+	if(OutRepresentation)
+	{
+		*OutRepresentation = (ERepresentation)State;
+	}
+
+	if(State == Double_Complete)
+	{
+		Ar << Transform;
+		return;
+	}
+	
+	Transform = FTransform::Identity;
+	FVector3f Position = FVector3f::ZeroVector;
+	FQuat4f Rotation = FQuat4f::Identity;
+	FVector3f Scale = FVector3f::OneVector;
+	float ScaleX = 1.f;
+
+	switch(State)
+	{
+		case Float_Zero_Identity_One:
+		{
+			// don't load anything
+			break;
+		}
+		case Float_Zero_Identity_Uniform:
+		{
+			Ar << ScaleX;
+			Transform.SetScale3D({ScaleX, ScaleX, ScaleX});
+			break;
+		}
+		case Float_Zero_Identity_NonUniform:
+		{
+			Ar << Scale;
+			Transform.SetScale3D(FVector(Scale));
+			break;
+		}
+		case Float_Zero_Quat_One:
+		{
+			Ar << Rotation;
+			Transform.SetRotation(FQuat(Rotation));
+			break;
+		}
+		case Float_Zero_Quat_Uniform:
+		{
+			Ar << Rotation;
+			Ar << ScaleX;
+			Transform.SetRotation(FQuat(Rotation));
+			Transform.SetScale3D({ScaleX, ScaleX, ScaleX});
+			break;
+		}
+		case Float_Zero_Quat_NonUniform:
+		{
+			Ar << Rotation;
+			Ar << Scale;
+			Transform.SetRotation(FQuat(Rotation));
+			Transform.SetScale3D(FVector(Scale));
+			break;
+		}
+		case Float_Position_Identity_One:
+		{
+			Ar << Position;
+			Transform.SetLocation(FVector(Position));
+			break;
+		}
+		case Float_Position_Identity_Uniform:
+		{
+			Ar << Position;
+			Ar << ScaleX;
+			Transform.SetLocation(FVector(Position));
+			Transform.SetScale3D({ScaleX, ScaleX, ScaleX});
+			break;
+		}
+		case Float_Position_Identity_NonUniform:
+		{
+			Ar << Position;
+			Ar << Scale;
+			Transform.SetLocation(FVector(Position));
+			Transform.SetScale3D(FVector(Scale));
+			break;
+		}
+		case Float_Position_Quat_One:
+		{
+			Ar << Position;
+			Ar << Rotation;
+			Transform.SetLocation(FVector(Position));
+			Transform.SetRotation(FQuat(Rotation));
+			break;
+		}
+		case Float_Position_Quat_Uniform:
+		{
+			Ar << Position;
+			Ar << Rotation;
+			Ar << ScaleX;
+			Transform.SetLocation(FVector(Position));
+			Transform.SetRotation(FQuat(Rotation));
+			Transform.SetScale3D({ScaleX, ScaleX, ScaleX});
+			break;
+		}
+		case Float_Position_Quat_NonUniform:
+		{
+			Ar << Position;
+			Ar << Rotation;
+			Ar << Scale;
+			Transform.SetLocation(FVector(Position));
+			Transform.SetRotation(FQuat(Rotation));
+			Transform.SetScale3D(FVector(Scale));
+			break;
+		}
+		case Double_Complete:
+		default:
+		{
+			Ar << Transform;
+			break;
+		}
+	}
+}
+
 ////////////////////////////////////////////////////////////////////////////////
 // FRigTransformDirtyState
 ////////////////////////////////////////////////////////////////////////////////
@@ -374,22 +646,36 @@ void FRigCurrentAndInitialDirtyState::UnlinkStorage(FRigReusableElementStorage<b
 // FRigComputedTransform
 ////////////////////////////////////////////////////////////////////////////////
 
-void FRigComputedTransform::Save(FArchive& Ar, const FRigTransformDirtyState& InDirtyState)
+void FRigComputedTransform::Save(FArchive& Ar, const FRigHierarchySerializationSettings& InSettings, const FRigTransformDirtyState& InDirtyState)
 {
 	FTransform Transform = Get();
 	bool bDirty = InDirtyState.Get();
-	
-	Ar << Transform;
+
 	Ar << bDirty;
+	if(!bDirty)
+	{
+		FRigCompactTransform(Transform).Save(Ar, InSettings);
+	}
 }
 
-void FRigComputedTransform::Load(FArchive& Ar, FRigTransformDirtyState& InDirtyState)
+void FRigComputedTransform::Load(FArchive& Ar, const FRigHierarchySerializationSettings& InSettings, FRigTransformDirtyState& InDirtyState)
 {
 	FTransform Transform = FTransform::Identity;
 	bool bDirty = false;
 
-	Ar << Transform;
-	Ar << bDirty;
+	if(InSettings.ControlRigVersion < FControlRigObjectVersion::RigHierarchyCompactTransformSerialization)
+	{
+		Ar << Transform;
+		Ar << bDirty;
+	}
+	else
+	{
+		Ar << bDirty;
+		if(!bDirty)
+		{
+			FRigCompactTransform(Transform).Load(Ar, InSettings);
+		}
+	}
 
 	Set(Transform);
 	(void)InDirtyState.Set(bDirty);
@@ -432,16 +718,46 @@ void FRigComputedTransform::UnlinkStorage(FRigReusableElementStorage<FTransform>
 // FRigLocalAndGlobalTransform
 ////////////////////////////////////////////////////////////////////////////////
 
-void FRigLocalAndGlobalTransform::Save(FArchive& Ar, const FRigLocalAndGlobalDirtyState& InDirtyState)
+void FRigLocalAndGlobalTransform::Save(FArchive& Ar, const FRigHierarchySerializationSettings& InSettings, const FRigLocalAndGlobalDirtyState& InDirtyState)
 {
-	Local.Save(Ar, InDirtyState.Local);
-	Global.Save(Ar, InDirtyState.Global);
+	if(InSettings.bSerializeLocalTransform)
+	{
+		if(!InSettings.bSerializeGlobalTransform)
+		{
+			check(!InDirtyState.Local.Get());
+		}
+		Local.Save(Ar, InSettings, InDirtyState.Local);
+	}
+	if(InSettings.bSerializeGlobalTransform)
+	{
+		if(!InSettings.bSerializeLocalTransform)
+		{
+			check(!InDirtyState.Global.Get());
+		}
+		Global.Save(Ar, InSettings, InDirtyState.Global);
+	}
 }
 
-void FRigLocalAndGlobalTransform::Load(FArchive& Ar, FRigLocalAndGlobalDirtyState& OutDirtyState)
+void FRigLocalAndGlobalTransform::Load(FArchive& Ar, const FRigHierarchySerializationSettings& InSettings, FRigLocalAndGlobalDirtyState& OutDirtyState)
 {
-	Local.Load(Ar, OutDirtyState.Local);
-	Global.Load(Ar, OutDirtyState.Global);
+	if(InSettings.bSerializeLocalTransform)
+	{
+		Local.Load(Ar, InSettings, OutDirtyState.Local);
+		if(!InSettings.bSerializeGlobalTransform)
+		{
+			check(!OutDirtyState.Local.Get());
+			OutDirtyState.Global.Set(true);
+		}
+	}
+	if(InSettings.bSerializeGlobalTransform)
+	{
+		Global.Load(Ar, InSettings, OutDirtyState.Global);
+		if(!InSettings.bSerializeLocalTransform)
+		{
+			check(!OutDirtyState.Global.Get());
+			OutDirtyState.Local.Set(true);
+		}
+	}
 }
 
 FRigLocalAndGlobalTransform& FRigLocalAndGlobalTransform::operator=(const FRigLocalAndGlobalTransform& InOther)
@@ -467,16 +783,33 @@ void FRigLocalAndGlobalTransform::UnlinkStorage(FRigReusableElementStorage<FTran
 // FRigCurrentAndInitialTransform
 ////////////////////////////////////////////////////////////////////////////////
 
-void FRigCurrentAndInitialTransform::Save(FArchive& Ar, const FRigCurrentAndInitialDirtyState& InDirtyState)
+void FRigCurrentAndInitialTransform::Save(FArchive& Ar, const FRigHierarchySerializationSettings& InSettings, const FRigCurrentAndInitialDirtyState& InDirtyState)
 {
-	Current.Save(Ar, InDirtyState.Current);
-	Initial.Save(Ar, InDirtyState.Initial);
+	if(InSettings.bSerializeCurrentTransform)
+	{
+		Current.Save(Ar, InSettings, InDirtyState.Current);
+	}
+	if(InSettings.bSerializeInitialTransform)
+	{
+		Initial.Save(Ar, InSettings, InDirtyState.Initial);
+	}
 }
 
-void FRigCurrentAndInitialTransform::Load(FArchive& Ar, FRigCurrentAndInitialDirtyState& OutDirtyState)
+void FRigCurrentAndInitialTransform::Load(FArchive& Ar, const FRigHierarchySerializationSettings& InSettings, FRigCurrentAndInitialDirtyState& OutDirtyState)
 {
-	Current.Load(Ar, OutDirtyState.Current);
-	Initial.Load(Ar, OutDirtyState.Initial);
+	if(InSettings.bSerializeCurrentTransform)
+	{
+		Current.Load(Ar, InSettings, OutDirtyState.Current);
+	}
+	if(InSettings.bSerializeInitialTransform)
+	{
+		Initial.Load(Ar, InSettings, OutDirtyState.Initial);
+		if(!InSettings.bSerializeCurrentTransform)
+		{
+			Current = Initial;
+			OutDirtyState.Current = OutDirtyState.Initial;
+		}
+	}
 }
 
 FRigCurrentAndInitialTransform& FRigCurrentAndInitialTransform::operator=(const FRigCurrentAndInitialTransform& InOther)
@@ -677,23 +1010,23 @@ FRigBaseComponent* FRigComponentHandle::Get()
 
 const FRigBaseElement::EElementIndex FRigTransformElement::ElementTypeIndex = TransformElement;
 
-void FRigTransformElement::Save(FArchive& Ar, ESerializationPhase SerializationPhase)
+void FRigTransformElement::Save(FArchive& Ar, const FRigHierarchySerializationSettings& InSettings)
 {
-	Super::Save(Ar, SerializationPhase);
+	Super::Save(Ar, InSettings);
 
-	if(SerializationPhase == ESerializationPhase::StaticData)
+	if(InSettings.SerializationPhase == FRigHierarchySerializationSettings::StaticData)
 	{
-		GetTransform().Save(Ar, GetDirtyState());
+		GetTransform().Save(Ar, InSettings, GetDirtyState());
 	}
 }
 
-void FRigTransformElement::Load(FArchive& Ar, ESerializationPhase SerializationPhase)
+void FRigTransformElement::Load(FArchive& Ar, const FRigHierarchySerializationSettings& InSettings)
 {
-	Super::Load(Ar, SerializationPhase);
+	Super::Load(Ar, InSettings);
 
-	if(SerializationPhase == ESerializationPhase::StaticData)
+	if(InSettings.SerializationPhase == FRigHierarchySerializationSettings::StaticData)
 	{
-		GetTransform().Load(Ar, GetDirtyState());
+		GetTransform().Load(Ar, InSettings, GetDirtyState());
 	}
 }
 
@@ -779,11 +1112,11 @@ void FRigTransformElement::UnlinkStorage(FRigReusableElementStorage<FTransform>&
 
 const FRigBaseElement::EElementIndex FRigSingleParentElement::ElementTypeIndex = SingleParentElement;
 
-void FRigSingleParentElement::Save(FArchive& Ar, ESerializationPhase SerializationPhase)
+void FRigSingleParentElement::Save(FArchive& Ar, const FRigHierarchySerializationSettings& InSettings)
 {
-	Super::Save(Ar, SerializationPhase);
+	Super::Save(Ar, InSettings);
 
-	if(SerializationPhase == ESerializationPhase::InterElementData)
+	if(InSettings.SerializationPhase == FRigHierarchySerializationSettings::InterElementData)
 	{
 		FRigElementKey ParentKey;
 		if(ParentElement)
@@ -794,11 +1127,11 @@ void FRigSingleParentElement::Save(FArchive& Ar, ESerializationPhase Serializati
 	}
 }
 
-void FRigSingleParentElement::Load(FArchive& Ar, ESerializationPhase SerializationPhase)
+void FRigSingleParentElement::Load(FArchive& Ar, const FRigHierarchySerializationSettings& InSettings)
 {
-	Super::Load(Ar, SerializationPhase);
+	Super::Load(Ar, InSettings);
 
-	if(SerializationPhase == ESerializationPhase::InterElementData)
+	if(InSettings.SerializationPhase == FRigHierarchySerializationSettings::InterElementData)
 	{
 		FRigElementKey ParentKey;
 		Ar << ParentKey;
@@ -832,16 +1165,16 @@ void FRigSingleParentElement::CopyFrom(const FRigBaseElement* InOther)
 
 const FRigBaseElement::EElementIndex FRigMultiParentElement::ElementTypeIndex = MultiParentElement;
 
-void FRigMultiParentElement::Save(FArchive& Ar, ESerializationPhase SerializationPhase)
+void FRigMultiParentElement::Save(FArchive& Ar, const FRigHierarchySerializationSettings& InSettings)
 {
-	Super::Save(Ar, SerializationPhase);
+	Super::Save(Ar, InSettings);
 
-	if(SerializationPhase == ESerializationPhase::StaticData)
+	if(InSettings.SerializationPhase == FRigHierarchySerializationSettings::StaticData)
 	{
 		int32 NumParents = ParentConstraints.Num();
 		Ar << NumParents;
 	}
-	else if(SerializationPhase == ESerializationPhase::InterElementData)
+	else if(InSettings.SerializationPhase == FRigHierarchySerializationSettings::InterElementData)
 	{
 		for(int32 ParentIndex = 0; ParentIndex < ParentConstraints.Num(); ParentIndex++)
 		{
@@ -858,17 +1191,17 @@ void FRigMultiParentElement::Save(FArchive& Ar, ESerializationPhase Serializatio
 	}
 }
 
-void FRigMultiParentElement::Load(FArchive& Ar, ESerializationPhase SerializationPhase)
+void FRigMultiParentElement::Load(FArchive& Ar, const FRigHierarchySerializationSettings& InSettings)
 {
-	Super::Load(Ar, SerializationPhase);
+	Super::Load(Ar, InSettings);
 
-	if(SerializationPhase == ESerializationPhase::StaticData)
+	if(InSettings.SerializationPhase == FRigHierarchySerializationSettings::StaticData)
 	{
 		if (Ar.CustomVer(FControlRigObjectVersion::GUID) < FControlRigObjectVersion::RemovedMultiParentParentCache)
 		{
 			FRigCurrentAndInitialTransform Parent;
 			FRigCurrentAndInitialDirtyState DirtyState;
-			Parent.Load(Ar, DirtyState);
+			Parent.Load(Ar, InSettings, DirtyState);
 		}
 
 		int32 NumParents = 0;
@@ -876,7 +1209,7 @@ void FRigMultiParentElement::Load(FArchive& Ar, ESerializationPhase Serializatio
 
 		ParentConstraints.SetNum(NumParents);
 	}
-	else if(SerializationPhase == ESerializationPhase::InterElementData)
+	else if(InSettings.SerializationPhase == FRigHierarchySerializationSettings::InterElementData)
 	{
 		for(int32 ParentIndex = 0; ParentIndex < ParentConstraints.Num(); ParentIndex++)
 		{
@@ -985,11 +1318,11 @@ static_assert(sizeof(FRigBoneElement) <= 736, "FRigBoneElement was optimized to
 
 const FRigBaseElement::EElementIndex FRigBoneElement::ElementTypeIndex = BoneElement;
 
-void FRigBoneElement::Save(FArchive& Ar, ESerializationPhase SerializationPhase)
+void FRigBoneElement::Save(FArchive& Ar, const FRigHierarchySerializationSettings& InSettings)
 {
-	Super::Save(Ar, SerializationPhase);
+	Super::Save(Ar, InSettings);
 
-	if(SerializationPhase == ESerializationPhase::StaticData)
+	if(InSettings.SerializationPhase == FRigHierarchySerializationSettings::StaticData)
 	{
 		static const UEnum* BoneTypeEnum = StaticEnum<ERigBoneType>();
 		FName TypeName = BoneTypeEnum->GetNameByValue((int64)BoneType);
@@ -997,11 +1330,11 @@ void FRigBoneElement::Save(FArchive& Ar, ESerializationPhase SerializationPhase)
 	}
 }
 
-void FRigBoneElement::Load(FArchive& Ar, ESerializationPhase SerializationPhase)
+void FRigBoneElement::Load(FArchive& Ar, const FRigHierarchySerializationSettings& InSettings)
 {
-	Super::Load(Ar, SerializationPhase);
+	Super::Load(Ar, InSettings);
 
-	if(SerializationPhase == ESerializationPhase::StaticData)
+	if(InSettings.SerializationPhase == FRigHierarchySerializationSettings::StaticData)
 	{
 		static const UEnum* BoneTypeEnum = StaticEnum<ERigBoneType>();
 		FName TypeName;
@@ -1530,28 +1863,28 @@ FRigCurrentAndInitialDirtyState& FRigControlElement::GetShapeDirtyState()
 
 const FRigBaseElement::EElementIndex FRigControlElement::ElementTypeIndex = ControlElement;
 
-void FRigControlElement::Save(FArchive& Ar, ESerializationPhase SerializationPhase)
+void FRigControlElement::Save(FArchive& Ar, const FRigHierarchySerializationSettings& InSettings)
 {
-	Super::Save(Ar, SerializationPhase);
+	Super::Save(Ar, InSettings);
 
-	if(SerializationPhase == ESerializationPhase::StaticData)
+	if(InSettings.SerializationPhase == FRigHierarchySerializationSettings::StaticData)
 	{
 		Settings.Save(Ar);
-		GetOffsetTransform().Save(Ar, GetOffsetDirtyState());
-		GetShapeTransform().Save(Ar, GetShapeDirtyState());
+		GetOffsetTransform().Save(Ar, InSettings, GetOffsetDirtyState());
+		GetShapeTransform().Save(Ar, InSettings, GetShapeDirtyState());
 		PreferredEulerAngles.Save(Ar);
 	}
 }
 
-void FRigControlElement::Load(FArchive& Ar, ESerializationPhase SerializationPhase)
+void FRigControlElement::Load(FArchive& Ar, const FRigHierarchySerializationSettings& InSettings)
 {
-	Super::Load(Ar, SerializationPhase);
+	Super::Load(Ar, InSettings);
 
-	if(SerializationPhase == ESerializationPhase::StaticData)
+	if(InSettings.SerializationPhase == FRigHierarchySerializationSettings::StaticData)
 	{
 		Settings.Load(Ar);
-		GetOffsetTransform().Load(Ar, GetOffsetDirtyState());
-		GetShapeTransform().Load(Ar, GetShapeDirtyState());
+		GetOffsetTransform().Load(Ar, InSettings, GetOffsetDirtyState());
+		GetShapeTransform().Load(Ar, InSettings, GetShapeDirtyState());
 
 		if (Ar.CustomVer(FControlRigObjectVersion::GUID) >= FControlRigObjectVersion::PreferredEulerAnglesForControls)
 		{
@@ -1629,11 +1962,11 @@ void FRigControlElement::UnlinkStorage(FRigReusableElementStorage<FTransform>& I
 
 const FRigBaseElement::EElementIndex FRigCurveElement::ElementTypeIndex = CurveElement;
 
-void FRigCurveElement::Save(FArchive& Ar, ESerializationPhase SerializationPhase)
+void FRigCurveElement::Save(FArchive& Ar, const FRigHierarchySerializationSettings& InSettings)
 {
-	Super::Save(Ar, SerializationPhase);
+	Super::Save(Ar, InSettings);
 
-	if(SerializationPhase == ESerializationPhase::StaticData)
+	if(InSettings.SerializationPhase == FRigHierarchySerializationSettings::StaticData)
 	{
 		float Value = Get();
 		Ar << bIsValueSet;
@@ -1641,11 +1974,11 @@ void FRigCurveElement::Save(FArchive& Ar, ESerializationPhase SerializationPhase
 	}
 }
 
-void FRigCurveElement::Load(FArchive& Ar, ESerializationPhase SerializationPhase)
+void FRigCurveElement::Load(FArchive& Ar, const FRigHierarchySerializationSettings& InSettings)
 {
-	Super::Load(Ar, SerializationPhase);
+	Super::Load(Ar, InSettings);
 
-	if(SerializationPhase == ESerializationPhase::StaticData)
+	if(InSettings.SerializationPhase == FRigHierarchySerializationSettings::StaticData)
 	{
 		if (Ar.CustomVer(FControlRigObjectVersion::GUID) >= FControlRigObjectVersion::CurveElementValueStateFlag)
 		{
@@ -1795,22 +2128,22 @@ void FRigPhysicsSettings::Load(FArchive& Ar)
 
 const FRigBaseElement::EElementIndex FRigPhysicsElement::ElementTypeIndex = PhysicsElement;
 
-void FRigPhysicsElement::Save(FArchive& Ar, ESerializationPhase SerializationPhase)
+void FRigPhysicsElement::Save(FArchive& Ar, const FRigHierarchySerializationSettings& InSettings)
 {
-	Super::Save(Ar, SerializationPhase);
+	Super::Save(Ar, InSettings);
 
-	if(SerializationPhase == ESerializationPhase::StaticData)
+	if(InSettings.SerializationPhase == FRigHierarchySerializationSettings::StaticData)
 	{
 		Ar << Solver;
 		Settings.Save(Ar);
 	}
 }
 
-void FRigPhysicsElement::Load(FArchive& Ar, ESerializationPhase SerializationPhase)
+void FRigPhysicsElement::Load(FArchive& Ar, const FRigHierarchySerializationSettings& InSettings)
 {
-	Super::Load(Ar, SerializationPhase);
+	Super::Load(Ar, InSettings);
 
-	if(SerializationPhase == ESerializationPhase::StaticData)
+	if(InSettings.SerializationPhase == FRigHierarchySerializationSettings::StaticData)
 	{
 		Ar << Solver;
 		Settings.Load(Ar);
@@ -1832,14 +2165,14 @@ void FRigPhysicsElement::CopyFrom(const FRigBaseElement* InOther)
 
 const FRigBaseElement::EElementIndex FRigReferenceElement::ElementTypeIndex = ReferenceElement;
 
-void FRigReferenceElement::Save(FArchive& Ar, ESerializationPhase SerializationPhase)
+void FRigReferenceElement::Save(FArchive& Ar, const FRigHierarchySerializationSettings& InSettings)
 {
-	Super::Save(Ar, SerializationPhase);
+	Super::Save(Ar, InSettings);
 }
 
-void FRigReferenceElement::Load(FArchive& Ar, ESerializationPhase SerializationPhase)
+void FRigReferenceElement::Load(FArchive& Ar, const FRigHierarchySerializationSettings& InSettings)
 {
-	Super::Load(Ar, SerializationPhase);
+	Super::Load(Ar, InSettings);
 }
 
 void FRigReferenceElement::CopyFrom(const FRigBaseElement* InOther)
@@ -1982,21 +2315,21 @@ uint32 GetTypeHash(const FRigConnectorSettings& Settings)
 
 const FRigBaseElement::EElementIndex FRigConnectorElement::ElementTypeIndex = ConnectorElement;
 
-void FRigConnectorElement::Save(FArchive& Ar, ESerializationPhase SerializationPhase)
+void FRigConnectorElement::Save(FArchive& Ar, const FRigHierarchySerializationSettings& InSettings)
 {
-	Super::Save(Ar, SerializationPhase);
+	Super::Save(Ar, InSettings);
 
-	if(SerializationPhase == ESerializationPhase::StaticData)
+	if(InSettings.SerializationPhase == FRigHierarchySerializationSettings::StaticData)
 	{
 		Settings.Save(Ar);
 	}
 }
 
-void FRigConnectorElement::Load(FArchive& Ar, ESerializationPhase SerializationPhase)
+void FRigConnectorElement::Load(FArchive& Ar, const FRigHierarchySerializationSettings& InSettings)
 {
-	Super::Load(Ar, SerializationPhase);
+	Super::Load(Ar, InSettings);
 
-	if(SerializationPhase == ESerializationPhase::StaticData)
+	if(InSettings.SerializationPhase == FRigHierarchySerializationSettings::StaticData)
 	{
 		Settings.Load(Ar);
 	}
@@ -2036,14 +2369,14 @@ const FName FRigSocketElement::DescriptionMetaName = TEXT("SocketDescription");
 const FName FRigSocketElement::DesiredParentMetaName = TEXT("SocketDesiredParent");
 const FLinearColor FRigSocketElement::SocketDefaultColor = FLinearColor::White;
 
-void FRigSocketElement::Save(FArchive& Ar, ESerializationPhase SerializationPhase)
+void FRigSocketElement::Save(FArchive& Ar, const FRigHierarchySerializationSettings& InSettings)
 {
-	Super::Save(Ar, SerializationPhase);
+	Super::Save(Ar, InSettings);
 }
 
-void FRigSocketElement::Load(FArchive& Ar, ESerializationPhase SerializationPhase)
+void FRigSocketElement::Load(FArchive& Ar, const FRigHierarchySerializationSettings& InSettings)
 {
-	Super::Load(Ar, SerializationPhase);
+	Super::Load(Ar, InSettings);
 }
 
 FRigSocketState FRigSocketElement::GetSocketState(const URigHierarchy* InHierarchy) const
diff --git a/Engine/Plugins/Animation/ControlRig/Source/ControlRig/Public/ControlRigObjectVersion.h b/Engine/Plugins/Animation/ControlRig/Source/ControlRig/Public/ControlRigObjectVersion.h
--- a/Engine/Plugins/Animation/ControlRig/Source/ControlRig/Public/ControlRigObjectVersion.h
+++ b/Engine/Plugins/Animation/ControlRig/Source/ControlRig/Public/ControlRigObjectVersion.h
@@ -113,6 +113,9 @@ struct CONTROLRIG_API FControlRigObjectVersion
 		// Added the notion of components to the rig hierarchy
 		RigHierarchyStoresComponents,
 
+		// Improve transform compactness when serializing the hierarchy
+		RigHierarchyCompactTransformSerialization,
+
 		// -----<new versions can be added above this line>-------------------------------------------------
 		VersionPlusOne,
 		LatestVersion = VersionPlusOne - 1
diff --git a/Engine/Plugins/Animation/ControlRig/Source/ControlRig/Public/Rigs/RigHierarchyDefines.h b/Engine/Plugins/Animation/ControlRig/Source/ControlRig/Public/Rigs/RigHierarchyDefines.h
--- a/Engine/Plugins/Animation/ControlRig/Source/ControlRig/Public/Rigs/RigHierarchyDefines.h
+++ b/Engine/Plugins/Animation/ControlRig/Source/ControlRig/Public/Rigs/RigHierarchyDefines.h
@@ -5,6 +5,7 @@
 #include "CoreMinimal.h"
 #include "TransformNoScale.h"
 #include "EulerTransform.h"
+#include "ControlRigObjectVersion.h"
 #include "RigHierarchyDefines.generated.h"
 
 class URigHierarchy;
@@ -148,6 +149,30 @@ enum class ERigEvent : uint8
 	Max UMETA(Hidden),
 };
 
+struct CONTROLRIG_API FRigHierarchySerializationSettings
+{
+	FRigHierarchySerializationSettings(const FArchive& InArchive);
+	
+	enum ESerializationPhase
+	{
+		StaticData,
+		InterElementData
+	};
+
+	void Save(FArchive& InArchive);
+	void Load(FArchive& InArchive);
+
+	FControlRigObjectVersion::Type ControlRigVersion = FControlRigObjectVersion::LatestVersion;
+	bool bIsSerializingToPackage = false;
+	bool bUseCompressedArchive = false;
+	bool bStoreCompactTransforms = true;
+	bool bSerializeLocalTransform = true;
+	bool bSerializeGlobalTransform = true;
+	bool bSerializeInitialTransform = true;
+	bool bSerializeCurrentTransform = true;
+	ESerializationPhase SerializationPhase = ESerializationPhase::StaticData;
+};
+
 /** When setting control values what to do with regards to setting key.*/
 UENUM()
 enum class EControlRigSetKey : uint8
diff --git a/Engine/Plugins/Animation/ControlRig/Source/ControlRig/Public/Rigs/RigHierarchyElements.h b/Engine/Plugins/Animation/ControlRig/Source/ControlRig/Public/Rigs/RigHierarchyElements.h
--- a/Engine/Plugins/Animation/ControlRig/Source/ControlRig/Public/Rigs/RigHierarchyElements.h
+++ b/Engine/Plugins/Animation/ControlRig/Source/ControlRig/Public/Rigs/RigHierarchyElements.h
@@ -259,6 +259,39 @@ namespace ERigTransformStorageType
 	};
 }
 
+class CONTROLRIG_API FRigCompactTransform
+{
+public:
+	FRigCompactTransform(FTransform& InTransform);
+
+	enum ERepresentation : uint8
+	{
+		Float_Zero_Identity_One = 0,
+		Float_Zero_Identity_Uniform = 1,
+		Float_Zero_Identity_NonUniform = 2,
+		Float_Zero_Quat_One = 3,
+		Float_Zero_Quat_Uniform = 4,
+		Float_Zero_Quat_NonUniform = 5,
+		Float_Position_Identity_One = 6,
+		Float_Position_Identity_Uniform = 7,
+		Float_Position_Identity_NonUniform = 8,
+		Float_Position_Quat_One = 9,
+		Float_Position_Quat_Uniform = 10,
+		Float_Position_Quat_NonUniform = 11,
+		Double_Complete = 12,
+		Last = Double_Complete,
+		Max = Last + 1 
+	};
+
+	void Serialize(FArchive& Ar, const FRigHierarchySerializationSettings& InSettings, ERepresentation* OutRepresentation = nullptr);
+	void Save(FArchive& Ar, const FRigHierarchySerializationSettings& InSettings, ERepresentation* OutRepresentation = nullptr);
+	void Load(FArchive& Ar, const FRigHierarchySerializationSettings& InSettings, ERepresentation* OutRepresentation = nullptr);
+
+private:
+
+	FTransform& Transform;
+};
+
 USTRUCT(BlueprintType)
 struct CONTROLRIG_API FRigTransformDirtyState
 {
@@ -419,8 +452,8 @@ public:
 	, Storage(nullptr)
 	{}
 
-	void Save(FArchive& Ar, const FRigTransformDirtyState& InDirtyState);
-	void Load(FArchive& Ar, FRigTransformDirtyState& InDirtyState);
+	void Save(FArchive& Ar, const FRigHierarchySerializationSettings& InSettings, const FRigTransformDirtyState& InDirtyState);
+	void Load(FArchive& Ar, const FRigHierarchySerializationSettings& InSettings, FRigTransformDirtyState& InDirtyState);
 	
 	const FTransform& Get() const;
 
@@ -485,8 +518,8 @@ public:
     , Global()
 	{}
 
-	void Save(FArchive& Ar, const FRigLocalAndGlobalDirtyState& InDirtyState);
-	void Load(FArchive& Ar, FRigLocalAndGlobalDirtyState& OutDirtyState);
+	void Save(FArchive& Ar, const FRigHierarchySerializationSettings& InSettings, const FRigLocalAndGlobalDirtyState& InDirtyState);
+	void Load(FArchive& Ar, const FRigHierarchySerializationSettings& InSettings, FRigLocalAndGlobalDirtyState& OutDirtyState);
 
 	FRigLocalAndGlobalTransform& operator =(const FRigLocalAndGlobalTransform& InOther);
 
@@ -574,8 +607,8 @@ public:
 		operator[](InTransformType).Set(InTransform);
 	}
 
-	void Save(FArchive& Ar, const FRigCurrentAndInitialDirtyState& InDirtyState);
-	void Load(FArchive& Ar, FRigCurrentAndInitialDirtyState& OutDirtyState);
+	void Save(FArchive& Ar, const FRigHierarchySerializationSettings& InSettings, const FRigCurrentAndInitialDirtyState& InDirtyState);
+	void Load(FArchive& Ar, const FRigHierarchySerializationSettings& InSettings, FRigCurrentAndInitialDirtyState& OutDirtyState);
 
 	bool operator == (const FRigCurrentAndInitialTransform& Other) const
 	{
@@ -816,12 +849,6 @@ public:
 	virtual int32 GetElementTypeIndex() const { return ElementTypeIndex; }
 	static int32 GetElementTypeCount() { return EElementIndex::Max; }
 
-	enum ESerializationPhase
-	{
-		StaticData,
-		InterElementData
-	};
-
 protected:
 	// Only derived types should be able to construct this one.
 	explicit FRigBaseElement(URigHierarchy* InOwner, ERigElementType InElementType)
@@ -871,9 +898,9 @@ protected:
 public:
 
 	virtual UScriptStruct* GetScriptStruct() const { return FRigBaseElement::StaticStruct(); }
-	void Serialize(FArchive& Ar, ESerializationPhase SerializationPhase);
-	virtual void Save(FArchive& Ar, ESerializationPhase SerializationPhase);
-	virtual void Load(FArchive& Ar, ESerializationPhase SerializationPhase);
+	void Serialize(FArchive& Ar, const FRigHierarchySerializationSettings& InSettings);
+	virtual void Save(FArchive& Ar, const FRigHierarchySerializationSettings& InSettings);
+	virtual void Load(FArchive& Ar, const FRigHierarchySerializationSettings& InSettings);
 
 	const FName& GetFName() const { return Key.Name; }
 	const FString& GetName() const
@@ -1027,8 +1054,8 @@ struct CONTROLRIG_API FRigTransformElement : public FRigBaseElement
 	}
 	virtual ~FRigTransformElement() override {}
 
-	virtual void Save(FArchive& A, ESerializationPhase SerializationPhase) override;
-	virtual void Load(FArchive& Ar, ESerializationPhase SerializationPhase) override;
+	virtual void Save(FArchive& A, const FRigHierarchySerializationSettings& InSettings) override;
+	virtual void Load(FArchive& Ar, const FRigHierarchySerializationSettings& InSettings) override;
 
 	virtual void CopyPose(FRigBaseElement* InOther, bool bCurrent, bool bInitial, bool bWeights) override;
 
@@ -1104,8 +1131,8 @@ struct CONTROLRIG_API FRigSingleParentElement : public FRigTransformElement
 	FRigSingleParentElement() = default;
 	virtual ~FRigSingleParentElement() override {}
 
-	virtual void Save(FArchive& A, ESerializationPhase SerializationPhase) override;
-	virtual void Load(FArchive& Ar, ESerializationPhase SerializationPhase) override;
+	virtual void Save(FArchive& A, const FRigHierarchySerializationSettings& InSettings) override;
+	virtual void Load(FArchive& Ar, const FRigHierarchySerializationSettings& InSettings) override;
 
 	FRigTransformElement* ParentElement = nullptr;
 
@@ -1246,8 +1273,8 @@ struct CONTROLRIG_API FRigMultiParentElement : public FRigTransformElement
 	FRigMultiParentElement() = default;	
 	virtual ~FRigMultiParentElement() override {}
 
-	virtual void Save(FArchive& A, ESerializationPhase SerializationPhase) override;
-	virtual void Load(FArchive& Ar, ESerializationPhase SerializationPhase) override;
+	virtual void Save(FArchive& A, const FRigHierarchySerializationSettings& InSettings) override;
+	virtual void Load(FArchive& Ar, const FRigHierarchySerializationSettings& InSettings) override;
 
 	virtual void CopyPose(FRigBaseElement* InOther, bool bCurrent, bool bInitial, bool bWeights) override;
 	
@@ -1296,8 +1323,8 @@ public:
 	
 	virtual ~FRigBoneElement() override {}
 
-	virtual void Save(FArchive& A, ESerializationPhase SerializationPhase) override;
-	virtual void Load(FArchive& Ar, ESerializationPhase SerializationPhase) override;
+	virtual void Save(FArchive& A, const FRigHierarchySerializationSettings& InSettings) override;
+	virtual void Load(FArchive& Ar, const FRigHierarchySerializationSettings& InSettings) override;
 
 private:
 	explicit FRigBoneElement(URigHierarchy* InOwner)
@@ -1698,8 +1725,8 @@ struct CONTROLRIG_API FRigControlElement final : public FRigMultiParentElement
 		return true;
 	}
 
-	virtual void Save(FArchive& A, ESerializationPhase SerializationPhase) override;
-	virtual void Load(FArchive& Ar, ESerializationPhase SerializationPhase) override;
+	virtual void Save(FArchive& A, const FRigHierarchySerializationSettings& InSettings) override;
+	virtual void Load(FArchive& Ar, const FRigHierarchySerializationSettings& InSettings) override;
 
 	virtual void CopyPose(FRigBaseElement* InOther, bool bCurrent, bool bInitial, bool bWeights) override;
 
@@ -1745,8 +1772,8 @@ struct CONTROLRIG_API FRigCurveElement final : public FRigBaseElement
 
 	virtual int32 GetNumCurves() const override { return 1; }
 
-	virtual void Save(FArchive& A, ESerializationPhase SerializationPhase) override;
-	virtual void Load(FArchive& Ar, ESerializationPhase SerializationPhase) override;
+	virtual void Save(FArchive& A, const FRigHierarchySerializationSettings& InSettings) override;
+	virtual void Load(FArchive& Ar, const FRigHierarchySerializationSettings& InSettings) override;
 
 	virtual void CopyPose(FRigBaseElement* InOther, bool bCurrent, bool bInitial, bool bWeights) override;
 
@@ -1935,8 +1962,8 @@ public:
 	
 	virtual ~FRigPhysicsElement() override {}
 
-	virtual void Save(FArchive& A, ESerializationPhase SerializationPhase) override;
-	virtual void Load(FArchive& Ar, ESerializationPhase SerializationPhase) override;
+	virtual void Save(FArchive& A, const FRigHierarchySerializationSettings& InSettings) override;
+	virtual void Load(FArchive& Ar, const FRigHierarchySerializationSettings& InSettings) override;
 
 private:
 	explicit FRigPhysicsElement(URigHierarchy* InOwner)
@@ -1965,8 +1992,8 @@ public:
 	
 	virtual ~FRigReferenceElement() override {}
 
-	virtual void Save(FArchive& A, ESerializationPhase SerializationPhase) override;
-	virtual void Load(FArchive& Ar, ESerializationPhase SerializationPhase) override;
+	virtual void Save(FArchive& A, const FRigHierarchySerializationSettings& InSettings) override;
+	virtual void Load(FArchive& Ar, const FRigHierarchySerializationSettings& InSettings) override;
 
 	FTransform GetReferenceWorldTransform(const FRigVMExecuteContext* InContext, bool bInitial) const;
 
@@ -2081,8 +2108,8 @@ struct CONTROLRIG_API FRigConnectorElement final : public FRigBaseElement
 
 	virtual ~FRigConnectorElement() override {}
 	
-	virtual void Save(FArchive& A, ESerializationPhase SerializationPhase) override;
-	virtual void Load(FArchive& Ar, ESerializationPhase SerializationPhase) override;
+	virtual void Save(FArchive& A, const FRigHierarchySerializationSettings& InSettings) override;
+	virtual void Load(FArchive& Ar, const FRigHierarchySerializationSettings& InSettings) override;
 
 	FRigConnectorState GetConnectorState(const URigHierarchy* InHierarchy) const;
 
@@ -2142,8 +2169,8 @@ struct CONTROLRIG_API FRigSocketElement final : public FRigSingleParentElement
 			
 	virtual ~FRigSocketElement() override {}
 	
-	virtual void Save(FArchive& A, ESerializationPhase SerializationPhase) override;
-	virtual void Load(FArchive& Ar, ESerializationPhase SerializationPhase) override;
+	virtual void Save(FArchive& A, const FRigHierarchySerializationSettings& InSettings) override;
+	virtual void Load(FArchive& Ar, const FRigHierarchySerializationSettings& InSettings) override;
 
 	FRigSocketState GetSocketState(const URigHierarchy* InHierarchy) const;
 
diff --git a/Engine/Plugins/Runtime/RigVM/Source/RigVM/Private/RigVMCore/RigVMDefines.cpp b/Engine/Plugins/Runtime/RigVM/Source/RigVM/Private/RigVMCore/RigVMDefines.cpp
--- a/Engine/Plugins/Runtime/RigVM/Source/RigVM/Private/RigVMCore/RigVMDefines.cpp
+++ b/Engine/Plugins/Runtime/RigVM/Source/RigVM/Private/RigVMCore/RigVMDefines.cpp
@@ -4,6 +4,8 @@
 
 #if UE_RIGVM_ARCHIVETRACE_ENABLE
 
+#include "RigVMModule.h"
+
 TMap<FArchive*, TSharedPtr<FRigVMArchiveTrace>> FRigVMArchiveTrace::ActiveTraces;
 FCriticalSection FRigVMArchiveTrace::AddRemoveTracesMutex;
 
