Control Rig: Clear hierarchy and curve import references on rig modules

#rnx

diff --git a/Engine/Plugins/Animation/ControlRig/Source/ControlRigDeveloper/Private/ControlRigBlueprint.cpp b/Engine/Plugins/Animation/ControlRig/Source/ControlRigDeveloper/Private/ControlRigBlueprint.cpp
--- a/Engine/Plugins/Animation/ControlRig/Source/ControlRigDeveloper/Private/ControlRigBlueprint.cpp
+++ b/Engine/Plugins/Animation/ControlRig/Source/ControlRigDeveloper/Private/ControlRigBlueprint.cpp
@@ -842,6 +842,12 @@ void UControlRigBlueprint::SetPreviewMesh(USkeletalMesh* PreviewMesh, bool bMark
 	}
 
 	PreviewSkeletalMesh = PreviewMesh;
+
+	if(IsControlRigModule())
+	{
+		SourceHierarchyImport.Reset();
+		SourceCurveImport.Reset();
+	}
 #endif
 }
 
@@ -941,6 +947,9 @@ void UControlRigBlueprint::PreSave(FObjectPreSaveContext ObjectSaveContext)
 		TGuardValue<bool> SuspendNotifGuardOnDebuggedHierarchy(DebuggedHierarchy->GetSuspendNotificationsFlag(), true);
 
 		UpdateExposedModuleConnectors();
+
+		SourceHierarchyImport.Reset();
+		SourceCurveImport.Reset();
 	}
 
 	if (IsControlRigModule())
diff --git a/Engine/Plugins/Animation/ControlRig/Source/ControlRigDeveloper/Public/ControlRigBlueprint.h b/Engine/Plugins/Animation/ControlRig/Source/ControlRigDeveloper/Public/ControlRigBlueprint.h
--- a/Engine/Plugins/Animation/ControlRig/Source/ControlRigDeveloper/Public/ControlRigBlueprint.h
+++ b/Engine/Plugins/Animation/ControlRig/Source/ControlRigDeveloper/Public/ControlRigBlueprint.h
@@ -314,7 +314,7 @@ private:
 
 	/** The default skeletal mesh to use when previewing this asset */
 #if WITH_EDITORONLY_DATA
-	UPROPERTY(AssetRegistrySearchable)
+	UPROPERTY(AssetRegistrySearchable, EditAnywhere, Category="Control Rig Blueprint")
 	TSoftObjectPtr<USkeletalMesh> PreviewSkeletalMesh;
 #endif
 
