Control Rig: Pose Adapter to only re-link if it is out of date

#rnx

diff --git a/Engine/Plugins/Animation/ControlRig/Source/ControlRig/Private/Rigs/RigHierarchy.cpp b/Engine/Plugins/Animation/ControlRig/Source/ControlRig/Private/Rigs/RigHierarchy.cpp
--- a/Engine/Plugins/Animation/ControlRig/Source/ControlRig/Private/Rigs/RigHierarchy.cpp
+++ b/Engine/Plugins/Animation/ControlRig/Source/ControlRig/Private/Rigs/RigHierarchy.cpp
@@ -3830,6 +3830,8 @@ void URigHierarchy::LinkPoseAdapter(TSharedPtr<FRigHierarchyPoseAdapter> InPoseA
 	{
 		PoseAdapter->PreUnlinked(this);
 		PoseAdapter->bLinked = false;
+		PoseAdapter->WeakHierarchy.Reset();;
+		PoseAdapter->LastTopologyVersion = INDEX_NONE;
 		PoseAdapter.Reset();
 	}
 
@@ -3838,6 +3840,9 @@ void URigHierarchy::LinkPoseAdapter(TSharedPtr<FRigHierarchyPoseAdapter> InPoseA
 		PoseAdapter = InPoseAdapter;
 		PoseAdapter->PostLinked(this);
 		PoseAdapter->bLinked = true;
+		PoseAdapter->WeakHierarchy = this;
+		PoseAdapter->LastTopologyVersion = GetTopologyVersion();
+
 	}
 }
 
diff --git a/Engine/Plugins/Animation/ControlRig/Source/ControlRig/Private/Rigs/RigHierarchyPoseAdapter.cpp b/Engine/Plugins/Animation/ControlRig/Source/ControlRig/Private/Rigs/RigHierarchyPoseAdapter.cpp
--- a/Engine/Plugins/Animation/ControlRig/Source/ControlRig/Private/Rigs/RigHierarchyPoseAdapter.cpp
+++ b/Engine/Plugins/Animation/ControlRig/Source/ControlRig/Private/Rigs/RigHierarchyPoseAdapter.cpp
@@ -41,7 +41,7 @@ bool FRigHierarchyPoseAdapter::IsUpdateToDate(const URigHierarchy* InHierarchy)
 	{
 		return false;
 	}
-	return LastTopologyVersion != InHierarchy->GetTopologyVersion();
+	return LastTopologyVersion == InHierarchy->GetTopologyVersion();
 }
 
 TTuple<FRigComputedTransform*, FRigTransformDirtyState*> FRigHierarchyPoseAdapter::GetElementTransformStorage(
