Control Rig: Fix foot connection disconnecting in leg module when foot module is added
#rnx

diff --git a/Engine/Plugins/Animation/ControlRig/Source/ControlRig/Private/ModularRigController.cpp b/Engine/Plugins/Animation/ControlRig/Source/ControlRig/Private/ModularRigController.cpp
--- a/Engine/Plugins/Animation/ControlRig/Source/ControlRig/Private/ModularRigController.cpp
+++ b/Engine/Plugins/Animation/ControlRig/Source/ControlRig/Private/ModularRigController.cpp
@@ -1346,11 +1346,13 @@ FString UModularRigController::ReparentModule(const FString& InModulePath, const
 	{
 		for (FModularRigSingleConnection& Connection : Model->Connections)
 		{
-			if (Connection.Connector.Name.ToString().StartsWith(OldPath, ESearchCase::IgnoreCase))
+			const TPair<FRigName, FRigName> ConnectorNamespace = URigHierarchy::SplitNameSpace(Connection.Connector.Name);
+			if (ConnectorNamespace.Key.ToString().StartsWith(OldPath, ESearchCase::IgnoreCase))
 			{
 				Connection.Connector.Name = *FString::Printf(TEXT("%s%s"), *NewPath, *Connection.Connector.Name.ToString().RightChop(OldPath.Len()));
 			}
-			if (Connection.Target.Name.ToString().StartsWith(OldPath, ESearchCase::IgnoreCase))
+			const TPair<FRigName, FRigName> TargetNamespace = URigHierarchy::SplitNameSpace(Connection.Target.Name);
+			if (TargetNamespace.Key.ToString().StartsWith(OldPath, ESearchCase::IgnoreCase))
 			{
 				Connection.Target.Name = *FString::Printf(TEXT("%s%s"), *NewPath, *Connection.Target.Name.ToString().RightChop(OldPath.Len()));
 			}
