Control Rig: Avoid accessing deleted rig element
#rnx

diff --git a/Engine/Plugins/Animation/ControlRig/Source/ControlRig/Private/Rigs/RigHierarchyController.cpp b/Engine/Plugins/Animation/ControlRig/Source/ControlRig/Private/Rigs/RigHierarchyController.cpp
--- a/Engine/Plugins/Animation/ControlRig/Source/ControlRig/Private/Rigs/RigHierarchyController.cpp
+++ b/Engine/Plugins/Animation/ControlRig/Source/ControlRig/Private/Rigs/RigHierarchyController.cpp
@@ -2396,6 +2396,7 @@ int32 URigHierarchyController::AddElement(FRigBaseElement* InElementToAdd, FRigB
 	ensure(IsValid());
 
 	URigHierarchy* Hierarchy = GetHierarchy();
+	FScopeLock Lock(&Hierarchy->ElementsLock);
 
 	InElementToAdd->CachedNameString.Reset();
 	InElementToAdd->SubIndex = Hierarchy->Num(InElementToAdd->Key.Type);
@@ -2526,6 +2527,7 @@ bool URigHierarchyController::RemoveElement(FRigBaseElement* InElement)
 	}
 
 	URigHierarchy* Hierarchy = GetHierarchy();
+	FScopeLock Lock(&Hierarchy->ElementsLock);
 
 	// make sure this element is part of this hierarchy
 	ensure(Hierarchy->FindChecked(InElement->Key) == InElement);
