Preventing FAnimNode_ControlRig destructor from resolving a remote object

diff --git a/Engine/Plugins/Animation/ControlRig/Source/ControlRig/Private/AnimNode_ControlRig.cpp b/Engine/Plugins/Animation/ControlRig/Source/ControlRig/Private/AnimNode_ControlRig.cpp
--- a/Engine/Plugins/Animation/ControlRig/Source/ControlRig/Private/AnimNode_ControlRig.cpp
+++ b/Engine/Plugins/Animation/ControlRig/Source/ControlRig/Private/AnimNode_ControlRig.cpp
@@ -36,7 +36,7 @@ FAnimNode_ControlRig::FAnimNode_ControlRig()
 
 FAnimNode_ControlRig::~FAnimNode_ControlRig()
 {
-	if(ControlRig)
+	if(ControlRig && ControlRig.IsResolved())
 	{
 		ControlRig->OnInitialized_AnyThread().RemoveAll(this);
 	}
