Control Rig: Make sure to relink the pose adapter after the hierarchy got reset

#rnx

diff --git a/Engine/Plugins/Animation/ControlRig/Source/ControlRig/Private/AnimNode_ControlRigBase.cpp b/Engine/Plugins/Animation/ControlRig/Source/ControlRig/Private/AnimNode_ControlRigBase.cpp
--- a/Engine/Plugins/Animation/ControlRig/Source/ControlRig/Private/AnimNode_ControlRigBase.cpp
+++ b/Engine/Plugins/Animation/ControlRig/Source/ControlRig/Private/AnimNode_ControlRigBase.cpp
@@ -146,6 +146,15 @@ void FAnimNode_ControlRig_PoseAdapter::PreUnlinked(URigHierarchy* InHierarchy)
 	FRigHierarchyPoseAdapter::PreUnlinked(InHierarchy);
 }
 
+bool FAnimNode_ControlRig_PoseAdapter::IsUpdateToDate(const URigHierarchy* InHierarchy) const
+{
+	if(!FRigHierarchyPoseAdapter::IsUpdateToDate(InHierarchy))
+	{
+		return false;
+	}
+	return !LocalPose.IsEmpty() && !GlobalPose.IsEmpty(); 
+}
+
 void FAnimNode_ControlRig_PoseAdapter::UnlinkTransformStorage()
 {
 	check(Hierarchy);
diff --git a/Engine/Plugins/Animation/ControlRig/Source/ControlRig/Public/AnimNode_ControlRigBase.h b/Engine/Plugins/Animation/ControlRig/Source/ControlRig/Public/AnimNode_ControlRigBase.h
--- a/Engine/Plugins/Animation/ControlRig/Source/ControlRig/Public/AnimNode_ControlRigBase.h
+++ b/Engine/Plugins/Animation/ControlRig/Source/ControlRig/Public/AnimNode_ControlRigBase.h
@@ -70,7 +70,8 @@ protected:
 
 	virtual void PostLinked(URigHierarchy* InHierarchy) override;
 	virtual void PreUnlinked(URigHierarchy* InHierarchy) override;
-
+	virtual bool IsUpdateToDate(const URigHierarchy* InHierarchy) const override;
+	
 	void ConvertToLocalPose();
 	void ConvertToGlobalPose();
 	const FTransform& GetLocalTransform(int32 InIndex);
diff --git a/Engine/Plugins/Animation/ControlRig/Source/ControlRig/Public/Rigs/RigHierarchyPoseAdapter.h b/Engine/Plugins/Animation/ControlRig/Source/ControlRig/Public/Rigs/RigHierarchyPoseAdapter.h
--- a/Engine/Plugins/Animation/ControlRig/Source/ControlRig/Public/Rigs/RigHierarchyPoseAdapter.h
+++ b/Engine/Plugins/Animation/ControlRig/Source/ControlRig/Public/Rigs/RigHierarchyPoseAdapter.h
@@ -24,7 +24,7 @@ protected:
 	virtual void PostUnlinked(URigHierarchy* InHierarchy);
 	bool IsLinked() const { return bLinked; }
 	bool IsLinkedTo(const URigHierarchy* InHierarchy) const;
-	bool IsUpdateToDate(const URigHierarchy* InHierarchy) const;
+	virtual bool IsUpdateToDate(const URigHierarchy* InHierarchy) const;
 
 	TTuple<FRigComputedTransform*, FRigTransformDirtyState*> GetElementTransformStorage(const FRigElementKeyAndIndex& InKeyAndIndex, ERigTransformType::Type InTransformType, ERigTransformStorageType::Type InStorageType) const;
 
