Sequencer: Introduced ECS codepath for Control Rig evaluation.

- This codepath is significantly faster than the legacy template evaluation, and enables us to forgo the finalization phase in cases where Control Rig is the only thing pulling in the legacy template codepath.
- Added the ability to passthrough the result of a filter to an entity system task. This is useful for branching based on the presence of a component or tag within a task but doesn't require the filter to be executed every time the task is run (it is only run once for each allocation when the task graph is initialized).
- Reworked initial value processors to be shareable, allowing for the same initial value processor to be used for multiple different filter matches. This enables control rig to have a single processor that branches on the control type, vastly reducing the amount of boilerplate and code duplication.
- Added ECS components for common parameter types (separate from property types)

diff --git a/Engine/Plugins/Animation/ControlRig/Source/ControlRig/Private/Sequencer/MovieSceneControlRigComponentTypes.cpp b/Engine/Plugins/Animation/ControlRig/Source/ControlRig/Private/Sequencer/MovieSceneControlRigComponentTypes.cpp
new file mode 100644
--- /dev/null
+++ b/Engine/Plugins/Animation/ControlRig/Source/ControlRig/Private/Sequencer/MovieSceneControlRigComponentTypes.cpp
@@ -0,0 +1,66 @@
+// Copyright Epic Games, Inc. All Rights Reserved.
+
+#include "Sequencer/MovieSceneControlRigComponentTypes.h"
+#include "Sequencer/MovieSceneControlRigSpaceChannel.h"
+
+#include "EntitySystem/BuiltInComponentTypes.h"
+#include "EntitySystem/MovieSceneComponentRegistry.h"
+#include "EntitySystem/MovieSceneEntityFactoryTemplates.h"
+#include "EntitySystem/MovieSceneEntitySystemLinker.h"
+
+namespace UE::MovieScene
+{
+
+
+static TUniquePtr<FControlRigComponentTypes> GControlRigComponentTypes;
+
+static bool GControlRigComponentTypesDestroyed = false;
+
+
+FControlRigComponentTypes* FControlRigComponentTypes::Get()
+{
+	if (!GControlRigComponentTypes.IsValid())
+	{
+		check(!GControlRigComponentTypesDestroyed);
+		GControlRigComponentTypes.Reset(new FControlRigComponentTypes);
+	}
+	return GControlRigComponentTypes.Get();
+}
+
+void FControlRigComponentTypes::Destroy()
+{
+	GControlRigComponentTypes.Reset();
+	GControlRigComponentTypesDestroyed = true;
+}
+
+FControlRigComponentTypes::FControlRigComponentTypes()
+{
+	using namespace UE::MovieScene;
+
+	FComponentRegistry* ComponentRegistry = UMovieSceneEntitySystemLinker::GetComponents();
+	FBuiltInComponentTypes* BuiltInComponents = FBuiltInComponentTypes::Get();
+
+	ComponentRegistry->NewComponentType(&ControlRigSource, TEXT("Control Rig Source"), EComponentTypeFlags::CopyToOutput);
+	ComponentRegistry->Factories.DuplicateChildComponent(ControlRigSource);
+	ComponentRegistry->Factories.DefineMutuallyInclusiveComponent(ControlRigSource, BuiltInComponents->EvalSeconds);
+
+	ComponentRegistry->NewComponentType(&BaseControlRigEvalData, TEXT("Base Control Rig Eval Data"));
+
+	ComponentRegistry->NewComponentType(&SpaceChannel, TEXT("Space Channel"));
+	ComponentRegistry->NewComponentType(&SpaceResult, TEXT("Space Result"));
+	ComponentRegistry->Factories.DefineMutuallyInclusiveComponent(SpaceChannel, BuiltInComponents->EvalTime);
+	ComponentRegistry->Factories.DefineMutuallyInclusiveComponent(SpaceChannel, SpaceResult);
+
+	Tags.BaseControlRig = ComponentRegistry->NewTag(TEXT("Base Control Rig"), EComponentTypeFlags::CopyToChildren);
+
+	Tags.Space = ComponentRegistry->NewTag(TEXT("Control Rig Space"), EComponentTypeFlags::CopyToChildren);
+
+	Tags.ControlRigParameter = ComponentRegistry->NewTag(TEXT("Control Rig Parameter"), EComponentTypeFlags::CopyToChildren | EComponentTypeFlags::CopyToOutput);
+
+	ComponentRegistry->NewComponentType(&AccumulatedControlEntryIndex, TEXT("Accumulated Control Entry Index"));
+
+	ComponentRegistry->Factories.DefineMutuallyInclusiveComponent(Tags.BaseControlRig, BaseControlRigEvalData);
+}
+
+
+} // namespace UE::MovieScene
\ No newline at end of file
diff --git a/Engine/Plugins/Animation/ControlRig/Source/ControlRig/Private/Sequencer/MovieSceneControlRigParameterBuffer.cpp b/Engine/Plugins/Animation/ControlRig/Source/ControlRig/Private/Sequencer/MovieSceneControlRigParameterBuffer.cpp
new file mode 100644
--- /dev/null
+++ b/Engine/Plugins/Animation/ControlRig/Source/ControlRig/Private/Sequencer/MovieSceneControlRigParameterBuffer.cpp
@@ -0,0 +1,1064 @@
+// Copyright Epic Games, Inc. All Rights Reserved.
+
+#include "Sequencer/MovieSceneControlRigParameterBuffer.h"
+#include "Sequencer/MovieSceneControlRigParameterSection.h"
+#include "Sequencer/MovieSceneControlRigParameterTrack.h"
+#include "Sequencer/MovieSceneControlRigSpaceChannel.h"
+#include "Sequencer/MovieSceneControlRigSystem.h"
+#include "Evaluation/PreAnimatedState/MovieScenePreAnimatedStateTypes.h"
+
+#include "Algo/IndexOf.h"
+
+
+namespace UE::MovieScene
+{
+
+/** This function is a temporary stand-in in the absence of a proper ZipSort algorithm in the core library,
+ *     which would remove the need to copy everything into a temporary array */
+template<typename T>
+void SortParametersByName(TArrayView<FName> Names, TArrayView<T> Parameters)
+{
+	check(Names.Num() == Parameters.Num());
+
+	static constexpr int32 FixedSize = 128;
+	const int32 Num = Names.Num();
+	if (Num <= 1)
+	{
+		return;
+	}
+
+	using PairType = TPair<FName, T>;
+
+	TArray<PairType, TFixedAllocator<FixedSize>> TempFixedSize;
+	TArray<PairType> TempHeap;
+
+	TArrayView<PairType> ZippedArray;
+	if (Num <= FixedSize)
+	{
+		TempFixedSize.SetNumUninitialized(Num);
+		ZippedArray = TempFixedSize;
+	}
+	else
+	{
+		TempHeap.SetNumUninitialized(Num);
+		ZippedArray = TempHeap;
+	}
+
+	for (int32 Index = 0; Index < Num; ++Index)
+	{
+		new (&ZippedArray[Index]) PairType(MakeTuple(MoveTemp(Names[Index]), MoveTemp(Parameters[Index])));
+	}
+
+	// Sort by name
+	Algo::SortBy(ZippedArray, [](const PairType& In) { return In.Key; }, FNameFastLess());
+
+	for (int32 Index = 0; Index < Num; ++Index)
+	{
+		Names[Index] = MoveTemp(ZippedArray[Index].Key);
+		Parameters[Index] = MoveTemp(ZippedArray[Index].Value);
+	}
+}
+
+FControlRigParameterValueHeader::FControlRigParameterValueHeader(EControlRigControlType InType, EControlRigParameterBufferIndexStability IndexStability)
+	: Data(nullptr)
+	, Capacity(0)
+	, NumElements(0)
+	, bStableIndices(IndexStability == EControlRigParameterBufferIndexStability::Stable)
+	, Type(InType)
+{
+	static_assert(std::is_trivially_copyable_v<FMovieSceneControlRigSpaceBaseKey>, "Parameter value types must be trivially copyable to be used inside a parameter buffer");
+	static_assert(std::is_trivially_copyable_v<bool>,                              "Parameter value types must be trivially copyable to be used inside a parameter buffer");
+	static_assert(std::is_trivially_copyable_v<uint8>,                             "Parameter value types must be trivially copyable to be used inside a parameter buffer");
+	static_assert(std::is_trivially_copyable_v<int32>,                             "Parameter value types must be trivially copyable to be used inside a parameter buffer");
+	static_assert(std::is_trivially_copyable_v<float>,                             "Parameter value types must be trivially copyable to be used inside a parameter buffer");
+	static_assert(std::is_trivially_copyable_v<FVector3f>,                         "Parameter value types must be trivially copyable to be used inside a parameter buffer");
+	static_assert(std::is_trivially_copyable_v<FEulerTransform>,                   "Parameter value types must be trivially copyable to be used inside a parameter buffer");
+
+	switch (InType)
+	{
+		case EControlRigControlType::Space:
+			Alignment = alignof(FMovieSceneControlRigSpaceBaseKey);
+			break;
+		case EControlRigControlType::Parameter_Bool:
+			Alignment = alignof(bool);
+			break;
+		case EControlRigControlType::Parameter_Enum:
+			Alignment = alignof(uint8);
+			break;
+		case EControlRigControlType::Parameter_Integer:
+			Alignment = alignof(int32);
+			break;
+		case EControlRigControlType::Parameter_Scalar:
+			Alignment = alignof(float);
+			break;
+		case EControlRigControlType::Parameter_Vector:
+			Alignment = alignof(FVector3f);
+			break;
+		case EControlRigControlType::Parameter_Transform:
+			Alignment = alignof(FEulerTransform);
+			break;
+		default:
+			check(false);
+			break;
+	}
+}
+
+FControlRigParameterValueHeader::FControlRigParameterValueHeader(FControlRigParameterValueHeader&& RHS)
+	: Data(RHS.Data)
+	, Capacity(RHS.Capacity)
+	, NumElements(RHS.NumElements)
+	, Alignment(RHS.Alignment)
+	, bStableIndices(RHS.bStableIndices)
+	, Type(RHS.Type)
+{
+	RHS.Data            = nullptr;
+	RHS.Capacity        = 0;
+	RHS.NumElements     = 0;
+	RHS.Alignment       = 0;
+	RHS.bStableIndices  = 0;
+}
+
+FControlRigParameterValueHeader& FControlRigParameterValueHeader::operator=(FControlRigParameterValueHeader&& RHS)
+{
+	if (Data)
+	{
+		FMemory::Free(Data);
+	}
+
+	Data            = RHS.Data;
+	Capacity        = RHS.Capacity;
+	NumElements     = RHS.NumElements;
+	Alignment       = RHS.Alignment;
+	bStableIndices  = RHS.bStableIndices;
+	Type            = RHS.Type;
+
+	RHS.Data            = nullptr;
+	RHS.Capacity        = 0;
+	RHS.NumElements     = 0;
+	RHS.Alignment       = 0;
+	RHS.bStableIndices  = 0;
+
+	return *this;
+}
+
+FControlRigParameterValueHeader::FControlRigParameterValueHeader(const FControlRigParameterValueHeader& RHS)
+	: Data(nullptr)
+	, Capacity(0)
+	, NumElements(0)
+	, Alignment(RHS.Alignment)
+	, bStableIndices(RHS.bStableIndices)
+	, Type(RHS.Type)
+{
+	Reserve(RHS.Capacity);
+	NumElements = RHS.NumElements;
+
+	if (NumElements > 0)
+	{
+		size_t AlignmentOffset = Alignment != alignof(FName) ? Alignment : 0;
+		size_t SizeInBytes     = Capacity*sizeof(FName) + AlignmentOffset + GetParameterSize()*Capacity;
+
+		FMemory::Memcpy(Data, RHS.Data, SizeInBytes);
+	}
+}
+
+FControlRigParameterValueHeader& FControlRigParameterValueHeader::operator=(const FControlRigParameterValueHeader& RHS)
+{
+	if (Data)
+	{
+		FMemory::Free(Data);
+		Data            = nullptr;
+		Capacity        = 0;
+		NumElements     = 0;
+		Alignment       = 0;
+		bStableIndices  = 0;
+	}
+
+	Alignment       = RHS.Alignment;
+	bStableIndices  = RHS.bStableIndices;
+	Type            = RHS.Type;
+
+	Reserve(RHS.Capacity);
+	NumElements = RHS.NumElements;
+
+	if (NumElements > 0)
+	{
+		size_t AlignmentOffset = Alignment != alignof(FName) ? Alignment : 0;
+		size_t SizeInBytes     = Capacity*sizeof(FName) + AlignmentOffset + GetParameterSize()*Capacity;
+
+		FMemory::Memcpy(Data, RHS.Data, SizeInBytes);
+	}
+	return *this;
+}
+
+
+FControlRigParameterValueHeader::~FControlRigParameterValueHeader()
+{
+	if (Data)
+	{
+		FMemory::Free(Data);
+	}
+}
+
+void FControlRigParameterValueHeader::Reset()
+{
+	NumElements = 0;
+}
+
+EControlRigControlType FControlRigParameterValueHeader::GetType() const
+{
+	return Type;
+}
+
+int32 FControlRigParameterValueHeader::Num() const
+{
+	return NumElements;
+}
+
+
+TArrayView<const FName> FControlRigParameterValueHeader::GetNames() const
+{
+	if (NumElements == 0)
+	{
+		return TArrayView<const FName>();
+	}
+
+	const FName* Names = reinterpret_cast<const FName*>(Data);
+	return TArrayView<const FName>(Names, NumElements);
+}
+
+uint8* FControlRigParameterValueHeader::GetParameterBuffer() const
+{
+	return Data ? GetParameterBuffer(Data, Capacity, Alignment) : nullptr;
+}
+
+void* FControlRigParameterValueHeader::GetParameter(int32 Index) const
+{
+	check(Index < NumElements);
+	return GetParameterBuffer() + Index*GetParameterSize();
+}
+
+int32 FControlRigParameterValueHeader::Add_GetIndex(FName InName)
+{
+	int32 Index = INDEX_NONE;
+	if (bStableIndices)
+	{
+		// If we need stable indices then we should only ever grow
+		Index = NumElements;
+		InsertDefaulted(InName, Index);
+	}
+	else
+	{
+		TArrayView<const FName> Names = GetNames();
+		Index = Algo::LowerBound(Names, InName, FNameFastLess());
+
+		if (Index >= NumElements || Names[Index] != InName)
+		{
+			InsertDefaulted(InName, Index);
+		}
+	}
+
+	return Index;
+}
+
+void* FControlRigParameterValueHeader::Add_GetPtr(FName InName)
+{
+	const int32 Index = Add_GetIndex(InName);
+	return GetParameterBuffer() + GetParameterSize()*Index;
+}
+
+void FControlRigParameterValueHeader::Add(FName InName, const void* Value)
+{
+	void* DestBuffer = Add_GetPtr(InName);
+	FMemory::Memcpy(DestBuffer, Value, GetParameterSize());
+}
+
+void FControlRigParameterValueHeader::Remove(FName InName)
+{
+	const int32 Index = Find(InName);
+	if (Index != INDEX_NONE)
+	{
+		RemoveAt(Index);
+	}
+}
+bool FControlRigParameterValueHeader::Contains(FName InName) const
+{
+	return Find(InName) != INDEX_NONE;
+}
+
+int32 FControlRigParameterValueHeader::Find(FName InName) const
+{
+	TArrayView<const FName> Names = GetNames();
+	if (bStableIndices)
+	{
+		return Algo::IndexOf(Names, InName);
+	}
+	else
+	{
+		return Algo::BinarySearch(Names, InName, FNameFastLess());
+	}
+}
+
+void FControlRigParameterValueHeader::OptimizeForLookup()
+{
+	if (bStableIndices == true)
+	{
+		bStableIndices = false;
+		Resort();
+	}
+}
+
+
+void FControlRigParameterValueHeader::Resort()
+{
+	check(bStableIndices == false);
+
+	FControlRigParameterValueHeader* NonConstThis = const_cast<FControlRigParameterValueHeader*>(this);
+
+	TArrayView<FName> Names = NonConstThis->GetMutableNames();
+	switch (Type)
+	{
+	case EControlRigControlType::Space:
+		{
+			TArrayView<FMovieSceneControlRigSpaceBaseKey> Parameters = NonConstThis->GetParameters<FMovieSceneControlRigSpaceBaseKey>();
+			SortParametersByName(Names, Parameters);
+		}
+		break;
+	case EControlRigControlType::Parameter_Bool:
+		{
+			TArrayView<bool> Parameters = NonConstThis->GetParameters<bool>();
+			SortParametersByName(Names, Parameters);
+		}
+		break;
+	case EControlRigControlType::Parameter_Enum:
+		{
+			TArrayView<uint8> Parameters = NonConstThis->GetParameters<uint8>();
+			SortParametersByName(Names, Parameters);
+		}
+		break;
+	case EControlRigControlType::Parameter_Integer:
+		{
+			TArrayView<int32> Parameters = NonConstThis->GetParameters<int32>();
+			SortParametersByName(Names, Parameters);
+		}
+		break;
+	case EControlRigControlType::Parameter_Scalar:
+		{
+			TArrayView<float> Parameters = NonConstThis->GetParameters<float>();
+			SortParametersByName(Names, Parameters);
+		}
+		break;
+	case EControlRigControlType::Parameter_Vector:
+		{
+			TArrayView<FVector3f> Parameters = NonConstThis->GetParameters<FVector3f>();
+			SortParametersByName(Names, Parameters);
+		}
+		break;
+	case EControlRigControlType::Parameter_Transform:
+		{
+			TArrayView<FEulerTransform> Parameters = NonConstThis->GetParameters<FEulerTransform>();
+			SortParametersByName(Names, Parameters);
+		}
+		break;
+	}
+}
+
+TArrayView<FName> FControlRigParameterValueHeader::GetMutableNames()
+{
+	if (NumElements == 0)
+	{
+		return TArrayView<FName>();
+	}
+
+	FName* Names = reinterpret_cast<FName*>(Data);
+	return TArrayView<FName>(Names, NumElements);
+}
+
+void FControlRigParameterValueHeader::RemoveAt(int32 Index)
+{
+	check(Index >= 0 && Index < NumElements);
+
+	FName* Names  = reinterpret_cast<FName*>(Data);
+	uint8* Values = GetParameterBuffer(Data, Capacity, Alignment);
+
+	// Move the tail first
+	const int32 TailNum = int32(NumElements) - Index - 1;
+	if (TailNum > 0)
+	{
+		size_t ParameterSize = GetParameterSize();
+		FMemory::Memmove(Names + Index, Names+(Index+1), TailNum*sizeof(FName));
+		FMemory::Memmove(Values + Index*ParameterSize, Values + (Index+1)*ParameterSize, TailNum*ParameterSize);
+	}
+
+	--NumElements;
+}
+
+void FControlRigParameterValueHeader::InsertDefaulted(FName Name, int32 Index)
+{
+	Reserve(NumElements+1);
+	check(Capacity >= NumElements+1);
+
+	++NumElements;
+
+	FName* Names  = reinterpret_cast<FName*>(Data);
+	uint8* Values = GetParameterBuffer(Data, Capacity, Alignment);
+
+	// Move the tail first
+	const int32 TailNum = int32(NumElements) - Index - 1;
+	if (TailNum > 0)
+	{
+		size_t ParameterSize = GetParameterSize();
+		FMemory::Memmove(Names + (Index+1), Names + Index, TailNum*sizeof(FName));
+		FMemory::Memmove(Values + (Index+1)*ParameterSize, Values + Index*ParameterSize, TailNum*ParameterSize);
+	}
+
+	// Emplace the new name at the correct location
+	new (Names + Index) FName(MoveTemp(Name));
+
+	switch (Type)
+	{
+	case EControlRigControlType::Space:
+		new (Values + sizeof(FMovieSceneControlRigSpaceBaseKey)*Index) FMovieSceneControlRigSpaceBaseKey();
+		break;
+	case EControlRigControlType::Parameter_Bool:
+		new (Values + sizeof(bool)*Index) bool(false);
+		break;
+	case EControlRigControlType::Parameter_Enum:
+		new (Values + sizeof(uint8)*Index) uint8(0);
+		break;
+	case EControlRigControlType::Parameter_Integer:
+		new (Values + sizeof(int32)*Index) int32(0);
+		break;
+	case EControlRigControlType::Parameter_Scalar:
+		new (Values + sizeof(float)*Index) float(0.f);
+		break;
+	case EControlRigControlType::Parameter_Vector:
+		new (Values + sizeof(FVector3f)*Index) FVector3f(0.f, 0.f, 0.f);
+		break;
+	case EControlRigControlType::Parameter_Transform:
+		new (Values + sizeof(FEulerTransform)*Index) FEulerTransform(FEulerTransform::Identity);
+		break;
+	default:
+		check(false);
+		break;
+	}
+}
+
+uint8* FControlRigParameterValueHeader::GetParameterBuffer(uint8* BasePtr, uint16 InCapacity, uint8 InAlignment)
+{
+	return Align(BasePtr + InCapacity*sizeof(FName), InAlignment);
+}
+
+size_t FControlRigParameterValueHeader::GetParameterSize() const
+{
+	switch (Type)
+	{
+	case EControlRigControlType::Space:
+		return sizeof(FMovieSceneControlRigSpaceBaseKey);
+	case EControlRigControlType::Parameter_Bool:
+		return sizeof(bool);
+	case EControlRigControlType::Parameter_Enum:
+		return sizeof(uint8);
+	case EControlRigControlType::Parameter_Integer:
+		return sizeof(int32);
+	case EControlRigControlType::Parameter_Scalar:
+		return sizeof(float);
+	case EControlRigControlType::Parameter_Vector:
+		return sizeof(FVector3f);
+	case EControlRigControlType::Parameter_Transform:
+		return sizeof(FEulerTransform);
+	default:
+		check(false);
+		return 0;
+	}
+}
+
+void FControlRigParameterValueHeader::Reserve(size_t NewCapacity)
+{
+	const size_t Size = GetParameterSize();
+
+	// Allocate capacity in blocks of 8
+	NewCapacity = Align(NewCapacity, 8);
+
+	size_t AlignmentOffset = Alignment != alignof(FName) ? Alignment : 0;
+
+	uint8* OldData = Data;
+	if (NewCapacity > Capacity)
+	{
+		const int32 RequiredCapacityBytes = NewCapacity*sizeof(FName) + AlignmentOffset + Size*NewCapacity;
+		Data = (uint8*)FMemory::Malloc(RequiredCapacityBytes, alignof(FName));
+
+		if (NumElements > 0)
+		{
+			// Copy the names over
+			FMemory::Memcpy(Data, OldData, sizeof(FName)*NumElements);
+
+			// Copy the values over
+			uint8* OldValuePtr = GetParameterBuffer(OldData, Capacity, Alignment);
+			uint8* NewValuePtr = GetParameterBuffer(Data, NewCapacity, Alignment);
+
+			FMemory::Memcpy(NewValuePtr, OldValuePtr, Size*NumElements);
+			FMemory::Free(OldData);
+		}
+
+		Capacity = NewCapacity;
+	}
+}
+
+
+FControlRigParameterValues::FControlRigParameterValues(EControlRigParameterBufferIndexStability IndexStability)
+	: Headers{
+		FControlRigParameterValueHeader(EControlRigControlType::Space, IndexStability),
+		FControlRigParameterValueHeader(EControlRigControlType::Parameter_Bool, IndexStability),
+		FControlRigParameterValueHeader(EControlRigControlType::Parameter_Enum, IndexStability),
+		FControlRigParameterValueHeader(EControlRigControlType::Parameter_Integer, IndexStability),
+		FControlRigParameterValueHeader(EControlRigControlType::Parameter_Scalar, IndexStability),
+		FControlRigParameterValueHeader(EControlRigControlType::Parameter_Vector, IndexStability),
+		FControlRigParameterValueHeader(EControlRigControlType::Parameter_Transform, IndexStability),
+	}
+{}
+
+FControlRigParameterValueHeader& FControlRigParameterValues::GetHeader(EControlRigControlType InType)
+{
+	return Headers[(uint8)InType];
+}
+
+const FControlRigParameterValueHeader& FControlRigParameterValues::GetHeader(EControlRigControlType InType) const
+{
+	return Headers[(uint8)InType];
+}
+
+FControlRigParameterValues::~FControlRigParameterValues()
+{
+}
+
+void FControlRigParameterValues::OptimizeForLookup()
+{
+	for (FControlRigParameterValueHeader& Header : Headers)
+	{
+		Header.OptimizeForLookup();
+	}
+}
+
+void FControlRigParameterValues::Reset()
+{
+	for (FControlRigParameterValueHeader& Header : Headers)
+	{
+		Header.Reset();
+	}
+}
+
+bool FControlRigParameterValues::InitializeParameters(UControlRig* Rig, FPreAnimatedControlRigParameterStorage& Storage)
+{
+	check(Rig);
+
+	// Make space for the headers
+	int32 TotalNum = 0;
+
+
+	FObjectKey RigKey(Rig);
+	for (FControlRigParameterValueHeader& Header : Headers)
+	{
+		for (FName ParameterName : Header.GetNames())
+		{
+			FPreAnimatedStorageIndex StorageIndex = Storage.FindStorageIndex(MakeTuple(RigKey, ParameterName));
+			if (StorageIndex && !Storage.IsStorageRequirementSatisfied(StorageIndex, EPreAnimatedStorageRequirement::Transient))
+			{
+				Storage.AssignPreAnimatedValue(StorageIndex, EPreAnimatedStorageRequirement::Transient, Storage.Traits.CachePreAnimatedValue(Rig, ParameterName));
+			}
+		}
+
+		TotalNum += Header.Num();
+	}
+	return TotalNum != 0;
+}
+
+void FControlRigParameterValues::ApplyAndRemove(UControlRig* Rig, FName InName)
+{
+	for (FControlRigParameterValueHeader& Header : Headers)
+	{
+		Header.ApplyAndRemove(Rig, InName);
+	}
+}
+
+void FControlRigParameterValues::AddCurrentValue(UControlRig* Rig, FRigControlElement* ControlElement)
+{
+	const FName ControlName = ControlElement->GetFName();
+
+	switch (ControlElement->Settings.ControlType)
+	{
+	case ERigControlType::Bool:
+	{
+		const bool Val = Rig->GetControlValue(ControlElement, ERigControlValueType::Current).Get<bool>();
+		Add(ControlName, Val);
+		break;
+	}
+
+	case ERigControlType::Float:
+	case ERigControlType::ScaleFloat:
+	{
+		const float Val = Rig->GetControlValue(ControlElement, ERigControlValueType::Current).Get<float>();
+		Add(ControlName, Val);
+		break;
+	}
+
+	case ERigControlType::Integer:
+	{
+		const int32 Val = Rig->GetControlValue(ControlElement, ERigControlValueType::Current).Get<int32>();
+		Add(ControlName, Val);
+		break;
+	}
+
+	case ERigControlType::Vector2D:
+	{
+		const FVector3f Val = Rig->GetControlValue(ControlElement, ERigControlValueType::Current).Get<FVector3f>();
+		Add(ControlName, Val);
+		break;
+	}
+
+	case ERigControlType::Position:
+	case ERigControlType::Scale:
+	case ERigControlType::Rotator:
+	{
+		FMovieSceneControlRigSpaceBaseKey SpaceValue;
+
+		//@helge how can we get the correct current space here? this is for restoring it.
+		//for now just using parent space
+		//FRigElementKey DefaultParent = RigHierarchy->GetFirstParent(ControlElement->GetKey());
+		SpaceValue.ControlRigElement = ControlElement->GetKey();
+		SpaceValue.SpaceType = EMovieSceneControlRigSpaceType::Parent;
+		Add(ControlName, SpaceValue);
+
+		FVector3f Val = Rig->GetControlValue(ControlElement, ERigControlValueType::Current).Get<FVector3f>();
+		if (ControlElement->Settings.ControlType == ERigControlType::Rotator)
+		{
+			FVector Vector = Rig->GetControlSpecifiedEulerAngle(ControlElement);
+			Val = FVector3f(Vector.X, Vector.Y, Vector.Z);
+		}
+		Add(ControlName, Val);
+		//mz todo specify rotator special so we can do quat interps
+		break;
+	}
+
+	case ERigControlType::Transform:
+	{
+		FMovieSceneControlRigSpaceBaseKey SpaceValue;
+		//@helge how can we get the correct current space here? this is for restoring it.
+		//for now just using parent space
+		//FRigElementKey DefaultParent = RigHierarchy->GetFirstParent(ControlElement->GetKey());
+		SpaceValue.ControlRigElement = ControlElement->GetKey();
+		SpaceValue.SpaceType = EMovieSceneControlRigSpaceType::Parent;
+		Add(ControlName, SpaceValue);
+
+		const FTransform Val = Rig->GetControlValue(ControlElement, ERigControlValueType::Current).Get<FRigControlValue::FTransform_Float>().ToTransform();
+		FEulerTransform EulerTransform(Val);
+		FVector Vector = Rig->GetControlSpecifiedEulerAngle(ControlElement);
+		EulerTransform.Rotation = FRotator(Vector.Y, Vector.Z, Vector.X);
+		Add(ControlName, EulerTransform);
+		break;
+	}
+	case ERigControlType::TransformNoScale:
+	{
+		const FTransformNoScale NoScale =
+			Rig
+			->GetControlValue(ControlElement, ERigControlValueType::Current).Get<FRigControlValue::FTransformNoScale_Float>().ToTransform();
+		FEulerTransform EulerTransform(NoScale.ToFTransform());
+		FVector Vector = Rig->GetControlSpecifiedEulerAngle(ControlElement);
+		EulerTransform.Rotation = FRotator(Vector.Y, Vector.Z, Vector.X);
+		Add(ControlName, EulerTransform);
+		break;
+	}
+	case ERigControlType::EulerTransform:
+	{
+		FEulerTransform EulerTransform =
+			Rig
+			->GetControlValue(ControlElement, ERigControlValueType::Current).Get<FRigControlValue::FEulerTransform_Float>().ToTransform();
+		FVector Vector = Rig->GetControlSpecifiedEulerAngle(ControlElement);
+		EulerTransform.Rotation = FRotator(Vector.Y, Vector.Z, Vector.X);
+		Add(ControlName, EulerTransform);
+		break;
+	}
+
+	} // end switch()
+}
+
+void FControlRigParameterValues::CopyFrom(const FControlRigParameterValues& Other, FName ControlName)
+{
+	for (int32 Index = 0; Index < UE_ARRAY_COUNT(Headers); ++Index)
+	{
+		const FControlRigParameterValueHeader& OtherHeader = Other.Headers[Index];
+		FControlRigParameterValueHeader&       ThisHeader  = this->Headers[Index];
+
+		const int32 EntryIndex = OtherHeader.Find(ControlName);
+		if (EntryIndex != INDEX_NONE)
+		{
+			ThisHeader.Add(ControlName, OtherHeader.GetParameter(EntryIndex));
+		}
+	}
+}
+
+void FControlRigParameterValues::ApplyTo(UControlRig* Rig) const
+{
+	for (const FControlRigParameterValueHeader& Header : Headers)
+	{
+		Header.Apply(Rig);
+	}
+}
+
+void FControlRigParameterValues::PopulateFrom(UControlRig* Rig)
+{
+	for (FRigControlElement* ControlElement : Rig->AvailableControls())
+	{
+		AddCurrentValue(Rig, ControlElement);
+	}
+
+	for (FControlRigParameterValueHeader& Header : Headers)
+	{
+		Header.Resort();
+	}
+}
+
+
+FControlRigParameterBuffer::FControlRigParameterBuffer(UControlRig* InRig, EControlRigParameterBufferIndexStability IndexStability)
+	: WeakControlRig(InRig)
+	, Values(IndexStability)
+{
+}
+
+void FAccumulatedControlRigValues::InitializeParameters(FPreAnimatedControlRigParameterStorage& Storage)
+{
+	for (FEntry& Entry : ValuesArray)
+	{
+		if (UControlRig* Rig = Entry.WeakControlRig.Get())
+		{
+			Entry.Values.InitializeParameters(Rig, Storage);
+		}
+	}
+}
+
+void FAccumulatedControlRigValues::Compact()
+{
+	for (int32 Index = NextValidIndex; Index < ValuesArray.Num(); ++Index)
+	{
+		ParameterValuesByTrack.Remove(ValuesArray[Index].Track.Get());
+	}
+	ValuesArray.RemoveAtSwap(NextValidIndex, ValuesArray.Num()-NextValidIndex, EAllowShrinking::Yes);
+}
+
+void FAccumulatedControlRigValues::MarkAsActive(int32 Index)
+{
+	FEntry& Entry = ValuesArray[Index];
+	if (Entry.bIsActive)
+	{
+		return;
+	}
+
+	check(Index >= NextValidIndex);
+
+	Entry.bIsActive = true;
+
+	// If it's already at the head just increment
+	if (Index == NextValidIndex)
+	{
+		++NextValidIndex;
+		return;
+	}
+
+	// Move it to the head
+	ParameterValuesByTrack.Add(ValuesArray[NextValidIndex].Track, Index);
+	ParameterValuesByTrack.Add(ValuesArray[Index].Track, NextValidIndex);
+	Swap(ValuesArray[NextValidIndex], ValuesArray[Index]);
+	++NextValidIndex;
+}
+
+int32 FAccumulatedControlRigValues::InitializeRig(UMovieSceneControlRigParameterTrack* Track, UControlRig* Rig)
+{
+	const int32* Existing = ParameterValuesByTrack.Find(Track);
+	if (Existing)
+	{
+		MarkAsActive(*Existing);
+		ValuesArray[*Existing].WeakControlRig = Rig;
+		return *Existing;
+	}
+
+	const int32 Index = ValuesArray.Num();
+	ValuesArray.Emplace(Track, Rig);
+
+	ParameterValuesByTrack.Add(Track, Index);
+	MarkAsActive(Index);
+	return Index;
+}
+
+FAccumulatedControlRigValues::FEntry::FEntry(UMovieSceneControlRigParameterTrack* InTrack, UControlRig* Rig)
+	: FControlRigParameterBuffer(Rig, EControlRigParameterBufferIndexStability::Stable)
+	, Track(InTrack)
+
+{}
+
+UControlRig* FAccumulatedControlRigValues::FindControlRig(FAccumulatedControlEntryIndex Entry) const
+{
+	return Entry.IsValid()
+		? ValuesArray[Entry.EntryIndex].WeakControlRig.Get()
+		: nullptr;
+}
+
+UControlRig* FAccumulatedControlRigValues::FindControlRig(UMovieSceneControlRigParameterTrack* Track) const
+{
+	const int32* EntryIndex = ParameterValuesByTrack.Find(Track);
+
+	return EntryIndex
+		? ValuesArray[*EntryIndex].WeakControlRig.Get()
+		: nullptr;
+}
+
+const FControlRigParameterBuffer* FAccumulatedControlRigValues::FindParameterBuffer(UMovieSceneControlRigParameterTrack* Track) const
+{
+	const int32* EntryIndex = ParameterValuesByTrack.Find(Track);
+
+	return EntryIndex
+		? &ValuesArray[*EntryIndex]
+		: nullptr;
+}
+
+bool FAccumulatedControlRigValues::DoesEntryExistForTrack(UMovieSceneControlRigParameterTrack* Track)
+{
+	return ParameterValuesByTrack.Find(Track) != nullptr;
+}
+
+FAccumulatedControlEntryIndex FAccumulatedControlRigValues::AllocateEntryIndex(UMovieSceneControlRigParameterTrack* Track, FName Name, EControlRigControlType ControlType)
+{
+	const int32* Existing = ParameterValuesByTrack.Find(Track);
+	if (!Existing)
+	{
+		return FAccumulatedControlEntryIndex();
+	}
+
+	FEntry& Entry = ValuesArray[*Existing];
+
+	const int32 AccumulatorIndex = Entry.Values.GetHeader(ControlType).Add_GetIndex(Name);
+	return FAccumulatedControlEntryIndex{ uint16(*Existing), uint16(AccumulatorIndex), ControlType };
+}
+
+void FAccumulatedControlRigValues::PrimeForInstantiation()
+{
+	NextValidIndex = 0;
+	for (FAccumulatedControlRigValues::FEntry& Entry : ValuesArray)
+	{
+		Entry.bIsActive = false;
+		Entry.Values.Reset();
+	}
+}
+
+void FAccumulatedControlRigValues::Apply() const
+{
+	for (const FEntry& Accumulator : ValuesArray)
+	{
+		Accumulator.Apply();
+	}
+}
+
+void* FAccumulatedControlRigValues::GetData(FAccumulatedControlEntryIndex Entry)
+{
+	return ValuesArray[Entry.EntryIndex].Values.GetHeader(Entry.ControlType).GetParameterBuffer();
+}
+
+void FAccumulatedControlRigValues::Store(FAccumulatedControlEntryIndex Entry, FMovieSceneControlRigSpaceBaseKey InValue)
+{
+	FMovieSceneControlRigSpaceBaseKey* Ptr = static_cast<FMovieSceneControlRigSpaceBaseKey*>(GetData(Entry));
+	Ptr[Entry.AccumulatorIndex] = InValue;
+}
+void FAccumulatedControlRigValues::Store(FAccumulatedControlEntryIndex Entry, bool InValue)
+{
+	bool* Ptr = static_cast<bool*>(GetData(Entry));
+	Ptr[Entry.AccumulatorIndex] = InValue;
+}
+void FAccumulatedControlRigValues::Store(FAccumulatedControlEntryIndex Entry, uint8 InValue)
+{
+	uint8* Ptr = static_cast<uint8*>(GetData(Entry));
+	Ptr[Entry.AccumulatorIndex] = InValue;
+}
+void FAccumulatedControlRigValues::Store(FAccumulatedControlEntryIndex Entry, int32 InValue)
+{
+	int32* Ptr = static_cast<int32*>(GetData(Entry));
+	Ptr[Entry.AccumulatorIndex] = InValue;
+}
+void FAccumulatedControlRigValues::Store(FAccumulatedControlEntryIndex Entry, float InValue)
+{
+	float* Ptr = static_cast<float*>(GetData(Entry));
+	Ptr[Entry.AccumulatorIndex] = InValue;
+}
+void FAccumulatedControlRigValues::Store(FAccumulatedControlEntryIndex Entry, FVector3f InValue)
+{
+	FVector3f* Ptr = static_cast<FVector3f*>(GetData(Entry));
+	Ptr[Entry.AccumulatorIndex] = InValue;
+}
+void FAccumulatedControlRigValues::Store(FAccumulatedControlEntryIndex Entry, FEulerTransform InValue)
+{
+	FEulerTransform* Ptr = static_cast<FEulerTransform*>(GetData(Entry));
+	Ptr[Entry.AccumulatorIndex] = InValue;
+}
+
+void FControlRigParameterValueHeader::ApplyAndRemove(UControlRig* Rig, FName InName)
+{
+	const int32 Index = Find(InName);
+	if (Index != INDEX_NONE)
+	{
+		Apply(Rig, InName, Index, GetParameterBuffer());
+		RemoveAt(Index);
+	}
+}
+
+void FControlRigParameterValueHeader::Apply(UControlRig* Rig, FName Name, int32 Index, const void* ParameterBuffer) const
+{
+	constexpr bool bSetupUndo = false;
+	constexpr bool bNotify    = true;
+	FRigControlModifiedContext ModifiedContext(EControlRigSetKey::Never);
+
+	switch (Type)
+	{
+		case EControlRigControlType::Space:
+		{
+			const FMovieSceneControlRigSpaceBaseKey* Ptr = static_cast<const FMovieSceneControlRigSpaceBaseKey*>(ParameterBuffer);
+			URigHierarchy* RigHierarchy = Rig->GetHierarchy();
+			FRigControlElement* RigControl = Rig->FindControl(Name);
+			if (!RigHierarchy || !RigControl)
+			{
+				return;
+			}
+
+			const FRigElementKey ControlKey = RigControl->GetKey();
+
+			FMovieSceneControlRigSpaceBaseKey Value = Ptr[Index];
+			switch (Value.SpaceType)
+			{
+			case EMovieSceneControlRigSpaceType::Parent:
+				Rig->SwitchToParent(ControlKey, RigHierarchy->GetDefaultParent(ControlKey), false, true);
+				break;
+			case EMovieSceneControlRigSpaceType::World:
+				Rig->SwitchToParent(ControlKey, RigHierarchy->GetWorldSpaceReferenceKey(), false, true);
+				break;
+			case EMovieSceneControlRigSpaceType::ControlRig:
+				Rig->SwitchToParent(ControlKey, Value.ControlRigElement, false, true);
+				break;
+			}
+			break;
+		}
+		case EControlRigControlType::Parameter_Bool:
+		{
+			const bool* Ptr = static_cast<const bool*>(ParameterBuffer);
+			Rig->SetControlValue<bool>(Name, Ptr[Index], bNotify, ModifiedContext, bSetupUndo);
+			break;
+		}
+		case EControlRigControlType::Parameter_Enum:
+		{
+			const uint8* Ptr = static_cast<const uint8*>(ParameterBuffer);
+			Rig->SetControlValue<int32>(Name, Ptr[Index], bNotify, ModifiedContext, bSetupUndo);
+			break;
+		}
+		case EControlRigControlType::Parameter_Integer:
+		{
+			const int32* Ptr = static_cast<const int32*>(ParameterBuffer);
+			Rig->SetControlValue<int32>(Name, Ptr[Index], bNotify, ModifiedContext, bSetupUndo);
+			break;
+		}
+		case EControlRigControlType::Parameter_Scalar:
+		{
+			const float* Ptr = static_cast<const float*>(ParameterBuffer);
+			Rig->SetControlValue<float>(Name, Ptr[Index], bNotify, ModifiedContext, bSetupUndo);
+			break;
+		}
+		case EControlRigControlType::Parameter_Vector:
+		{
+			const FVector3f* Ptr = static_cast<const FVector3f*>(ParameterBuffer);
+			Rig->SetControlValue<FVector3f>(Name, Ptr[Index], bNotify, ModifiedContext, bSetupUndo);
+			break;
+		}
+		case EControlRigControlType::Parameter_Transform:
+		{
+			URigHierarchy* RigHierarchy = Rig->GetHierarchy();
+			if (!RigHierarchy)
+			{
+				return;
+			}
+
+			const FEulerTransform* Ptr         = static_cast<const FEulerTransform*>(ParameterBuffer);
+			FEulerTransform        Transform   = Ptr[Index];
+
+			FRigControlElement* ControlElement = Rig->FindControl(Name);
+			if (ControlElement)
+			{
+				switch (ControlElement->Settings.ControlType)
+				{
+				case ERigControlType::Transform:
+				{
+					FVector EulerAngle(Transform.Rotation.Roll, Transform.Rotation.Pitch, Transform.Rotation.Yaw);
+					RigHierarchy->SetControlSpecifiedEulerAngle(ControlElement, EulerAngle);
+					Rig->SetControlValue<FRigControlValue::FTransform_Float>(Name, Transform.ToFTransform(), bNotify, ModifiedContext, bSetupUndo);
+					break;
+				}
+				case ERigControlType::TransformNoScale:
+				{
+					FTransformNoScale NoScale = Transform.ToFTransform();
+					FVector EulerAngle(Transform.Rotation.Roll, Transform.Rotation.Pitch, Transform.Rotation.Yaw);
+					RigHierarchy->SetControlSpecifiedEulerAngle(ControlElement, EulerAngle);
+					Rig->SetControlValue<FRigControlValue::FTransformNoScale_Float>(Name, NoScale, bNotify, ModifiedContext, bSetupUndo);
+					break;
+				}
+				case ERigControlType::EulerTransform:
+				{
+					FVector EulerAngle(Transform.Rotation.Roll, Transform.Rotation.Pitch, Transform.Rotation.Yaw);
+					FQuat Quat = RigHierarchy->GetControlQuaternion(ControlElement, EulerAngle);
+					RigHierarchy->SetControlSpecifiedEulerAngle(ControlElement, EulerAngle);
+					FRotator UERotator(Quat);
+					Transform.Rotation = UERotator;
+					Rig->SetControlValue<FRigControlValue::FEulerTransform_Float>(Name, Transform, bNotify, ModifiedContext, bSetupUndo);
+					RigHierarchy->SetControlSpecifiedEulerAngle(ControlElement, EulerAngle);
+					break;
+				}
+				default:
+					break;
+				}
+			}
+			break;
+		}
+	}
+}
+
+void FControlRigParameterValueHeader::Apply(UControlRig* Rig) const
+{
+	TArrayView<const FName> Names = GetNames();
+	const uint8* ParameterBuffer = GetParameterBuffer();
+
+	for (int32 Index = 0; Index < NumElements; ++Index)
+	{
+		Apply(Rig, Names[Index], Index, ParameterBuffer);
+	}
+}
+
+void FControlRigParameterBuffer::Populate()
+{
+	if (UControlRig* Rig = WeakControlRig.Get())
+	{
+		Values.PopulateFrom(Rig);
+	}
+}
+
+void FControlRigParameterBuffer::Apply() const
+{
+	if (UControlRig* Rig = WeakControlRig.Get())
+	{
+#if STATS || ENABLE_STATNAMEDEVENTS
+	#if STATS
+		TStatId StatID = FDynamicStats::CreateStatId<STAT_GROUP_TO_FStatGroup(STATGROUP_MovieSceneECS)>(Rig->GetName());
+	#else
+		// Just use the base UObject stat ID if we only have named events
+		TStatId StatID = Rig->GetStatID(true /* bForDeferredUse */);
+	#endif
+
+		FScopeCycleCounter Scope(StatID);
+#endif
+
+		Values.ApplyTo(Rig);
+	}
+}
+
+} // namespace UE::MovieScene
diff --git a/Engine/Plugins/Animation/ControlRig/Source/ControlRig/Private/Sequencer/MovieSceneControlRigParameterSection.cpp b/Engine/Plugins/Animation/ControlRig/Source/ControlRig/Private/Sequencer/MovieSceneControlRigParameterSection.cpp
--- a/Engine/Plugins/Animation/ControlRig/Source/ControlRig/Private/Sequencer/MovieSceneControlRigParameterSection.cpp
+++ b/Engine/Plugins/Animation/ControlRig/Source/ControlRig/Private/Sequencer/MovieSceneControlRigParameterSection.cpp
@@ -24,7 +24,11 @@
 #include "TransformableHandle.h"
 #include "AnimationCoreLibrary.h"
 #include "UObject/ObjectSaveContext.h"
-
+#include "Evaluation/MovieSceneEvaluationField.h"
+#include "EntitySystem/BuiltInComponentTypes.h"
+#include "MovieSceneTracksComponentTypes.h"
+#include "Sequencer/MovieSceneControlRigSystem.h"
+#include "ControlRigObjectBinding.h"
 
 #include UE_INLINE_GENERATED_CPP_BY_NAME(MovieSceneControlRigParameterSection)
 
@@ -78,6 +82,7 @@ struct FParameterFloatChannelEditorData
 
 		UMovieSceneTrack* Track = SectionToKey->GetTypedOuter<UMovieSceneTrack>();
 
+		/*
 		if (Track)
 		{
 			FMovieSceneEvaluationTrack EvalTrack = CastChecked<IMovieSceneTrackTemplateProducer>(Track)->GenerateTrackTemplate(Track);
@@ -98,6 +103,7 @@ struct FParameterFloatChannelEditorData
 			}
 			OutValue = Val;
 		}
+		*/
 		OutWeight = MovieSceneHelpers::CalculateWeightForBlending(SectionToKey, KeyTime);
 	}
 
@@ -234,6 +240,7 @@ struct FParameterVectorChannelEditorData
 			return;
 		}
 
+		/*
 		UMovieSceneTrack* Track = SectionToKey->GetTypedOuter<UMovieSceneTrack>();
 
 		if (Track)
@@ -293,40 +300,10 @@ struct FParameterVectorChannelEditorData
 					break;
 				}
 			}
-			break;
-			case 4:
-			{
-				/* No Interrogation for Vector4, todo if we do add later
-				FVector4 Val(0.0f, 0.0f, 0.0f, 0.0f);
-				for (const FVector4InterrogationData& InVector : InterrogationData.Iterate<FVector4InterrogationData>(UMovieSceneControlRigParameterSection::GetVector4InterrogationKey()))
-				{
-					if (InVector.ParameterName == ParameterName)
-					{
-						Val = InVector.Val;
-						break;
-					}
-				}
-				switch (Index)
-				{
-				case 0:
-					OutValue = Val.X;
-					break;
-				case 1:
-					OutValue = Val.Y;
-					break;
-				case 2:
-					OutValue = Val.Z;
-					break;
-				case 3:
-					OutValue = Val.W;
-					break;
-				}
-				*/
-			}
-			
 			break;
 			}
 		}
+		*/
 		OutWeight = MovieSceneHelpers::CalculateWeightForBlending(SectionToKey, KeyTime);
 	}
 	FText							GroupName;
@@ -735,6 +712,7 @@ struct FParameterTransformChannelEditorData
 		float& OutValue, float& OutWeight)
 	{
 		UMovieSceneTrack* Track = SectionToKey->GetTypedOuter<UMovieSceneTrack>();
+		/*
 		FMovieSceneEvaluationTrack EvalTrack = CastChecked<UMovieSceneControlRigParameterTrack>(Track)->GenerateTrackTemplate(Track);
 		FMovieSceneInterrogationData InterrogationData;
 		RootTemplate.CopyActuators(InterrogationData.GetAccumulator());
@@ -788,6 +766,7 @@ struct FParameterTransformChannelEditorData
 			break;
 
 		}
+		*/
 		OutWeight = MovieSceneHelpers::CalculateWeightForBlending(SectionToKey, KeyTime);
 	}
 		
@@ -802,6 +781,7 @@ public:
 
 #endif // WITH_EDITOR
 
+
 UMovieSceneControlRigParameterSection::UMovieSceneControlRigParameterSection() :bDoNotKey(false)
 {
 	// Section template relies on always restoring state for objects when they are no longer animating. This is how it releases animation control.
@@ -1086,6 +1066,7 @@ void UMovieSceneControlRigParameterSection::SetBlendType(EMovieSceneBlendType In
 							FloatChannels[ChannelIndex + 8]->SetDefault(CurrentVector.Z);
 						}
 					}
+					
 				}
 				break;
 				}
@@ -3913,6 +3894,386 @@ FIntegerParameterNameAndCurve::FIntegerParameterNameAndCurve(FName InParameterNa
 	ParameterName = InParameterName;
 }
 
+namespace UE::MovieScene
+{
+
+	enum class EControlRigEntityType : uint8
+	{
+		Base,
+		Space,
+		BoolParameter,
+		EnumParameter,
+		IntegerParameter,
+		ScalarParameter,
+		Vector2Parameter,
+		Vector3Parameter,
+		TransformParameter,
+	};
+	/* Entity IDs are an encoded type and index, with the upper 8 bits being the type, and the lower 24 bits as the index */
+	uint32 EncodeControlRigEntityID(int32 InIndex, EControlRigEntityType InType)
+	{
+		check(InIndex >= 0 && InIndex < int32(0x00FFFFFF));
+		return static_cast<uint32>(InIndex) | (uint32(InType) << 24);
+	}
+	void DecodeControlRigEntityID(uint32 InEntityID, int32& OutIndex, EControlRigEntityType& OutType)
+	{
+		// Mask out the type to get the index
+		OutIndex = static_cast<int32>(InEntityID & 0x00FFFFFF);
+		OutType = (EControlRigEntityType)(InEntityID >> 24);
+	}
+
+
+}// namespace UE::MovieScene
+
+void UMovieSceneControlRigParameterSection::ImportEntityImpl(UMovieSceneEntitySystemLinker* EntityLinker, const FEntityImportParams& Params, FImportedEntity* OutImportedEntity)
+{
+	using namespace UE::MovieScene;
+
+	if (UMovieSceneControlRigParameterTrack::ShouldUseLegacyTemplate())
+	{
+		return;
+	}
+
+	const bool bExternalBlending = (BlendType.Get() == EMovieSceneBlendType::Absolute);
+
+	EControlRigEntityType EntityType;
+	int32 EntityIndex = 0;
+	DecodeControlRigEntityID(Params.EntityID, EntityIndex, EntityType);
+
+	FBuiltInComponentTypes* BuiltInComponentTypes = FBuiltInComponentTypes::Get();
+	FMovieSceneTracksComponentTypes* TracksComponentTypes = FMovieSceneTracksComponentTypes::Get();
+	FControlRigComponentTypes* ControlRigComponents = FControlRigComponentTypes::Get();
+
+	FGuid ObjectBindingID = Params.GetObjectBindingID();
+	FControlRigSourceData ControlRigSource;
+	ControlRigSource.Track = GetTypedOuter<UMovieSceneControlRigParameterTrack>();
+
+	if (!ensure(ControlRigSource.Track))
+	{
+		return;
+	}
+
+	EMovieSceneTransformChannel ChannelMask = TransformMask.GetChannels();
+
+	switch (EntityType)
+	{
+		case EControlRigEntityType::Base:
+		{
+			auto NeverResolve = [](UObject*)->UObject*{ return nullptr; };
+
+			OutImportedEntity->AddBuilder(
+				FEntityBuilder()
+				.Add(BuiltInComponentTypes->GenericObjectBinding, ObjectBindingID)
+				.Add(BuiltInComponentTypes->BoundObjectResolver, NeverResolve)
+				.Add(ControlRigComponents->ControlRigSource, ControlRigSource)
+				.Add(ControlRigComponents->BaseControlRigEvalData, FBaseControlRigEvalData{ this })
+				.AddDefaulted(BuiltInComponentTypes->EvalTime)
+				.AddDefaulted(BuiltInComponentTypes->EvalSeconds)
+				.AddTag(ControlRigComponents->Tags.BaseControlRig)
+				.AddMutualComponents()
+			);
+			break;
+		}
+
+		case EControlRigEntityType::Space:
+		{
+			const FSpaceControlNameAndChannel& Space = SpaceChannels[EntityIndex];
+			OutImportedEntity->AddBuilder(
+				FEntityBuilder()
+				.Add(ControlRigComponents->ControlRigSource, ControlRigSource)
+				.Add(TracksComponentTypes->GenericParameterName, Space.ControlName)
+				.Add(ControlRigComponents->SpaceChannel, &Space.SpaceCurve)
+				.AddTag(ControlRigComponents->Tags.ControlRigParameter)
+				.AddTag(ControlRigComponents->Tags.Space)
+				.AddMutualComponents()
+			);
+			break;
+		}
+
+		case EControlRigEntityType::BoolParameter:
+		{
+			const FBoolParameterNameAndCurve& Bool = BoolParameterNamesAndCurves[EntityIndex];
+			OutImportedEntity->AddBuilder(
+				FEntityBuilder()
+				.Add(ControlRigComponents->ControlRigSource, ControlRigSource)
+				.Add(TracksComponentTypes->GenericParameterName, Bool.ParameterName)
+				.Add(BuiltInComponentTypes->BoolChannel, &Bool.ParameterCurve)
+				.AddTag(ControlRigComponents->Tags.ControlRigParameter)
+				.AddTag(TracksComponentTypes->Parameters.Bool.PropertyTag)
+				.AddMutualComponents()
+			);
+			break;
+		}
+
+		case EControlRigEntityType::EnumParameter:
+		{
+			const FEnumParameterNameAndCurve& Enum = EnumParameterNamesAndCurves[EntityIndex];
+			OutImportedEntity->AddBuilder(
+				FEntityBuilder()
+				.Add(ControlRigComponents->ControlRigSource, ControlRigSource)
+				.Add(TracksComponentTypes->GenericParameterName, Enum.ParameterName)
+				.Add(ControlRigComponents->ControlRigSource, ControlRigSource)
+				.AddTag(ControlRigComponents->Tags.ControlRigParameter)
+				.AddTag(TracksComponentTypes->Parameters.Byte.PropertyTag)
+				.Add(BuiltInComponentTypes->ByteChannel, &Enum.ParameterCurve)
+				.AddMutualComponents()
+			);
+			break;
+		}
+
+		case EControlRigEntityType::IntegerParameter:
+		{
+			const FIntegerParameterNameAndCurve& Integer = IntegerParameterNamesAndCurves[EntityIndex];
+			OutImportedEntity->AddBuilder(
+				FEntityBuilder()
+				.Add(ControlRigComponents->ControlRigSource, ControlRigSource)
+				.Add(TracksComponentTypes->GenericParameterName, Integer.ParameterName)
+				.AddTag(ControlRigComponents->Tags.ControlRigParameter)
+				.AddTag(TracksComponentTypes->Parameters.Integer.PropertyTag)
+				.Add(BuiltInComponentTypes->IntegerChannel, &Integer.ParameterCurve)
+				.AddMutualComponents()
+			);
+			break;
+		}
+
+		case EControlRigEntityType::ScalarParameter:
+		{
+			const FScalarParameterNameAndCurve& Scalar = ScalarParameterNamesAndCurves[EntityIndex];
+			OutImportedEntity->AddBuilder(
+				FEntityBuilder()
+				.Add(ControlRigComponents->ControlRigSource, ControlRigSource)
+				.Add(TracksComponentTypes->GenericParameterName, Scalar.ParameterName)
+				.AddTag(ControlRigComponents->Tags.ControlRigParameter)
+				.AddTag(TracksComponentTypes->Parameters.Scalar.PropertyTag)
+				.Add(BuiltInComponentTypes->FloatChannel[0], &Scalar.ParameterCurve)
+				.AddTagConditional(BuiltInComponentTypes->Tags.ExternalBlending, bExternalBlending)
+				.AddConditional(BuiltInComponentTypes->WeightChannel, &Weight, EnumHasAnyFlags(ChannelMask, EMovieSceneTransformChannel::Weight) && Weight.HasAnyData())
+				.AddMutualComponents()
+			);
+			break;
+		}
+
+		case EControlRigEntityType::Vector2Parameter:
+		{
+			const FVector2DParameterNameAndCurves& Vector2D = Vector2DParameterNamesAndCurves[EntityIndex];
+			OutImportedEntity->AddBuilder(
+				FEntityBuilder()
+				.Add(ControlRigComponents->ControlRigSource, ControlRigSource)
+				.Add(TracksComponentTypes->GenericParameterName, Vector2D.ParameterName)
+				.AddTag(ControlRigComponents->Tags.ControlRigParameter)
+				.AddTag(TracksComponentTypes->Parameters.Vector3.PropertyTag)
+				.AddConditional(BuiltInComponentTypes->FloatChannel[0], &Vector2D.XCurve, Vector2D.XCurve.HasAnyData())
+				.AddConditional(BuiltInComponentTypes->FloatChannel[1], &Vector2D.YCurve, Vector2D.YCurve.HasAnyData())
+				.AddConditional(BuiltInComponentTypes->WeightChannel, &Weight, EnumHasAnyFlags(ChannelMask, EMovieSceneTransformChannel::Weight) && Weight.HasAnyData())
+				.AddTagConditional(BuiltInComponentTypes->Tags.ExternalBlending, bExternalBlending)
+				.AddMutualComponents()
+			);
+			break;
+		}
+		
+		case EControlRigEntityType::Vector3Parameter:
+		{
+			const FVectorParameterNameAndCurves& Vector = VectorParameterNamesAndCurves[EntityIndex];
+			OutImportedEntity->AddBuilder(
+				FEntityBuilder()
+				.Add(ControlRigComponents->ControlRigSource, ControlRigSource)
+				.Add(TracksComponentTypes->GenericParameterName, Vector.ParameterName)
+				.AddTag(ControlRigComponents->Tags.ControlRigParameter)
+				.AddTag(TracksComponentTypes->Parameters.Vector3.PropertyTag)
+				.AddConditional(BuiltInComponentTypes->FloatChannel[0], &Vector.XCurve, Vector.XCurve.HasAnyData())
+				.AddConditional(BuiltInComponentTypes->FloatChannel[1], &Vector.YCurve, Vector.YCurve.HasAnyData())
+				.AddConditional(BuiltInComponentTypes->FloatChannel[2], &Vector.ZCurve, Vector.ZCurve.HasAnyData())
+				.AddConditional(BuiltInComponentTypes->WeightChannel, &Weight, EnumHasAnyFlags(ChannelMask, EMovieSceneTransformChannel::Weight) && Weight.HasAnyData())
+				.AddTagConditional(BuiltInComponentTypes->Tags.ExternalBlending, bExternalBlending)
+				.AddMutualComponents()
+			);
+			break;
+		}
+
+		case EControlRigEntityType::TransformParameter:
+		{
+			const FTransformParameterNameAndCurves& Transform = TransformParameterNamesAndCurves[EntityIndex];
+
+			OutImportedEntity->AddBuilder(
+				FEntityBuilder()
+				.Add(ControlRigComponents->ControlRigSource, ControlRigSource)
+				.Add(TracksComponentTypes->GenericParameterName, Transform.ParameterName)
+				.AddTag(ControlRigComponents->Tags.ControlRigParameter)
+				.AddTag(TracksComponentTypes->Parameters.Transform.PropertyTag)
+				.AddConditional(BuiltInComponentTypes->FloatChannel[0], &Transform.Translation[0], EnumHasAllFlags(ChannelMask, EMovieSceneTransformChannel::TranslationX) && Transform.Translation[0].HasAnyData())
+				.AddConditional(BuiltInComponentTypes->FloatChannel[1], &Transform.Translation[1], EnumHasAllFlags(ChannelMask, EMovieSceneTransformChannel::TranslationY) && Transform.Translation[1].HasAnyData())
+				.AddConditional(BuiltInComponentTypes->FloatChannel[2], &Transform.Translation[2], EnumHasAllFlags(ChannelMask, EMovieSceneTransformChannel::TranslationZ) && Transform.Translation[2].HasAnyData())
+				.AddConditional(BuiltInComponentTypes->FloatChannel[3], &Transform.Rotation[0],    EnumHasAllFlags(ChannelMask, EMovieSceneTransformChannel::RotationX) && Transform.Rotation[0].HasAnyData())
+				.AddConditional(BuiltInComponentTypes->FloatChannel[4], &Transform.Rotation[1],    EnumHasAllFlags(ChannelMask, EMovieSceneTransformChannel::RotationY) && Transform.Rotation[1].HasAnyData())
+				.AddConditional(BuiltInComponentTypes->FloatChannel[5], &Transform.Rotation[2],    EnumHasAllFlags(ChannelMask, EMovieSceneTransformChannel::RotationZ) && Transform.Rotation[2].HasAnyData())
+				.AddConditional(BuiltInComponentTypes->FloatChannel[6], &Transform.Scale[0],       EnumHasAllFlags(ChannelMask, EMovieSceneTransformChannel::ScaleX) && Transform.Scale[0].HasAnyData())
+				.AddConditional(BuiltInComponentTypes->FloatChannel[7], &Transform.Scale[1],       EnumHasAllFlags(ChannelMask, EMovieSceneTransformChannel::ScaleY) && Transform.Scale[1].HasAnyData())
+				.AddConditional(BuiltInComponentTypes->FloatChannel[8], &Transform.Scale[2],       EnumHasAllFlags(ChannelMask, EMovieSceneTransformChannel::ScaleZ) && Transform.Scale[2].HasAnyData())
+				.AddConditional(BuiltInComponentTypes->WeightChannel,   &Weight,				   EnumHasAnyFlags(ChannelMask, EMovieSceneTransformChannel::Weight) && Weight.HasAnyData())
+
+				.AddTagConditional(BuiltInComponentTypes->Tags.ExternalBlending, bExternalBlending)
+				.AddMutualComponents()
+			);
+			break;
+		}
+	}
+}
+
+
+bool UMovieSceneControlRigParameterSection::PopulateEvaluationFieldImpl(const TRange<FFrameNumber>& EffectiveRange, const FMovieSceneEvaluationFieldEntityMetaData& InMetaData, FMovieSceneEntityComponentFieldBuilder* OutFieldBuilder)
+{
+	using namespace UE::MovieScene;
+
+	const int32 MetaDataIndex = OutFieldBuilder->AddMetaData(InMetaData);
+
+	// We use the top 8 bits of EntityID to encode the type of parameter
+	const int32 NumSpaceID = SpaceChannels.Num();
+	const int32 NumBoolID = BoolParameterNamesAndCurves.Num();
+	const int32 NumEnumID = EnumParameterNamesAndCurves.Num();
+	const int32 NumIntegerID = IntegerParameterNamesAndCurves.Num();
+	const int32 NumScalarID = ScalarParameterNamesAndCurves.Num();
+	const int32 NumVector2ID = Vector2DParameterNamesAndCurves.Num();
+	const int32 NumVector3ID = VectorParameterNamesAndCurves.Num();
+	const int32 NumTransformID = TransformParameterNamesAndCurves.Num();
+
+	{
+		// In the event there are multiple we will just pick one at runtime, but we need one entity per section for pre-animated state
+		const int32 EntityIndex = OutFieldBuilder->FindOrAddEntity(this, EncodeControlRigEntityID(0, EControlRigEntityType::Base));
+		OutFieldBuilder->AddPersistentEntity(EffectiveRange, EntityIndex, MetaDataIndex);
+	}
+
+	for (int32 Index = 0; Index < NumSpaceID; ++Index)
+	{
+		const FSpaceControlNameAndChannel& Space = SpaceChannels[Index];
+
+		const int32 EntityIndex = OutFieldBuilder->FindOrAddEntity(this, EncodeControlRigEntityID(Index, EControlRigEntityType::Space));
+		OutFieldBuilder->AddPersistentEntity(EffectiveRange, EntityIndex, MetaDataIndex);
+	}
+
+	for (int32 Index = 0; Index < NumBoolID; ++Index)
+	{
+		const FBoolParameterNameAndCurve& Bool = BoolParameterNamesAndCurves[Index];
+
+		const FChannelMapInfo* ChannelInfo = ControlChannelMap.Find(Bool.ParameterName);
+		if (!ChannelInfo || (ControlsMask.IsValidIndex(ChannelInfo->MaskIndex) && !ControlsMask[ChannelInfo->MaskIndex]))
+		{
+			continue;
+		}
+
+		if (Bool.ParameterCurve.HasAnyData())
+		{
+			const int32 EntityIndex = OutFieldBuilder->FindOrAddEntity(this, EncodeControlRigEntityID(Index, EControlRigEntityType::BoolParameter));
+			OutFieldBuilder->AddPersistentEntity(EffectiveRange, EntityIndex, MetaDataIndex);
+		}
+	}
+
+	for (int32 Index = 0; Index < NumEnumID; ++Index)
+	{
+		const FEnumParameterNameAndCurve& Enum = EnumParameterNamesAndCurves[Index];
+
+		const FChannelMapInfo* ChannelInfo = ControlChannelMap.Find(Enum.ParameterName);
+		if (!ChannelInfo || (ControlsMask.IsValidIndex(ChannelInfo->MaskIndex) && !ControlsMask[ChannelInfo->MaskIndex]))
+		{
+			continue;
+		}
+
+		if (Enum.ParameterCurve.HasAnyData())
+		{
+			const int32 EntityIndex = OutFieldBuilder->FindOrAddEntity(this, EncodeControlRigEntityID(Index, EControlRigEntityType::EnumParameter));
+			OutFieldBuilder->AddPersistentEntity(EffectiveRange, EntityIndex, MetaDataIndex);
+		}
+	}
+
+	for (int32 Index = 0; Index < NumIntegerID; ++Index)
+	{
+		const FIntegerParameterNameAndCurve& Integer = IntegerParameterNamesAndCurves[Index];
+
+		const FChannelMapInfo* ChannelInfo = ControlChannelMap.Find(Integer.ParameterName);
+		if (!ChannelInfo || (ControlsMask.IsValidIndex(ChannelInfo->MaskIndex) && !ControlsMask[ChannelInfo->MaskIndex]))
+		{
+			continue;
+		}
+
+		if (Integer.ParameterCurve.HasAnyData())
+		{
+			const int32 EntityIndex = OutFieldBuilder->FindOrAddEntity(this, EncodeControlRigEntityID(Index, EControlRigEntityType::IntegerParameter));
+			OutFieldBuilder->AddPersistentEntity(EffectiveRange, EntityIndex, MetaDataIndex);
+		}
+	}
+
+	for (int32 Index = 0; Index < NumScalarID; ++Index)
+	{
+		const FScalarParameterNameAndCurve& Scalar = ScalarParameterNamesAndCurves[Index];
+
+		const FChannelMapInfo* ChannelInfo = ControlChannelMap.Find(Scalar.ParameterName);
+		if (!ChannelInfo || (ControlsMask.IsValidIndex(ChannelInfo->MaskIndex) && !ControlsMask[ChannelInfo->MaskIndex]))
+		{
+			continue;
+		}
+
+		if (Scalar.ParameterCurve.HasAnyData())
+		{
+			const int32 EntityIndex = OutFieldBuilder->FindOrAddEntity(this, EncodeControlRigEntityID(Index, EControlRigEntityType::ScalarParameter));
+			OutFieldBuilder->AddPersistentEntity(EffectiveRange, EntityIndex, MetaDataIndex);
+		}
+	}
+
+	for (int32 Index = 0; Index < NumVector2ID; ++Index)
+	{
+		const FVector2DParameterNameAndCurves& Vector2D = Vector2DParameterNamesAndCurves[Index];
+
+		const FChannelMapInfo* ChannelInfo = ControlChannelMap.Find(Vector2D.ParameterName);
+		if (!ChannelInfo || (ControlsMask.IsValidIndex(ChannelInfo->MaskIndex) && !ControlsMask[ChannelInfo->MaskIndex]))
+		{
+			continue;
+		}
+
+		if (Vector2D.XCurve.HasAnyData() || Vector2D.YCurve.HasAnyData())
+		{
+			const int32 EntityIndex = OutFieldBuilder->FindOrAddEntity(this, EncodeControlRigEntityID(Index, EControlRigEntityType::Vector2Parameter));
+			OutFieldBuilder->AddPersistentEntity(EffectiveRange, EntityIndex, MetaDataIndex);
+		}
+	}
+
+	for (int32 Index = 0; Index < NumVector3ID; ++Index)
+	{
+		const FVectorParameterNameAndCurves& Vector = VectorParameterNamesAndCurves[Index];
+
+		const FChannelMapInfo* ChannelInfo = ControlChannelMap.Find(Vector.ParameterName);
+		if (!ChannelInfo || (ControlsMask.IsValidIndex(ChannelInfo->MaskIndex) && !ControlsMask[ChannelInfo->MaskIndex]))
+		{
+			continue;
+		}
+
+		if (Vector.XCurve.HasAnyData() || Vector.YCurve.HasAnyData() || Vector.ZCurve.HasAnyData())
+		{
+			const int32 EntityIndex = OutFieldBuilder->FindOrAddEntity(this, EncodeControlRigEntityID(Index, EControlRigEntityType::Vector3Parameter));
+			OutFieldBuilder->AddPersistentEntity(EffectiveRange, EntityIndex, MetaDataIndex);
+		}
+	}
+
+	for (int32 Index = 0; Index < NumTransformID; ++Index)
+	{
+		const FTransformParameterNameAndCurves& Transform = TransformParameterNamesAndCurves[Index];
+
+		const FChannelMapInfo* ChannelInfo = ControlChannelMap.Find(Transform.ParameterName);
+		if (!ChannelInfo || (ControlsMask.IsValidIndex(ChannelInfo->MaskIndex) && !ControlsMask[ChannelInfo->MaskIndex]))
+		{
+			continue;
+		}
+
+		if (Transform.Translation[0].HasAnyData() || Transform.Translation[1].HasAnyData() || Transform.Translation[2].HasAnyData() ||
+			Transform.Rotation[0].HasAnyData() || Transform.Rotation[1].HasAnyData() || Transform.Rotation[2].HasAnyData() ||
+			Transform.Scale[0].HasAnyData() || Transform.Scale[1].HasAnyData() || Transform.Scale[2].HasAnyData())
+		{
+			const int32 EntityIndex = OutFieldBuilder->FindOrAddEntity(this, EncodeControlRigEntityID(Index, EControlRigEntityType::TransformParameter));
+			OutFieldBuilder->AddPersistentEntity(EffectiveRange, EntityIndex, MetaDataIndex);
+		}
+	}
+
+	return true;
+}
+
 #undef LOCTEXT_NAMESPACE 
 
 
diff --git a/Engine/Plugins/Animation/ControlRig/Source/ControlRig/Private/Sequencer/MovieSceneControlRigParameterTemplate.cpp b/Engine/Plugins/Animation/ControlRig/Source/ControlRig/Private/Sequencer/MovieSceneControlRigParameterTemplate.cpp
--- a/Engine/Plugins/Animation/ControlRig/Source/ControlRig/Private/Sequencer/MovieSceneControlRigParameterTemplate.cpp
+++ b/Engine/Plugins/Animation/ControlRig/Source/ControlRig/Private/Sequencer/MovieSceneControlRigParameterTemplate.cpp
@@ -1900,6 +1900,11 @@ struct TControlRigParameterActuatorTransform : TMovieSceneBlendingActuator<FCont
 
 void FMovieSceneControlRigParameterTemplate::Evaluate(const FMovieSceneEvaluationOperand& Operand, const FMovieSceneContext& Context, const FPersistentEvaluationData& PersistentData, FMovieSceneExecutionTokens& ExecutionTokens) const
 {
+	if (!UMovieSceneControlRigParameterTrack::ShouldUseLegacyTemplate())
+	{
+		return;
+	}
+
 	const FFrameTime Time = Context.GetTime();
 
 	const UMovieSceneControlRigParameterSection* Section = Cast<UMovieSceneControlRigParameterSection>(GetSourceSection());
@@ -2440,6 +2445,11 @@ void FMovieSceneControlRigParameterTemplate::Interrogate(const FMovieSceneContex
 {
 	MOVIESCENE_DETAILED_SCOPE_CYCLE_COUNTER(MovieSceneEval_ControlRigTemplateParameter_Evaluate)
 
+	if (!UMovieSceneControlRigParameterTrack::ShouldUseLegacyTemplate())
+	{
+		return;
+	}
+
 	const FFrameTime Time = Context.GetTime();
 
 	const UMovieSceneControlRigParameterSection* Section = Cast<UMovieSceneControlRigParameterSection>(GetSourceSection());
diff --git a/Engine/Plugins/Animation/ControlRig/Source/ControlRig/Private/Sequencer/MovieSceneControlRigParameterTrack.cpp b/Engine/Plugins/Animation/ControlRig/Source/ControlRig/Private/Sequencer/MovieSceneControlRigParameterTrack.cpp
--- a/Engine/Plugins/Animation/ControlRig/Source/ControlRig/Private/Sequencer/MovieSceneControlRigParameterTrack.cpp
+++ b/Engine/Plugins/Animation/ControlRig/Source/ControlRig/Private/Sequencer/MovieSceneControlRigParameterTrack.cpp
@@ -4,6 +4,7 @@
 #include "Compilation/MovieSceneCompilerRules.h"
 #include "Evaluation/MovieSceneEvaluationTrack.h"
 #include "MovieSceneCommonHelpers.h"
+#include "EntitySystem/MovieSceneEntitySystem.h"
 #include "Sequencer/MovieSceneControlRigParameterTemplate.h"
 #include "MovieScene.h"
 #include "MovieSceneTimeHelpers.h"
@@ -12,6 +13,7 @@
 #include "Rigs/RigHierarchyController.h"
 #include "UObject/Package.h"
 #include "Async/Async.h"
+#include "Sequencer/MovieSceneControlRigSystem.h"
 
 #if WITH_EDITOR
 #include "ScopedTransaction.h"
@@ -21,6 +23,32 @@
 
 #define LOCTEXT_NAMESPACE "MovieSceneParameterControlRigTrack"
 
+namespace UE::MovieScene
+{
+	bool GUseLegacyControlRigTemplate = true;
+	FAutoConsoleVariableRef CVarUseLegacyControlRigTemplate(
+		TEXT("ControlRig.UseLegacySequencerTemplate"),
+		GUseLegacyControlRigTemplate,
+		TEXT("(Default: true) Specifies whether to use legacy template evaluation for control rig tracks in Sequencer."),
+		ECVF_Default
+	);
+
+	bool GCanUseLegacyControlRigTemplate = true;
+	FAutoConsoleVariableRef CVarCanUseLegacyControlRigTemplate(
+		TEXT("ControlRig.CanUseLegacySequencerTemplate"),
+		GCanUseLegacyControlRigTemplate,
+		TEXT("(Default: true) Specifies whether control rig tracks compile their legacy template as a back up. When disabled, ControlRig.UseLegacySequencerTemplate has no effect."),
+		ECVF_Default
+	);
+
+} // namespace UE::MovieScene
+
+bool UMovieSceneControlRigParameterTrack::ShouldUseLegacyTemplate()
+{
+	// We use the legacy template if we were asked to, or if we're not using the custom scheduler (since the code is not implemented there....)
+	return UE::MovieScene::GUseLegacyControlRigTemplate || !UMovieSceneEntitySystem::IsCustomSchedulingEnabled();
+}
+
 FColor UMovieSceneControlRigParameterTrack::AbsoluteRigTrackColor = FColor(65, 89, 194, 65);
 FColor UMovieSceneControlRigParameterTrack::LayeredRigTrackColor = FColor(173, 151, 114);
 
@@ -44,7 +72,10 @@ UMovieSceneControlRigParameterTrack::UMovieSceneControlRigParameterTrack(const F
 	SupportedBlendTypes = FMovieSceneBlendTypeField::None();
 	SupportedBlendTypes.Add(EMovieSceneBlendType::Additive);
 	SupportedBlendTypes.Add(EMovieSceneBlendType::Absolute);
-	SupportedBlendTypes.Add(EMovieSceneBlendType::Override);
+	if (ShouldUseLegacyTemplate() == false)
+	{
+		SupportedBlendTypes.Add(EMovieSceneBlendType::Override);
+	}
 }
 
 void UMovieSceneControlRigParameterTrack::BeginDestroy()
@@ -56,9 +87,15 @@ void UMovieSceneControlRigParameterTrack::BeginDestroy()
 	}
 }
 
+// UMovieSceneTrack interface
 FMovieSceneEvalTemplatePtr UMovieSceneControlRigParameterTrack::CreateTemplateForSection(const UMovieSceneSection& InSection) const
 {
-	return FMovieSceneControlRigParameterTemplate(*CastChecked<UMovieSceneControlRigParameterSection>(&InSection), *this);
+	if (UE::MovieScene::GCanUseLegacyControlRigTemplate)
+	{
+		return FMovieSceneControlRigParameterTemplate(*CastChecked<UMovieSceneControlRigParameterSection>(&InSection), *this);
+	}
+
+	return FMovieSceneEvalTemplatePtr();
 }
 
 bool UMovieSceneControlRigParameterTrack::SupportsType(TSubclassOf<UMovieSceneSection> SectionClass) const
@@ -710,6 +747,11 @@ void UMovieSceneControlRigParameterTrack::RenameParameterName(const FName& OldPa
 
 void UMovieSceneControlRigParameterTrack::ReplaceControlRig(UControlRig* NewControlRig, bool RecreateChannels)
 {
+	if (ControlRig == NewControlRig)
+	{
+		return;
+	}
+
 	if (IsValid(ControlRig))
 	{
 		ControlRig->OnPostConstruction_AnyThread().RemoveAll(this);
diff --git a/Engine/Plugins/Animation/ControlRig/Source/ControlRig/Private/Sequencer/MovieSceneControlRigSystem.cpp b/Engine/Plugins/Animation/ControlRig/Source/ControlRig/Private/Sequencer/MovieSceneControlRigSystem.cpp
new file mode 100644
--- /dev/null
+++ b/Engine/Plugins/Animation/ControlRig/Source/ControlRig/Private/Sequencer/MovieSceneControlRigSystem.cpp
@@ -0,0 +1,1914 @@
+// Copyright Epic Games, Inc. All Rights Reserved.
+
+#include "Sequencer/MovieSceneControlRigSystem.h"
+#include "Sequencer/MovieSceneControlRigParameterSection.h"
+#include "Sequencer/MovieSceneControlRigParameterTrack.h"
+#include "Sequencer/MovieSceneControlRigParameterBuffer.h"
+#include "Sequencer/MovieSceneControlRigComponentTypes.h"
+#include "Algo/IndexOf.h"
+#include "ConstraintsManager.h"
+#include "MovieSceneTracksComponentTypes.h"
+#include "EntitySystem/MovieSceneInterrogation.h"
+#include "EntitySystem/MovieSceneEntitySystemRunner.h"
+#include "EntitySystem/MovieSceneBoundObjectInstantiator.h"
+#include "EntitySystem/MovieSceneBoundSceneComponentInstantiator.h"
+#include "EntitySystem/MovieSceneEntitySystemTask.h"
+#include "EntitySystem/MovieSceneComponentRegistry.h"
+#include "EntitySystem/MovieSceneRootInstantiatorSystem.h"
+#include "EntitySystem/MovieSceneEntityMutations.h"
+#include "Evaluation/PreAnimatedState/MovieScenePreAnimatedObjectStorage.h"
+#include "Evaluation/PreAnimatedState/MovieScenePreAnimatedStorageID.inl"
+#include "EntitySystem/MovieScenePreAnimatedStateSystem.h"
+#include "EntitySystem/BuiltInComponentTypes.h"
+#include "MovieSceneTracksComponentTypes.h"
+#include "Systems/FloatChannelEvaluatorSystem.h"
+#include "Systems/DoubleChannelEvaluatorSystem.h"
+#include "Systems/MovieScenePiecewiseDoubleBlenderSystem.h"
+#include "Systems/MovieSceneHierarchicalBiasSystem.h"
+#include "Systems/MovieSceneInitialValueSystem.h"
+#include "Systems/WeightAndEasingEvaluatorSystem.h"
+#include "EntitySystem/MovieSceneEntityGroupingSystem.h"
+#include "ControlRig.h"
+#include "Components/SkeletalMeshComponent.h"
+#include "AnimCustomInstanceHelper.h"
+#include "Sequencer/ControlRigLayerInstance.h"
+#include "IControlRigObjectBinding.h"
+#include "ControlRigObjectBinding.h"
+#include "TransformNoScale.h"
+#include "ControlRigComponent.h"
+#include "SkeletalMeshRestoreState.h"
+#include "Rigs/FKControlRig.h"
+#include "UObject/UObjectAnnotation.h"
+#include "Rigs/RigHierarchy.h"
+#include "TransformConstraint.h"
+#include "TransformableHandle.h"
+#include "AnimationCoreLibrary.h"
+
+#include UE_INLINE_GENERATED_CPP_BY_NAME(MovieSceneControlRigSystem)
+
+namespace UE::MovieScene
+{
+
+TAutoRegisterPreAnimatedStorageID<FPreAnimatedControlRigStorage> FPreAnimatedControlRigStorage::StorageID;
+TAutoRegisterPreAnimatedStorageID<FPreAnimatedControlRigParameterStorage> FPreAnimatedControlRigParameterStorage::StorageID;
+
+
+
+UControlRig* GetControlRig(UMovieSceneControlRigParameterTrack* Track, UObject* BoundObject)
+{
+	UWorld* GameWorld = (BoundObject && BoundObject->GetWorld() && BoundObject->GetWorld()->IsGameWorld()) ? BoundObject->GetWorld() : nullptr;
+
+	UControlRig* ControlRig = GameWorld ? Track->GetGameWorldControlRig(GameWorld) : Track->GetControlRig();
+	if (!ControlRig)
+	{
+		return nullptr;
+	}
+
+	if (ControlRig->GetObjectBinding())
+	{
+		if (UControlRigComponent* ControlRigComponent = Cast<UControlRigComponent>(ControlRig->GetObjectBinding()->GetBoundObject()))
+		{
+			if (AActor* Actor = Cast<AActor>(BoundObject))
+			{
+				if (UControlRigComponent* NewControlRigComponent = Actor->FindComponentByClass<UControlRigComponent>())
+				{
+					if (NewControlRigComponent->GetWorld())
+					{
+						if (NewControlRigComponent->GetWorld()->IsGameWorld())
+						{
+							ControlRig = NewControlRigComponent->GetControlRig();
+							if (!ControlRig)
+							{
+								NewControlRigComponent->Initialize();
+								ControlRig = NewControlRigComponent->GetControlRig();
+							}
+							if (ControlRig)
+							{
+								if (ControlRig->GetObjectBinding() == nullptr)
+								{
+									ControlRig->SetObjectBinding(MakeShared<FControlRigObjectBinding>());
+								}
+								if (ControlRig->GetObjectBinding()->GetBoundObject() != NewControlRigComponent)
+								{
+									ControlRig->GetObjectBinding()->BindToObject(BoundObject);
+								}
+							}
+						}
+						else if (NewControlRigComponent != ControlRigComponent)
+						{
+							NewControlRigComponent->SetControlRig(ControlRig);
+						}
+					}
+				}
+			}
+		}
+		else if (UControlRigComponent* NewControlRigComponent = Cast<UControlRigComponent>(BoundObject))
+		{
+			if (NewControlRigComponent->GetWorld())
+			{
+				if (NewControlRigComponent->GetWorld()->IsGameWorld())
+				{
+					ControlRig = NewControlRigComponent->GetControlRig();
+					if (!ControlRig)
+					{
+						NewControlRigComponent->Initialize();
+						ControlRig = NewControlRigComponent->GetControlRig();
+					}
+					if (ControlRig)
+					{
+						if (ControlRig->GetObjectBinding() == nullptr)
+						{
+							ControlRig->SetObjectBinding(MakeShared<FControlRigObjectBinding>());
+						}
+						if (ControlRig->GetObjectBinding()->GetBoundObject() != NewControlRigComponent)
+						{
+							ControlRig->GetObjectBinding()->BindToObject(BoundObject);
+						}
+					}
+				}
+				else if (NewControlRigComponent != ControlRigComponent)
+				{
+					NewControlRigComponent->SetControlRig(ControlRig);
+				}
+			}
+		}
+	}
+
+	return ControlRig;
+}
+
+UTickableConstraint* CreateConstraintIfNeeded(UWorld* InWorld, UTickableConstraint* Constraint, UMovieSceneControlRigParameterSection* InSection)
+{
+	if (!Constraint)
+	{
+		return nullptr;
+	}
+
+	// it's possible that we have it but it's not in the manager, due to manager not being saved with it (due to spawning or undo/redo).
+	if (InWorld)
+	{
+		const FConstraintsManagerController& Controller = FConstraintsManagerController::Get(InWorld);
+		if (Controller.GetConstraint(Constraint->ConstraintID) == nullptr)
+		{
+			Controller.AddConstraint(Constraint);
+			//need to reconstuct channels here.. note this is now lazy and so will recreate it next time view requests it
+			//but only do it if the control rig has a valid world it may not for example in PIE
+			if (InSection->GetControlRig() && InSection->GetControlRig()->GetWorld())
+			{
+				InSection->ReconstructChannelProxy();
+				InSection->MarkAsChanged();
+			}
+		}
+	}
+
+	return Constraint;
+}
+
+
+/**
+ * Initial value processor for control rig parameters.
+ * Responsible for caching the current value of scalar, vector and transform parameters into
+ * either an initial value cache index, initial value component, or both.
+ */
+struct FControlRigInitialValueProcessor : IInitialValueProcessor
+{
+	FBuiltInComponentTypes* BuiltInComponents;
+	FControlRigComponentTypes* ControlRigComponents;
+	FMovieSceneTracksComponentTypes* TracksComponents;
+
+	/*~ Transient properties reset on each usage (through Initialize/Finalize). */
+	IInterrogationExtension* Interrogation;
+	FInitialValueCache* InitialValueCache;
+	FEntityAllocationWriteContext WriteContext;
+	UMovieSceneControlRigParameterEvaluatorSystem* ParameterSystem;
+
+	FControlRigInitialValueProcessor()
+		: BuiltInComponents(FBuiltInComponentTypes::Get())
+		, ControlRigComponents(FControlRigComponentTypes::Get())
+		, TracksComponents(FMovieSceneTracksComponentTypes::Get())
+		, WriteContext(FEntityAllocationWriteContext::NewAllocation())
+	{
+		Interrogation     = nullptr;
+		InitialValueCache = nullptr;
+	}
+
+	virtual void Initialize(UMovieSceneEntitySystemLinker* Linker, FInitialValueCache* InInitialValueCache) override
+	{
+		Interrogation     = Linker->FindExtension<IInterrogationExtension>();
+		ParameterSystem   = Linker->FindSystem<UMovieSceneControlRigParameterEvaluatorSystem>();
+		InitialValueCache = InInitialValueCache;
+		WriteContext      = FEntityAllocationWriteContext(Linker->EntityManager);
+	}
+
+	virtual void PopulateFilter(FEntityComponentFilter& OutFilter) const override
+	{
+		OutFilter.All({ ControlRigComponents->ControlRigSource, TracksComponents->GenericParameterName, ControlRigComponents->Tags.ControlRigParameter });
+	}
+
+	virtual void Process(const FEntityAllocation* Allocation, const FComponentMask& AllocationType) override
+	{
+		if (!ParameterSystem)
+		{
+			return;
+		}
+		if (Interrogation && AllocationType.Contains(BuiltInComponents->Interrogation.OutputKey))
+		{
+			VisitInterrogationAllocation(Allocation);
+		}
+		else
+		{
+			VisitAllocation(Allocation);
+		}
+	}
+
+	virtual void Finalize() override
+	{
+		Interrogation     = nullptr;
+		InitialValueCache = nullptr;
+		ParameterSystem   = nullptr;
+	}
+
+	void VisitAllocation(const FEntityAllocation* Allocation)
+	{
+		if (TryCacheInitialValues<FFloatPropertyTraits, float>(Allocation, TracksComponents->Parameters.Scalar.InitialValue))
+		{
+			return;
+		}
+		if (TryCacheInitialValues<FFloatVectorPropertyTraits, FVector3f>(Allocation, TracksComponents->Parameters.Vector3.InitialValue))
+		{
+			return;
+		}
+		if (TryCacheInitialValues<FEulerTransformPropertyTraits, FEulerTransform>(Allocation, TracksComponents->Parameters.Transform.InitialValue))
+		{
+			return;
+		}
+	}
+
+	void VisitInterrogationAllocation(const FEntityAllocation* Allocation)
+	{
+		if (TryInterrogateValues<FFloatPropertyTraits, float>(Allocation, TracksComponents->Parameters.Scalar.InitialValue))
+		{
+			return;
+		}
+		if (TryInterrogateValues<FFloatVectorPropertyTraits, FVector3f>(Allocation, TracksComponents->Parameters.Vector3.InitialValue))
+		{
+			return;
+		}
+		if (TryInterrogateValues<FEulerTransformPropertyTraits, FEulerTransform>(Allocation, TracksComponents->Parameters.Transform.InitialValue))
+		{
+			return;
+		}
+	}
+
+private:
+
+	/** Default cache function for retrieving a control rig variable */
+	template<typename RigControlType, typename InitialValueType>
+	void ControlValueToInitialValue(UControlRig* Rig, FRigControlElement* Control, InitialValueType& OutInitialValue)
+	{
+		OutInitialValue = Rig->GetControlValue(Control, ERigControlValueType::Current).Get<RigControlType>();
+	}
+
+	/** Specialization for vector controls */
+	template<>
+	void ControlValueToInitialValue<FVector3f, FFloatIntermediateVector>(UControlRig* Rig, FRigControlElement* Control, FFloatIntermediateVector& OutInitialValue)
+	{
+		FVector3f Value = Rig->GetControlValue(Control, ERigControlValueType::Current).Get<FVector3f>();
+		OutInitialValue = FFloatIntermediateVector(Value.X, Value.Y, Value.Z);
+	}
+
+	/** Specialization for transform controls */
+	template<>
+	void ControlValueToInitialValue<FEulerTransform, FIntermediate3DTransform>(UControlRig* Rig, FRigControlElement* Control, FIntermediate3DTransform& OutInitialValue)
+	{
+		FEulerTransform EulerTransform;
+
+		switch (Control->Settings.ControlType)
+		{
+			case ERigControlType::Transform:
+			{
+				const FTransform Val = Rig->GetControlValue(Control, ERigControlValueType::Current)
+					.Get<FRigControlValue::FTransform_Float>().ToTransform();
+
+				EulerTransform = FEulerTransform(Val);
+				break;
+			}
+			case ERigControlType::TransformNoScale:
+			{
+				const FTransformNoScale NoScale = Rig->GetControlValue(Control, ERigControlValueType::Current)
+					.Get<FRigControlValue::FTransformNoScale_Float>().ToTransform();
+
+				EulerTransform = FEulerTransform(NoScale.ToFTransform());
+				break;
+			}
+			case ERigControlType::EulerTransform:
+			{
+				EulerTransform = Rig->GetControlValue(Control, ERigControlValueType::Current)
+					.Get<FRigControlValue::FEulerTransform_Float>().ToTransform();
+				break;
+			}
+		}
+
+		FVector Vector = Rig->GetControlSpecifiedEulerAngle(Control);
+		OutInitialValue = FIntermediate3DTransform(EulerTransform.Location, FRotator(Vector.Y, Vector.Z, Vector.X), EulerTransform.Scale);
+	}
+
+	template<typename PropertyTraits, typename RigControlType, typename InitialValueType>
+	bool TryCacheInitialValues(const FEntityAllocation* Allocation, TComponentTypeID<InitialValueType> InitialValue)
+	{
+		TOptionalComponentWriter<InitialValueType> InitialValues = Allocation->TryWriteComponents(InitialValue, WriteContext);
+		if (!InitialValues)
+		{
+			return false;
+		}
+
+		TPropertyValueStorage<PropertyTraits>* CacheStorage = InitialValueCache ? InitialValueCache->GetStorage<PropertyTraits>(InitialValue) : nullptr;
+
+		TComponentReader<FControlRigSourceData> ControlRigSources = Allocation->ReadComponents(ControlRigComponents->ControlRigSource);
+		TComponentReader<FName>                    ParameterNames  = Allocation->ReadComponents(TracksComponents->GenericParameterName);
+
+		const int32 Num = Allocation->Num();
+		if (CacheStorage)
+		{
+			TComponentWriter<FInitialValueIndex> InitialValueIndices = Allocation->WriteComponents(BuiltInComponents->InitialValueIndex, WriteContext);
+
+			for (int32 Index = 0; Index < Num; ++Index)
+			{
+				UControlRig* Rig = ParameterSystem->GetRigFromTrack(ControlRigSources[Index].Track);
+				FRigControlElement* ControlElement = Rig ? Rig->FindControl(ParameterNames[Index]) : nullptr;
+
+				if (Rig && ControlElement)
+				{
+					TOptional<FInitialValueIndex> ExistingIndex = CacheStorage->FindPropertyIndex(Rig, ParameterNames[Index]);
+					if (ExistingIndex)
+					{
+						InitialValues[Index] = CacheStorage->GetCachedValue(ExistingIndex.GetValue());
+					}
+					else
+					{
+						ControlValueToInitialValue<RigControlType>(Rig, ControlElement, InitialValues[Index]);
+						InitialValueIndices[Index] = CacheStorage->AddInitialValue(Rig, InitialValues[Index], ParameterNames[Index]);
+					}
+				}
+			}
+		}
+		else
+		{
+			for (int32 Index = 0; Index < Num; ++Index)
+			{
+				UControlRig* Rig = ParameterSystem->GetRigFromTrack(ControlRigSources[Index].Track);
+				FRigControlElement* ControlElement = Rig ? Rig->FindControl(ParameterNames[Index]) : nullptr;
+
+				if (Rig && ControlElement)
+				{
+					ControlValueToInitialValue<RigControlType>(Rig, ControlElement, InitialValues[Index]);
+				}
+			}
+		}
+
+		return true;
+	}
+
+
+	template<typename PropertyTraits, typename RigControlType, typename InitialValueType>
+	bool TryInterrogateValues(const FEntityAllocation* Allocation, TComponentTypeID<InitialValueType> InitialValue)
+	{
+		TOptionalComponentWriter<InitialValueType> InitialValues = Allocation->TryWriteComponents(InitialValue, WriteContext);
+		if (!InitialValues)
+		{
+			return false;
+		}
+
+		const int32 Num = Allocation->Num();
+
+		TComponentReader<FInterrogationKey>    OutputKeys   = Allocation->ReadComponents(BuiltInComponents->Interrogation.OutputKey);
+		TPropertyValueStorage<PropertyTraits>* CacheStorage = InitialValueCache ? InitialValueCache->GetStorage<PropertyTraits>(InitialValue) : nullptr;
+
+		const FSparseInterrogationChannelInfo& SparseChannelInfo = Interrogation->GetSparseChannelInfo();
+
+		for (int32 Index = 0; Index < Num; ++Index)
+		{
+			FInterrogationChannel Channel = OutputKeys[Index].Channel;
+
+			const FInterrogationChannelInfo* ChannelInfo = SparseChannelInfo.Find(Channel);
+			UControlRig* Rig = ChannelInfo ? Cast<UControlRig>(ChannelInfo->WeakObject.Get()) : nullptr;
+			if (!ChannelInfo || !Rig || ChannelInfo->PropertyBinding.PropertyName.IsNone())
+			{
+				continue;
+			}
+
+			FRigControlElement* ControlElement = Rig->FindControl(ChannelInfo->PropertyBinding.PropertyName);
+			if (!ControlElement)
+			{
+				continue;
+			}
+
+			// Retrieve a cached value if possible
+			if (CacheStorage)
+			{
+				const InitialValueType* CachedValue = CacheStorage->FindCachedValue(Rig, ChannelInfo->PropertyBinding.PropertyName);
+				if (CachedValue)
+				{
+					InitialValues[Index] = *CachedValue;
+					continue;
+				}
+			}
+
+			// No cached value available, must retrieve it now
+			ControlValueToInitialValue<RigControlType>(Rig, ControlElement, InitialValues[Index]);
+		};
+		return true;
+	}
+};
+
+
+
+
+/**
+ * Mutation that adds and assigns accumulation buffer indices to parameters
+ */
+struct FControlRigAccumulationEntryIndexMutation : IMovieSceneEntityMutation
+{
+	FAccumulatedControlRigValues* AccumulatedValues;
+	FPreAnimatedControlRigParameterStorage* PreAnimatedParameters;
+
+	FControlRigAccumulationEntryIndexMutation(FAccumulatedControlRigValues* InAccumulatedValues, FPreAnimatedControlRigParameterStorage* InPreAnimatedParameters)
+		: AccumulatedValues(InAccumulatedValues)
+		, PreAnimatedParameters(InPreAnimatedParameters)
+	{
+	}
+	void CreateMutation(FEntityManager* EntityManager, FComponentMask* InOutEntityComponentTypes) const override
+	{
+		FControlRigComponentTypes* ControlRigComponents = FControlRigComponentTypes::Get();
+		InOutEntityComponentTypes->Set(ControlRigComponents->AccumulatedControlEntryIndex);
+	}
+	void InitializeAllocation(FEntityAllocation* Allocation, const FComponentMask& AllocationType) const override
+	{
+		ProcessAllocation(Allocation, AllocationType);
+	}
+	void InitializeUnmodifiedAllocation(FEntityAllocation* Allocation, const FComponentMask& AllocationType) const override
+	{
+		ProcessAllocation(Allocation, AllocationType);
+	}
+	void ProcessAllocation(FEntityAllocation* Allocation, const FComponentMask& AllocationType) const
+	{
+		FEntityAllocationWriteContext WriteContext = FEntityAllocationWriteContext::NewAllocation();
+
+		FControlRigComponentTypes*       ControlRigComponents = FControlRigComponentTypes::Get();
+		FMovieSceneTracksComponentTypes* TracksComponents     = FMovieSceneTracksComponentTypes::Get();
+
+		TComponentReader<FName>                 ParameterNames    = Allocation->ReadComponents(TracksComponents->GenericParameterName);
+		TComponentWriter<FControlRigSourceData> ControlRigSources = Allocation->WriteComponents(ControlRigComponents->ControlRigSource, WriteContext);
+
+		EControlRigControlType ControlType;
+		if (AllocationType.Contains(TracksComponents->Parameters.Bool.PropertyTag))
+		{
+			ControlType = EControlRigControlType::Parameter_Bool;
+		}
+		else if (AllocationType.Contains(TracksComponents->Parameters.Byte.PropertyTag))
+		{
+			ControlType = EControlRigControlType::Parameter_Enum;
+		}
+		else if (AllocationType.Contains(TracksComponents->Parameters.Integer.PropertyTag))
+		{
+			ControlType = EControlRigControlType::Parameter_Integer;
+		}
+		else if (AllocationType.Contains(TracksComponents->Parameters.Scalar.PropertyTag))
+		{
+			ControlType = EControlRigControlType::Parameter_Scalar;
+		}
+		else if (AllocationType.Contains(TracksComponents->Parameters.Vector3.PropertyTag))
+		{
+			ControlType = EControlRigControlType::Parameter_Vector;
+		}
+		else if (AllocationType.Contains(TracksComponents->Parameters.Transform.PropertyTag))
+		{
+			ControlType = EControlRigControlType::Parameter_Transform;
+		}
+		else if (AllocationType.Contains(ControlRigComponents->Tags.Space))
+		{
+			ControlType = EControlRigControlType::Space;
+		}
+		else
+		{
+			return;
+		}
+
+		TComponentWriter<FAccumulatedControlEntryIndex> AccumulatorEntryIndices = Allocation->WriteComponents(ControlRigComponents->AccumulatedControlEntryIndex, WriteContext);
+
+		const int32 Num = Allocation->Num();
+
+		for (int32 Index = 0; Index < Num; ++Index)
+		{
+			AccumulatorEntryIndices[Index] = AccumulatedValues->AllocateEntryIndex(ControlRigSources[Index].Track, ParameterNames[Index], ControlType);
+		}
+	}
+};
+
+
+struct FRemoveInvalidControlRigAccumulationComponents : IMovieSceneConditionalEntityMutation
+{
+	FAccumulatedControlRigValues* AccumulatedValues;
+
+	FRemoveInvalidControlRigAccumulationComponents(FAccumulatedControlRigValues* InAccumulatedValues)
+		: AccumulatedValues(InAccumulatedValues)
+	{
+	}
+	void MarkAllocation(FEntityAllocation* Allocation, TBitArray<>& OutEntitiesToMutate) const override
+	{
+		FControlRigComponentTypes*              ControlRigComponents = FControlRigComponentTypes::Get();
+		TComponentReader<FControlRigSourceData> ControlRigSources    = Allocation->ReadComponents(ControlRigComponents->ControlRigSource);
+
+		const int32 Num = Allocation->Num();
+		for (int32 Index = 0; Index < Num; ++Index)
+		{
+			// If the entity does not have a valid entry, set the bit, resulting in the component being removed
+			if (!AccumulatedValues->DoesEntryExistForTrack(ControlRigSources[Index].Track))
+			{
+				OutEntitiesToMutate.PadToNum(Index+1, false);
+				OutEntitiesToMutate[Index] = true;
+			}
+		}
+	}
+	void CreateMutation(FEntityManager* EntityManager, FComponentMask* InOutEntityComponentTypes) const override
+	{
+		FControlRigComponentTypes* ControlRigComponents = FControlRigComponentTypes::Get();
+		InOutEntityComponentTypes->Remove(ControlRigComponents->AccumulatedControlEntryIndex);
+	}
+};
+
+
+struct FInitialControlRigParameterValueMutation : IMovieSceneEntityMutation
+{
+	FAccumulatedControlRigValues* AccumulatedValues;
+	FInitialControlRigParameterValueMutation(FAccumulatedControlRigValues* InAccumulatedValues)
+		: AccumulatedValues(InAccumulatedValues)
+	{
+	}
+	void CreateMutation(FEntityManager* EntityManager, FComponentMask* InOutEntityComponentTypes) const override
+	{
+		FControlRigComponentTypes*       ControlRigComponents = FControlRigComponentTypes::Get();
+		FMovieSceneTracksComponentTypes* TracksComponents     = FMovieSceneTracksComponentTypes::Get();
+
+		if (InOutEntityComponentTypes->Contains(TracksComponents->Parameters.Bool.PropertyTag))
+		{
+			InOutEntityComponentTypes->Set(TracksComponents->Parameters.Bool.InitialValue);
+		}
+		else if (InOutEntityComponentTypes->Contains(TracksComponents->Parameters.Byte.PropertyTag))
+		{
+			InOutEntityComponentTypes->Set(TracksComponents->Parameters.Byte.InitialValue);
+		}
+		else if (InOutEntityComponentTypes->Contains(TracksComponents->Parameters.Integer.PropertyTag))
+		{
+			InOutEntityComponentTypes->Set(TracksComponents->Parameters.Integer.InitialValue);
+		}
+		else if (InOutEntityComponentTypes->Contains(TracksComponents->Parameters.Scalar.PropertyTag))
+		{
+			InOutEntityComponentTypes->Set(TracksComponents->Parameters.Scalar.InitialValue);
+		}
+		else if (InOutEntityComponentTypes->Contains(TracksComponents->Parameters.Vector3.PropertyTag))
+		{
+			InOutEntityComponentTypes->Set(TracksComponents->Parameters.Vector3.InitialValue);
+		}
+		else if (InOutEntityComponentTypes->Contains(TracksComponents->Parameters.Transform.PropertyTag))
+		{
+			InOutEntityComponentTypes->Set(TracksComponents->Parameters.Transform.InitialValue);
+		}
+	}
+	void InitializeAllocation(FEntityAllocation* Allocation, const FComponentMask& AllocationType) const override
+	{
+		FEntityAllocationWriteContext WriteContext = FEntityAllocationWriteContext::NewAllocation();
+
+		FControlRigComponentTypes*       ControlRigComponents = FControlRigComponentTypes::Get();
+		FMovieSceneTracksComponentTypes* TracksComponents     = FMovieSceneTracksComponentTypes::Get();
+
+		TComponentReader<FAccumulatedControlEntryIndex> AccumulatorEntryIndices = Allocation->ReadComponents(ControlRigComponents->AccumulatedControlEntryIndex);
+		TComponentReader<FName> ParameterNames = Allocation->ReadComponents(TracksComponents->GenericParameterName);
+
+		const int32 Num = Allocation->Num();
+		if (AllocationType.Contains(TracksComponents->Parameters.Bool.PropertyTag))
+		{
+			TComponentWriter<bool> InitialValues = Allocation->WriteComponents(TracksComponents->Parameters.Bool.InitialValue, WriteContext);
+			for (int32 Index = 0; Index < Num; ++Index)
+			{
+				UControlRig* Rig = AccumulatedValues->FindControlRig(AccumulatorEntryIndices[Index]);
+				if (FRigControlElement* ControlElement = Rig->FindControl(ParameterNames[Index]))
+				{
+					InitialValues[Index] = Rig->GetControlValue(ControlElement, ERigControlValueType::Current).Get<bool>();
+				}
+			}
+		}
+		else if (AllocationType.Contains(TracksComponents->Parameters.Byte.PropertyTag))
+		{
+			TComponentWriter<uint8> InitialValues = Allocation->WriteComponents(TracksComponents->Parameters.Byte.InitialValue, WriteContext);
+			for (int32 Index = 0; Index < Num; ++Index)
+			{
+				UControlRig* Rig = AccumulatedValues->FindControlRig(AccumulatorEntryIndices[Index]);
+				if (FRigControlElement* ControlElement = Rig->FindControl(ParameterNames[Index]))
+				{
+					const int32 EnumAsInt = Rig->GetControlValue(ControlElement, ERigControlValueType::Current).Get<int32>();
+					InitialValues[Index] = static_cast<uint8>(EnumAsInt);
+				}
+			}
+		}
+		else if (AllocationType.Contains(TracksComponents->Parameters.Integer.PropertyTag))
+		{
+			TComponentWriter<int32> InitialValues = Allocation->WriteComponents(TracksComponents->Parameters.Integer.InitialValue, WriteContext);
+			for (int32 Index = 0; Index < Num; ++Index)
+			{
+				UControlRig* Rig = AccumulatedValues->FindControlRig(AccumulatorEntryIndices[Index]);
+				if (FRigControlElement* ControlElement = Rig->FindControl(ParameterNames[Index]))
+				{
+					InitialValues[Index] = Rig->GetControlValue(ControlElement, ERigControlValueType::Current).Get<int32>();
+				}
+			}
+		}
+		else if (AllocationType.Contains(TracksComponents->Parameters.Scalar.PropertyTag))
+		{
+			TComponentWriter<double> InitialValues = Allocation->WriteComponents(TracksComponents->Parameters.Scalar.InitialValue, WriteContext);
+			for (int32 Index = 0; Index < Num; ++Index)
+			{
+				UControlRig* Rig = AccumulatedValues->FindControlRig(AccumulatorEntryIndices[Index]);
+				if (FRigControlElement* ControlElement = Rig->FindControl(ParameterNames[Index]))
+				{
+					InitialValues[Index] = Rig->GetControlValue(ControlElement, ERigControlValueType::Current).Get<float>();
+				}
+			}
+		}
+		else if (AllocationType.Contains(TracksComponents->Parameters.Vector3.PropertyTag))
+		{
+			TComponentWriter<FFloatIntermediateVector> InitialValues = Allocation->WriteComponents(TracksComponents->Parameters.Vector3.InitialValue, WriteContext);
+			for (int32 Index = 0; Index < Num; ++Index)
+			{
+				UControlRig* Rig = AccumulatedValues->FindControlRig(AccumulatorEntryIndices[Index]);
+				if (FRigControlElement* ControlElement = Rig->FindControl(ParameterNames[Index]))
+				{
+					FVector3f Value = Rig->GetControlValue(ControlElement, ERigControlValueType::Current).Get<FVector3f>();
+					InitialValues[Index] = FFloatIntermediateVector(Value.X, Value.Y, Value.Z);
+				}
+			}
+		}
+		else if (AllocationType.Contains(TracksComponents->Parameters.Transform.PropertyTag))
+		{
+			TComponentWriter<FIntermediate3DTransform> InitialValues = Allocation->WriteComponents(TracksComponents->Parameters.Transform.InitialValue, WriteContext);
+			for (int32 Index = 0; Index < Num; ++Index)
+			{
+				UControlRig* Rig = AccumulatedValues->FindControlRig(AccumulatorEntryIndices[Index]);
+				if (FRigControlElement* ControlElement = Rig->FindControl(ParameterNames[Index]))
+				{
+					switch (ControlElement->Settings.ControlType)
+					{
+
+						case ERigControlType::Transform:
+						{
+							const FTransform Val = Rig->GetControlValue(ControlElement, ERigControlValueType::Current).Get<FRigControlValue::FTransform_Float>().ToTransform();
+							FEulerTransform EulerTransform(Val);
+							FVector Vector = Rig->GetControlSpecifiedEulerAngle(ControlElement);
+							EulerTransform.Rotation = FRotator(Vector.Y, Vector.Z, Vector.X);
+							InitialValues[Index] = FIntermediate3DTransform(EulerTransform.Location, EulerTransform.Rotation, EulerTransform.Scale);
+							break;
+						}
+						case ERigControlType::TransformNoScale:
+						{
+							const FTransformNoScale NoScale =
+								Rig
+								->GetControlValue(ControlElement, ERigControlValueType::Current).Get<FRigControlValue::FTransformNoScale_Float>().ToTransform();
+							FEulerTransform EulerTransform(NoScale.ToFTransform());
+							FVector Vector = Rig->GetControlSpecifiedEulerAngle(ControlElement);
+							EulerTransform.Rotation = FRotator(Vector.Y, Vector.Z, Vector.X);
+							InitialValues[Index] = FIntermediate3DTransform(EulerTransform.Location, EulerTransform.Rotation, EulerTransform.Scale);
+							break;
+						}
+						case ERigControlType::EulerTransform:
+						{
+							FEulerTransform EulerTransform =
+								Rig
+								->GetControlValue(ControlElement, ERigControlValueType::Current).Get<FRigControlValue::FEulerTransform_Float>().ToTransform();
+							FVector Vector = Rig->GetControlSpecifiedEulerAngle(ControlElement);
+							EulerTransform.Rotation = FRotator(Vector.Y, Vector.Z, Vector.X);
+							InitialValues[Index] = FIntermediate3DTransform(EulerTransform.Location, EulerTransform.Rotation, EulerTransform.Scale);
+							break;
+						}
+					}
+				}
+			}
+		}
+	}
+	void InitializeUnmodifiedAllocation(FEntityAllocation* Allocation, const FComponentMask& AllocationType) const override
+	{
+	}
+};
+
+
+FAnimatedControlRigParameterInfo::~FAnimatedControlRigParameterInfo()
+{
+}
+
+void CollectGarbageForOutput(FAnimatedControlRigParameterInfo* Output)
+{
+	// This should only happen during garbage collection
+	if (Output->OutputEntityID.IsValid())
+	{
+		UMovieSceneEntitySystemLinker* Linker = Output->WeakLinker.Get();
+		if (Linker)
+		{
+			Linker->EntityManager.AddComponent(Output->OutputEntityID, FBuiltInComponentTypes::Get()->Tags.NeedsUnlink);
+		}
+
+		Output->OutputEntityID = FMovieSceneEntityID();
+	}
+
+	if (Output->BlendChannelID.IsValid())
+	{
+		UMovieSceneBlenderSystem* BlenderSystem = Output->WeakBlenderSystem.Get();
+		if (BlenderSystem)
+		{
+			BlenderSystem->ReleaseBlendChannel(Output->BlendChannelID);
+		}
+		Output->BlendChannelID = FMovieSceneBlendChannelID();
+	}
+}
+
+
+TSharedPtr<FPreAnimatedControlRigParameterTraits::FPreAnimatedBufferPairs> FPreAnimatedControlRigParameterTraits::GetBuffers(UControlRig* Rig)
+{
+	check(Rig);
+
+	TWeakPtr<FPreAnimatedBufferPairs>& WeakBuffers = PreAnimatedBuffers.FindOrAdd(Rig);
+	TSharedPtr<FPreAnimatedBufferPairs> Buffers = WeakBuffers.Pin();
+	if (!Buffers)
+	{
+		Buffers = MakeShared<FPreAnimatedBufferPairs>();
+		WeakBuffers = Buffers;
+	}
+	return Buffers;
+}
+
+FPreAnimatedControlRigParameterTraits::FPreAnimatedParameterKey FPreAnimatedControlRigParameterTraits::CachePreAnimatedValue(UControlRig* Rig, FName ParameterName)
+{
+	check(Rig);
+	TSharedPtr<FPreAnimatedBufferPairs> Buffers = GetBuffers(Rig);
+
+	if (FRigControlElement* ControlElement = Rig->FindControl(ParameterName))
+	{
+		Buffers->Transient.AddCurrentValue(Rig, ControlElement);
+	}
+	return FPreAnimatedControlRigParameterTraits::FPreAnimatedParameterKey(Buffers, ParameterName);
+}
+
+void FPreAnimatedControlRigParameterTraits::RestorePreAnimatedValue(const TTuple<FObjectKey, FName>& InKey, FPreAnimatedParameterKey& ParameterBufferKey, const FRestoreStateParams& Params)
+{
+	UControlRig* Rig = Cast<UControlRig>(InKey.Get<0>().ResolveObjectPtr());
+	if (Rig)
+	{
+		if (ParameterBufferKey.bPersistent)
+		{
+			ParameterBufferKey.Buffer->Persistent.ApplyAndRemove(Rig, InKey.Get<1>());
+		}
+		else
+		{
+			ParameterBufferKey.Buffer->Transient.ApplyAndRemove(Rig, InKey.Get<1>());
+		}
+	}
+}
+
+FPreAnimatedControlRigState FPreAnimatedControlRigTraits::CachePreAnimatedValue(FObjectKey InObject)
+{
+	UControlRig* ControlRig = Cast<UControlRig>(InObject.ResolveObjectPtr());
+	check(ControlRig);
+
+	FPreAnimatedControlRigState State;
+	USkeletalMeshComponent* SkeletalMeshComponent = Cast<USkeletalMeshComponent>(ControlRig->GetObjectBinding()->GetBoundObject());
+	if (SkeletalMeshComponent)
+	{
+		State.SetSkelMesh(SkeletalMeshComponent);
+	}
+	return State;
+}
+
+void FPreAnimatedControlRigTraits::RestorePreAnimatedValue(const FObjectKey& InObject, FPreAnimatedControlRigState& State, const FRestoreStateParams& Params)
+{
+	if (UControlRig* ControlRig = Cast<UControlRig>(InObject.ResolveObjectPtr()))
+	{
+		ControlRig->Evaluate_AnyThread();
+
+		//unbind instances and reset animbp
+		FControlRigBindingHelper::UnBindFromSequencerInstance(ControlRig);
+
+		//do a tick and restore skel mesh
+		if (USkeletalMeshComponent* SkeletalMeshComponent = Cast<USkeletalMeshComponent>(ControlRig->GetObjectBinding()->GetBoundObject()))
+		{
+			// If the skel mesh comp owner has been removed from the world, no need to restore anything
+			if (SkeletalMeshComponent->IsRegistered())
+			{
+				// Restore pose after unbinding to force the restored pose
+				SkeletalMeshComponent->SetUpdateAnimationInEditor(true);
+				SkeletalMeshComponent->SetUpdateClothInEditor(true);
+				if (!SkeletalMeshComponent->IsPostEvaluatingAnimation())
+				{
+					SkeletalMeshComponent->TickAnimation(0.f, false);
+					SkeletalMeshComponent->RefreshBoneTransforms();
+					SkeletalMeshComponent->RefreshFollowerComponents();
+					SkeletalMeshComponent->UpdateComponentToWorld();
+					SkeletalMeshComponent->FinalizeBoneTransform();
+					SkeletalMeshComponent->MarkRenderTransformDirty();
+					SkeletalMeshComponent->MarkRenderDynamicDataDirty();
+				}
+				State.SkeletalMeshRestoreState.RestoreState(SkeletalMeshComponent);
+
+				if (SkeletalMeshComponent->GetAnimationMode() != State.AnimationMode)
+				{
+					SkeletalMeshComponent->SetAnimationMode(State.AnimationMode);
+				}
+			}
+		}
+
+		//only unbind if not a component
+		if (Cast<UControlRigComponent>(ControlRig->GetObjectBinding()->GetBoundObject()) == nullptr)
+		{
+			ControlRig->GetObjectBinding()->UnbindFromObject();
+		}
+	}
+}
+
+
+
+struct FControlRigDataGroupingPolicy
+{
+	using GroupKeyType = TTuple<FObjectKey, FName, EControlRigControlType>;
+
+	void InitializeGroupKeys(
+		TEntityGroupingHandlerBase<FControlRigDataGroupingPolicy>& Handler,
+		FEntityGroupBuilder* Builder,
+		FEntityAllocationIteratorItem Item,
+		FReadEntityIDs EntityIDs,
+		TWrite<FEntityGroupID> GroupIDs,
+		TRead<FControlRigSourceData> ControlRigSources,
+		TRead<FName> ParameterNames)
+	{
+		FMovieSceneTracksComponentTypes* TracksComponents     = FMovieSceneTracksComponentTypes::Get();
+		FControlRigComponentTypes*       ControlRigComponents = FControlRigComponentTypes::Get();
+
+		const FEntityAllocation* Allocation     = Item.GetAllocation();
+		const FComponentMask&    AllocationType = Item.GetAllocationType();
+		const int32              Num            = Allocation->Num();
+
+		EControlRigControlType ControlType;
+		if (AllocationType.Contains(TracksComponents->Parameters.Bool.PropertyTag))
+		{
+			ControlType = EControlRigControlType::Parameter_Bool;
+		}
+		else if (AllocationType.Contains(TracksComponents->Parameters.Byte.PropertyTag))
+		{
+			ControlType = EControlRigControlType::Parameter_Enum;
+		}
+		else if (AllocationType.Contains(TracksComponents->Parameters.Integer.PropertyTag))
+		{
+			ControlType = EControlRigControlType::Parameter_Integer;
+		}
+		else if (AllocationType.Contains(TracksComponents->Parameters.Scalar.PropertyTag))
+		{
+			ControlType = EControlRigControlType::Parameter_Scalar;
+		}
+		else if (AllocationType.Contains(TracksComponents->Parameters.Vector3.PropertyTag))
+		{
+			ControlType = EControlRigControlType::Parameter_Vector;
+		}
+		else if (AllocationType.Contains(TracksComponents->Parameters.Transform.PropertyTag))
+		{
+			ControlType = EControlRigControlType::Parameter_Transform;
+		}
+		else if (AllocationType.Contains(ControlRigComponents->Tags.Space))
+		{
+			ControlType = EControlRigControlType::Space;
+		}
+		else
+		{
+			check(false);
+		}
+
+		for (int32 Index = 0; Index < Num; ++Index)
+		{
+			if (GroupIDs[Index].IsValid())
+			{
+				continue;
+			}
+			GroupKeyType Key = MakeTuple(ControlRigSources[Index].Track, ParameterNames[Index], ControlType);
+
+			const int32    NewGroupIndex = Handler.GetOrAllocateGroupIndex(Key, Builder);
+			FEntityGroupID NewGroupID = Builder->MakeGroupID(NewGroupIndex);
+
+			Builder->AddEntityToGroup(EntityIDs[Index], NewGroupID);
+			// Write out the group ID component
+			GroupIDs[Index] = NewGroupID;
+		}
+	}
+
+#if WITH_EDITOR
+	bool OnObjectsReplaced(GroupKeyType& InOutKey, const TMap<UObject*, UObject*>& ReplacementMap)
+	{
+		return false;
+	}
+#endif
+};
+
+struct FOverlappingControlRigParameterHandler
+{
+	UMovieSceneEntitySystemLinker* Linker;
+	UMovieSceneControlRigParameterEvaluatorSystem* System;
+
+	FOverlappingControlRigParameterHandler(UMovieSceneControlRigParameterEvaluatorSystem* InSystem)
+		: Linker(InSystem->GetLinker())
+		, System(InSystem)
+	{}
+
+	void InitializeOutput(FEntityGroupID GroupID, TArrayView<const FMovieSceneEntityID> Inputs, FAnimatedControlRigParameterInfo* Output, FEntityOutputAggregate Aggregate)
+	{
+		UpdateOutput(GroupID, Inputs, Output, Aggregate);
+	}
+
+	void UpdateOutput(FEntityGroupID GroupID, TArrayView<const FMovieSceneEntityID> Inputs, FAnimatedControlRigParameterInfo* Output, FEntityOutputAggregate Aggregate)
+	{
+		using namespace UE::MovieScene;
+
+		FBuiltInComponentTypes* BuiltInComponents = FBuiltInComponentTypes::Get();
+		FMovieSceneTracksComponentTypes* TracksComponents = FMovieSceneTracksComponentTypes::Get();
+		FControlRigComponentTypes* ControlRigComponents = FControlRigComponentTypes::Get();
+
+		const int32 NumContributors = Inputs.Num();
+		if (!ensure(NumContributors != 0))
+		{
+			return;
+		}
+
+		bool bUseBlending = NumContributors > 1 ||
+			!Linker->EntityManager.HasComponent(Inputs[0], BuiltInComponents->Tags.AbsoluteBlend) ||
+			 Linker->EntityManager.HasComponent(Inputs[0], BuiltInComponents->WeightAndEasingResult);
+
+		if (!bUseBlending)
+		{
+			// Check if the parameter is fully keyed
+			const FComponentMask& EntityType = Linker->EntityManager.GetEntityType(Inputs[0]);
+			TArrayView<const FPropertyDefinition> Properties = BuiltInComponents->PropertyRegistry.GetProperties();
+
+			const int32 PropertyDefinitionIndex = Algo::IndexOfByPredicate(Properties, [=](const FPropertyDefinition& InDefinition){ return EntityType.Contains(InDefinition.PropertyType); });
+			if (PropertyDefinitionIndex != INDEX_NONE)
+			{
+				const FPropertyDefinition& Property = Properties[PropertyDefinitionIndex];
+
+				TArrayView<const FPropertyCompositeDefinition> Composites = BuiltInComponents->PropertyRegistry.GetComposites(Property);
+				for (int32 CompositeIndex = 0; CompositeIndex < Property.CompositeSize; ++CompositeIndex)
+				{
+					if (!EntityType.Contains(Composites[CompositeIndex].ComponentTypeID))
+					{
+						bUseBlending = true;
+						break;
+					}
+				}
+			}
+		}
+
+		if (bUseBlending || Output->OutputEntityID)
+		{
+			if (!Output->OutputEntityID)
+			{
+				if (!System->DoubleBlenderSystem)
+				{
+					System->DoubleBlenderSystem = Linker->LinkSystem<UMovieScenePiecewiseDoubleBlenderSystem>();
+					Linker->SystemGraph.AddReference(System, System->DoubleBlenderSystem);
+				}
+
+				Output->WeakLinker = Linker;
+				Output->WeakBlenderSystem = System->DoubleBlenderSystem;
+
+				// Initialize the blend channel ID
+				Output->BlendChannelID = System->DoubleBlenderSystem->AllocateBlendChannel();
+			}
+
+			const FComponentTypeID BlenderTypeTag = System->DoubleBlenderSystem->GetBlenderTypeTag();
+			InitializeBlendOutput(BlenderTypeTag, Inputs, Output);
+
+			for (FMovieSceneEntityID Input : Inputs)
+			{
+				if (!Linker->EntityManager.HasComponent(Input, BuiltInComponents->BlendChannelInput))
+				{
+					Linker->EntityManager.AddComponent(Input, BuiltInComponents->BlendChannelInput, Output->BlendChannelID);
+				}
+				else
+				{
+					// If the bound material changed, we might have been re-assigned a different blend channel so make sure it's up to date
+					Linker->EntityManager.WriteComponentChecked(Input, BuiltInComponents->BlendChannelInput, Output->BlendChannelID);
+				}
+
+
+				// Ensure we have the blender type tag on the inputs.
+				Linker->EntityManager.AddComponent(Input, BlenderTypeTag);
+			}
+		}
+		else if (!Output->OutputEntityID && Inputs.Num() == 1)
+		{
+			Linker->EntityManager.RemoveComponent(Inputs[0], BuiltInComponents->BlendChannelInput);
+		}
+
+		Output->NumContributors = NumContributors;
+	}
+
+	void DestroyOutput(FEntityGroupID GroupID, FAnimatedControlRigParameterInfo* Output, FEntityOutputAggregate Aggregate)
+	{
+		if (Output->OutputEntityID)
+		{
+			FBuiltInComponentTypes* BuiltInComponents = FBuiltInComponentTypes::Get();
+
+			Linker->EntityManager.AddComponent(Output->OutputEntityID, BuiltInComponents->Tags.NeedsUnlink);
+			Output->OutputEntityID = FMovieSceneEntityID();
+
+			if (System->DoubleBlenderSystem)
+			{
+				System->DoubleBlenderSystem->ReleaseBlendChannel(Output->BlendChannelID);
+			}
+			Output->BlendChannelID = FMovieSceneBlendChannelID();
+		}
+	}
+
+	void InitializeBlendOutput(FComponentTypeID BlenderTypeTag, TArrayView<const FMovieSceneEntityID> Inputs, FAnimatedControlRigParameterInfo* Output)
+	{
+		FBuiltInComponentTypes*          BuiltInComponents    = FBuiltInComponentTypes::Get();
+		FMovieSceneTracksComponentTypes* TracksComponents     = FMovieSceneTracksComponentTypes::Get();
+		FControlRigComponentTypes*       ControlRigComponents = FControlRigComponentTypes::Get();
+
+		FComponentMask ChannelResults = {
+			BuiltInComponents->DoubleResult[0],
+			BuiltInComponents->DoubleResult[1],
+			BuiltInComponents->DoubleResult[2],
+			BuiltInComponents->DoubleResult[3],
+			BuiltInComponents->DoubleResult[4],
+			BuiltInComponents->DoubleResult[5],
+			BuiltInComponents->DoubleResult[6],
+			BuiltInComponents->DoubleResult[7],
+			BuiltInComponents->DoubleResult[8]
+		};
+
+		FTypelessMutation Mutation;
+		for (FMovieSceneEntityID Input : Inputs)
+		{
+			FComponentMask Type = Linker->EntityManager.GetEntityType(Input);
+			Type.CombineWithBitwiseAND(ChannelResults, EBitwiseOperatorFlags::MinSize);
+			Mutation.AddMask.CombineWithBitwiseOR(Type, EBitwiseOperatorFlags::MaxSize);
+		}
+		
+		// Remove any channels not present in any of the inputs by taking an XOR of all the channels with the AddMask
+		Mutation.RemoveMask = MoveTemp(ChannelResults);
+		Mutation.RemoveMask.CombineWithBitwiseXOR(Mutation.AddMask, EBitwiseOperatorFlags::MaintainSize);
+
+		auto Builder = FEntityBuilder()
+			.Add(BuiltInComponents->BlendChannelOutput, Output->BlendChannelID)
+			.AddTag(ControlRigComponents->Tags.ControlRigParameter)
+			.AddTag(BlenderTypeTag)
+			.AddTag(BuiltInComponents->Tags.NeedsLink)
+			.AddMutualComponents();
+
+		if (Output->OutputEntityID)
+		{
+			Builder.MutateExisting(&Linker->EntityManager, Output->OutputEntityID, Mutation);
+		}
+		else
+		{
+			Output->OutputEntityID = Builder.CreateEntity(&Linker->EntityManager, Mutation.AddMask);
+			Linker->EntityManager.CopyComponents(Inputs[0], Output->OutputEntityID, Linker->EntityManager.GetComponents()->GetCopyAndMigrationMask());
+		}
+	}
+
+};
+
+
+
+struct FEvaluateBaseControlRigs
+{
+	FInstanceRegistry* InstanceRegistry;
+
+	FEvaluateBaseControlRigs(FInstanceRegistry* InInstanceRegistry)
+		: InstanceRegistry(InInstanceRegistry)
+	{
+	}
+
+	void ForEachEntity(
+		FRootInstanceHandle RootInstanceHandle,
+		const FMovieSceneSequenceID* OptSequenceID,
+		FFrameTime EvalTime,
+		float EvalTimeSeconds,
+		const double* OptWeightAndEasing,
+		const FControlRigSourceData& ControlRigSource,
+		FBaseControlRigEvalData& InOutBaseEvalData) const
+	{
+		UControlRig*   ControlRig   = InOutBaseEvalData.WeakControlRig.Get();
+		URigHierarchy* RigHierarchy = ControlRig ? ControlRig->GetHierarchy() : nullptr;
+
+		if (ControlRig == nullptr || RigHierarchy == nullptr)
+		{
+			return;
+		}
+
+		InOutBaseEvalData.bWasDoNotKey = InOutBaseEvalData.Section->GetDoNotKey();
+		InOutBaseEvalData.Section->SetDoNotKey(false);
+
+		if (InOutBaseEvalData.bIsActive)
+		{
+			if (USkeletalMeshComponent* SkeletalMeshComponent = Cast<USkeletalMeshComponent>(ControlRig->GetObjectBinding()->GetBoundObject()))
+			{
+				if (UControlRigLayerInstance* AnimInstance = Cast<UControlRigLayerInstance>(SkeletalMeshComponent->GetAnimInstance()))
+				{
+					const float Weight = static_cast<float>(OptWeightAndEasing ? *OptWeightAndEasing : 1.0);
+
+					FControlRigIOSettings InputSettings;
+					InputSettings.bUpdateCurves = true;
+					InputSettings.bUpdatePose = true;
+
+					AnimInstance->UpdateControlRigTrack(ControlRig->GetUniqueID(), Weight, InputSettings, true);
+				}
+			}
+
+			const bool bSetupUndo = false;
+			ControlRig->SetAbsoluteTime(EvalTimeSeconds);
+		}
+
+		UObject* BoundObject = ControlRig->GetObjectBinding()->GetBoundObject();
+
+		// when playing animation, instead of scrubbing/stepping thru frames, InTime might have a subframe of 0.999928
+		// leading to a decimal value of 24399.999928 (for example). This results in evaluating one frame less than expected
+		// (24399 instead of 24400) and leads to spaces and constraints switching parents/state after the control changes
+		// its transform. (float/double channels will interpolate to a value pretty close to the one at 24400 as its based
+		// on that 0.999928 subframe value.
+		const FFrameTime RoundTime = EvalTime.RoundToFrame();
+
+		TSharedRef<FSharedPlaybackState> SharedPlaybackState = InstanceRegistry->GetInstance(RootInstanceHandle).GetSharedPlaybackState();
+
+		FMovieSceneSequenceID SequenceID = OptSequenceID ? *OptSequenceID : MovieSceneSequenceID::Root;
+
+		for (const FConstraintAndActiveChannel& ConstraintAndActiveChannel : InOutBaseEvalData.Section->GetConstraintsChannels())
+		{
+			bool Value = false;
+			ConstraintAndActiveChannel.ActiveChannel.Evaluate(RoundTime, Value);
+			CreateConstraintIfNeeded(BoundObject->GetWorld(), ConstraintAndActiveChannel.GetConstraint().Get(), InOutBaseEvalData.Section);
+
+			if (ConstraintAndActiveChannel.GetConstraint())
+			{
+				//For Control Rig we may need to explicitly set the control rig
+				if (UTickableTransformConstraint* TransformConstraint = Cast<UTickableTransformConstraint>(ConstraintAndActiveChannel.GetConstraint()))
+				{
+					TransformConstraint->InitConstraint(BoundObject->GetWorld());
+				}
+				ConstraintAndActiveChannel.GetConstraint()->ResolveBoundObjects(SequenceID, SharedPlaybackState, ControlRig);
+				ConstraintAndActiveChannel.GetConstraint()->SetActive(Value);
+			}
+		}
+
+		//unfortunately for Constraints with ControlRig we need to resolve all Parents also. Don't need to do children since they wil be handled by
+		//the channel resolve above
+		const FConstraintsManagerController& Controller = FConstraintsManagerController::Get(BoundObject->GetWorld());
+		const TArray< TWeakObjectPtr<UTickableConstraint>> AllConstraints = Controller.GetAllConstraints();
+		for (const TWeakObjectPtr<UTickableConstraint>& TickConstraint : AllConstraints)
+		{
+			if (UTickableTransformConstraint* TransformConstraint = Cast< UTickableTransformConstraint>(TickConstraint.Get()))
+			{
+				if (TransformConstraint->ParentTRSHandle)
+				{
+					TransformConstraint->ParentTRSHandle->ResolveBoundObjects(SequenceID, SharedPlaybackState, ControlRig);
+					TransformConstraint->EnsurePrimaryDependency(BoundObject->GetWorld());
+				}
+			}
+		}
+	}
+};
+
+
+/** This should logically be in its own system, but that is unlikely to be of use anywhere else ,
+, currently not called since we do space eval in the base*/
+struct FEvaluateControlRigSpaceChannels
+{
+	static void ForEachEntity(FFrameTime EvalTime, const FMovieSceneControlRigSpaceChannel* Channel, FMovieSceneControlRigSpaceBaseKey& OutValue)
+	{
+		// when playing animation, instead of scrubbing/stepping thru frames, InTime might have a subframe of 0.999928
+		// leading to a decimal value of 24399.999928 (for example). This results in evaluating one frame less than expected
+		// (24399 instead of 24400) and leads to spaces and constraints switching parents/state after the control changes
+		// its transform. (float/double channels will interpolate to a value pretty close to the one at 24400 as its based
+		// on that 0.999928 subframe value.
+		Channel->Evaluate(EvalTime.RoundToFrame(), OutValue);
+	}
+};
+
+
+
+struct FGatherControlRigParameterValues
+{
+	FAccumulatedControlRigValues* AccumulatedValues;
+	FControlRigComponentTypes* ControlRigComponents = FControlRigComponentTypes::Get();
+	FMovieSceneTracksComponentTypes* TracksComponents = FMovieSceneTracksComponentTypes::Get();
+
+	void ForEachAllocation(FEntityAllocationIteratorItem Item, const FAccumulatedControlEntryIndex* EntryIndices,
+		const FName* ParameterNames,
+		const FMovieSceneControlRigSpaceBaseKey* OptSpaceResults, const bool* OptBoolResults, const uint8* OptByteResults, const int32* OptIntegerResults,
+		const double* OptDoubleResults0, const double* OptDoubleResults1, const double* OptDoubleResults2,
+		const double* OptDoubleResults3, const double* OptDoubleResults4, const double* OptDoubleResults5,
+		const double* OptDoubleResults6, const double* OptDoubleResults7, const double* OptDoubleResults8
+		) const
+	{
+		const FComponentMask&    AllocationType = Item.GetAllocationType();
+		const FEntityAllocation* Allocation     = Item.GetAllocation();
+
+		const int32 Num = Allocation->Num();
+
+		if (AllocationType.Contains(ControlRigComponents->Tags.Space))
+		{
+			check(OptSpaceResults);
+			for (int32 Index = 0; Index < Num; ++Index)
+			{
+				AccumulatedValues->Store(EntryIndices[Index], OptSpaceResults[Index]);
+			}
+		}
+		else if (AllocationType.Contains(TracksComponents->Parameters.Bool.PropertyTag))
+		{
+			check(OptBoolResults);
+			for (int32 Index = 0; Index < Num; ++Index)
+			{
+				AccumulatedValues->Store(EntryIndices[Index], OptBoolResults[Index]);
+			}
+		}
+		else if (AllocationType.Contains(TracksComponents->Parameters.Byte.PropertyTag))
+		{
+			check(OptByteResults);
+			for (int32 Index = 0; Index < Num; ++Index)
+			{
+				AccumulatedValues->Store(EntryIndices[Index], OptByteResults[Index]);
+			}
+		}
+		else if (AllocationType.Contains(TracksComponents->Parameters.Integer.PropertyTag))
+		{
+			check(OptIntegerResults);
+			for (int32 Index = 0; Index < Num; ++Index)
+			{
+				AccumulatedValues->Store(EntryIndices[Index], OptIntegerResults[Index]);
+			}
+		}
+		else if (AllocationType.Contains(TracksComponents->Parameters.Scalar.PropertyTag))
+		{
+			check(OptDoubleResults0);
+			for (int32 Index = 0; Index < Num; ++Index)
+			{
+				AccumulatedValues->Store(EntryIndices[Index], static_cast<float>(OptDoubleResults0[Index]));
+			}
+		}
+		// Vector2 is applied as vector3 in Control Rig
+		else if (AllocationType.Contains(TracksComponents->Parameters.Vector3.PropertyTag))
+		{
+			const bool bFullyAnimated = OptDoubleResults0 && OptDoubleResults1 && OptDoubleResults2;
+
+			if (bFullyAnimated)
+			{
+				for (int32 Index = 0; Index < Num; ++Index)
+				{
+					FVector3f Result(
+						OptDoubleResults0[Index],
+						OptDoubleResults1[Index],
+						OptDoubleResults2[Index]
+					);
+					AccumulatedValues->Store(EntryIndices[Index], Result);
+				}
+			}
+			else
+			{
+				TOptionalComponentReader<FFloatIntermediateVector> OptInitialValues =
+					Allocation->TryReadComponents(TracksComponents->Parameters.Vector3.InitialValue);
+
+				for (int32 Index = 0; Index < Num; ++Index)
+				{
+					FVector3f Vector = OptInitialValues
+						? OptInitialValues[Index].AsVector3f()
+						: FVector3f(0.f, 0.f, 0.f);
+
+					if (OptDoubleResults0)
+					{
+						Vector.X = OptDoubleResults0[Index];
+					}
+					if (OptDoubleResults1)
+					{
+						Vector.Y = OptDoubleResults1[Index];
+					}
+					if (OptDoubleResults2)
+					{
+						Vector.Z = OptDoubleResults2[Index];
+					}
+					AccumulatedValues->Store(EntryIndices[Index], Vector);
+				}
+			}
+		}
+		else if (AllocationType.Contains(TracksComponents->Parameters.Transform.PropertyTag))
+		{
+			const bool bFullyAnimated =
+				OptDoubleResults0 && OptDoubleResults1 && OptDoubleResults2 &&
+				OptDoubleResults3 && OptDoubleResults4 && OptDoubleResults5 &&
+				OptDoubleResults6 && OptDoubleResults7 && OptDoubleResults8;
+
+			if (bFullyAnimated)
+			{
+				for (int32 Index = 0; Index < Num; ++Index)
+				{
+					FEulerTransform Result(
+						FVector( OptDoubleResults0[Index], OptDoubleResults1[Index], OptDoubleResults2[Index]),
+						FRotator(OptDoubleResults4[Index], OptDoubleResults5[Index], OptDoubleResults3[Index]),
+						FVector( OptDoubleResults6[Index], OptDoubleResults7[Index], OptDoubleResults8[Index])
+					);
+					AccumulatedValues->Store(EntryIndices[Index], Result);
+				}
+			}
+			else
+			{
+				TOptionalComponentReader<FIntermediate3DTransform> OptInitialValues =
+					Allocation->TryReadComponents(TracksComponents->Parameters.Transform.InitialValue);
+
+				for (int32 Index = 0; Index < Num; ++Index)
+				{
+					FEulerTransform Transform = OptInitialValues
+						? OptInitialValues[Index].AsEuler()
+						: FEulerTransform::Identity;
+
+					if (OptDoubleResults0) { Transform.Location.X     = OptDoubleResults0[Index]; }
+					if (OptDoubleResults1) { Transform.Location.Y     = OptDoubleResults1[Index]; }
+					if (OptDoubleResults2) { Transform.Location.Z     = OptDoubleResults2[Index]; }
+
+					if (OptDoubleResults3) { Transform.Rotation.Roll  = OptDoubleResults3[Index]; }
+					if (OptDoubleResults4) { Transform.Rotation.Pitch = OptDoubleResults4[Index]; }
+					if (OptDoubleResults5) { Transform.Rotation.Yaw   = OptDoubleResults5[Index]; }
+
+					if (OptDoubleResults6) { Transform.Scale.X        = OptDoubleResults6[Index]; }
+					if (OptDoubleResults7) { Transform.Scale.Y        = OptDoubleResults7[Index]; }
+					if (OptDoubleResults8) { Transform.Scale.Z        = OptDoubleResults8[Index]; }
+
+					AccumulatedValues->Store(EntryIndices[Index], Transform);
+				}
+			}
+		}
+	}
+};
+
+
+struct FApplyControlRigParameterValuesTask
+{
+	const FAccumulatedControlRigValues* AccumulatedValues;
+	bool bApplyRigs = true;
+
+	FApplyControlRigParameterValuesTask(const FAccumulatedControlRigValues* InAccumulatedValues, bool bInApplyRigs)
+		: AccumulatedValues(InAccumulatedValues)
+		, bApplyRigs(bInApplyRigs)
+	{
+	}
+
+	void Run(FEntityAllocationWriteContext WriteContext) const
+	{
+		if (bApplyRigs)
+		{
+			AccumulatedValues->Apply();
+		}
+	}
+};
+
+
+struct FResetDoNotKey
+{
+	void ForEachEntity(const FControlRigSourceData& ControlRigSource, const FBaseControlRigEvalData& InOutBaseEvalData) const
+	{
+		InOutBaseEvalData.Section->SetDoNotKey(InOutBaseEvalData.bWasDoNotKey);
+	}
+};
+
+} // namespace UE::MovieScene
+
+
+UMovieSceneControlRigParameterEvaluatorSystem::UMovieSceneControlRigParameterEvaluatorSystem(const FObjectInitializer& ObjInit)
+	: Super(ObjInit)
+{
+	using namespace UE::MovieScene;
+
+	Phase = ESystemPhase::Instantiation | ESystemPhase::Scheduling;
+
+	if (HasAnyFlags(RF_ClassDefaultObject))
+	{
+		FBuiltInComponentTypes* BuiltInComponents = FBuiltInComponentTypes::Get();
+		FControlRigComponentTypes* ControlRigComponents = FControlRigComponentTypes::Get();
+		FMovieSceneTracksComponentTypes* TracksComponents = FMovieSceneTracksComponentTypes::Get();
+
+		TSharedRef<FControlRigInitialValueProcessor> InitialValueProcessor =
+			MakeShared<FControlRigInitialValueProcessor>();
+
+		UMovieSceneInitialValueSystem::RegisterProcessor(TracksComponents->Parameters.Scalar.InitialValue, InitialValueProcessor);
+		UMovieSceneInitialValueSystem::RegisterProcessor(TracksComponents->Parameters.Vector2.InitialValue, InitialValueProcessor);
+		UMovieSceneInitialValueSystem::RegisterProcessor(TracksComponents->Parameters.Vector3.InitialValue, InitialValueProcessor);
+		UMovieSceneInitialValueSystem::RegisterProcessor(TracksComponents->Parameters.Color.InitialValue, InitialValueProcessor);
+		UMovieSceneInitialValueSystem::RegisterProcessor(TracksComponents->Parameters.Transform.InitialValue, InitialValueProcessor);
+
+
+		DefineComponentConsumer(GetClass(), ControlRigComponents->ControlRigSource);
+		DefineComponentConsumer(GetClass(), BuiltInComponents->HierarchicalBlendTarget);
+		DefineComponentConsumer(GetClass(), BuiltInComponents->BoundObject);
+
+		DefineImplicitPrerequisite(UFloatChannelEvaluatorSystem::StaticClass(), GetClass());
+		DefineImplicitPrerequisite(UDoubleChannelEvaluatorSystem::StaticClass(), GetClass());
+
+
+		for (int32 Index = 0; Index < UE_ARRAY_COUNT(BuiltInComponents->DoubleResult); ++Index)
+		{
+			DefineComponentConsumer(GetClass(), BuiltInComponents->DoubleResult[Index]);
+		}
+		DefineImplicitPrerequisite(UMovieScenePiecewiseDoubleBlenderSystem::StaticClass(), GetClass());
+		DefineImplicitPrerequisite(UMovieSceneHierarchicalEasingInstantiatorSystem::StaticClass(), GetClass());
+		DefineImplicitPrerequisite(GetClass(), UMovieSceneInitialValueSystem::StaticClass());
+		DefineImplicitPrerequisite(GetClass(), UMovieSceneRestorePreAnimatedStateSystem::StaticClass());
+	}
+}
+
+bool UMovieSceneControlRigParameterEvaluatorSystem::IsRelevantImpl(UMovieSceneEntitySystemLinker* InLinker) const
+{
+	using namespace UE::MovieScene;
+
+	const bool bCanBeEnabled = true;// !UMovieSceneControlRigParameterTrack::ShouldUseLegacyTemplate();
+	return bCanBeEnabled && InLinker->EntityManager.ContainsComponent(FControlRigComponentTypes::Get()->ControlRigSource);
+}
+
+void UMovieSceneControlRigParameterEvaluatorSystem::OnLink()
+{
+	using namespace UE::MovieScene;
+
+	ControlRigParameterTracker.Initialize(this);
+
+	ControlRigStorage = Linker->PreAnimatedState.GetOrCreateStorage<FPreAnimatedControlRigStorage>();
+	ControlRigParameterStorage = Linker->PreAnimatedState.GetOrCreateStorage<FPreAnimatedControlRigParameterStorage>();
+
+	FControlRigComponentTypes* ControlRigComponents = FControlRigComponentTypes::Get();
+	FMovieSceneTracksComponentTypes* TracksComponents = FMovieSceneTracksComponentTypes::Get();
+	FBuiltInComponentTypes* BuiltInComponents = FBuiltInComponentTypes::Get();
+
+	UMovieSceneEntityGroupingSystem* GroupingSystem = Linker->LinkSystem<UMovieSceneEntityGroupingSystem>();
+
+	ParameterGroupingKey = GroupingSystem->AddGrouping(
+		FControlRigDataGroupingPolicy(),
+		ControlRigComponents->ControlRigSource,
+		TracksComponents->GenericParameterName
+	);
+}
+
+UControlRig* UMovieSceneControlRigParameterEvaluatorSystem::GetRigFromTrack(UMovieSceneControlRigParameterTrack* Track) const
+{
+	return AccumulatedValues.FindControlRig(Track);
+}
+
+const UE::MovieScene::FControlRigParameterBuffer* UMovieSceneControlRigParameterEvaluatorSystem::FindParameters(UMovieSceneControlRigParameterTrack* Track) const
+{
+	return AccumulatedValues.FindParameterBuffer(Track);
+}
+
+void UMovieSceneControlRigParameterEvaluatorSystem::OnUnlink()
+{
+	using namespace UE::MovieScene;
+
+	// Always reset the float blender system on link to ensure that recycled systems are correctly initialized.
+	DoubleBlenderSystem = nullptr;
+
+	ControlRigParameterTracker.Destroy(FOverlappingControlRigParameterHandler(this));
+
+	UMovieSceneEntityGroupingSystem* GroupingSystem = Linker->FindSystem<UMovieSceneEntityGroupingSystem>();
+	if (ensure(GroupingSystem))
+	{
+		GroupingSystem->RemoveGrouping(ParameterGroupingKey);
+	}
+	ParameterGroupingKey = FEntityGroupingPolicyKey();
+
+}
+
+void UMovieSceneControlRigParameterEvaluatorSystem::OnRun(FSystemTaskPrerequisites& InPrerequisites, FSystemSubsequentTasks& Subsequents)
+{
+	using namespace UE::MovieScene;
+
+	TSharedRef<FMovieSceneEntitySystemRunner> Runner = Linker->GetRunner();
+
+	ESystemPhase CurrentPhase = Runner->GetCurrentPhase();
+	if (CurrentPhase == ESystemPhase::Instantiation)
+	{
+		OnInstantiation();
+	}
+	else if (CurrentPhase == ESystemPhase::Evaluation)
+	{
+		OnEvaluation(InPrerequisites, Subsequents);
+	}
+}
+
+void UMovieSceneControlRigParameterEvaluatorSystem::OnInstantiation()
+{
+	using namespace UE::MovieScene;
+
+	FBuiltInComponentTypes* BuiltInComponents = FBuiltInComponentTypes::Get();
+	FControlRigComponentTypes* ControlRigComponents = FControlRigComponentTypes::Get();
+	FMovieSceneTracksComponentTypes* TracksComponents = FMovieSceneTracksComponentTypes::Get();
+
+
+	FEntityComponentFilter ChangedFilter;
+	ChangedFilter.Any({ ControlRigComponents->BaseControlRigEvalData, ControlRigComponents->Tags.ControlRigParameter });
+	ChangedFilter.Any({ BuiltInComponents->Tags.NeedsLink, BuiltInComponents->Tags.NeedsUnlink });
+	if (!Linker->EntityManager.Contains(ChangedFilter))
+	{
+		return;
+	}
+
+	// Prime the existing container so we can track what needs to be destroyed
+	AccumulatedValues.PrimeForInstantiation();
+
+	// Keep track of which base eval component is the 'active' one based on whether it is blended or not
+	TMap<UControlRig*, FBaseControlRigEvalData*> EncounteredRigs;
+
+	// --------------------------------------------------------------------------------------------------------
+	// Initialize base control rig components.
+	if (Linker->GetLinkerRole() == EEntitySystemLinkerRole::Interrogation)
+	{
+		IInterrogationExtension* Interrogation = Linker->FindExtension<IInterrogationExtension>();
+		check(Interrogation);
+
+		FEntityTaskBuilder()
+		.ReadEntityIDs()
+		.Read(BuiltInComponents->Interrogation.InputKey)
+		.Read(ControlRigComponents->ControlRigSource)
+		.Write(ControlRigComponents->BaseControlRigEvalData)
+		.PassthroughFilter(FEntityComponentFilter().All({ BuiltInComponents->WeightAndEasingResult }))
+		.FilterNone({ BuiltInComponents->Tags.NeedsUnlink })
+		.Iterate_PerEntity(&Linker->EntityManager, [this, Interrogation, &EncounteredRigs](FMovieSceneEntityID EntityID, FInterrogationKey InterrogationKey, const FControlRigSourceData& ControlRigSource, FBaseControlRigEvalData& OutBaseData, bool bHasWeight)
+			{
+				UControlRig* Rig = Cast<UControlRig>(Interrogation->GetSparseChannelInfo().FindObject(InterrogationKey.Channel));
+
+				OutBaseData.bIsActive = false;
+				OutBaseData.bHasWeight = bHasWeight;
+				OutBaseData.WeakControlRig = Rig;
+
+				FBaseControlRigEvalData** ExistingData = EncounteredRigs.Find(Rig);
+				if (ExistingData)
+				{
+					if ((*ExistingData)->bHasWeight)
+					{
+						return;
+					}
+
+					// Previous one wasn't weighted but this is, this should take over
+					if (bHasWeight)
+					{
+						(*ExistingData)->bIsActive = false;
+						OutBaseData.bIsActive = true;
+						*ExistingData = &OutBaseData;
+
+						// Don't reinitialize
+						return;
+					}
+				}
+				else
+				{
+					OutBaseData.bIsActive = true;
+					EncounteredRigs.Add(Rig, &OutBaseData);
+				}
+
+				this->AccumulatedValues.InitializeRig(ControlRigSource.Track, Rig);
+			}
+		);
+	}
+	else
+	{
+		FEntityTaskBuilder()
+		.ReadEntityIDs()
+		.Read(BuiltInComponents->GenericObjectBinding)
+		.Read(BuiltInComponents->RootInstanceHandle)
+		.ReadOptional(BuiltInComponents->SequenceID)
+		.Read(ControlRigComponents->ControlRigSource)
+		.Write(ControlRigComponents->BaseControlRigEvalData)
+		.PassthroughFilter(FEntityComponentFilter().All({ BuiltInComponents->Tags.RestoreState }))
+		.PassthroughFilter(FEntityComponentFilter().All({ BuiltInComponents->WeightAndEasingResult }))
+		.FilterNone({ BuiltInComponents->Tags.NeedsUnlink })
+		.Iterate_PerEntity(&Linker->EntityManager, [this, &EncounteredRigs](FMovieSceneEntityID EntityID, const FGuid& ObjectBindingID, FRootInstanceHandle RootInstanceHandle, const FMovieSceneSequenceID* OptSequenceID,
+				const FControlRigSourceData& ControlRigSource, FBaseControlRigEvalData& OutBaseData, bool bWantsRestoreState, bool bHasWeight)
+			{
+
+				OutBaseData.bIsActive = false;
+				OutBaseData.bHasWeight = bHasWeight;
+				OutBaseData.WeakControlRig = nullptr;
+
+				FMovieSceneSequenceID SequenceID = OptSequenceID ? *OptSequenceID : MovieSceneSequenceID::Root;
+				TArrayView<TWeakObjectPtr<>> BoundObjects = Linker->GetInstanceRegistry()->GetInstance(RootInstanceHandle).GetSharedPlaybackState()->FindBoundObjects(ObjectBindingID, SequenceID);
+
+				for (TWeakObjectPtr<> WeakBoundObject : BoundObjects)
+				{
+					UObject* BoundObject = WeakBoundObject.Get();
+					if (BoundObject)
+					{
+						this->InitializeBaseRigComponent(
+							BoundObject,
+							EntityID,
+							RootInstanceHandle,
+							bWantsRestoreState,
+							bHasWeight,
+							ControlRigSource,
+							OutBaseData,
+							EncounteredRigs);
+						return;
+					}
+				}
+			}
+		);
+	}
+
+	// Compact the accululation buffers if necessary
+	AccumulatedValues.Compact();
+
+	// --------------------------------------------------------------------------------------------------------
+	// Process overlapping control rig parameter entities that animate the same parameter based on the group ID
+	{
+		FEntityTaskBuilder()
+		.Read(BuiltInComponents->Group)
+		.PassthroughFilter(FEntityComponentFilter().All({ BuiltInComponents->Tags.NeedsLink }))
+		.FilterAny({ BuiltInComponents->Tags.NeedsLink, BuiltInComponents->Tags.NeedsUnlink })
+		.FilterAll({ ControlRigComponents->ControlRigSource, TracksComponents->GenericParameterName, ControlRigComponents->Tags.ControlRigParameter })
+		.Iterate_PerAllocation(&Linker->EntityManager, [this](const FEntityAllocation* Allocation, TRead<FEntityGroupID> Group, bool bIsBeingLinked)
+			{
+				if (bIsBeingLinked)
+				{
+					this->ControlRigParameterTracker.VisitActiveAllocation(Allocation, Group);
+				}
+				else
+				{
+					this->ControlRigParameterTracker.VisitUnlinkedAllocation(Allocation);
+				}
+			}
+		);
+
+		FOverlappingControlRigParameterHandler Handler(this);
+		ControlRigParameterTracker.ProcessInvalidatedOutputs(Linker, Handler);
+	}
+
+	// --------------------------------------------------------------------------------------------------------
+	// Track pre-animated state for parameters
+	FEntityTaskBuilder()
+	.ReadEntityIDs()
+	.Read(BuiltInComponents->RootInstanceHandle)
+	.Read(ControlRigComponents->ControlRigSource)
+	.Read(TracksComponents->GenericParameterName)
+	.PassthroughFilter(FEntityComponentFilter().All({ BuiltInComponents->Tags.RestoreState }))
+	.FilterAll({ BuiltInComponents->Tags.NeedsLink })
+	.FilterNone({ BuiltInComponents->Tags.Ignored })
+	.Iterate_PerAllocation(&Linker->EntityManager, [this](FEntityAllocationIteratorItem Item,
+		TRead<FMovieSceneEntityID> EntityIDs,
+		TRead<FRootInstanceHandle> RootInstanceHandles,
+		TRead<FControlRigSourceData> ControlRigSources,
+		TRead<FName> ParameterNames,
+		const bool bWantsRestore)
+	{
+		if (this->ControlRigParameterStorage->IsCapturingGlobalState() || bWantsRestore)
+		{
+			const int32 Num = Item.GetAllocation()->Num();
+
+			for (int32 Index = 0; Index < Num; ++Index)
+			{
+				UControlRig* Rig = this->GetRigFromTrack(ControlRigSources[Index].Track);
+				if (Rig)
+				{
+					this->ControlRigParameterStorage->BeginTrackingEntity(EntityIDs[Index], bWantsRestore, RootInstanceHandles[Index], Rig, ParameterNames[Index]);
+				}
+			}
+		}
+	});
+
+
+	// --------------------------------------------------------------------------------------------------------
+	// Regather all active parameters into our allocation buffers
+	{
+		FEntityComponentFilter MutationFilter;
+		MutationFilter.All({
+			TracksComponents->GenericParameterName,
+			ControlRigComponents->Tags.ControlRigParameter,
+		});
+		
+		MutationFilter.None({ BuiltInComponents->BlendChannelInput, BuiltInComponents->Tags.NeedsUnlink });
+
+		// Now mutate it all to assign new allocation entries
+		FControlRigAccumulationEntryIndexMutation AccumulationEntryMutation(&AccumulatedValues, ControlRigParameterStorage.Get());
+		Linker->EntityManager.MutateAll(MutationFilter, AccumulationEntryMutation);
+
+		FRemoveInvalidControlRigAccumulationComponents RemoveInvalidMutation(&AccumulatedValues);
+		Linker->EntityManager.MutateConditional(MutationFilter, RemoveInvalidMutation);
+
+		// Initialize the parameter buffers
+		AccumulatedValues.InitializeParameters(*this->ControlRigParameterStorage.Get());
+	}
+
+	// --------------------------------------------------------------------------------------------------------
+	// Initialize initial values
+	{
+		FEntityComponentFilter MutationFilter;
+		MutationFilter.All({
+			ControlRigComponents->ControlRigSource,
+			ControlRigComponents->AccumulatedControlEntryIndex,
+			TracksComponents->GenericParameterName,
+			ControlRigComponents->Tags.ControlRigParameter,
+			BuiltInComponents->Tags.NeedsLink,
+		});
+		MutationFilter.None({ BuiltInComponents->BlendChannelInput });
+
+		FInitialControlRigParameterValueMutation InitialValueMutation(&AccumulatedValues);
+		Linker->EntityManager.MutateAll(MutationFilter, InitialValueMutation);
+	}
+}
+
+
+void UMovieSceneControlRigParameterEvaluatorSystem::OnSchedulePersistentTasks(UE::MovieScene::IEntitySystemScheduler* TaskScheduler)
+{
+	using namespace UE::MovieScene;
+
+	FBuiltInComponentTypes*          BuiltInComponents    = FBuiltInComponentTypes::Get();
+	FMovieSceneTracksComponentTypes* TracksComponents     = FMovieSceneTracksComponentTypes::Get();
+	FControlRigComponentTypes*       ControlRigComponents = FControlRigComponentTypes::Get();
+
+	// Evaluate base rigs before anything else
+	FTaskID EvalBaseRigs = FEntityTaskBuilder()
+	.Read(BuiltInComponents->RootInstanceHandle)
+	.ReadOptional(BuiltInComponents->SequenceID)
+	.Read(BuiltInComponents->EvalTime)
+	.Read(BuiltInComponents->EvalSeconds)
+	.ReadOptional(BuiltInComponents->WeightAndEasingResult)
+	.Read(ControlRigComponents->ControlRigSource)
+	.Write(ControlRigComponents->BaseControlRigEvalData)
+	.SetDesiredThread(Linker->EntityManager.GetGatherThread())
+	.Schedule_PerEntity<FEvaluateBaseControlRigs>(&Linker->EntityManager, TaskScheduler, Linker->GetInstanceRegistry());
+
+	// Evaluate space channels
+	FTaskID EvalSpaces = FEntityTaskBuilder()
+	.Read(BuiltInComponents->EvalTime)
+	.Read(ControlRigComponents->SpaceChannel)
+	.Write(ControlRigComponents->SpaceResult)
+	.Fork_PerEntity<FEvaluateControlRigSpaceChannels>(&Linker->EntityManager, TaskScheduler);
+
+	// Gather all (potentially blended) parameter values
+	FTaskID GatherAnimatedControlRigs = FEntityTaskBuilder()
+	.Read(ControlRigComponents->AccumulatedControlEntryIndex)
+	.Read(TracksComponents->GenericParameterName)
+	.ReadAnyOf(
+		ControlRigComponents->SpaceResult, BuiltInComponents->BoolResult, BuiltInComponents->ByteResult, BuiltInComponents->IntegerResult,
+		BuiltInComponents->DoubleResult[0], BuiltInComponents->DoubleResult[1], BuiltInComponents->DoubleResult[2],
+		BuiltInComponents->DoubleResult[3], BuiltInComponents->DoubleResult[4], BuiltInComponents->DoubleResult[5],
+		BuiltInComponents->DoubleResult[6], BuiltInComponents->DoubleResult[7], BuiltInComponents->DoubleResult[8])
+	.FilterAll({ ControlRigComponents->Tags.ControlRigParameter })
+	.FilterNone({ BuiltInComponents->BlendChannelInput, BuiltInComponents->Tags.Ignored })
+	.Schedule_PerAllocation<FGatherControlRigParameterValues>(&Linker->EntityManager, TaskScheduler, &AccumulatedValues);
+
+	FTaskID ApplyTask = TaskScheduler->AddTask<FApplyControlRigParameterValuesTask>(
+		FTaskParams(TEXT("Apply Control Rig Parameters")).ForceGameThread(),
+		&AccumulatedValues,
+		Linker->GetLinkerRole() != EEntitySystemLinkerRole::Interrogation
+	);
+
+	// Reset Do Not Key states on any thread
+	FTaskID ResetDoNotKey = FEntityTaskBuilder()
+	.Read(ControlRigComponents->ControlRigSource)
+	.Read(ControlRigComponents->BaseControlRigEvalData)
+	.Fork_PerEntity<FResetDoNotKey>(&Linker->EntityManager, TaskScheduler);
+
+	// Spaces must be evaluated before we gather results
+	TaskScheduler->AddPrerequisite(EvalSpaces, GatherAnimatedControlRigs);
+	// We must finish gathering parameter values before we apply them
+	TaskScheduler->AddPrerequisite(GatherAnimatedControlRigs, ApplyTask);
+	// Base control rigs must be evaluated before we apply parameters
+	TaskScheduler->AddPrerequisite(EvalBaseRigs, ApplyTask);
+	// Reset do not key states has to happen last (after parameters have been applied)
+	TaskScheduler->AddPrerequisite(ApplyTask, ResetDoNotKey);
+}
+
+void UMovieSceneControlRigParameterEvaluatorSystem::OnEvaluation(FSystemTaskPrerequisites& InPrerequisites, FSystemSubsequentTasks& Subsequents)
+{
+	// Not enabled for this codepath
+	ensure(false);
+}
+
+void UMovieSceneControlRigParameterEvaluatorSystem::InitializeBaseRigComponent(
+	UObject* BoundObject,
+	UE::MovieScene::FMovieSceneEntityID EntityID,
+	UE::MovieScene::FRootInstanceHandle RootInstanceHandle,
+	bool bWantsRestoreState,
+	bool bHasWeight,
+	UE::MovieScene::FControlRigSourceData ControlRigSource,
+	UE::MovieScene::FBaseControlRigEvalData& OutBaseData,
+	TMap<UControlRig*, UE::MovieScene::FBaseControlRigEvalData*>& OutBaseComponentTracker)
+{
+	using namespace UE::MovieScene;
+
+	UMovieSceneControlRigParameterTrack* Track = ControlRigSource.Track;
+	if (!Track || !BoundObject)
+	{
+		return;
+	}
+
+	UWorld* GameWorld = (BoundObject->GetWorld() && BoundObject->GetWorld()->IsGameWorld()) ? BoundObject->GetWorld() : nullptr;
+
+	UControlRig* ControlRig = Cast<UControlRig>(BoundObject);
+	if (!ControlRig)
+	{
+		if (GameWorld)
+		{
+			ControlRig = Track->GetGameWorldControlRig(GameWorld);
+		}
+		else
+		{
+			ControlRig = Track->GetControlRig();
+		}
+	}
+
+	if (!ControlRig)
+	{
+		return;
+	}
+
+	// Begin tracking this entity for the base pre-animated state that keeps the rig bound to our target object
+	this->ControlRigStorage->BeginTrackingEntity(
+		EntityID,
+		bWantsRestoreState,
+		RootInstanceHandle,
+		ControlRig);
+
+	if (!ControlRig->GetObjectBinding())
+	{
+		ControlRig->SetObjectBinding(MakeShared<FControlRigObjectBinding>());
+	}
+
+	if (ControlRig->GetObjectBinding()->GetBoundObject() != FControlRigObjectBinding::GetBindableObject(BoundObject))
+	{ 
+		ControlRig->GetObjectBinding()->BindToObject(BoundObject);
+
+		TArray<FName> SelectedControls = ControlRig->CurrentControlSelection();
+		ControlRig->Initialize();
+		if (USkeletalMeshComponent* SkeletalMeshComponent = Cast<USkeletalMeshComponent>(FControlRigObjectBinding::GetBindableObject(BoundObject)))
+		{
+			ControlRig->RequestInit();
+			ControlRig->SetBoneInitialTransformsFromSkeletalMeshComponent(SkeletalMeshComponent, true);
+			ControlRig->Evaluate_AnyThread();
+		};
+		if (GameWorld == nullptr && ControlRig->IsA<UFKControlRig>())// mz only in editor replace Fk Control rig, will look post 29.20 to see if this really needed but want to unblock folks
+		{
+			Track->ReplaceControlRig(ControlRig, true);
+		}
+		TArray<FName> NewSelectedControls = ControlRig->CurrentControlSelection();
+		if (SelectedControls != NewSelectedControls)
+		{
+			if (ControlRig)
+			{
+				ControlRig->ClearControlSelection();
+				for (const FName& Name : SelectedControls)
+				{
+					ControlRig->SelectControl(Name, true);
+				}
+			}
+		}
+	}
+
+	// make sure to pick the correct CR instance for the  Components to bind.
+	// In case of PIE + Spawnable Actor + CR component, sequencer should grab
+	// CR component's CR instance for evaluation, see comment in BindToSequencerInstance
+	// i.e. CR component should bind to the instance that it owns itself.
+	ControlRig = GetControlRig(Track, BoundObject);
+	if (!ControlRig)
+	{
+		return;
+	}
+
+	OutBaseData.WeakControlRig = ControlRig;
+
+	// Cache pre-animated value now that the control rig is bound to the correct object
+	this->ControlRigStorage->CachePreAnimatedValue(ControlRig);
+
+	FBaseControlRigEvalData** ExistingData = OutBaseComponentTracker.Find(ControlRig);
+	if (ExistingData)
+	{
+		if ((*ExistingData)->bHasWeight)
+		{
+			return;
+		}
+
+		// Previous one wasn't weighted but this is, this should take over
+		if (bHasWeight)
+		{
+			(*ExistingData)->bIsActive = false;
+			OutBaseData.bIsActive = true;
+			*ExistingData = &OutBaseData;
+
+			// Don't reinitialize
+			return;
+		}
+	}
+	else
+	{
+		OutBaseData.bIsActive = true;
+		OutBaseComponentTracker.Add(ControlRig, &OutBaseData);
+	}
+
+	FControlRigBindingHelper::BindToSequencerInstance(ControlRig);
+	this->AccumulatedValues.InitializeRig(Track, ControlRig);
+}
+
diff --git a/Engine/Plugins/Animation/ControlRig/Source/ControlRig/Public/Sequencer/MovieSceneControlRigComponentTypes.h b/Engine/Plugins/Animation/ControlRig/Source/ControlRig/Public/Sequencer/MovieSceneControlRigComponentTypes.h
new file mode 100644
--- /dev/null
+++ b/Engine/Plugins/Animation/ControlRig/Source/ControlRig/Public/Sequencer/MovieSceneControlRigComponentTypes.h
@@ -0,0 +1,73 @@
+// Copyright Epic Games, Inc. All Rights Reserved.
+
+#pragma once
+
+#include "CoreTypes.h"
+#include "Sequencer/MovieSceneControlRigParameterBuffer.h"
+#include "EntitySystem/MovieSceneEntityIDs.h"
+
+struct FMovieSceneControlRigSpaceChannel;
+struct FMovieSceneControlRigSpaceBaseKey;
+
+class UMovieSceneControlRigParameterTrack;
+class UMovieSceneControlRigParameterSection;
+
+
+namespace UE::MovieScene
+{
+
+struct FAccumulatedControlEntryIndex;
+
+/** Component data present aon all base and parameter Control Rig entities */
+struct FControlRigSourceData
+{
+	UMovieSceneControlRigParameterTrack* Track = nullptr;
+};
+
+/**
+ * Component that exists for base-eval control rig entities
+ */
+struct FBaseControlRigEvalData
+{
+	UMovieSceneControlRigParameterSection* Section = nullptr;
+
+	TWeakObjectPtr<UControlRig> WeakControlRig;
+	uint8 bIsActive : 1 = true;
+	uint8 bHasWeight : 1 = false;
+	uint8 bWasDoNotKey : 1 = false;
+};
+
+/**
+ * Singleton Control Rig component types
+ */
+struct FControlRigComponentTypes
+{
+public:
+	static FControlRigComponentTypes* Get();
+	static void Destroy();
+
+	TComponentTypeID<FControlRigSourceData> ControlRigSource;
+
+	TComponentTypeID<FBaseControlRigEvalData> BaseControlRigEvalData;
+
+	TComponentTypeID<FAccumulatedControlEntryIndex> AccumulatedControlEntryIndex;
+
+	TComponentTypeID<const FMovieSceneControlRigSpaceChannel*> SpaceChannel;
+	TComponentTypeID<FMovieSceneControlRigSpaceBaseKey> SpaceResult;
+
+	struct
+	{
+		FComponentTypeID BaseControlRig;
+		FComponentTypeID ControlRigParameter;
+		FComponentTypeID Space;
+
+		FComponentTypeID IgnoredBaseControlRig;
+	} Tags;
+
+private:
+	FControlRigComponentTypes();
+};
+
+
+} // namespace UE::MovieScene
+
diff --git a/Engine/Plugins/Animation/ControlRig/Source/ControlRig/Public/Sequencer/MovieSceneControlRigParameterBuffer.h b/Engine/Plugins/Animation/ControlRig/Source/ControlRig/Public/Sequencer/MovieSceneControlRigParameterBuffer.h
new file mode 100644
--- /dev/null
+++ b/Engine/Plugins/Animation/ControlRig/Source/ControlRig/Public/Sequencer/MovieSceneControlRigParameterBuffer.h
@@ -0,0 +1,324 @@
+// Copyright Epic Games, Inc. All Rights Reserved.
+
+#pragma once
+
+#include "CoreTypes.h"
+#include "ControlRig.h"
+#include "MovieSceneControlRigComponentTypes.h"
+#include "Rigs/RigHierarchyDefines.h"
+
+struct FRigControlElement;
+struct FMovieSceneControlRigSpaceBaseKey;
+
+class UMovieSceneControlRigParameterTrack;
+
+namespace UE::MovieScene
+{
+
+struct FPreAnimatedControlRigParameterStorage;
+
+enum class EControlRigControlType : uint8
+{
+	Space,
+	Parameter_Bool,
+	Parameter_Enum,
+	Parameter_Integer,
+	Parameter_Scalar,
+	Parameter_Vector,
+	Parameter_Transform,
+
+	Num,
+};
+
+namespace Private
+{
+	/** Template specializations that produce the correct enum value for a given c++ control type */
+	template<typename T>
+	static constexpr EControlRigControlType ControlRigTypeFromNativeType()
+	{
+		static_assert(!std::is_same_v<T, T>, "This type is not supported");
+		return EControlRigControlType::Num;
+	}
+
+	template<> constexpr EControlRigControlType ControlRigTypeFromNativeType<FMovieSceneControlRigSpaceBaseKey>(){ return EControlRigControlType::Space; }
+	template<> constexpr EControlRigControlType ControlRigTypeFromNativeType<bool>()                             { return EControlRigControlType::Parameter_Bool; }
+	template<> constexpr EControlRigControlType ControlRigTypeFromNativeType<uint8>()                            { return EControlRigControlType::Parameter_Enum; }
+	template<> constexpr EControlRigControlType ControlRigTypeFromNativeType<int32>()                            { return EControlRigControlType::Parameter_Integer; }
+	template<> constexpr EControlRigControlType ControlRigTypeFromNativeType<float>()                            { return EControlRigControlType::Parameter_Scalar; }
+	template<> constexpr EControlRigControlType ControlRigTypeFromNativeType<FVector3f>()                        { return EControlRigControlType::Parameter_Vector; }
+	template<> constexpr EControlRigControlType ControlRigTypeFromNativeType<FEulerTransform>()                  { return EControlRigControlType::Parameter_Transform; }
+}
+
+
+/**
+ * Enumeration that determines the stability of entries allocated within a FControlRigParameterBuffer.
+ * Depending on the use case of the buffer, it can be preferable to have entries stable or unstable.
+ * Typically:
+ *   - Stable is used when the buffer is used as a container that is only populated and assigned to rig.
+ *   - Unstable is used when parameter values need to be retrieved from the buffer by name.
+ */
+enum class EControlRigParameterBufferIndexStability
+{
+	/** Once allocated, the name:value index must not be changed. Typically used for fast, index-based assignment. */
+	Stable,
+	/** Entries are freely relocatable: allows for more efficient binary searching of parameters typically required in editor or introspection code. */
+	Unstable,
+};
+
+/**
+ * Map-like structure containing key/value pairs of FName:T using a binary-search.
+ * Parameter names and values are allocated contiguously rather than as pairs.
+ */
+struct CONTROLRIG_API FControlRigParameterValueHeader
+{
+	/**
+	 * Construct this header for a specific control type
+	 * */
+	FControlRigParameterValueHeader(EControlRigControlType InType, EControlRigParameterBufferIndexStability IndexStability);
+
+	/**
+	 * Move construction/assignment
+	 */
+	FControlRigParameterValueHeader(FControlRigParameterValueHeader&& RHS);
+	FControlRigParameterValueHeader& operator=(FControlRigParameterValueHeader&& RHS);
+
+	/**
+	 * Copy construction/assignment
+	 */
+	FControlRigParameterValueHeader(const FControlRigParameterValueHeader& RHS);
+	FControlRigParameterValueHeader& operator=(const FControlRigParameterValueHeader& RHS);
+
+	/**
+	 * Destructor
+	 */
+	~FControlRigParameterValueHeader();
+
+	void Reset();
+	EControlRigControlType GetType() const;
+	int32 Num() const;
+	void Apply(UControlRig* Rig) const;
+	void Apply(UControlRig* Rig, FName Name, int32 Index, const void* ParameterBuffer) const;
+	void ApplyAndRemove(UControlRig* Rig, FName InName);
+
+	TArrayView<const FName> GetNames() const;
+
+	uint8* GetParameterBuffer() const;
+
+	void* GetParameter(int32 Index) const;
+
+	template<typename T>
+	const T* GetParameterBuffer() const
+	{
+		check(Private::ControlRigTypeFromNativeType<T>() == Type);
+		return reinterpret_cast<const T*>(GetParameterBuffer());
+	}
+
+	template<typename T>
+	T* GetParameterBuffer()
+	{
+		check(Private::ControlRigTypeFromNativeType<T>() == Type);
+		return reinterpret_cast<T*>(GetParameterBuffer());
+	}
+
+	template<typename T>
+	TArrayView<const T> GetParameters() const
+	{
+		return TArrayView<const T>(GetParameterBuffer<T>(), NumElements);
+	}
+
+	template<typename T>
+	TArrayView<T> GetParameters()
+	{
+		return TArrayView<T>(GetParameterBuffer<T>(), NumElements);
+	}
+
+	int32 Add_GetIndex(FName InName);
+
+	void* Add_GetPtr(FName InName);
+
+	void Add(FName InName, const void* Value);
+
+	void Remove(FName InName);
+
+	bool Contains(FName InName) const;
+
+	int32 Find(FName InName) const;
+
+	void OptimizeForLookup();
+	void Resort();
+
+private:
+	TArrayView<FName> GetMutableNames();
+	void RemoveAt(int32 Index);
+	void InsertDefaulted(FName Name, int32 Index);
+	static uint8* GetParameterBuffer(uint8* BasePtr, uint16 Capacity, uint8 InAlignment);
+	size_t GetParameterSize() const;
+	void Reserve(size_t NewCapacity);
+
+private:
+
+	uint8* Data = nullptr;
+
+	uint16 Capacity;
+	uint16 NumElements;
+
+	uint8 Alignment : 7;
+	uint8 bStableIndices : 1;
+
+	EControlRigControlType Type;
+};
+
+
+/**
+ * Parameter buffer structure containing any number of control values arranged by type and name
+ */
+struct CONTROLRIG_API FControlRigParameterValues
+{
+	FControlRigParameterValues(EControlRigParameterBufferIndexStability IndexStability);
+
+	FControlRigParameterValueHeader& GetHeader(EControlRigControlType InType);
+	const FControlRigParameterValueHeader& GetHeader(EControlRigControlType InType) const;
+
+
+	template<typename T>
+	T* Add(FName InName)
+	{
+		FControlRigParameterValueHeader& Header = GetHeader(Private::ControlRigTypeFromNativeType<T>());
+		return static_cast<T*>(Header.Add_GetPtr(InName));
+	}
+
+	template<typename T>
+	void Add(FName InName, T&& InValue)
+	{
+		using DecayedType = std::decay_t<T>;
+		*Add<DecayedType>(InName) = Forward<T>(InValue);
+	}
+	template<typename T>
+	const T* Find(FName InName) const
+	{
+		const FControlRigParameterValueHeader& Header = GetHeader(Private::ControlRigTypeFromNativeType<T>());
+		const int32 Index = Header.Find(InName);
+		if (Index != INDEX_NONE)
+		{
+			return &Header.GetParameters<T>()[Index];
+		}
+		return nullptr;
+	}
+
+	template<typename T>
+	bool Find(FName InName, T& OutValue) const
+	{
+		const FControlRigParameterValueHeader& Header = GetHeader(Private::ControlRigTypeFromNativeType<T>());
+		const int32 Index = Header.Find(InName);
+		if (Index != INDEX_NONE)
+		{
+			OutValue = Header.GetParameters<T>()[Index];
+			return true;
+		}
+		return false;
+	}
+
+	~FControlRigParameterValues();
+
+	void CopyFrom(const FControlRigParameterValues& Other, FName ControlName);
+
+	void AddCurrentValue(UControlRig* Rig, FRigControlElement* ControlElement);
+
+	void ApplyTo(UControlRig* Rig) const;
+
+	void ApplyAndRemove(UControlRig* Rig, FName Name);
+
+	void PopulateFrom(UControlRig* Rig);
+
+	bool InitializeParameters(UControlRig* Rig, FPreAnimatedControlRigParameterStorage& Storage);
+
+	void OptimizeForLookup();
+
+	void Reset();
+
+private:
+
+	FControlRigParameterValueHeader Headers[(uint8)EControlRigControlType::Num];
+};
+
+
+struct CONTROLRIG_API FAccumulatedControlEntryIndex
+{
+	FAccumulatedControlEntryIndex() : EntryIndex(uint16(-1)), AccumulatorIndex(uint16(-1)) {};
+	FAccumulatedControlEntryIndex(uint16 InEntryIndex, uint16 InAccumulatorIndex, EControlRigControlType InControlType)
+		: EntryIndex(InEntryIndex)
+		, AccumulatorIndex(InAccumulatorIndex) 
+		, ControlType(InControlType)
+	{};
+
+	bool IsValid() const
+	{
+		return (EntryIndex != (uint16(-1)) && AccumulatorIndex != (uint16(-1)));
+	}
+
+	uint16 EntryIndex;
+	uint16 AccumulatorIndex;
+	EControlRigControlType ControlType;
+};
+
+struct CONTROLRIG_API FControlRigParameterBuffer
+{
+	FControlRigParameterBuffer(UControlRig* InRig, EControlRigParameterBufferIndexStability IndexStability);
+
+	TWeakObjectPtr<UControlRig> WeakControlRig;
+	FControlRigParameterValues Values;
+
+	void AddCurrentValue(UControlRig* Rig, FRigControlElement* ControlElement);
+
+	void Apply() const;
+	void Populate();
+};
+
+
+struct FAccumulatedControlRigValues
+{
+	void Apply() const;
+
+	UControlRig* FindControlRig(FAccumulatedControlEntryIndex Entry) const;
+	UControlRig* FindControlRig(UMovieSceneControlRigParameterTrack* Track) const;
+	const FControlRigParameterBuffer* FindParameterBuffer(UMovieSceneControlRigParameterTrack* Track) const;
+
+	bool DoesEntryExistForTrack(UMovieSceneControlRigParameterTrack* Track);
+	int32 InitializeRig(UMovieSceneControlRigParameterTrack* Track, UControlRig* InRig);
+	void MarkAsActive(int32 Index);
+
+	FAccumulatedControlEntryIndex AllocateEntryIndex(UMovieSceneControlRigParameterTrack* Track, FName Name, EControlRigControlType Type);
+
+	void InitializeParameters(FPreAnimatedControlRigParameterStorage& Storage);
+
+	void Store(FAccumulatedControlEntryIndex Entry, FMovieSceneControlRigSpaceBaseKey InValue);
+	void Store(FAccumulatedControlEntryIndex Entry, bool InValue);
+	void Store(FAccumulatedControlEntryIndex Entry, uint8 InValue);
+	void Store(FAccumulatedControlEntryIndex Entry, int32 InValue);
+	void Store(FAccumulatedControlEntryIndex Entry, float InValue);
+	void Store(FAccumulatedControlEntryIndex Entry, FVector3f InValue);
+	void Store(FAccumulatedControlEntryIndex Entry, FEulerTransform InValue);
+
+	void Compact();
+
+	void PrimeForInstantiation();
+
+private:
+
+	struct FEntry : FControlRigParameterBuffer
+	{
+		FEntry(UMovieSceneControlRigParameterTrack* InTrack, UControlRig* InRig);
+
+		TObjectPtr<UMovieSceneControlRigParameterTrack> Track;
+		bool bIsActive = false;
+	};
+
+	void* GetData(FAccumulatedControlEntryIndex Entry);
+
+	TMap<FObjectKey, int32> ParameterValuesByTrack;
+	TArray<FEntry> ValuesArray;
+	int32 NextValidIndex = 0;
+};
+
+
+} // namespace UE::MovieScene
diff --git a/Engine/Plugins/Animation/ControlRig/Source/ControlRig/Public/Sequencer/MovieSceneControlRigParameterSection.h b/Engine/Plugins/Animation/ControlRig/Source/ControlRig/Public/Sequencer/MovieSceneControlRigParameterSection.h
--- a/Engine/Plugins/Animation/ControlRig/Source/ControlRig/Public/Sequencer/MovieSceneControlRigParameterSection.h
+++ b/Engine/Plugins/Animation/ControlRig/Source/ControlRig/Public/Sequencer/MovieSceneControlRigParameterSection.h
@@ -20,6 +20,7 @@
 #include "Sequencer/MovieSceneControlRigSpaceChannel.h"
 #include "ConstraintChannel.h"
 #include "KeyParams.h"
+#include "EntitySystem/IMovieSceneEntityProvider.h"
 #include "MovieSceneControlRigParameterSection.generated.h"
 
 class UAnimSequence;
@@ -177,7 +178,9 @@ struct FMovieSceneControlRigSpaceChannel;
  * Movie scene section that controls animation controller animation
  */
 UCLASS()
-class CONTROLRIG_API UMovieSceneControlRigParameterSection : public UMovieSceneParameterSection, public IMovieSceneConstrainedSection
+class CONTROLRIG_API UMovieSceneControlRigParameterSection : public UMovieSceneParameterSection
+															, public IMovieSceneConstrainedSection
+
 {
 	GENERATED_BODY()
 
@@ -215,6 +218,10 @@ public:
 	void ChangeControlRotationOrder(const FName& InControlName, const TOptional<EEulerRotationOrder>& CurrentOrder, 
 		const  TOptional<EEulerRotationOrder>& NewOrder, EMovieSceneKeyInterpolation Interpolation);
 
+	virtual void ImportEntityImpl(UMovieSceneEntitySystemLinker* EntityLinker, const FEntityImportParams& Params, FImportedEntity* OutImportedEntity) override;
+	virtual bool PopulateEvaluationFieldImpl(const TRange<FFrameNumber>& EffectiveRange, const FMovieSceneEvaluationFieldEntityMetaData& InMetaData, FMovieSceneEntityComponentFieldBuilder* OutFieldBuilder) override;
+	//mz todo virtual void InterrogateEntityImpl(UMovieSceneEntitySystemLinker* EntityLinker, const FEntityImportParams& Params, FImportedEntity* OutImportedEntity) override;
+
 private:
 
 	FSpaceChannelAddedEvent OnSpaceChannelAdded;
@@ -548,6 +555,10 @@ protected:
 	virtual void OnBindingIDsUpdated(const TMap<UE::MovieScene::FFixedObjectBindingID, UE::MovieScene::FFixedObjectBindingID>& OldFixedToNewFixedMap, FMovieSceneSequenceID LocalSequenceID, TSharedRef<UE::MovieScene::FSharedPlaybackState> SharedPlaybackState) override;
 	virtual void GetReferencedBindings(TArray<FGuid>& OutBindings) override;
 	virtual void PreSave(FObjectPreSaveContext SaveContext) override;
+	virtual bool IsBlendingHandledExternally() const
+	{
+		return BlendType.Get() == EMovieSceneBlendType::Absolute;
+	}
 
 #if WITH_EDITOR
 	virtual void PostDuplicate(bool bDuplicateForPIE) override;
diff --git a/Engine/Plugins/Animation/ControlRig/Source/ControlRig/Public/Sequencer/MovieSceneControlRigParameterTrack.h b/Engine/Plugins/Animation/ControlRig/Source/ControlRig/Public/Sequencer/MovieSceneControlRigParameterTrack.h
--- a/Engine/Plugins/Animation/ControlRig/Source/ControlRig/Public/Sequencer/MovieSceneControlRigParameterTrack.h
+++ b/Engine/Plugins/Animation/ControlRig/Source/ControlRig/Public/Sequencer/MovieSceneControlRigParameterTrack.h
@@ -48,6 +48,9 @@ class UMovieSceneControlRigParameterTrack
 	GENERATED_UCLASS_BODY()
 
 public:
+
+	CONTROLRIG_API static bool ShouldUseLegacyTemplate();
+
 	// UMovieSceneTrack interface
 	virtual FMovieSceneEvalTemplatePtr CreateTemplateForSection(const UMovieSceneSection& InSection) const override;
 	virtual bool SupportsType(TSubclassOf<UMovieSceneSection> SectionClass) const override;
@@ -60,6 +63,7 @@ public:
 	virtual bool IsEmpty() const override;
 	virtual const TArray<UMovieSceneSection*>& GetAllSections() const override;
 	virtual FName GetTrackName() const override { return TrackName; }
+
 	// UObject
 	virtual void PostLoad() override;
 	virtual void BeginDestroy() override;
diff --git a/Engine/Plugins/Animation/ControlRig/Source/ControlRig/Public/Sequencer/MovieSceneControlRigSystem.h b/Engine/Plugins/Animation/ControlRig/Source/ControlRig/Public/Sequencer/MovieSceneControlRigSystem.h
new file mode 100644
--- /dev/null
+++ b/Engine/Plugins/Animation/ControlRig/Source/ControlRig/Public/Sequencer/MovieSceneControlRigSystem.h
@@ -0,0 +1,244 @@
+// Copyright Epic Games, Inc. All Rights Reserved.
+
+#pragma once
+
+#include "EntitySystem/MovieSceneEntityIDs.h"
+#include "EntitySystem/MovieSceneEntityInstantiatorSystem.h"
+#include "EntitySystem/MovieSceneEntitySystem.h"
+#include "EntitySystem/MovieSceneOverlappingEntityTracker.h"
+#include "Evaluation/PreAnimatedState/MovieScenePreAnimatedObjectStorage.h"
+#include "Evaluation/PreAnimatedState/MovieScenePreAnimatedStateStorage.h"
+#include "Sequencer/MovieSceneControlRigComponentTypes.h"
+#include "Sequencer/MovieSceneControlRigParameterBuffer.h"
+#include "SkeletalMeshRestoreState.h"
+#include "UObject/ObjectKey.h"
+
+#include "MovieSceneControlRigSystem.generated.h"
+
+class UMovieSceneControlRigParameterTrack;
+class UMovieSceneControlRigParameterSection;
+class UMovieSceneBlenderSystem;
+class UMovieScenePiecewiseDoubleBlenderSystem;
+
+namespace UE::MovieScene
+{
+	struct FBaseControlRigEvalData;
+	struct FControlRigParameterValue;
+	struct FEntityGroupID;
+	struct FEvaluateControlRigChannels;
+	struct FMovieSceneEntityID;
+	struct FPreAnimatedControlRigParameterStorage;
+	struct FPreAnimatedControlRigStorage;
+	struct FRootInstanceHandle;
+
+	struct FAnimatedControlRigParameterInfo
+	{
+		/** Weak linker ptr - only assigned if the output entity is ever allocated */
+		TWeakObjectPtr<UMovieSceneEntitySystemLinker> WeakLinker;
+		/** Weak blender system ptr - only assigned if the blend channel is ever allocated */
+		TWeakObjectPtr<UMovieSceneBlenderSystem> WeakBlenderSystem;
+		int32 NumContributors = 0;
+		FMovieSceneEntityID OutputEntityID;
+		FMovieSceneBlendChannelID BlendChannelID;
+
+		~FAnimatedControlRigParameterInfo();
+	};
+
+	CONTROLRIG_API void CollectGarbageForOutput(FAnimatedControlRigParameterInfo* Output);
+
+
+	struct FPreAnimatedControlRigState
+	{
+		FSkeletalMeshRestoreState SkeletalMeshRestoreState;
+		EAnimationMode::Type AnimationMode;
+		void SetSkelMesh(USkeletalMeshComponent* InComponent)
+		{
+			SkeletalMeshRestoreState.SaveState(InComponent);
+			AnimationMode = InComponent->GetAnimationMode();
+		}
+	};
+
+	/** Pre-animation traits for control rig bindings */
+	struct FPreAnimatedControlRigTraits : FBoundObjectPreAnimatedStateTraits
+	{
+		using KeyType = FObjectKey;
+		using StorageType = FPreAnimatedControlRigState;
+
+		static FPreAnimatedControlRigState CachePreAnimatedValue(FObjectKey InObject);
+		static void RestorePreAnimatedValue(const FObjectKey& Object, FPreAnimatedControlRigState& State, const FRestoreStateParams& Params);
+	};
+
+	/** Pre-animation traits for control rig parameters */
+	struct FPreAnimatedControlRigParameterTraits : FBoundObjectPreAnimatedStateTraits
+	{
+		using FBoundObjectPreAnimatedStateTraits::ResolveComponent;
+
+		struct FPreAnimatedBufferPairs
+		{
+			FControlRigParameterValues Transient = FControlRigParameterValues(EControlRigParameterBufferIndexStability::Unstable);
+			FControlRigParameterValues Persistent = FControlRigParameterValues(EControlRigParameterBufferIndexStability::Unstable);
+		};
+
+		struct FPreAnimatedParameterKey
+		{
+			TSharedPtr<FPreAnimatedBufferPairs> Buffer;
+			FName ParameterName;
+			bool bPersistent;
+
+			friend uint32 GetTypeHash(const FPreAnimatedParameterKey& In)
+			{
+				return GetTypeHash(In.Buffer) ^ GetTypeHash(In.ParameterName) + (uint32)In.bPersistent;
+			}
+
+			friend bool operator==(const FPreAnimatedParameterKey& A, const FPreAnimatedParameterKey& B)
+			{
+				return A.Buffer == B.Buffer && A.ParameterName == B.ParameterName && A.bPersistent == B.bPersistent;
+			}
+
+			FPreAnimatedParameterKey() = default;
+			FPreAnimatedParameterKey(TSharedPtr<FPreAnimatedBufferPairs> InBuffer, FName InName)
+				: Buffer(InBuffer)
+				, ParameterName(InName)
+				, bPersistent(false)
+			{}
+
+			FPreAnimatedParameterKey(FPreAnimatedParameterKey&&) = default;
+			FPreAnimatedParameterKey& operator=(FPreAnimatedParameterKey&&) = default;
+
+			// Keys can only be copied when being elevated to persistent
+			FPreAnimatedParameterKey(const FPreAnimatedParameterKey& RHS)
+				: Buffer(RHS.Buffer)
+				, bPersistent(true)
+			{
+				// Copy the parameter value from the transient to the persistent
+				if (!RHS.bPersistent)
+				{
+					Buffer->Persistent.CopyFrom(Buffer->Transient, ParameterName);
+				}
+			}
+			FPreAnimatedParameterKey& operator=(const FPreAnimatedParameterKey& RHS)
+			{
+				Buffer = RHS.Buffer;
+				bPersistent = true;
+				if (!RHS.bPersistent)
+				{
+					Buffer->Persistent.CopyFrom(Buffer->Transient, ParameterName);
+				}
+				return *this;
+			}
+		};
+		using KeyType = TTuple<FObjectKey, FName>;
+		using StorageType = FPreAnimatedParameterKey;
+
+		TMap<FObjectKey, TWeakPtr<FPreAnimatedBufferPairs>> PreAnimatedBuffers;
+
+		TSharedPtr<FPreAnimatedBufferPairs> GetBuffers(UControlRig* Rig);
+
+		FPreAnimatedParameterKey CachePreAnimatedValue(UControlRig* Rig, FName ParameterName);
+		void RestorePreAnimatedValue(const TTuple<FObjectKey, FName>& Rig, FPreAnimatedParameterKey& ParameterBuffer, const FRestoreStateParams& Params);
+	};
+
+	/** Preanimated storage that manages unbinding the rig from the target component */
+	struct FPreAnimatedControlRigStorage
+		: public TPreAnimatedStateStorage<FPreAnimatedControlRigTraits>
+	{
+		static CONTROLRIG_API TAutoRegisterPreAnimatedStorageID<FPreAnimatedControlRigStorage> StorageID;
+	};
+
+	/** Preanimated storage that manages resetting parameters from the rig */
+	struct FPreAnimatedControlRigParameterStorage
+		: public TPreAnimatedStateStorage<FPreAnimatedControlRigParameterTraits>
+	{
+		static CONTROLRIG_API TAutoRegisterPreAnimatedStorageID<FPreAnimatedControlRigParameterStorage> StorageID;
+	};
+
+} // namespace UE::MovieScene
+
+
+/**
+ * System that tracks and evaluates control rig control parameters and spaces.
+ * 
+ * This system is relatively complex, and implements the following behavior:
+ * 
+ * Instantiation:
+ *   - Initializes rigs by ensuring the anim instance set up and bound to the correct rig
+ *   - Tracks contriguting control rig parameters by rg
+ *   - Caches pre-animated state
+ *   - Assigns accumulation indices for each parameter
+ *   - Caches initial values if necessary
+
+ * Evaluation:
+ *   - (GameThread) Evaluates any base rigs by calling SetDoNotKey(true) on the section, and updating the control rig track in the anim instance with its weight (if active)
+ *   - (GameThread) Evaluates constraints
+ *   - (Async)      Evaluates space channels
+ *   - (Async)      Accumulates all control rig parameters and spaces into a single parameter buffer for each control rig
+ *   - (GameThread) Applies the control rig parameter buffer to the rig
+ *   - (Async)      Resets the 'do not key' states on the sections
+ */
+UCLASS(MinimalAPI)
+class UMovieSceneControlRigParameterEvaluatorSystem : public UMovieSceneEntitySystem
+{
+
+public:
+
+	GENERATED_BODY()
+	UMovieSceneControlRigParameterEvaluatorSystem(const FObjectInitializer& ObjInit);
+
+	CONTROLRIG_API UControlRig* GetRigFromTrack(UMovieSceneControlRigParameterTrack* Track) const;
+
+	CONTROLRIG_API const UE::MovieScene::FControlRigParameterBuffer* FindParameters(UMovieSceneControlRigParameterTrack* Track) const;
+
+private:
+
+	virtual bool IsRelevantImpl(UMovieSceneEntitySystemLinker* InLinker) const override;
+	virtual void OnLink() override;
+	virtual void OnUnlink() override;
+	virtual void OnSchedulePersistentTasks(UE::MovieScene::IEntitySystemScheduler* TaskScheduler) override;
+	virtual void OnRun(FSystemTaskPrerequisites& InPrerequisites, FSystemSubsequentTasks& Subsequents) override;
+
+	void OnInstantiation();
+	void OnEvaluation(FSystemTaskPrerequisites& InPrerequisites, FSystemSubsequentTasks& Subsequents);
+
+private:
+
+	void InitializeBaseRigComponent(
+		UObject* BoundObject,
+		UE::MovieScene::FMovieSceneEntityID EntityID,
+		UE::MovieScene::FRootInstanceHandle RootInstanceHandle,
+		bool bWantsRestoreState,
+		bool bHasWeight,
+		UE::MovieScene::FControlRigSourceData ControlRigSource,
+		UE::MovieScene::FBaseControlRigEvalData& OutBaseData,
+		TMap<UControlRig*, UE::MovieScene::FBaseControlRigEvalData*>& OutBaseComponentTracker
+	);
+
+private:
+
+	/**
+	 * Tracker that operates on group IDs to initialize blend outputs for each control rig parameter type
+	 */
+	UE::MovieScene::TOverlappingEntityTracker<UE::MovieScene::FAnimatedControlRigParameterInfo, UE::MovieScene::FEntityGroupID> ControlRigParameterTracker;
+
+	/**
+	 * Persistent accumulation buffer for all known rig values.
+	 * FControlRigComponentTypes::AccumulatedControlEntryIndex components directly reference entries within this structure.
+	 * Care is taken to avoid async read/write through task dependencies in OnEvaluation.
+	 * */
+	UE::MovieScene::FAccumulatedControlRigValues AccumulatedValues;
+
+public:
+
+	/** Pre-animated state storage for all control rig parameters */
+	TSharedPtr<UE::MovieScene::FPreAnimatedControlRigParameterStorage> ControlRigParameterStorage;
+
+	/** Pre-animated state storage for control rigs themselves (stores which rigs are currently being animated along with their anim instance and binding) */
+	TSharedPtr<UE::MovieScene::FPreAnimatedControlRigStorage> ControlRigStorage;
+
+	/** Registered key for the grouping policy that groups control rig parameters together by name and type */
+	UE::MovieScene::FEntityGroupingPolicyKey ParameterGroupingKey;
+
+	/** Cached blender system used for blending parameters */
+	UPROPERTY()
+	TObjectPtr<UMovieScenePiecewiseDoubleBlenderSystem> DoubleBlenderSystem;
+};
+
diff --git a/Engine/Plugins/Animation/ControlRig/Source/ControlRigEditor/Private/Sequencer/ControlRigParameterTrackEditor.cpp b/Engine/Plugins/Animation/ControlRig/Source/ControlRigEditor/Private/Sequencer/ControlRigParameterTrackEditor.cpp
--- a/Engine/Plugins/Animation/ControlRig/Source/ControlRigEditor/Private/Sequencer/ControlRigParameterTrackEditor.cpp
+++ b/Engine/Plugins/Animation/ControlRig/Source/ControlRigEditor/Private/Sequencer/ControlRigParameterTrackEditor.cpp
@@ -1,6 +1,7 @@
 // Copyright Epic Games, Inc. All Rights Reserved.
 
 #include "Sequencer/ControlRigParameterTrackEditor.h"
+#include "Sequencer/MovieSceneControlRigSystem.h"
 #include "MVVM/Extensions/ITrackExtension.h"
 #include "MVVM/ViewModels/OutlinerViewModel.h"
 #include "MVVM/ViewModels/TrackModel.h"
@@ -10,6 +11,7 @@
 #include "MVVM/ViewModels/ViewModelIterators.h"
 #include "MVVM/ViewModels/OutlinerViewModel.h"
 #include "MVVM/Selection/Selection.h"
+#include "EntitySystem/Interrogation/MovieSceneInterrogationLinker.h"
 #include "Animation/AnimMontage.h"
 #include "Sequencer/MovieSceneControlRigParameterTrack.h"
 #include "Sequencer/MovieSceneControlRigParameterSection.h"
@@ -4346,44 +4348,105 @@ void FControlRigParameterTrackEditor::AddControlKeys(
 
 bool FControlRigParameterTrackEditor::ModifyOurGeneratedKeysByCurrentAndWeight(UObject* Object, UControlRig* InControlRig, FName RigControlName, UMovieSceneTrack* Track, UMovieSceneSection* SectionToKey, FFrameNumber EvaluateTime, FGeneratedTrackKeys& GeneratedTotalKeys, float Weight) const
 {
-	FFrameRate TickResolution = GetSequencer()->GetFocusedTickResolution();
-	FMovieSceneEvaluationTrack EvalTrack = CastChecked<UMovieSceneControlRigParameterTrack>(Track)->GenerateTrackTemplate(Track);
+	using namespace UE::MovieScene;
+
+	// Start off with stable names for population since we shouldn't encounter any duplicates
+	FControlRigParameterValues ParameterBufferValues(EControlRigParameterBufferIndexStability::Stable);
+
+	if (UMovieSceneControlRigParameterTrack::ShouldUseLegacyTemplate())
+	{
+		FFrameRate TickResolution = GetSequencer()->GetFocusedTickResolution();
+		FMovieSceneEvaluationTrack EvalTrack = CastChecked<UMovieSceneControlRigParameterTrack>(Track)->GenerateTrackTemplate(Track);
+		FMovieSceneInterrogationData InterrogationData;
+		GetSequencer()->GetEvaluationTemplate().CopyActuators(InterrogationData.GetAccumulator());
+		//use the EvaluateTime to do the evaluation, may be different than the actually time we key
+		FMovieSceneContext Context(FMovieSceneEvaluationRange(EvaluateTime, GetSequencer()->GetFocusedTickResolution()));
+		EvalTrack.Interrogate(Context, InterrogationData, Object);
+		int32 ChannelIndex = 0;
+		FChannelMapInfo* pChannelIndex = nullptr;
+
+		// Add the legacy interrogated data to the parameter buffer
+		for (const FFloatInterrogationData& Val : InterrogationData.Iterate<FFloatInterrogationData>(UMovieSceneControlRigParameterSection::GetFloatInterrogationKey()))
+		{
+			ParameterBufferValues.Add(Val.ParameterName, Val.Val);
+		}
+		for (const FVector2DInterrogationData& Val : InterrogationData.Iterate<FVector2DInterrogationData>(UMovieSceneControlRigParameterSection::GetVector2DInterrogationKey()))
+		{
+			FVector3f Vector((float)Val.Val.X, (float)Val.Val.Y, 0.f);
+			ParameterBufferValues.Add(Val.ParameterName, Vector);
+		}
+		for (const FVectorInterrogationData& Val : InterrogationData.Iterate<FVectorInterrogationData>(UMovieSceneControlRigParameterSection::GetVectorInterrogationKey()))
+		{
+			FVector3f Vector((float)Val.Val.X, (float)Val.Val.Y, (float)Val.Val.Z);
+			ParameterBufferValues.Add(Val.ParameterName, Vector);
+		}
+		for (const FEulerTransformInterrogationData& Val : InterrogationData.Iterate<FEulerTransformInterrogationData>(UMovieSceneControlRigParameterSection::GetTransformInterrogationKey()))
+		{
+			ParameterBufferValues.Add(Val.ParameterName, Val.Val);
+		}
+	}
+	else
+	{
+		// Create the interrogator.
+		FSystemInterrogator Interrogator;
+		Interrogator.TrackImportedEntities(true);
 
-	FMovieSceneInterrogationData InterrogationData;
-	GetSequencer()->GetEvaluationTemplate().CopyActuators(InterrogationData.GetAccumulator());
+		TGuardValue<FEntityManager*> DebugVizGuard(GEntityManagerForDebuggingVisualizers, &Interrogator.GetLinker()->EntityManager);
+
+
+		// Run an interrogation on the track at the specified time.
+		FInterrogationKey InterrogationKey(FInterrogationKey::Default());
+		FInterrogationChannel InterrogationChannel = Interrogator.AllocateChannel(InControlRig, FMovieScenePropertyBinding());
+		InterrogationKey.Channel = InterrogationChannel;
+		Interrogator.ImportTrack(Track, InterrogationChannel);
+
+		Interrogator.AddInterrogation(EvaluateTime);
+		Interrogator.Update();
+
+		// Find the CR System
+		UMovieSceneControlRigParameterEvaluatorSystem* ControlRigSystem   = Interrogator.GetLinker()->FindSystem<UMovieSceneControlRigParameterEvaluatorSystem>();
+		const FControlRigParameterBuffer*              ParameterBufferPtr = ControlRigSystem ? ControlRigSystem->FindParameters(CastChecked<UMovieSceneControlRigParameterTrack>(Track)) : nullptr;
+
+		if (!ensure(ParameterBufferPtr))
+		{
+			return false;
+		}
+
+		ParameterBufferValues = ParameterBufferPtr->Values;
+	}
+
+	// Make searching faster by hashing the values
+	ParameterBufferValues.OptimizeForLookup();
+
+	FFrameRate TickResolution = GetSequencer()->GetFocusedTickResolution();
 
-	//use the EvaluateTime to do the evaluation, may be different than the actually time we key
-	FMovieSceneContext Context(FMovieSceneEvaluationRange(EvaluateTime, GetSequencer()->GetFocusedTickResolution()));
-	EvalTrack.Interrogate(Context, InterrogationData, Object);
 	TArray<FRigControlElement*> Controls = InControlRig->AvailableControls();
 	UMovieSceneControlRigParameterSection* Section = Cast<UMovieSceneControlRigParameterSection>(SectionToKey);
 	FMovieSceneChannelProxy& Proxy = SectionToKey->GetChannelProxy();
-	int32 ChannelIndex = 0;
-	FChannelMapInfo* pChannelIndex = nullptr;
+
 	for (FRigControlElement* ControlElement : Controls)
 	{
 		if (!InControlRig->GetHierarchy()->IsAnimatable(ControlElement))
 		{
 			continue;
 		}
+
+		FName ControlName = ControlElement->GetFName();
+		const FChannelMapInfo* pChannelIndex = Section->ControlChannelMap.Find(ControlName);
+		if (!pChannelIndex || pChannelIndex->GeneratedKeyIndex == INDEX_NONE)
+		{
+			continue;
+		}
+
 		switch (ControlElement->Settings.ControlType)
 		{
 		case ERigControlType::Float:
 		case ERigControlType::ScaleFloat:
 		{
-			for (const FFloatInterrogationData& Val : InterrogationData.Iterate<FFloatInterrogationData>(UMovieSceneControlRigParameterSection::GetFloatInterrogationKey()))
+			float Value = 0.f;
+			if (ParameterBufferValues.Find<float>(ControlName, Value))
 			{
-				if ((Val.ParameterName == ControlElement->GetFName()))
-				{
-					pChannelIndex = Section->ControlChannelMap.Find(ControlElement->GetFName());
-					if (pChannelIndex && pChannelIndex->GeneratedKeyIndex != INDEX_NONE)
-					{
-						ChannelIndex = pChannelIndex->GeneratedKeyIndex;
-						float FVal = (float)Val.Val;
-						GeneratedTotalKeys[ChannelIndex]->ModifyByCurrentAndWeight(Proxy, EvaluateTime, (void *)&FVal, Weight);
-					}
-					break;
-				}
+				GeneratedTotalKeys[pChannelIndex->GeneratedKeyIndex]->ModifyByCurrentAndWeight(Proxy, EvaluateTime, (void *)&Value, Weight);
 			}
 			break;
 		}
@@ -4396,21 +4459,11 @@ bool FControlRigParameterTrackEditor::ModifyOurGeneratedKeysByCurrentAndWeight(U
 		}
 		case ERigControlType::Vector2D:
 		{
-			for (const FVector2DInterrogationData& Val : InterrogationData.Iterate<FVector2DInterrogationData>(UMovieSceneControlRigParameterSection::GetVector2DInterrogationKey()))
+			FVector3f Value;
+			if (ParameterBufferValues.Find(ControlName, Value))
 			{
-				if ((Val.ParameterName == ControlElement->GetFName()))
-				{
-					pChannelIndex = Section->ControlChannelMap.Find(ControlElement->GetFName());
-					if (pChannelIndex && pChannelIndex->GeneratedKeyIndex != INDEX_NONE)
-					{
-						ChannelIndex = pChannelIndex->GeneratedKeyIndex;
-						float FVal = (float)Val.Val.X;
-						GeneratedTotalKeys[ChannelIndex]->ModifyByCurrentAndWeight(Proxy, EvaluateTime, (void *)&FVal, Weight);
-						FVal = (float)Val.Val.Y;
-						GeneratedTotalKeys[ChannelIndex + 1]->ModifyByCurrentAndWeight(Proxy, EvaluateTime, (void *)&FVal, Weight);
-					}
-					break;
-				}
+				GeneratedTotalKeys[pChannelIndex->GeneratedKeyIndex  ]->ModifyByCurrentAndWeight(Proxy, EvaluateTime, (void *)&Value.X, Weight);
+				GeneratedTotalKeys[pChannelIndex->GeneratedKeyIndex+1]->ModifyByCurrentAndWeight(Proxy, EvaluateTime, (void *)&Value.Y, Weight);
 			}
 			break;
 		}
@@ -4418,23 +4471,12 @@ bool FControlRigParameterTrackEditor::ModifyOurGeneratedKeysByCurrentAndWeight(U
 		case ERigControlType::Scale:
 		case ERigControlType::Rotator:
 		{
-			for (const FVectorInterrogationData& Val : InterrogationData.Iterate<FVectorInterrogationData>(UMovieSceneControlRigParameterSection::GetVectorInterrogationKey()))
+			FVector3f Value;
+			if (ParameterBufferValues.Find(ControlName, Value))
 			{
-				if ((Val.ParameterName == ControlElement->GetFName()))
-				{
-					pChannelIndex = Section->ControlChannelMap.Find(ControlElement->GetFName());
-					if (pChannelIndex && pChannelIndex->GeneratedKeyIndex != INDEX_NONE)
-					{
-						ChannelIndex = pChannelIndex->GeneratedKeyIndex;
-						float FVal = (float)Val.Val.X;
-						GeneratedTotalKeys[ChannelIndex]->ModifyByCurrentAndWeight(Proxy, EvaluateTime, (void*)&FVal, Weight);
-						FVal = (float)Val.Val.Y;
-						GeneratedTotalKeys[ChannelIndex + 1]->ModifyByCurrentAndWeight(Proxy, EvaluateTime, (void*)&FVal, Weight);
-						FVal = (float)Val.Val.Z;
-						GeneratedTotalKeys[ChannelIndex + 2]->ModifyByCurrentAndWeight(Proxy, EvaluateTime, (void*)&FVal, Weight);			
-					}
-					break;
-				}
+				GeneratedTotalKeys[pChannelIndex->GeneratedKeyIndex  ]->ModifyByCurrentAndWeight(Proxy, EvaluateTime, (void *)&Value.X, Weight);
+				GeneratedTotalKeys[pChannelIndex->GeneratedKeyIndex+1]->ModifyByCurrentAndWeight(Proxy, EvaluateTime, (void *)&Value.Y, Weight);
+				GeneratedTotalKeys[pChannelIndex->GeneratedKeyIndex+2]->ModifyByCurrentAndWeight(Proxy, EvaluateTime, (void *)&Value.Z, Weight);
 			}
 			break;
 		}
@@ -4443,46 +4485,36 @@ bool FControlRigParameterTrackEditor::ModifyOurGeneratedKeysByCurrentAndWeight(U
 		case ERigControlType::TransformNoScale:
 		case ERigControlType::EulerTransform:
 		{
-			for (const FEulerTransformInterrogationData& Val : InterrogationData.Iterate<FEulerTransformInterrogationData>(UMovieSceneControlRigParameterSection::GetTransformInterrogationKey()))
+
+			FEulerTransform Value;
+			if (ParameterBufferValues.Find(ControlName, Value))
 			{
+				FVector3f  CurrentPos(Value.GetLocation());
+				FRotator3f CurrentRot(Value.Rotator());
 
-				if ((Val.ParameterName == ControlElement->GetFName()))
-				{
-					pChannelIndex = Section->ControlChannelMap.Find(ControlElement->GetFName());
-					if (pChannelIndex && pChannelIndex->GeneratedKeyIndex != INDEX_NONE)
-					{
-						ChannelIndex = pChannelIndex->GeneratedKeyIndex;
+				int32 ChannelIndex = pChannelIndex->bDoesHaveSpace ? pChannelIndex->GeneratedKeyIndex + 1 : pChannelIndex->GeneratedKeyIndex;
 
-						if (pChannelIndex->bDoesHaveSpace)
-						{
-							++ChannelIndex;
-						}
-						
-						FVector3f CurrentPos = (FVector3f)Val.Val.GetLocation();
-						FRotator3f CurrentRot = FRotator3f(Val.Val.Rotator());
-						GeneratedTotalKeys[ChannelIndex]->ModifyByCurrentAndWeight(Proxy, EvaluateTime, (void*)&CurrentPos.X, Weight);
-						GeneratedTotalKeys[ChannelIndex + 1]->ModifyByCurrentAndWeight(Proxy, EvaluateTime, (void*)&CurrentPos.Y, Weight);
-						GeneratedTotalKeys[ChannelIndex + 2]->ModifyByCurrentAndWeight(Proxy, EvaluateTime, (void*)&CurrentPos.Z, Weight);
+				GeneratedTotalKeys[ChannelIndex]->ModifyByCurrentAndWeight(Proxy, EvaluateTime, (void*)&CurrentPos.X, Weight);
+				GeneratedTotalKeys[ChannelIndex + 1]->ModifyByCurrentAndWeight(Proxy, EvaluateTime, (void*)&CurrentPos.Y, Weight);
+				GeneratedTotalKeys[ChannelIndex + 2]->ModifyByCurrentAndWeight(Proxy, EvaluateTime, (void*)&CurrentPos.Z, Weight);
 
-						GeneratedTotalKeys[ChannelIndex + 3]->ModifyByCurrentAndWeight(Proxy, EvaluateTime, (void*)&CurrentRot.Roll, Weight);
-						GeneratedTotalKeys[ChannelIndex + 4]->ModifyByCurrentAndWeight(Proxy, EvaluateTime, (void*)&CurrentRot.Pitch, Weight);
-						GeneratedTotalKeys[ChannelIndex + 5]->ModifyByCurrentAndWeight(Proxy, EvaluateTime, (void*)&CurrentRot.Yaw, Weight);
+				GeneratedTotalKeys[ChannelIndex + 3]->ModifyByCurrentAndWeight(Proxy, EvaluateTime, (void*)&CurrentRot.Roll, Weight);
+				GeneratedTotalKeys[ChannelIndex + 4]->ModifyByCurrentAndWeight(Proxy, EvaluateTime, (void*)&CurrentRot.Pitch, Weight);
+				GeneratedTotalKeys[ChannelIndex + 5]->ModifyByCurrentAndWeight(Proxy, EvaluateTime, (void*)&CurrentRot.Yaw, Weight);
 
-						if (ControlElement->Settings.ControlType == ERigControlType::Transform || ControlElement->Settings.ControlType == ERigControlType::EulerTransform)
-						{
-							FVector3f CurrentScale = (FVector3f)Val.Val.GetScale3D();
-							GeneratedTotalKeys[ChannelIndex + 6]->ModifyByCurrentAndWeight(Proxy, EvaluateTime, (void*)&CurrentScale.X, Weight);
-							GeneratedTotalKeys[ChannelIndex + 7]->ModifyByCurrentAndWeight(Proxy, EvaluateTime, (void*)&CurrentScale.Y, Weight);
-							GeneratedTotalKeys[ChannelIndex + 8]->ModifyByCurrentAndWeight(Proxy, EvaluateTime, (void*)&CurrentScale.Z, Weight);
-						}
-					}
-					break;
+				if (ControlElement->Settings.ControlType == ERigControlType::Transform || ControlElement->Settings.ControlType == ERigControlType::EulerTransform)
+				{
+					FVector3f CurrentScale(Value.GetScale3D());
+					GeneratedTotalKeys[ChannelIndex + 6]->ModifyByCurrentAndWeight(Proxy, EvaluateTime, (void*)&CurrentScale.X, Weight);
+					GeneratedTotalKeys[ChannelIndex + 7]->ModifyByCurrentAndWeight(Proxy, EvaluateTime, (void*)&CurrentScale.Y, Weight);
+					GeneratedTotalKeys[ChannelIndex + 8]->ModifyByCurrentAndWeight(Proxy, EvaluateTime, (void*)&CurrentScale.Z, Weight);
 				}
 			}
 			break;
 		}
 		}
 	}
+
 	return true;
 }
 
diff --git a/Engine/Source/Runtime/MovieScene/Private/EntitySystem/BuiltInComponentTypes.cpp b/Engine/Source/Runtime/MovieScene/Private/EntitySystem/BuiltInComponentTypes.cpp
--- a/Engine/Source/Runtime/MovieScene/Private/EntitySystem/BuiltInComponentTypes.cpp
+++ b/Engine/Source/Runtime/MovieScene/Private/EntitySystem/BuiltInComponentTypes.cpp
@@ -3,12 +3,13 @@
 #include "EntitySystem/BuiltInComponentTypes.h"
 #include "Channels/MovieSceneDoubleChannel.h"
 #include "Channels/MovieSceneFloatChannel.h"
+#include "Channels/MovieSceneInterpolation.h"
+#include "EntitySystem/BuiltInComponentTypes.h"
+#include "EntitySystem/MovieSceneComponentRegistry.h"
+#include "EntitySystem/MovieSceneEntityFactoryTemplates.h"
 #include "EntitySystem/MovieSceneEntityManager.h"
 #include "EntitySystem/MovieSceneEntitySystemLinker.h"
-#include "EntitySystem/MovieSceneComponentRegistry.h"
 #include "EntitySystem/MovieScenePropertyBinding.h"
-#include "EntitySystem/MovieSceneEntityFactoryTemplates.h"
-#include "Channels/MovieSceneInterpolation.h"
 
 #include UE_INLINE_GENERATED_CPP_BY_NAME(BuiltInComponentTypes)
 
@@ -286,7 +287,7 @@ FBuiltInComponentTypes::FBuiltInComponentTypes()
 	ComponentRegistry->NewComponentType(&EvalTime,              TEXT("Eval Time"));
 	ComponentRegistry->NewComponentType(&EvalSeconds,           TEXT("Eval Seconds"));
 
-	ComponentRegistry->NewComponentType(&Group,					TEXT("Entity Group"));
+	ComponentRegistry->NewComponentType(&Group,					TEXT("Entity Group"), EComponentTypeFlags::Preserved);
 
 	ComponentRegistry->NewComponentType(&BoundObjectKey,        TEXT("Bound Object Key"));
 	// Intentionally hidden from the reference graph because they are always accompanied by a BoundObjectKey which is used for garbage collection
@@ -347,7 +348,7 @@ FBuiltInComponentTypes::FBuiltInComponentTypes()
 	ComponentRegistry->NewComponentType(&BlendChannelInput,     TEXT("Blend Channel Input"));
 	ComponentRegistry->NewComponentType(&HierarchicalBias,      TEXT("Hierarchical Bias"));
 	ComponentRegistry->NewComponentType(&BlendChannelOutput,    TEXT("Blend Channel Output"));
-	ComponentRegistry->NewComponentType(&InitialValueIndex,     TEXT("Initial Value Index"));
+	ComponentRegistry->NewComponentType(&InitialValueIndex,     TEXT("Initial Value Index"), EComponentTypeFlags::MigrateToOutput);
 
 	ComponentRegistry->NewComponentType(&CustomPropertyIndex,   TEXT("Custom Property Index"));			// Not EComponentTypeFlags::Preserved because the system property manager will always ensure that the component is added to the correct entity
 	ComponentRegistry->NewComponentType(&FastPropertyOffset,    TEXT("Fast Property Offset"));			// Not EComponentTypeFlags::Preserved because the system property manager will always ensure that the component is added to the correct entity
@@ -411,7 +412,7 @@ FBuiltInComponentTypes::FBuiltInComponentTypes()
 	Tags.NeedsLink               = ComponentRegistry->NewTag(TEXT("Needs Link"));
 	Tags.NeedsUnlink             = ComponentRegistry->NewTag(TEXT("Needs Unlink"));
 	Tags.HasUnresolvedBinding    = ComponentRegistry->NewTag(TEXT("Has Unresolved Binding"));
-	Tags.HasAssignedInitialValue = ComponentRegistry->NewTag(TEXT("Has Assigned Initial Value"));
+	Tags.HasAssignedInitialValue = ComponentRegistry->NewTag(TEXT("Has Assigned Initial Value"), EComponentTypeFlags::MigrateToOutput);
 	Tags.Root                    = ComponentRegistry->NewTag(TEXT("Root"));
 	Tags.SubInstance             = ComponentRegistry->NewTag(TEXT("Sub Instance"), EComponentTypeFlags::CopyToChildren);
 	Tags.ImportedEntity          = ComponentRegistry->NewTag(TEXT("Imported Entity"));
@@ -421,8 +422,9 @@ FBuiltInComponentTypes::FBuiltInComponentTypes()
 	Tags.FixedTime               = ComponentRegistry->NewTag(TEXT("Fixed Time"));
 	Tags.PreRoll                 = ComponentRegistry->NewTag(TEXT("Pre Roll"));
 	Tags.SectionPreRoll          = ComponentRegistry->NewTag(TEXT("Section Pre Roll"));
-	Tags.AlwaysCacheInitialValue = ComponentRegistry->NewTag(TEXT("Always Cache Initial Value"));
-	Tags.OldStyleSpawnable		= ComponentRegistry->NewTag(TEXT("Old Style Spawnable"));
+	Tags.AlwaysCacheInitialValue = ComponentRegistry->NewTag(TEXT("Always Cache Initial Value"), EComponentTypeFlags::CopyToOutput);
+	Tags.OldStyleSpawnable       = ComponentRegistry->NewTag(TEXT("Old Style Spawnable"));
+	Tags.ExternalBlending        = ComponentRegistry->NewTag(TEXT("External Blending"));
 
 	SymbolicTags.CreatesEntities = ComponentRegistry->NewTag(TEXT("~~ SYMBOLIC ~~ Creates Entities"));
 
@@ -469,7 +471,7 @@ FBuiltInComponentTypes::FBuiltInComponentTypes()
 	ComponentRegistry->Factories.DuplicateChildComponent(BindingLifetime);
 	ComponentRegistry->Factories.ConditionallyDuplicateChildComponent(GenericObjectBinding, FComponentMask({BindingLifetime}));
 
-	
+
 	// Children always need a Parent - these are initialized by the tasks that create them
 	{
 		ComponentRegistry->Factories.DefineChildComponent(FComponentTypeID::Invalid(), ParentEntity);
diff --git a/Engine/Source/Runtime/MovieScene/Private/EntitySystem/MovieSceneEntityGroupingSystem.cpp b/Engine/Source/Runtime/MovieScene/Private/EntitySystem/MovieSceneEntityGroupingSystem.cpp
--- a/Engine/Source/Runtime/MovieScene/Private/EntitySystem/MovieSceneEntityGroupingSystem.cpp
+++ b/Engine/Source/Runtime/MovieScene/Private/EntitySystem/MovieSceneEntityGroupingSystem.cpp
@@ -40,7 +40,7 @@ bool FEntityGroupBuilder::RemoveEntityFromGroup(const FMovieSceneEntityID& InEnt
 			using FEntityGroupingHandlerInfo = UMovieSceneEntityGroupingSystem::FEntityGroupingHandlerInfo;
 
 			const int32 NumRemoved = PreviousGroup->Entities.Remove(InEntity);
-			ensure(NumRemoved == 1);
+			ensureAlwaysMsgf(NumRemoved == 1, TEXT("Expected 1 entry to be removed from group by got %d instead."), NumRemoved);
 			if (PreviousGroup->Entities.Num() == 0)
 			{
 				// The group is now empty! Let's remove it.
diff --git a/Engine/Source/Runtime/MovieScene/Private/EntitySystem/MovieSceneEntityIDs.cpp b/Engine/Source/Runtime/MovieScene/Private/EntitySystem/MovieSceneEntityIDs.cpp
--- a/Engine/Source/Runtime/MovieScene/Private/EntitySystem/MovieSceneEntityIDs.cpp
+++ b/Engine/Source/Runtime/MovieScene/Private/EntitySystem/MovieSceneEntityIDs.cpp
@@ -13,5 +13,11 @@ bool FEntityHandle::IsValid(const FEntityManager& InEntityManager) const
 	return InEntityManager.IsHandleValid(*this);
 }
 
+bool FComponentTypeIDFilter::Passes(const FComponentMask& Type) const
+{
+	FComponentTypeID ConditionComponent = GetComponentType();
+	return !ConditionComponent || Type.Contains(ConditionComponent) == PassFlag;
+}
+
 } // namespace MovieScene
 } // namespace UE
diff --git a/Engine/Source/Runtime/MovieScene/Private/EntitySystem/MovieSceneEntitySystem.cpp b/Engine/Source/Runtime/MovieScene/Private/EntitySystem/MovieSceneEntitySystem.cpp
--- a/Engine/Source/Runtime/MovieScene/Private/EntitySystem/MovieSceneEntitySystem.cpp
+++ b/Engine/Source/Runtime/MovieScene/Private/EntitySystem/MovieSceneEntitySystem.cpp
@@ -4,6 +4,7 @@
 #include "EntitySystem/MovieSceneEntitySystemTypes.h"
 #include "EntitySystem/MovieSceneEntitySystemLinker.h"
 #include "EntitySystem/MovieSceneComponentRegistry.h"
+#include "EntitySystem/MovieSceneTaskScheduler.h"
 
 #include UE_INLINE_GENERATED_CPP_BY_NAME(MovieSceneEntitySystem)
 
@@ -633,6 +634,11 @@ UE::MovieScene::EEntitySystemCategory UMovieSceneEntitySystem::RegisterCustomSys
 	return Result;
 }
 
+bool UMovieSceneEntitySystem::IsCustomSchedulingEnabled()
+{
+	return UE::MovieScene::FEntitySystemScheduler::IsCustomSchedulingEnabled();
+}
+
 void UMovieSceneEntitySystem::SortByFlowOrder(TArray<uint16>& InOutGlobalNodeIDs)
 {
 	using namespace UE::MovieScene;
diff --git a/Engine/Source/Runtime/MovieScene/Private/EntitySystem/MovieSceneInitialValueSystem.cpp b/Engine/Source/Runtime/MovieScene/Private/EntitySystem/MovieSceneInitialValueSystem.cpp
new file mode 100644
--- /dev/null
+++ b/Engine/Source/Runtime/MovieScene/Private/EntitySystem/MovieSceneInitialValueSystem.cpp
@@ -0,0 +1,204 @@
+// Copyright Epic Games, Inc. All Rights Reserved.
+
+#include "Systems/MovieSceneInitialValueSystem.h"
+#include "EntitySystem/MovieSceneEntitySystemLinker.h"
+#include "EntitySystem/MovieScenePropertyRegistry.h"
+#include "EntitySystem/MovieSceneEntityMutations.h"
+#include "EntitySystem/BuiltInComponentTypes.h"
+#include "EntitySystem/MovieSceneInitialValueCache.h"
+#include "EntitySystem/Interrogation/MovieSceneInterrogationExtension.h"
+
+#include "Systems/MovieScenePropertyInstantiator.h"
+#include "EntitySystem/Interrogation/MovieSceneInterrogatedPropertyInstantiator.h"
+
+#include UE_INLINE_GENERATED_CPP_BY_NAME(MovieSceneInitialValueSystem)
+
+namespace UE
+{
+namespace MovieScene
+{
+
+struct FInitialValueProcessorEntry
+{
+	FComponentTypeID InitialValueType;
+	FEntityComponentFilter Filter;
+	TSharedPtr<IInitialValueProcessor> Processor;
+};
+
+TArray<FInitialValueProcessorEntry> GInitialValueProcessors;
+
+struct FInitialValueMutation : IMovieSceneEntityMutation
+{
+	TMultiMap<FComponentTypeID, int32> InitialValueTypeToProcessor;
+	FInitialValueCache* InitialValueCache;
+	FBuiltInComponentTypes* BuiltInComponents;
+	FComponentMask AnyInitialValue;
+
+	FInitialValueMutation(UMovieSceneEntitySystemLinker* Linker)
+	{
+		BuiltInComponents = FBuiltInComponentTypes::Get();
+
+		InitialValueCache = Linker->FindExtension<FInitialValueCache>();
+
+		const int32 NumProcessors = GInitialValueProcessors.Num();
+		for (int32 Index = 0; Index < NumProcessors; ++Index)
+		{
+			const FInitialValueProcessorEntry& Entry = GInitialValueProcessors[Index];
+
+			if (Linker->EntityManager.ContainsComponent(Entry.InitialValueType) &&
+				(!Entry.Filter.IsValid() || Linker->EntityManager.Contains(Entry.Filter))
+				)
+			{
+				Entry.Processor->Initialize(Linker, InitialValueCache);
+				AnyInitialValue.Set(Entry.InitialValueType);
+				InitialValueTypeToProcessor.Add(Entry.InitialValueType, Index);
+			}
+		}
+	}
+
+	~FInitialValueMutation()
+	{
+		for (TPair<FComponentTypeID, int32> Pair : InitialValueTypeToProcessor)
+		{
+			GInitialValueProcessors[Pair.Value].Processor->Finalize();
+		}
+	}
+
+	bool IsCached() const
+	{
+		return InitialValueCache != nullptr;
+	}
+
+	virtual void CreateMutation(FEntityManager* EntityManager, FComponentMask* InOutEntityComponentTypes) const override
+	{
+		if (IsCached())
+		{
+			InOutEntityComponentTypes->Set(BuiltInComponents->InitialValueIndex);
+		}
+		InOutEntityComponentTypes->Set(BuiltInComponents->Tags.HasAssignedInitialValue);
+	}
+
+	virtual void InitializeAllocation(FEntityAllocation* Allocation, const FComponentMask& AllocationType) const
+	{
+		FComponentTypeID InitialValueType = FComponentMask::BitwiseAND(AllocationType, AnyInitialValue, EBitwiseOperatorFlags::MinSize).First();
+
+		for (auto IndexIt = InitialValueTypeToProcessor.CreateConstKeyIterator(InitialValueType); IndexIt; ++IndexIt)
+		{
+			const FInitialValueProcessorEntry& Entry = GInitialValueProcessors[IndexIt.Value()];
+			if (!Entry.Filter.IsValid() || Entry.Filter.Match(AllocationType))
+			{
+				Entry.Processor->Process(Allocation, AllocationType);
+			}
+		}
+	}
+};
+
+} // namespace MovieScene
+} // namespace UE
+
+UMovieSceneInitialValueSystem::UMovieSceneInitialValueSystem(const FObjectInitializer& ObjInit)
+	: Super(ObjInit)
+{
+	SystemCategories = UE::MovieScene::EEntitySystemCategory::Core;
+}
+
+void UMovieSceneInitialValueSystem::RegisterProcessor(
+	UE::MovieScene::FComponentTypeID InitialValueComponent,
+	TSharedPtr<UE::MovieScene::IInitialValueProcessor> Processor,
+	UE::MovieScene::FEntityComponentFilter&& OptionalFilter)
+{
+	using namespace UE::MovieScene;
+
+	Processor->PopulateFilter(OptionalFilter);
+
+	GInitialValueProcessors.Add(FInitialValueProcessorEntry{
+		InitialValueComponent,
+		MoveTemp(OptionalFilter),
+		MoveTemp(Processor)
+	});
+}
+
+bool UMovieSceneInitialValueSystem::IsRelevantImpl(UMovieSceneEntitySystemLinker* InLinker) const
+{
+	using namespace UE::MovieScene;
+
+	for (const FInitialValueProcessorEntry& Entry : GInitialValueProcessors)
+	{
+		if (InLinker->EntityManager.ContainsComponent(Entry.InitialValueType))
+		{
+			return true;
+		}
+	}
+	return false;
+}
+
+void UMovieSceneInitialValueSystem::OnLink()
+{
+
+}
+
+void UMovieSceneInitialValueSystem::OnUnlink()
+{
+
+}
+
+void UMovieSceneInitialValueSystem::OnRun(FSystemTaskPrerequisites& InPrerequisites, FSystemSubsequentTasks& Subsequents)
+{
+	using namespace UE::MovieScene;
+
+	FBuiltInComponentTypes* BuiltInComponents = FBuiltInComponentTypes::Get();
+
+	FInitialValueMutation Mutation(Linker);
+
+	// If we don't have any initial value processors, we've no work to do
+	if (Mutation.AnyInitialValue.NumComponents() == 0)
+	{
+		return;
+	}
+
+	if (Mutation.IsCached() && Linker->FindExtension<IInterrogationExtension>() == nullptr)
+	{
+		// When there is an initial value cache extension, we mutate anything with an initial value component on it by
+		// adding an additional index that refers to its cached position. This ensures we are able to clean up the cache
+		// easily.
+		{
+			FEntityComponentFilter Filter;
+			Filter.Any(Mutation.AnyInitialValue);
+			Filter.All({ BuiltInComponents->Tags.NeedsLink });
+			Filter.None({ BuiltInComponents->InitialValueIndex });
+			Filter.None({ BuiltInComponents->Tags.HasAssignedInitialValue, BuiltInComponents->Tags.Ignored });
+
+			Linker->EntityManager.MutateAll(Filter, Mutation);
+		}
+
+		// Clean up any stale cache entries
+		{
+			FEntityComponentFilter Filter;
+			Filter.Any(Mutation.AnyInitialValue);
+			Filter.All({ BuiltInComponents->InitialValueIndex, BuiltInComponents->Tags.NeedsUnlink });
+
+			for (FEntityAllocationIteratorItem Item : Linker->EntityManager.Iterate(&Filter))
+			{
+				const FEntityAllocation* Allocation     = Item.GetAllocation();
+				FComponentMask           AllocationType = Item.GetAllocationType();
+
+				FComponentTypeID InitialValueType = FComponentMask::BitwiseAND(AllocationType, Mutation.AnyInitialValue, EBitwiseOperatorFlags::MinSize).First();
+
+				TComponentReader<FInitialValueIndex> Indices = Allocation->ReadComponents(BuiltInComponents->InitialValueIndex);
+				Mutation.InitialValueCache->Reset(InitialValueType, Indices.AsArray(Allocation->Num()));
+			}
+		}
+	}
+	else
+	{
+		// When there is no caching extension, or we are interrogating we simply initialize any initial values directly without going through the cache
+		FEntityComponentFilter Filter;
+		Filter.Any(Mutation.AnyInitialValue);
+		//Filter.Any({ BuiltInComponents->Interrogation.OutputKey });
+		Filter.All({ BuiltInComponents->Tags.NeedsLink });
+		Filter.None({ BuiltInComponents->Tags.HasAssignedInitialValue, BuiltInComponents->Tags.Ignored });
+
+		Linker->EntityManager.MutateAll(Filter, Mutation);
+	}
+}
+
diff --git a/Engine/Source/Runtime/MovieScene/Private/EntitySystem/MovieSceneInterrogation.cpp b/Engine/Source/Runtime/MovieScene/Private/EntitySystem/MovieSceneInterrogation.cpp
new file mode 100644
--- /dev/null
+++ b/Engine/Source/Runtime/MovieScene/Private/EntitySystem/MovieSceneInterrogation.cpp
@@ -0,0 +1,17 @@
+// Copyright Epic Games, Inc. All Rights Reserved.
+
+#include "EntitySystem/MovieSceneInterrogation.h"
+#include "EntitySystem/MovieSceneEntitySystemLinker.h"
+
+namespace UE::MovieScene
+{
+
+
+TEntitySystemLinkerExtensionID<IInterrogationExtension> IInterrogationExtension::GetExtensionID()
+{
+	static TEntitySystemLinkerExtensionID<IInterrogationExtension> ID = UMovieSceneEntitySystemLinker::RegisterExtension<IInterrogationExtension>();
+	return ID;
+}
+
+
+} // namespace UE::MovieScene
diff --git a/Engine/Source/Runtime/MovieScene/Private/EntitySystem/MovieScenePropertyRegistry.cpp b/Engine/Source/Runtime/MovieScene/Private/EntitySystem/MovieScenePropertyRegistry.cpp
--- a/Engine/Source/Runtime/MovieScene/Private/EntitySystem/MovieScenePropertyRegistry.cpp
+++ b/Engine/Source/Runtime/MovieScene/Private/EntitySystem/MovieScenePropertyRegistry.cpp
@@ -5,6 +5,7 @@
 #include "UObject/Class.h"
 #include "UObject/Object.h"
 #include "EntitySystem/MovieScenePropertyBinding.h"
+#include "EntitySystem/MovieSceneInitialValueSystem.h"
 #include "MovieSceneCommonHelpers.h"
 
 
@@ -208,5 +209,19 @@ TOptional<FResolvedProperty> FPropertyRegistry::ResolveProperty(UObject* Object,
 }
 
 
+void FPropertyDefinition::SetupInitialValueProcessor() const
+{
+	if (InitialValueType)
+	{
+		TSharedPtr<IInitialValueProcessor> InitialValueProcessor = Handler->MakeInitialValueProcessor(*this);
+		if (InitialValueProcessor)
+		{
+			FEntityComponentFilter Filter;
+			InitialValueProcessor->PopulateFilter(Filter);
+			UMovieSceneInitialValueSystem::RegisterProcessor(InitialValueType, MoveTemp(InitialValueProcessor), MoveTemp(Filter));
+		}
+	}
+}
+
 } // namespace MovieScene
 } // namespace UE
\ No newline at end of file
diff --git a/Engine/Source/Runtime/MovieScene/Public/EntitySystem/BuiltInComponentTypes.h b/Engine/Source/Runtime/MovieScene/Public/EntitySystem/BuiltInComponentTypes.h
--- a/Engine/Source/Runtime/MovieScene/Public/EntitySystem/BuiltInComponentTypes.h
+++ b/Engine/Source/Runtime/MovieScene/Public/EntitySystem/BuiltInComponentTypes.h
@@ -773,6 +773,8 @@ public:
 
 		FComponentTypeID Ignored;
 
+		FComponentTypeID ExternalBlending;
+
 		FComponentTypeID AlwaysCacheInitialValue;
 
 		FComponentTypeID DontOptimizeConstants;
diff --git a/Engine/Source/Runtime/MovieScene/Public/EntitySystem/IMovieScenePropertyComponentHandler.h b/Engine/Source/Runtime/MovieScene/Public/EntitySystem/IMovieScenePropertyComponentHandler.h
--- a/Engine/Source/Runtime/MovieScene/Public/EntitySystem/IMovieScenePropertyComponentHandler.h
+++ b/Engine/Source/Runtime/MovieScene/Public/EntitySystem/IMovieScenePropertyComponentHandler.h
@@ -18,6 +18,7 @@ struct FEntityAllocation;
 struct FInitialValueCache;
 struct FPropertyDefinition;
 struct IPreAnimatedStorage;
+struct FEntityComponentFilter;
 struct FSystemSubsequentTasks;
 struct FSystemTaskPrerequisites;
 struct FValueDecompositionParams;
@@ -136,7 +137,9 @@ struct IInitialValueProcessor
 	virtual ~IInitialValueProcessor(){}
 
 	/** Initialize this processor before any allocations are visited */
-	virtual void Initialize(UMovieSceneEntitySystemLinker* Linker, const FPropertyDefinition* Definition, FInitialValueCache* InitialValueCache) = 0;
+	virtual void Initialize(UMovieSceneEntitySystemLinker* Linker, FInitialValueCache* InitialValueCache) = 0;
+	/** Populate a filter that controls whether this processor will run on a given allocation or not */
+	virtual void PopulateFilter(FEntityComponentFilter& OutFilter) const = 0;
 	/** Process all initial values for the specified allocation */
 	virtual void Process(const FEntityAllocation* Allocation, const FComponentMask& AllocationType) = 0;
 	/** Finish processing */
@@ -218,7 +221,7 @@ struct IPropertyComponentHandler
 	/**
 	 * Retrieve an initial value processor interface for this property type
 	 */
-	virtual IInitialValueProcessor* GetInitialValueProcessor() = 0;
+	virtual TSharedPtr<IInitialValueProcessor> MakeInitialValueProcessor(const FPropertyDefinition& Definition) = 0;
 };
 
 
diff --git a/Engine/Source/Runtime/MovieScene/Public/EntitySystem/MovieSceneComponentAccessors.h b/Engine/Source/Runtime/MovieScene/Public/EntitySystem/MovieSceneComponentAccessors.h
--- a/Engine/Source/Runtime/MovieScene/Public/EntitySystem/MovieSceneComponentAccessors.h
+++ b/Engine/Source/Runtime/MovieScene/Public/EntitySystem/MovieSceneComponentAccessors.h
@@ -51,9 +51,17 @@ struct FWriteAccess : FComponentAccess
 struct FOptionalComponentAccess
 {
 	FComponentTypeID ComponentType;
+	FComponentTypeIDFilter Condition;
 
-	FORCEINLINE void PreLockComponentData(const FEntityAllocation* Allocation, FPreLockedDataPtr* PrelockedComponentData) const
+	FORCEINLINE void PreLockComponentData(FEntityAllocationIteratorItem Item, FPreLockedDataPtr* PrelockedComponentData) const
 	{
+		if (!Condition.Passes(Item.GetAllocationType()))
+		{
+			return;
+		}
+
+		const FEntityAllocation* Allocation = Item.GetAllocation();
+
 		const FComponentHeader* Header = Allocation->FindComponentHeader(ComponentType);
 		if (Header)
 		{
@@ -64,11 +72,11 @@ struct FOptionalComponentAccess
 };
 struct FOptionalReadAccess : FOptionalComponentAccess
 {
-	FOptionalReadAccess(FComponentTypeID InComponentType) : FOptionalComponentAccess{ InComponentType } {}
+	FOptionalReadAccess(FComponentTypeID InComponentType, FComponentTypeIDFilter InConditionType = FComponentTypeIDFilter()) : FOptionalComponentAccess{ InComponentType, InConditionType } {}
 };
 struct FOptionalWriteAccess : FOptionalComponentAccess
 {
-	FOptionalWriteAccess(FComponentTypeID InComponentType) : FOptionalComponentAccess{ InComponentType } {}
+	FOptionalWriteAccess(FComponentTypeID InComponentType, FComponentTypeIDFilter InConditionType = FComponentTypeIDFilter()) : FOptionalComponentAccess{ InComponentType, InConditionType } {}
 };
 
 
@@ -95,7 +103,39 @@ struct FEntityIDAccess
 	}
 };
 
+struct FFilterMatchPassthrough
+{
+	using AccessType = bool;
+	//static constexpr int32 PreLockedDataNum = 1;
+
+	FORCEINLINE void PreLockComponentData(FEntityAllocationIteratorItem Item, FPreLockedDataPtr* PrelockedComponentData) const
+	{
+		// Must be relative to the base allocation ptr
+		const bool bFilterMatch = Filter.Match(Item.GetAllocationType());
+
+		// We abuse the pointer here by using either the allocation ptr itself as 'true' and a nullptr as 'false'
+		if (bFilterMatch)
+		{
+			PrelockedComponentData->Reset(nullptr);
+		}
+		else
+		{
+			const FEntityAllocation* Allocation = Item.GetAllocation();
+			PrelockedComponentData->Reset(Allocation, Allocation);
+		}
+	}
+	FORCEINLINE bool ResolvePreLockedComponentData(const FEntityAllocation* Allocation, const FPreLockedDataPtr* Ptr, FEntityAllocationWriteContext WriteContext) const
+	{
+		return Ptr->Resolve<FEntityAllocation>(Allocation) != nullptr;
+	}
 
+	FORCEINLINE bool LockComponentData(FEntityAllocationIteratorItem Item, FEntityAllocationWriteContext WriteContext) const
+	{
+		return Filter.Match(Item.GetAllocationType());
+	}
+
+	FEntityComponentFilter Filter;
+};
 
 template<typename T>
 struct TReadAccess : FReadAccess
@@ -147,12 +187,21 @@ struct FErasedOptionalReadAccess : FReadAccess
 {
 	static constexpr int32 PreLockedDataNum = 1;
 
-	FErasedOptionalReadAccess(FComponentTypeID InComponentTypeID)
+	FComponentTypeIDFilter Condition;
+
+	FErasedOptionalReadAccess(FComponentTypeID InComponentTypeID, FComponentTypeIDFilter InCondition)
 		: FReadAccess{ InComponentTypeID }
+		, Condition(InCondition)
 	{}
 
-	FORCEINLINE void PreLockComponentData(const FEntityAllocation* Allocation, FPreLockedDataPtr* PrelockedComponentData) const
+	FORCEINLINE void PreLockComponentData(FEntityAllocationIteratorItem Item, FPreLockedDataPtr* PrelockedComponentData) const
 	{
+		if (!Condition.Passes(Item.GetAllocationType()))
+		{
+			return;
+		}
+
+		const FEntityAllocation* Allocation = Item.GetAllocation();
 		const FComponentHeader* Header = Allocation->FindComponentHeader(ComponentType);
 		if (Header)
 		{
@@ -165,9 +214,14 @@ struct FErasedOptionalReadAccess : FReadAccess
 		// Resolve using the same base ptr as PreLockComponentData
 		return FReadErasedOptional(Ptr->Resolve<const FComponentHeader>(Allocation));
 	}
-	FORCEINLINE FOptionalComponentReader LockComponentData(const FEntityAllocation* Allocation, FEntityAllocationWriteContext WriteContext) const
+	FORCEINLINE FOptionalComponentReader LockComponentData(FEntityAllocationIteratorItem Item, FEntityAllocationWriteContext WriteContext) const
 	{
-		return Allocation->TryReadComponentsErased(ComponentType);
+		if (!Condition.Passes(Item.GetAllocationType()))
+		{
+			return FOptionalComponentReader();
+		}
+
+		return Item.GetAllocation()->TryReadComponentsErased(ComponentType);
 	}
 };
 
@@ -229,8 +283,8 @@ struct TOptionalReadAccess : FOptionalReadAccess
 	using AccessType = const T;
 	static constexpr int32 PreLockedDataNum = 1;
 
-	TOptionalReadAccess(FComponentTypeID InComponentTypeID)
-		: FOptionalReadAccess{InComponentTypeID}
+	TOptionalReadAccess(FComponentTypeID InComponentTypeID, FComponentTypeIDFilter InConditionType = FComponentTypeIDFilter())
+		: FOptionalReadAccess{InComponentTypeID, InConditionType}
 	{}
 
 	FORCEINLINE TReadOptional<T> ResolvePreLockedComponentData(const FEntityAllocation* Allocation, const FPreLockedDataPtr* Ptr, FEntityAllocationWriteContext WriteContext) const
@@ -253,8 +307,8 @@ struct TOptionalWriteAccess : FOptionalWriteAccess
 	using AccessType = T;
 	static constexpr int32 PreLockedDataNum = 1;
 
-	TOptionalWriteAccess(FComponentTypeID InComponentTypeID)
-		: FOptionalWriteAccess{ InComponentTypeID }
+	TOptionalWriteAccess(FComponentTypeID InComponentTypeID, FComponentTypeIDFilter InConditionType = FComponentTypeIDFilter())
+		: FOptionalWriteAccess{ InComponentTypeID, InConditionType }
 	{}
 
 	FORCEINLINE TWriteOptional<T> ResolvePreLockedComponentData(const FEntityAllocation* Allocation, const FPreLockedDataPtr* Ptr, FEntityAllocationWriteContext WriteContext) const
@@ -299,11 +353,11 @@ struct TReadOneOfAccessor
 		: ComponentTypes{ InComponentTypeIDs... }
 	{}
 
-	void PreLockComponentData(const FEntityAllocation* Allocation, FPreLockedDataPtr* PrelockedComponentData) const
+	void PreLockComponentData(FEntityAllocationIteratorItem Item, FPreLockedDataPtr* PrelockedComponentData) const
 	{
 		for (int32 Index = 0; Index < sizeof...(T); ++Index)
 		{
-			ComponentTypes[Index].PreLockComponentData(Allocation, &PrelockedComponentData[Index]);
+			ComponentTypes[Index].PreLockComponentData(Item, &PrelockedComponentData[Index]);
 		}
 	}
 
@@ -335,11 +389,11 @@ struct TReadOneOrMoreOfAccessor
 		: ComponentTypes{ InComponentTypeIDs... }
 	{}
 
-	void PreLockComponentData(const FEntityAllocation* Allocation, FPreLockedDataPtr* PrelockedComponentData) const
+	void PreLockComponentData(FEntityAllocationIteratorItem Item, FPreLockedDataPtr* PrelockedComponentData) const
 	{
 		for (int32 Index = 0; Index < sizeof...(T); ++Index)
 		{
-			ComponentTypes[Index].PreLockComponentData(Allocation, &PrelockedComponentData[Index]);
+			ComponentTypes[Index].PreLockComponentData(Item, &PrelockedComponentData[Index]);
 		}
 	}
 
@@ -362,6 +416,9 @@ struct TReadOneOrMoreOfAccessor
 inline void AddAccessorToFilter(const FEntityIDAccess*, FEntityComponentFilter* OutFilter)
 {
 }
+inline void AddAccessorToFilter(const FFilterMatchPassthrough*, FEntityComponentFilter* OutFilter)
+{
+}
 inline void AddAccessorToFilter(const FComponentAccess* In, FEntityComponentFilter* OutFilter)
 {
 	check(In->ComponentType);
@@ -542,12 +599,26 @@ bool HasBeenWrittenToSince(const TReadOneOrMoreOfAccessor<T...>* In, FEntityAllo
 	return bAnyWrittenTo;
 }
 
+template<typename T>
+auto GetComponentAtIndex(T* InAccessor, int32 Index)
+	-> decltype(DeclVal<T>().ComponentAtIndex(0))
+{
+	return InAccessor->ComponentAtIndex(Index);
+}
 
+inline bool GetComponentAtIndex(const bool* BoolPassthrough, int32 Index)
+{
+	return *BoolPassthrough;
+}
 
 inline bool IsAccessorValid(const FEntityIDAccess*)
 {
 	return true;
 }
+inline bool IsAccessorValid(const FFilterMatchPassthrough*)
+{
+	return true;
+}
 inline bool IsAccessorValid(const FComponentAccess* In)
 {
 	return In->ComponentType != FComponentTypeID::Invalid();
diff --git a/Engine/Source/Runtime/MovieScene/Public/EntitySystem/MovieSceneComponentRegistry.h b/Engine/Source/Runtime/MovieScene/Public/EntitySystem/MovieSceneComponentRegistry.h
--- a/Engine/Source/Runtime/MovieScene/Public/EntitySystem/MovieSceneComponentRegistry.h
+++ b/Engine/Source/Runtime/MovieScene/Public/EntitySystem/MovieSceneComponentRegistry.h
@@ -103,11 +103,11 @@ public:
 #if UE_MOVIESCENE_ENTITY_DEBUG
 		FString InitialValueDebugName = FString(TEXT("Initial ")) + DebugName;
 
-		OutComponents.PropertyTag = NewTag(DebugName, EComponentTypeFlags::CopyToChildren);
-		NewComponentType(&OutComponents.InitialValue, *InitialValueDebugName, EComponentTypeFlags::Preserved);
+		OutComponents.PropertyTag = NewTag(DebugName, EComponentTypeFlags::CopyToChildren | EComponentTypeFlags::CopyToOutput);
+		NewComponentType(&OutComponents.InitialValue, *InitialValueDebugName, EComponentTypeFlags::Preserved | EComponentTypeFlags::MigrateToOutput);
 #else
-		OutComponents.PropertyTag = NewTag(nullptr, EComponentTypeFlags::CopyToChildren);
-		NewComponentType(&OutComponents.InitialValue, nullptr, EComponentTypeFlags::Preserved);
+		OutComponents.PropertyTag = NewTag(nullptr, EComponentTypeFlags::CopyToChildren | EComponentTypeFlags::CopyToOutput);
+		NewComponentType(&OutComponents.InitialValue, nullptr, EComponentTypeFlags::Preserved | EComponentTypeFlags::MigrateToOutput);
 #endif
 	}
 
diff --git a/Engine/Source/Runtime/MovieScene/Public/EntitySystem/MovieSceneEntityIDs.h b/Engine/Source/Runtime/MovieScene/Public/EntitySystem/MovieSceneEntityIDs.h
--- a/Engine/Source/Runtime/MovieScene/Public/EntitySystem/MovieSceneEntityIDs.h
+++ b/Engine/Source/Runtime/MovieScene/Public/EntitySystem/MovieSceneEntityIDs.h
@@ -247,6 +247,37 @@ protected:
 	uint16 Value;
 };
 
+struct FComponentTypeIDFilter
+{
+	FComponentTypeIDFilter()
+		: Value(INVALID)
+		, PassFlag(false)
+	{}
+
+	FComponentTypeIDFilter(FComponentTypeID InComponentType, bool bInPassFlag)
+		: Value(InComponentType ? uint16(InComponentType.BitIndex()) : INVALID)
+		, PassFlag(bInPassFlag)
+	{
+	}
+
+	MOVIESCENE_API bool Passes(const FComponentMask& Type) const;
+
+	FComponentTypeID GetComponentType() const
+	{
+		if (Value == INVALID)
+		{
+			return FComponentTypeID();
+		}
+		return FComponentTypeID::FromBitIndex(Value);
+	}
+
+private:
+	static const uint16 INVALID = uint16(-1) >> 1u;
+
+	uint16 Value:15;
+	uint16 PassFlag:1;
+};
+
 template<typename T>
 struct TComponentTypeID : public FComponentTypeID
 {
diff --git a/Engine/Source/Runtime/MovieScene/Public/EntitySystem/MovieSceneEntitySystem.h b/Engine/Source/Runtime/MovieScene/Public/EntitySystem/MovieSceneEntitySystem.h
--- a/Engine/Source/Runtime/MovieScene/Public/EntitySystem/MovieSceneEntitySystem.h
+++ b/Engine/Source/Runtime/MovieScene/Public/EntitySystem/MovieSceneEntitySystem.h
@@ -215,6 +215,13 @@ public:
 	 */
 	MOVIESCENE_API void Disable();
 
+
+	/**
+	 * Checks whether custom task sheduling is enabled or not
+	 * @see FEntitySystemScheduler
+	 */
+	MOVIESCENE_API static bool IsCustomSchedulingEnabled();
+
 protected:
 
 	MOVIESCENE_API virtual bool IsReadyForFinishDestroy() override;
diff --git a/Engine/Source/Runtime/MovieScene/Public/EntitySystem/MovieSceneEntitySystemTask.h b/Engine/Source/Runtime/MovieScene/Public/EntitySystem/MovieSceneEntitySystemTask.h
--- a/Engine/Source/Runtime/MovieScene/Public/EntitySystem/MovieSceneEntitySystemTask.h
+++ b/Engine/Source/Runtime/MovieScene/Public/EntitySystem/MovieSceneEntitySystemTask.h
@@ -712,9 +712,9 @@ struct TEntityTaskComponentsImpl<TIntegerSequence<int, Indices...>, T...>
 	 *
 	 * @param ComponentType   A valid component type to read.
 	 */
-	TEntityTaskComponents<T..., FErasedOptionalReadAccess> ReadErasedOptional(FComponentTypeID ComponentType) const
+	TEntityTaskComponents<T..., FErasedOptionalReadAccess> ReadErasedOptional(FComponentTypeID ComponentType, FComponentTypeIDFilter InConditionType = FComponentTypeIDFilter()) const
 	{
-		return TEntityTaskComponents<T..., FErasedOptionalReadAccess >(CommonParams, Accessors.template Get<Indices>()..., FErasedOptionalReadAccess(ComponentType));
+		return TEntityTaskComponents<T..., FErasedOptionalReadAccess >(CommonParams, Accessors.template Get<Indices>()..., FErasedOptionalReadAccess(ComponentType, InConditionType));
 	}
 
 	/**
@@ -722,9 +722,9 @@ struct TEntityTaskComponentsImpl<TIntegerSequence<int, Indices...>, T...>
 	 * @note Always passed to the task as a const T* pointer which must be checked for null
 	 */
 	template<typename U>
-	TEntityTaskComponents<T..., TOptionalReadAccess<U>> ReadOptional(TComponentTypeID<U> ComponentType) const
+	TEntityTaskComponents<T..., TOptionalReadAccess<U>> ReadOptional(TComponentTypeID<U> ComponentType, FComponentTypeIDFilter InConditionType = FComponentTypeIDFilter()) const
 	{
-		return TEntityTaskComponents<T..., TOptionalReadAccess<U> >(CommonParams, Accessors.template Get<Indices>()..., TOptionalReadAccess<U>(ComponentType));
+		return TEntityTaskComponents<T..., TOptionalReadAccess<U> >(CommonParams, Accessors.template Get<Indices>()..., TOptionalReadAccess<U>(ComponentType, InConditionType));
 	}
 
 	/**
@@ -767,9 +767,14 @@ struct TEntityTaskComponentsImpl<TIntegerSequence<int, Indices...>, T...>
 	 * @note Always passed to the task as a T* pointer which must be checked for null
 	 */
 	template<typename U>
-	TEntityTaskComponents<T..., TOptionalWriteAccess<U>> WriteOptional(TComponentTypeID<U> ComponentType) const
+	TEntityTaskComponents<T..., TOptionalWriteAccess<U>> WriteOptional(TComponentTypeID<U> ComponentType, FComponentTypeIDFilter InConditionType = FComponentTypeIDFilter()) const
 	{
-		return TEntityTaskComponents<T..., TOptionalWriteAccess<U> >(CommonParams, Accessors.template Get<Indices>()..., TOptionalWriteAccess<U>(ComponentType));
+		return TEntityTaskComponents<T..., TOptionalWriteAccess<U> >(CommonParams, Accessors.template Get<Indices>()..., TOptionalWriteAccess<U>(ComponentType, InConditionType));
+	}
+
+	TEntityTaskComponents<T..., FFilterMatchPassthrough> PassthroughFilter(const FEntityComponentFilter& InFilter) const
+	{
+		return TEntityTaskComponents<T..., FFilterMatchPassthrough >(CommonParams, Accessors.template Get<Indices>()..., FFilterMatchPassthrough(InFilter));
 	}
 
 	bool HasBeenWrittenToSince(uint32 InSystemVersion)
@@ -834,12 +839,12 @@ struct TEntityTaskComponentsImpl<TIntegerSequence<int, Indices...>, T...>
 		(void)Temp;
 	}
 
-	void PreLockComponentHeaders(const FEntityAllocation* Allocation, TArray<FPreLockedDataPtr>& OutComponentHeaders) const
+	void PreLockComponentHeaders(FEntityAllocationIteratorItem Item, TArray<FPreLockedDataPtr>& OutComponentHeaders) const
 	{
 		constexpr TPrelockedDataOffsets<T...> PrelockedDataOffsets;
 		int32 StartIndex = OutComponentHeaders.Num();
 		OutComponentHeaders.AddDefaulted(PrelockedDataOffsets.Num);
-		( Accessors.template Get<Indices>().PreLockComponentData(Allocation, &OutComponentHeaders[StartIndex + PrelockedDataOffsets.StartOffset[Indices]]), ... );
+		( Accessors.template Get<Indices>().PreLockComponentData(Item, &OutComponentHeaders[StartIndex + PrelockedDataOffsets.StartOffset[Indices]]), ... );
 	}
 
 	void PopulateReadWriteDependencies(FComponentMask& OutReadDependencies, FComponentMask& OutWriteDependencies) const
@@ -886,7 +891,7 @@ struct TEntityTaskComponentsImpl<TIntegerSequence<int, Indices...>, T...>
 	template<typename Callback>
 	void Iterate_PerEntityImpl(FEntityManager* EntityManager, const FEntityComponentFilter& Filter, Callback&& InCallback) const
 	{
-		using TupleType = TTuple< decltype(Accessors.template Get<Indices>().LockComponentData(nullptr, DeclVal<FEntityAllocationWriteContext>()))... >;
+		using TupleType = TTuple< decltype(Accessors.template Get<Indices>().LockComponentData(DeclVal<FEntityAllocationIteratorItem>(), DeclVal<FEntityAllocationWriteContext>()))... >;
 
 		if (IsValid())
 		{
@@ -896,19 +901,20 @@ struct TEntityTaskComponentsImpl<TIntegerSequence<int, Indices...>, T...>
 				? EComponentHeaderLockMode::LockFree
 				: EComponentHeaderLockMode::Mutex;
 
-			for (FEntityAllocation* Allocation : EntityManager->Iterate(&Filter))
+			for (FEntityAllocationIteratorItem Item : EntityManager->Iterate(&Filter))
 			{
 				FEntityIterationResult Result;
 
+				FEntityAllocation* Allocation = Item.GetAllocation();
 				FEntityAllocationMutexGuard LockGuard(Allocation, LockMode);
 
 				// Lock the components we want to access
-				TupleType ComponentData( Accessors.template Get<Indices>().LockComponentData(Allocation, WriteContext)... );
+				TupleType ComponentData( Accessors.template Get<Indices>().LockComponentData(Item, WriteContext)... );
 
 				const int32 Num = Allocation->Num();
 				for (int32 ComponentOffset = 0; ComponentOffset < Num && Result.Value; ++ComponentOffset )
 				{
-					Result = (InCallback( ComponentData.template Get<Indices>().ComponentAtIndex(ComponentOffset)... ), Result);
+					Result = (InCallback( GetComponentAtIndex(&ComponentData.template Get<Indices>(), ComponentOffset)... ), Result);
 				}
 			}
 		}
@@ -934,11 +940,10 @@ struct TEntityTaskComponentsImpl<TIntegerSequence<int, Indices...>, T...>
 
 			for (FEntityAllocationIteratorItem Item : EntityManager->Iterate(&Filter))
 			{
-				FEntityAllocation* Allocation = Item;
 				FEntityAllocationMutexGuard LockGuard(Item.GetAllocation(), LockMode);
 
 				// Lock on the components we want to access
-				auto ComponentData = MakeTuple( Accessors.template Get<Indices>().LockComponentData(Allocation, WriteContext)... );
+				auto ComponentData = MakeTuple( Accessors.template Get<Indices>().LockComponentData(Item, WriteContext)... );
 
 				FEntityIterationResult Result = (InCallback(Item, ComponentData.template Get<Indices>()...), FEntityIterationResult{});
 
@@ -1553,10 +1558,10 @@ struct TEntityTaskBase
 			? EComponentHeaderLockMode::LockFree
 			: EComponentHeaderLockMode::Mutex;
 
-		for (FEntityAllocation* Allocation : EntityManager->Iterate(&FilteredTask.GetFilter()))
+		for (FEntityAllocationIteratorItem Item : EntityManager->Iterate(&FilteredTask.GetFilter()))
 		{
-			FEntityAllocationMutexGuard LockGuard(Allocation, LockMode);
-			Caller::ForEachEntityImpl(TaskImplInstance, Allocation, WriteContext, FilteredTask.GetComponents());
+			FEntityAllocationMutexGuard LockGuard(Item, LockMode);
+			Caller::ForEachEntityImpl(TaskImplInstance, Item, WriteContext, FilteredTask.GetComponents());
 		}
 
 		PostTask(&TaskImplInstance);
@@ -1855,32 +1860,32 @@ template<int... Indices>
 struct TEntityTaskCaller_AutoExpansion<TIntegerSequence<int, Indices...>>
 {
 	template<typename TaskImpl, typename... AccessorTypes>
-	static void ForEachEntityImpl(TaskImpl& TaskImplInstance, const FEntityAllocation* Allocation, TArrayView<const FPreLockedDataPtr> PreLockedData, FEntityAllocationWriteContext WriteContext, const TEntityTaskComponents<AccessorTypes...>& Components)
+	static void ForEachEntityImpl(TaskImpl& TaskImplInstance, FEntityAllocationIteratorItem Item, TArrayView<const FPreLockedDataPtr> PreLockedData, FEntityAllocationWriteContext WriteContext, const TEntityTaskComponents<AccessorTypes...>& Components)
 	{
 		FEntityIterationResult Result;
 
 		constexpr TPrelockedDataOffsets<AccessorTypes...> PrelockedDataOffsets;
-		auto ResolvedComponentData = MakeTuple( Components.template GetAccessor<Indices>().ResolvePreLockedComponentData(Allocation, &PreLockedData[PrelockedDataOffsets.StartOffset[Indices]], WriteContext)... );
+		auto ResolvedComponentData = MakeTuple( Components.template GetAccessor<Indices>().ResolvePreLockedComponentData(Item, &PreLockedData[PrelockedDataOffsets.StartOffset[Indices]], WriteContext)... );
 
-		const int32 Num = Allocation->Num();
+		const int32 Num = Item.GetAllocation()->Num();
 		for (int32 ComponentOffset = 0; ComponentOffset < Num && Result.Value; ++ComponentOffset )
 		{
-			Result = (TaskImplInstance.ForEachEntity(ResolvedComponentData.template Get<Indices>().ComponentAtIndex(ComponentOffset)... ), Result);
+			Result = (TaskImplInstance.ForEachEntity(GetComponentAtIndex(&ResolvedComponentData.template Get<Indices>(), ComponentOffset)... ), Result);
 		}
 	}
 
 	template<typename TaskImpl, typename... AccessorTypes>
-	static void ForEachEntityImpl(TaskImpl& TaskImplInstance, const FEntityAllocation* Allocation, FEntityAllocationWriteContext WriteContext, const TEntityTaskComponents<AccessorTypes...>& Components)
+	static void ForEachEntityImpl(TaskImpl& TaskImplInstance, FEntityAllocationIteratorItem Item, FEntityAllocationWriteContext WriteContext, const TEntityTaskComponents<AccessorTypes...>& Components)
 	{
 		FEntityIterationResult Result;
 
 		// Lock the components we want to access
-		auto LockedComponentData = MakeTuple( Components.template GetAccessor<Indices>().LockComponentData(Allocation, WriteContext)... );
+		auto LockedComponentData = MakeTuple( Components.template GetAccessor<Indices>().LockComponentData(Item, WriteContext)... );
 
-		const int32 Num = Allocation->Num();
+		const int32 Num = Item.GetAllocation()->Num();
 		for (int32 ComponentOffset = 0; ComponentOffset < Num && Result.Value; ++ComponentOffset )
 		{
-			Result = (TaskImplInstance.ForEachEntity(LockedComponentData.template Get<Indices>().ComponentAtIndex(ComponentOffset)... ), Result);
+			Result = (TaskImplInstance.ForEachEntity(GetComponentAtIndex(&LockedComponentData.template Get<Indices>(), ComponentOffset)... ), Result);
 		}
 	}
 
@@ -1889,13 +1894,13 @@ struct TEntityTaskCaller_AutoExpansion<TIntegerSequence<int, Indices...>>
 	{
 		constexpr TPrelockedDataOffsets<AccessorTypes...> PrelockedDataOffsets;
 		const FEntityAllocation* Allocation = Item.GetAllocation();
-		TaskImplInstance.ForEachAllocation(Item, Components.template GetAccessor<Indices>().ResolvePreLockedComponentData(Allocation, &PreLockedData[PrelockedDataOffsets.StartOffset[Indices]], WriteContext)...);
+		TaskImplInstance.ForEachAllocation(Item, Components.template GetAccessor<Indices>().ResolvePreLockedComponentData(Item, &PreLockedData[PrelockedDataOffsets.StartOffset[Indices]], WriteContext)...);
 	}
 
 	template<typename TaskImpl, typename... AccessorTypes>
 	static void ForEachAllocationImpl(TaskImpl& TaskImplInstance, FEntityAllocationIteratorItem Item, FEntityAllocationWriteContext WriteContext, const TEntityTaskComponents<AccessorTypes...>& Components)
 	{
-		auto LockedComponentData = MakeTuple( Components.template GetAccessor<Indices>().LockComponentData(Item.GetAllocation(), WriteContext)... );
+		auto LockedComponentData = MakeTuple( Components.template GetAccessor<Indices>().LockComponentData(Item, WriteContext)... );
 		TaskImplInstance.ForEachAllocation(Item, LockedComponentData.template Get<Indices>()...);
 	}
 };
diff --git a/Engine/Source/Runtime/MovieScene/Public/EntitySystem/MovieSceneInitialValueCache.h b/Engine/Source/Runtime/MovieScene/Public/EntitySystem/MovieSceneInitialValueCache.h
--- a/Engine/Source/Runtime/MovieScene/Public/EntitySystem/MovieSceneInitialValueCache.h
+++ b/Engine/Source/Runtime/MovieScene/Public/EntitySystem/MovieSceneInitialValueCache.h
@@ -210,6 +210,23 @@ struct TPropertyValueStorage : IPropertyValueStorage
 	}
 
 
+	/**
+	 * Add a cached value for the specified object and a propertypath, returning a unique index for the value
+	 * @note: Value must not have been cached previously - doing so will result in a failed assertion
+	 *
+	 * @param BoundObject            The object instance to cache the property from.
+	 * @param InValue                The value to cache
+	 * @param PropertyPath           The path of the property
+	 * @return A tuple containing the cached value and its index
+	 */
+	FInitialValueIndex AddInitialValue(UObject* BoundObject, const StorageType& InValue, const FName& PropertyPath)
+	{
+		FKeyType Key{ FObjectKey(BoundObject), FPropertyKey(TInPlaceType<FName>(), PropertyPath)};
+
+		return AddInitialValue(Key, InValue);
+	}
+
+
 	/**
 	 * Find an initial value index given its object and fast ptr offset
 	 */
diff --git a/Engine/Source/Runtime/MovieScene/Public/EntitySystem/MovieSceneInitialValueSystem.h b/Engine/Source/Runtime/MovieScene/Public/EntitySystem/MovieSceneInitialValueSystem.h
new file mode 100644
--- /dev/null
+++ b/Engine/Source/Runtime/MovieScene/Public/EntitySystem/MovieSceneInitialValueSystem.h
@@ -0,0 +1,41 @@
+// Copyright Epic Games, Inc. All Rights Reserved.
+
+#pragma once
+
+#include "EntitySystem/MovieSceneEntityInstantiatorSystem.h"
+#include "MovieSceneInitialValueSystem.generated.h"
+
+namespace UE::MovieScene
+{
+	struct FComponentTypeID;
+	struct IInitialValueProcessor;
+	struct FEntityComponentFilter;
+}
+
+
+/**
+ * System responsible for initializing initial values for all property types
+ * Will handle the presence of an FInitialValueCache extension on the linker
+ */
+UCLASS(MinimalAPI)
+class UMovieSceneInitialValueSystem
+	: public UMovieSceneEntityInstantiatorSystem
+{
+public:
+
+	GENERATED_BODY()
+
+	MOVIESCENE_API UMovieSceneInitialValueSystem(const FObjectInitializer& ObjInit);
+
+	MOVIESCENE_API static void RegisterProcessor(
+		UE::MovieScene::FComponentTypeID InitialValueComponent,
+		TSharedPtr<UE::MovieScene::IInitialValueProcessor> Processor,
+		UE::MovieScene::FEntityComponentFilter&& OptionalFilter = UE::MovieScene::FEntityComponentFilter());
+
+private:
+
+	virtual void OnRun(FSystemTaskPrerequisites& InPrerequisites, FSystemSubsequentTasks& Subsequents) override;
+	virtual bool IsRelevantImpl(UMovieSceneEntitySystemLinker* InLinker) const override;
+	virtual void OnLink() override;
+	virtual void OnUnlink() override;
+};
diff --git a/Engine/Source/Runtime/MovieScene/Public/EntitySystem/MovieSceneInterrogation.h b/Engine/Source/Runtime/MovieScene/Public/EntitySystem/MovieSceneInterrogation.h
new file mode 100644
--- /dev/null
+++ b/Engine/Source/Runtime/MovieScene/Public/EntitySystem/MovieSceneInterrogation.h
@@ -0,0 +1,77 @@
+// Copyright Epic Games, Inc. All Rights Reserved.
+
+#pragma once
+
+#include "CoreTypes.h"
+#include "UObject/WeakObjectPtrTemplates.h"
+#include "UObject/NameTypes.h"
+#include "Containers/SortedMap.h"
+#include "EntitySystem/MovieSceneEntitySystemTypes.h"
+#include "EntitySystem/MovieScenePropertyBinding.h"
+#include "EntitySystem/MovieSceneEntitySystemLinkerExtension.h"
+
+namespace UE::MovieScene
+{
+
+
+struct FInterrogationChannelInfo
+{
+	/** The object that relates to the channel */
+	TWeakObjectPtr<UObject> WeakObject;
+
+	/** (Optional) property path for this channel */
+	FMovieScenePropertyBinding PropertyBinding;
+
+	/** The channel's hierarchical 'parent' - generally only used for transforms */
+	FInterrogationChannel ParentChannel;
+};
+
+
+struct FInterrogationUpdateParams
+{
+
+	/** The channel's hierarchical 'parent' - generally only used for transforms */
+	FInterrogationChannel ParentChannel;
+};
+
+struct FSparseInterrogationChannelInfo
+{
+	FInterrogationChannelInfo& Get(FInterrogationChannel Channel) { return ChannelInfo.FindOrAdd(Channel); }
+	const FInterrogationChannelInfo& Get(FInterrogationChannel Channel) const { return ChannelInfo.FindChecked(Channel); }
+
+	FInterrogationChannelInfo* Find(FInterrogationChannel Channel) { return ChannelInfo.Find(Channel); }
+	const FInterrogationChannelInfo* Find(FInterrogationChannel Channel) const { return ChannelInfo.Find(Channel); }
+
+	UObject* FindObject(FInterrogationChannel Channel) const
+	{
+		return ChannelInfo.FindRef(Channel).WeakObject.Get();
+	}
+
+	FInterrogationChannel FindParent(FInterrogationChannel Channel) const
+	{
+		return ChannelInfo.FindRef(Channel).ParentChannel;
+	}
+
+	void Empty()
+	{
+		ChannelInfo.Empty();
+	}
+
+private:
+
+	TSortedMap<FInterrogationChannel, FInterrogationChannelInfo> ChannelInfo;
+};
+
+struct IInterrogationExtension
+{
+public:
+	MOVIESCENE_API static TEntitySystemLinkerExtensionID<IInterrogationExtension> GetExtensionID();
+
+	virtual const FSparseInterrogationChannelInfo& GetSparseChannelInfo() const = 0;
+
+protected:
+
+	virtual ~IInterrogationExtension() {}
+};
+
+} // namespace UE::MovieScene
\ No newline at end of file
diff --git a/Engine/Source/Runtime/MovieScene/Public/EntitySystem/MovieScenePropertyComponentHandler.h b/Engine/Source/Runtime/MovieScene/Public/EntitySystem/MovieScenePropertyComponentHandler.h
--- a/Engine/Source/Runtime/MovieScene/Public/EntitySystem/MovieScenePropertyComponentHandler.h
+++ b/Engine/Source/Runtime/MovieScene/Public/EntitySystem/MovieScenePropertyComponentHandler.h
@@ -45,41 +45,55 @@ struct TInitialValueProcessorImpl<PropertyTraits, TPropertyMetaData<MetaDataType
 	TSortedMap<FInterrogationChannel, StorageType> ValuesByChannel;
 
 	FBuiltInComponentTypes* BuiltInComponents;
-	IInterrogationExtension* Interrogation;
-	const FPropertyDefinition* PropertyDefinition;
+
+	ICustomPropertyRegistration* CustomPropertyRegistration;
+	TComponentTypeID<StorageType> InitialValueType;
+	TTuple<TComponentTypeID<MetaDataTypes>...> MetaDataComponents;
 	FCustomAccessorView CustomAccessors;
 
-	FEntityAllocationWriteContext WriteContext;
+	// Transient properties reset on each usage.
+	IInterrogationExtension* Interrogation;
 	TPropertyValueStorage<PropertyTraits>* CacheStorage;
 
-	TInitialValueProcessorImpl()
-		: WriteContext(FEntityAllocationWriteContext::NewAllocation())
-	{
-		BuiltInComponents = FBuiltInComponentTypes::Get();
+	FEntityAllocationWriteContext WriteContext;
 
+	TInitialValueProcessorImpl(
+		TComponentTypeID<StorageType> InInitialValueType,
+		TArrayView<const FComponentTypeID> InMetaDataComponents,
+		ICustomPropertyRegistration* InCustomPropertyRegistration
+	)
+		: BuiltInComponents(FBuiltInComponentTypes::Get())
+		, CustomPropertyRegistration(InCustomPropertyRegistration)
+		, InitialValueType(InInitialValueType)
+		, MetaDataComponents(InMetaDataComponents[MetaDataIndices].ReinterpretCast<MetaDataTypes>()...)
+		, WriteContext(FEntityAllocationWriteContext::NewAllocation())
+	{
 		Interrogation = nullptr;
 		CacheStorage = nullptr;
 	}
 
-	virtual void Initialize(UMovieSceneEntitySystemLinker* Linker, const FPropertyDefinition* Definition, FInitialValueCache* InitialValueCache) override
+	virtual void Initialize(UMovieSceneEntitySystemLinker* Linker, FInitialValueCache* InitialValueCache) override
 	{
-		PropertyDefinition = Definition;
 		Interrogation = Linker->FindExtension<IInterrogationExtension>();
 		WriteContext  = FEntityAllocationWriteContext(Linker->EntityManager);
 
-		check(PropertyDefinition->MetaDataTypes.Num() == PropertyTraits::MetaDataType::Num);
-
-		if (PropertyDefinition->CustomPropertyRegistration)
+		if (CustomPropertyRegistration)
 		{
-			CustomAccessors = PropertyDefinition->CustomPropertyRegistration->GetAccessors();
+			CustomAccessors = CustomPropertyRegistration->GetAccessors();
 		}
 
 		if (InitialValueCache)
 		{
-			CacheStorage = InitialValueCache->GetStorage<PropertyTraits>(Definition->InitialValueType);
+			CacheStorage = InitialValueCache->GetStorage<PropertyTraits>(InitialValueType);
 		}
 	}
 
+	virtual void PopulateFilter(FEntityComponentFilter& OutFilter) const override
+	{
+		// These initializers must have bound objects
+		OutFilter.All({ BuiltInComponents->BoundObject });
+	}
+
 	virtual void Process(const FEntityAllocation* Allocation, const FComponentMask& AllocationType) override
 	{
 		if (Interrogation && AllocationType.Contains(BuiltInComponents->Interrogation.OutputKey))
@@ -101,18 +115,17 @@ struct TInitialValueProcessorImpl<PropertyTraits, TPropertyMetaData<MetaDataType
 		ValuesByChannel.Empty();
 		Interrogation = nullptr;
 		CacheStorage = nullptr;
-		CustomAccessors = FCustomAccessorView();
 	}
 
 	void VisitAllocation(const FEntityAllocation* Allocation)
 	{
 		const int32 Num = Allocation->Num();
 
-		TComponentWriter<StorageType> InitialValues = Allocation->WriteComponents(PropertyDefinition->InitialValueType.ReinterpretCast<StorageType>(), WriteContext);
+		TComponentWriter<StorageType> InitialValues = Allocation->WriteComponents(InitialValueType, WriteContext);
 		TComponentReader<UObject*>    BoundObjects  = Allocation->ReadComponents(BuiltInComponents->BoundObject);
 
 		TTuple< TComponentReader<MetaDataTypes>... > MetaData(
-			Allocation->ReadComponents(PropertyDefinition->GetMetaDataComponent<MetaDataTypes>(MetaDataIndices))...
+			Allocation->ReadComponents(MetaDataComponents.template Get<MetaDataIndices>())...
 		);
 
 		if (TOptionalComponentReader<FCustomPropertyIndex> CustomIndices = Allocation->TryReadComponents(BuiltInComponents->CustomPropertyIndex))
@@ -148,11 +161,11 @@ struct TInitialValueProcessorImpl<PropertyTraits, TPropertyMetaData<MetaDataType
 		const int32 Num = Allocation->Num();
 
 		TComponentWriter<FInitialValueIndex> InitialValueIndices = Allocation->WriteComponents(BuiltInComponents->InitialValueIndex, WriteContext);
-		TComponentWriter<StorageType>        InitialValues       = Allocation->WriteComponents(PropertyDefinition->InitialValueType.ReinterpretCast<StorageType>(), WriteContext);
+		TComponentWriter<StorageType>        InitialValues       = Allocation->WriteComponents(InitialValueType, WriteContext);
 		TComponentReader<UObject*>           BoundObjects        = Allocation->ReadComponents(BuiltInComponents->BoundObject);
 
 		TTuple< TComponentReader<MetaDataTypes>... > MetaData(
-			Allocation->ReadComponents(PropertyDefinition->GetMetaDataComponent<MetaDataTypes>(MetaDataIndices))...
+			Allocation->ReadComponents(MetaDataComponents.template Get<MetaDataIndices>())...
 		);
 
 		if (TOptionalComponentReader<FCustomPropertyIndex> CustomIndices = Allocation->TryReadComponents(BuiltInComponents->CustomPropertyIndex))
@@ -223,11 +236,11 @@ struct TInitialValueProcessorImpl<PropertyTraits, TPropertyMetaData<MetaDataType
 	{
 		const int32 Num = Allocation->Num();
 
-		TComponentWriter<StorageType>       InitialValues = Allocation->WriteComponents(PropertyDefinition->InitialValueType.ReinterpretCast<StorageType>(), WriteContext);
+		TComponentWriter<StorageType>       InitialValues = Allocation->WriteComponents(InitialValueType, WriteContext);
 		TComponentReader<FInterrogationKey> OutputKeys    = Allocation->ReadComponents(BuiltInComponents->Interrogation.OutputKey);
 
 		TTuple< TComponentReader<MetaDataTypes>... > MetaData(
-			Allocation->ReadComponents(PropertyDefinition->GetMetaDataComponent<MetaDataTypes>(MetaDataIndices))...
+			Allocation->ReadComponents(MetaDataComponents.template Get<MetaDataIndices>())...
 		);
 
 		const FSparseInterrogationChannelInfo& SparseChannelInfo = Interrogation->GetSparseChannelInfo();
@@ -301,7 +314,18 @@ struct TInitialValueProcessorImpl<PropertyTraits, TPropertyMetaData<MetaDataType
 
 template<typename PropertyTraits>
 struct TInitialValueProcessor : TInitialValueProcessorImpl<PropertyTraits, typename PropertyTraits::MetaDataType, TMakeIntegerSequence<int, PropertyTraits::MetaDataType::Num>>
-{};
+{
+	using Super = TInitialValueProcessorImpl<PropertyTraits, typename PropertyTraits::MetaDataType, TMakeIntegerSequence<int, PropertyTraits::MetaDataType::Num>>;
+
+	TInitialValueProcessor() = delete;
+
+	TInitialValueProcessor(
+		TComponentTypeID<typename PropertyTraits::StorageType> InInitialValueType,
+		TArrayView<const FComponentTypeID> InMetaDataComponents,
+		ICustomPropertyRegistration* InCustomPropertyRegistration
+	) : Super(InInitialValueType, InMetaDataComponents, InCustomPropertyRegistration)
+	{}
+};
 
 template<typename T, typename U = decltype(T::bIsComposite)>
 constexpr bool IsCompositePropertyTraits(T*)
@@ -424,10 +448,13 @@ struct TPropertyComponentHandlerImpl<PropertyTraits, TPropertyMetaData<MetaDataT
 		}
 	}
 
-	virtual IInitialValueProcessor* GetInitialValueProcessor() override
+	virtual TSharedPtr<IInitialValueProcessor> MakeInitialValueProcessor(const FPropertyDefinition& Definition) override
 	{
-		static TInitialValueProcessor<PropertyTraits> Processor;
-		return &Processor;
+		return MakeShared<TInitialValueProcessor<PropertyTraits>>(
+			Definition.InitialValueType.ReinterpretCast<StorageType>(),
+			Definition.MetaDataTypes,
+			Definition.CustomPropertyRegistration
+			);
 	}
 
 	virtual void RecomposeBlendOperational(const FPropertyDefinition& PropertyDefinition, TArrayView<const FPropertyCompositeDefinition> Composites, const FValueDecompositionParams& InParams, UMovieSceneBlenderSystem* Blender, FConstPropertyComponentView InCurrentValue, FPropertyComponentArrayView OutResult) override
@@ -653,12 +680,14 @@ struct TPropertyDefinitionBuilder
 	void Commit()
 	{
 		Definition->Handler = TPropertyComponentHandler<PropertyTraits, typename PropertyTraits::StorageType>();
+		Definition->SetupInitialValueProcessor();
 	}
 
 	template<typename HandlerType>
 	void Commit(HandlerType&& InHandler)
 	{
 		Definition->Handler = Forward<HandlerType>(InHandler);
+		Definition->SetupInitialValueProcessor();
 	}
 
 protected:
@@ -744,12 +773,14 @@ struct TCompositePropertyDefinitionBuilder
 	void Commit()
 	{
 		Definition->Handler = TPropertyComponentHandler<PropertyTraits, Composites...>();
+		Definition->SetupInitialValueProcessor();
 	}
 	
 	template<typename HandlerType>
 	void Commit(HandlerType&& InHandler)
 	{
 		Definition->Handler = Forward<HandlerType>(InHandler);
+		Definition->SetupInitialValueProcessor();
 	}
 
 private:
diff --git a/Engine/Source/Runtime/MovieScene/Public/EntitySystem/MovieScenePropertyRegistry.h b/Engine/Source/Runtime/MovieScene/Public/EntitySystem/MovieScenePropertyRegistry.h
--- a/Engine/Source/Runtime/MovieScene/Public/EntitySystem/MovieScenePropertyRegistry.h
+++ b/Engine/Source/Runtime/MovieScene/Public/EntitySystem/MovieScenePropertyRegistry.h
@@ -90,6 +90,11 @@ struct FPropertyDefinition
 		return MetaDataTypes[Index].ReinterpretCast<T>();
 	}
 
+	/**
+	 * Sets up an initial value processor for this property type
+	 */
+	MOVIESCENE_API void SetupInitialValueProcessor() const;
+
 	/** Pointer to a custom getter/setter registry for short circuiting the UObject VM. Must outlive this definitions lifetime (usually these are static or singletons) */
 	ICustomPropertyRegistration* CustomPropertyRegistration = nullptr;
 
diff --git a/Engine/Source/Runtime/MovieScene/Public/Evaluation/PreAnimatedState/MovieScenePreAnimatedObjectStorage.h b/Engine/Source/Runtime/MovieScene/Public/Evaluation/PreAnimatedState/MovieScenePreAnimatedObjectStorage.h
--- a/Engine/Source/Runtime/MovieScene/Public/Evaluation/PreAnimatedState/MovieScenePreAnimatedObjectStorage.h
+++ b/Engine/Source/Runtime/MovieScene/Public/Evaluation/PreAnimatedState/MovieScenePreAnimatedObjectStorage.h
@@ -32,6 +32,12 @@ struct FBoundObjectPreAnimatedStateTraits : FPreAnimatedStateTraits
 
 	MOVIESCENE_API void Initialize(FPreAnimatedStorageID InStorageID, FPreAnimatedStateExtension* InParentExtension);
 
+	template<typename T>
+	T ResolveComponent(T&& In)
+	{
+		return Forward<T>(In);
+	}
+
 	/* Defined as a template rather than a variadic function to prevent error C4840 */
 	template<typename... T>
 	FPreAnimatedStorageGroupHandle FindGroup(UObject* BoundObject, T&&... Unused)
diff --git a/Engine/Source/Runtime/MovieScene/Public/Evaluation/PreAnimatedState/MovieScenePreAnimatedStateStorage.h b/Engine/Source/Runtime/MovieScene/Public/Evaluation/PreAnimatedState/MovieScenePreAnimatedStateStorage.h
--- a/Engine/Source/Runtime/MovieScene/Public/Evaluation/PreAnimatedState/MovieScenePreAnimatedStateStorage.h
+++ b/Engine/Source/Runtime/MovieScene/Public/Evaluation/PreAnimatedState/MovieScenePreAnimatedStateStorage.h
@@ -467,7 +467,7 @@ public:
 		FPreAnimatedStorageGroupHandle GroupHandle;
 		if constexpr (StorageTraits::SupportsGrouping)
 		{
-			GroupHandle = this->Traits.FindGroup(InKeyArgs...);
+			GroupHandle = this->Traits.FindGroup(this->Traits.ResolveComponent(InKeyArgs)...);
 		}
 
 		FPreAnimatedStorageIndex StorageIndex = this->FindStorageIndex(Key);
@@ -488,7 +488,7 @@ public:
 		FPreAnimatedStorageGroupHandle GroupHandle;
 		if constexpr(StorageTraits::SupportsGrouping)
 		{
-			GroupHandle = this->Traits.MakeGroup(InKeyArgs...);
+			GroupHandle = this->Traits.MakeGroup(this->Traits.ResolveComponent(InKeyArgs)...);
 		}
 
 		FPreAnimatedStorageIndex StorageIndex = this->GetOrCreateStorageIndex(Key);
@@ -496,6 +496,11 @@ public:
 		return FPreAnimatedStateEntry{ GroupHandle, FPreAnimatedStateCachedValueHandle{ this->StorageID, StorageIndex } };
 	}
 
+	bool IsCapturingGlobalState() const
+	{
+		return this->ParentExtension->IsCapturingGlobalState();
+	}
+
 	/**
 	 * Cause a previously cached value to always outlive any actively animating sources.
 	 * This is called when a Restore State track overlaps a Keep State track. The Restore State track will initially
@@ -667,7 +672,7 @@ public:
 				EPreAnimatedStorageRequirement StorageRequirement = this->ParentExtension->GetStorageRequirement(Entry);
 				if (!this->IsStorageRequirementSatisfied(Entry.ValueHandle.StorageIndex, StorageRequirement))
 				{
-					StorageType NewValue = this->Traits.CachePreAnimatedValue(Values[Index]...);
+					StorageType NewValue = this->Traits.CachePreAnimatedValue(this->Traits.ResolveComponent(Values[Index])...);
 					this->AssignPreAnimatedValue(Entry.ValueHandle.StorageIndex, StorageRequirement, MoveTemp(NewValue));
 				}
 			}
@@ -708,7 +713,7 @@ public:
 
 			if (!IsStorageRequirementSatisfied(Entry.ValueHandle.StorageIndex, Requirement))
 			{
-				StorageType NewValue = this->Traits.CachePreAnimatedValue(Values...);
+				StorageType NewValue = this->Traits.CachePreAnimatedValue(this->Traits.ResolveComponent(Values)...);
 				this->AssignPreAnimatedValue(Entry.ValueHandle.StorageIndex, Requirement, MoveTemp(NewValue));
 			}
 		}
@@ -791,7 +796,7 @@ public:
 				EPreAnimatedStorageRequirement StorageRequirement = this->ParentExtension->GetStorageRequirement(Entry);
 				if (!this->IsStorageRequirementSatisfied(Entry.ValueHandle.StorageIndex, StorageRequirement))
 				{
-					StorageType NewValue = this->Traits.CachePreAnimatedValue(Inputs[Index]...);
+					StorageType NewValue = this->Traits.CachePreAnimatedValue(this->Traits.ResolveComponent(Inputs[Index])...);
 					this->AssignPreAnimatedValue(Entry.ValueHandle.StorageIndex, StorageRequirement, MoveTemp(NewValue));
 				}
 			}
@@ -942,6 +947,8 @@ protected:
 	/** Our storage ID */
 	FPreAnimatedStorageID StorageID;
 
+public:
+
 	/** Traits instance - normally not required but some traits require state */
 	StorageTraits Traits;
 };
diff --git a/Engine/Source/Runtime/MovieSceneTracks/Private/EntitySystem/Interrogation/MovieSceneInterrogatedPropertyInstantiator.cpp b/Engine/Source/Runtime/MovieSceneTracks/Private/EntitySystem/Interrogation/MovieSceneInterrogatedPropertyInstantiator.cpp
--- a/Engine/Source/Runtime/MovieSceneTracks/Private/EntitySystem/Interrogation/MovieSceneInterrogatedPropertyInstantiator.cpp
+++ b/Engine/Source/Runtime/MovieSceneTracks/Private/EntitySystem/Interrogation/MovieSceneInterrogatedPropertyInstantiator.cpp
@@ -6,6 +6,7 @@
 #include "EntitySystem/MovieSceneEntityBuilder.h"
 #include "EntitySystem/MovieSceneEntitySystemLinker.h"
 #include "EntitySystem/MovieScenePropertyRegistry.h"
+#include "EntitySystem/MovieSceneInitialValueSystem.h"
 #include "Systems/MovieScenePiecewiseDoubleBlenderSystem.h"
 
 #include "Algo/IndexOf.h"
@@ -31,6 +32,8 @@ UMovieSceneInterrogatedPropertyInstantiatorSystem::UMovieSceneInterrogatedProper
 	{
 		DefineComponentProducer(GetClass(), BuiltInComponents->BlendChannelInput);
 		DefineComponentProducer(GetClass(), BuiltInComponents->SymbolicTags.CreatesEntities);
+
+		DefineImplicitPrerequisite(GetClass(), UMovieSceneInitialValueSystem::StaticClass());
 	}
 
 	CleanFastPathMask.Set(BuiltInComponents->Interrogation.OutputKey);
diff --git a/Engine/Source/Runtime/MovieSceneTracks/Private/EntitySystem/Interrogation/MovieSceneInterrogationLinker.cpp b/Engine/Source/Runtime/MovieSceneTracks/Private/EntitySystem/Interrogation/MovieSceneInterrogationLinker.cpp
--- a/Engine/Source/Runtime/MovieSceneTracks/Private/EntitySystem/Interrogation/MovieSceneInterrogationLinker.cpp
+++ b/Engine/Source/Runtime/MovieSceneTracks/Private/EntitySystem/Interrogation/MovieSceneInterrogationLinker.cpp
@@ -68,12 +68,6 @@ struct FSystemInterrogatorEntityTracker
 };
 
 
-TEntitySystemLinkerExtensionID<IInterrogationExtension> IInterrogationExtension::GetExtensionID()
-{
-	static TEntitySystemLinkerExtensionID<IInterrogationExtension> ID = UMovieSceneEntitySystemLinker::RegisterExtension<IInterrogationExtension>();
-	return ID;
-}
-
 FInterrogationChannels::FInterrogationChannels()
 {
 	// Always add a bit for the default channel
diff --git a/Engine/Source/Runtime/MovieSceneTracks/Private/MovieSceneTracksComponentTypes.cpp b/Engine/Source/Runtime/MovieSceneTracks/Private/MovieSceneTracksComponentTypes.cpp
--- a/Engine/Source/Runtime/MovieSceneTracks/Private/MovieSceneTracksComponentTypes.cpp
+++ b/Engine/Source/Runtime/MovieSceneTracks/Private/MovieSceneTracksComponentTypes.cpp
@@ -371,7 +371,7 @@ struct FColorHandler : TPropertyComponentHandler<FColorPropertyTraits, double, d
 
 struct FFloatParameterHandler : TPropertyComponentHandler<FFloatParameterTraits, double>
 {
-	virtual IInitialValueProcessor* GetInitialValueProcessor() override
+	virtual TSharedPtr<IInitialValueProcessor> MakeInitialValueProcessor(const FPropertyDefinition& Definition) override
 	{
 		return nullptr;
 	}
@@ -379,7 +379,7 @@ struct FFloatParameterHandler : TPropertyComponentHandler<FFloatParameterTraits,
 
 struct FColorParameterHandler : TPropertyComponentHandler<FColorParameterTraits, double, double, double, double>
 {
-	virtual IInitialValueProcessor* GetInitialValueProcessor() override
+	virtual TSharedPtr<IInitialValueProcessor> MakeInitialValueProcessor(const FPropertyDefinition& Definition) override
 	{
 		return nullptr;
 	}
@@ -568,6 +568,15 @@ FMovieSceneTracksComponentTypes::FMovieSceneTracksComponentTypes()
 	ComponentRegistry->NewPropertyType(FloatParameter, TEXT("float parameter"));
 	ComponentRegistry->NewPropertyType(ColorParameter, TEXT("color parameter"));
 
+	ComponentRegistry->NewPropertyType(Parameters.Bool, TEXT("Bool parameter"));
+	ComponentRegistry->NewPropertyType(Parameters.Byte, TEXT("Byte parameter"));
+	ComponentRegistry->NewPropertyType(Parameters.Integer, TEXT("Integer parameter"));
+	ComponentRegistry->NewPropertyType(Parameters.Scalar, TEXT("Scalar parameter"));
+	ComponentRegistry->NewPropertyType(Parameters.Vector2, TEXT("Vector2 parameter"));
+	ComponentRegistry->NewPropertyType(Parameters.Vector3, TEXT("Vector3 parameter"));
+	ComponentRegistry->NewPropertyType(Parameters.Color, TEXT("Color parameter"));
+	ComponentRegistry->NewPropertyType(Parameters.Transform, TEXT("Transform parameter"));
+
 	Bool.MetaDataComponents.Initialize(ComponentRegistry, TEXT("Bool Bitfield"));
 	Color.MetaDataComponents.Initialize(ComponentRegistry, TEXT("Color Type"));
 	FloatVector.MetaDataComponents.Initialize(ComponentRegistry, TEXT("Num Float Vector Channels"));
@@ -605,6 +614,7 @@ FMovieSceneTracksComponentTypes::FMovieSceneTracksComponentTypes()
 	ComponentRegistry->NewComponentType(&VectorParameterName,    TEXT("Vector Parameter Name"), EComponentTypeFlags::CopyToChildren | EComponentTypeFlags::CopyToOutput);
 	ComponentRegistry->NewComponentType(&ColorParameterName,     TEXT("Color Parameter Name"), EComponentTypeFlags::CopyToChildren | EComponentTypeFlags::CopyToOutput);
 	ComponentRegistry->NewComponentType(&TransformParameterName, TEXT("Transform Parameter Name"), EComponentTypeFlags::CopyToChildren | EComponentTypeFlags::CopyToOutput);
+	ComponentRegistry->NewComponentType(&GenericParameterName,   TEXT("Parameter Name"), EComponentTypeFlags::CopyToChildren | EComponentTypeFlags::CopyToOutput);
 
 	ComponentRegistry->NewComponentType(&ScalarMaterialParameterInfo, TEXT("Scalar Material Parameter Info"), EComponentTypeFlags::CopyToChildren | EComponentTypeFlags::CopyToOutput);
 	ComponentRegistry->NewComponentType(&ColorMaterialParameterInfo, TEXT("Color Material Parameter Info"), EComponentTypeFlags::CopyToChildren | EComponentTypeFlags::CopyToOutput);
@@ -738,6 +748,11 @@ FMovieSceneTracksComponentTypes::FMovieSceneTracksComponentTypes()
 	.SetBlenderSystem<UMovieScenePiecewiseDoubleBlenderSystem>()
 	.Commit(FFloatParameterHandler());
 
+	BuiltInComponents->PropertyRegistry.DefineProperty(Parameters.Scalar, TEXT("Scalar Parameters"))
+	.AddSoleChannel(BuiltInComponents->DoubleResult[0])
+	.SetBlenderSystem<UMovieScenePiecewiseDoubleBlenderSystem>()
+	.Commit();
+
 	// --------------------------------------------------------------------------------------------
 	// Set up color parameters
 	BuiltInComponents->PropertyRegistry.DefineCompositeProperty(ColorParameter, TEXT("Apply Color Parameters"))
@@ -748,6 +763,14 @@ FMovieSceneTracksComponentTypes::FMovieSceneTracksComponentTypes()
 	.SetBlenderSystem<UMovieScenePiecewiseDoubleBlenderSystem>()
 	.Commit(FColorParameterHandler());
 
+	BuiltInComponents->PropertyRegistry.DefineCompositeProperty(Parameters.Color, TEXT("Color Parameters"))
+	.AddComposite(BuiltInComponents->DoubleResult[0], &FIntermediateColor::R)
+	.AddComposite(BuiltInComponents->DoubleResult[1], &FIntermediateColor::G)
+	.AddComposite(BuiltInComponents->DoubleResult[2], &FIntermediateColor::B)
+	.AddComposite(BuiltInComponents->DoubleResult[3], &FIntermediateColor::A)
+	.SetBlenderSystem<UMovieScenePiecewiseDoubleBlenderSystem>()
+	.Commit();
+
 	Accessors.Byte.Add(
 		USkeletalMeshComponent::StaticClass(), USkeletalMeshComponent::GetAnimationModePropertyNameChecked(),
 		GetSkeletalMeshAnimationMode, SetSkeletalMeshAnimationMode);
@@ -766,6 +789,19 @@ FMovieSceneTracksComponentTypes::FMovieSceneTracksComponentTypes()
 	.SetCustomAccessors(&Accessors.FloatVector)
 	.Commit(FFloatVectorHandler());
 
+	BuiltInComponents->PropertyRegistry.DefineCompositeProperty(Parameters.Vector2, TEXT("Vector2 Parameters"))
+	.AddComposite(BuiltInComponents->DoubleResult[0], &FFloatIntermediateVector::X)
+	.AddComposite(BuiltInComponents->DoubleResult[1], &FFloatIntermediateVector::Y)
+	.SetBlenderSystem<UMovieScenePiecewiseDoubleBlenderSystem>()
+	.Commit();
+
+	BuiltInComponents->PropertyRegistry.DefineCompositeProperty(Parameters.Vector3, TEXT("Vector3 Parameters"))
+	.AddComposite(BuiltInComponents->DoubleResult[0], &FFloatIntermediateVector::X)
+	.AddComposite(BuiltInComponents->DoubleResult[1], &FFloatIntermediateVector::Y)
+	.AddComposite(BuiltInComponents->DoubleResult[2], &FFloatIntermediateVector::Z)
+	.SetBlenderSystem<UMovieScenePiecewiseDoubleBlenderSystem>()
+	.Commit();
+
 	BuiltInComponents->PropertyRegistry.DefineCompositeProperty(DoubleVector, TEXT("Apply DoubleVector Properties"))
 	.AddComposite(BuiltInComponents->DoubleResult[0], &FDoubleIntermediateVector::X)
 	.AddComposite(BuiltInComponents->DoubleResult[1], &FDoubleIntermediateVector::Y)
@@ -790,6 +826,19 @@ FMovieSceneTracksComponentTypes::FMovieSceneTracksComponentTypes()
 	.SetBlenderSystem<UMovieScenePiecewiseDoubleBlenderSystem>()
 	.Commit();
 
+	BuiltInComponents->PropertyRegistry.DefineCompositeProperty(Parameters.Transform, TEXT("Transform Parameters"))
+	.AddComposite(BuiltInComponents->DoubleResult[0], &FIntermediate3DTransform::T_X)
+	.AddComposite(BuiltInComponents->DoubleResult[1], &FIntermediate3DTransform::T_Y)
+	.AddComposite(BuiltInComponents->DoubleResult[2], &FIntermediate3DTransform::T_Z)
+	.AddComposite(BuiltInComponents->DoubleResult[3], &FIntermediate3DTransform::R_X)
+	.AddComposite(BuiltInComponents->DoubleResult[4], &FIntermediate3DTransform::R_Y)
+	.AddComposite(BuiltInComponents->DoubleResult[5], &FIntermediate3DTransform::R_Z)
+	.AddComposite(BuiltInComponents->DoubleResult[6], &FIntermediate3DTransform::S_X)
+	.AddComposite(BuiltInComponents->DoubleResult[7], &FIntermediate3DTransform::S_Y)
+	.AddComposite(BuiltInComponents->DoubleResult[8], &FIntermediate3DTransform::S_Z)
+	.SetBlenderSystem<UMovieScenePiecewiseDoubleBlenderSystem>()
+	.Commit();
+
 	// --------------------------------------------------------------------------------------------
 	// Set up component transforms
 	{
diff --git a/Engine/Source/Runtime/MovieSceneTracks/Private/Systems/MovieSceneInitialValueSystem.cpp b/Engine/Source/Runtime/MovieSceneTracks/Private/Systems/MovieSceneInitialValueSystem.cpp
deleted file mode 100644
diff --git a/Engine/Source/Runtime/MovieSceneTracks/Private/Systems/MovieScenePiecewiseDoubleBlenderSystem.cpp b/Engine/Source/Runtime/MovieSceneTracks/Private/Systems/MovieScenePiecewiseDoubleBlenderSystem.cpp
--- a/Engine/Source/Runtime/MovieSceneTracks/Private/Systems/MovieScenePiecewiseDoubleBlenderSystem.cpp
+++ b/Engine/Source/Runtime/MovieSceneTracks/Private/Systems/MovieScenePiecewiseDoubleBlenderSystem.cpp
@@ -651,6 +651,9 @@ void UMovieScenePiecewiseDoubleBlenderSystem::OnSchedulePersistentTasks(UE::Movi
 
 	FTaskID SyncTask = TaskScheduler->AddNullTask();
 
+	// Do not use weights when ExternalBlending is present
+	FComponentTypeIDFilter WeightFilter(BuiltInComponents->Tags.ExternalBlending, false);
+
 	if (AccumulationBuffers.Absolute.Num() != 0)
 	{
 		for (TPair<FComponentTypeID, TArray<FBlendResult>>& Pair : AccumulationBuffers.Absolute)
@@ -658,7 +661,7 @@ void UMovieScenePiecewiseDoubleBlenderSystem::OnSchedulePersistentTasks(UE::Movi
 			FTaskID AbsoluteTask = FEntityTaskBuilder()
 			.Read(Pair.Key.ReinterpretCast<double>())
 			.Read(BuiltInComponents->BlendChannelInput)
-			.ReadOptional(BuiltInComponents->WeightAndEasingResult)
+			.ReadOptional(BuiltInComponents->WeightAndEasingResult, WeightFilter)
 			.FilterAll({ BuiltInComponents->Tags.AbsoluteBlend, GetBlenderTypeTag() })
 			.FilterAny(BlendedResultMask)
 			.FilterNone({ BuiltInComponents->Tags.Ignored })
@@ -678,7 +681,7 @@ void UMovieScenePiecewiseDoubleBlenderSystem::OnSchedulePersistentTasks(UE::Movi
 			FTaskID RelativeTask = FEntityTaskBuilder()
 			.Read(Pair.Key.ReinterpretCast<double>())
 			.Read(BuiltInComponents->BlendChannelInput)
-			.ReadOptional(BuiltInComponents->WeightAndEasingResult)
+			.ReadOptional(BuiltInComponents->WeightAndEasingResult, WeightFilter)
 			.FilterAll({ BuiltInComponents->Tags.RelativeBlend, GetBlenderTypeTag() })
 			.FilterAny(BlendedResultMask)
 			.FilterNone({ BuiltInComponents->Tags.Ignored })
@@ -698,7 +701,7 @@ void UMovieScenePiecewiseDoubleBlenderSystem::OnSchedulePersistentTasks(UE::Movi
 			FTaskID AdditiveTask = FEntityTaskBuilder()
 			.Read(Pair.Key.ReinterpretCast<double>())
 			.Read(BuiltInComponents->BlendChannelInput)
-			.ReadOptional(BuiltInComponents->WeightAndEasingResult)
+			.ReadOptional(BuiltInComponents->WeightAndEasingResult, WeightFilter)
 			.FilterAll({ BuiltInComponents->Tags.AdditiveBlend, GetBlenderTypeTag() })
 			.FilterAny(BlendedResultMask)
 			.FilterNone({ BuiltInComponents->Tags.Ignored })
@@ -715,7 +718,7 @@ void UMovieScenePiecewiseDoubleBlenderSystem::OnSchedulePersistentTasks(UE::Movi
 	{
 		FTaskID AdditiveFromBaseTask = FEntityTaskBuilder()
 		.Read(BuiltInComponents->BlendChannelInput)
-		.ReadOptional(BuiltInComponents->WeightAndEasingResult)
+		.ReadOptional(BuiltInComponents->WeightAndEasingResult, WeightFilter)
 		.FilterAll({ BuiltInComponents->Tags.AdditiveFromBaseBlend, GetBlenderTypeTag() })
 		.FilterAny(BlendedResultMask)
 		.FilterNone({ BuiltInComponents->Tags.Ignored })
@@ -831,12 +834,15 @@ void UMovieScenePiecewiseDoubleBlenderSystem::OnRun(FSystemTaskPrerequisites& In
 
 	FBuiltInComponentTypes* BuiltInComponents = FBuiltInComponentTypes::Get();
 
+	// Do not use weights when ExternalBlending is present
+	FComponentTypeIDFilter WeightFilter(BuiltInComponents->Tags.ExternalBlending, false);
+
 	FSystemTaskPrerequisites Prereqs;
 	if (AccumulationBuffers.Absolute.Num() != 0)
 	{
 		FGraphEventRef Task = FEntityTaskBuilder()
 		.Read(BuiltInComponents->BlendChannelInput)
-		.ReadOptional(BuiltInComponents->WeightAndEasingResult)
+		.ReadOptional(BuiltInComponents->WeightAndEasingResult, WeightFilter)
 		.FilterAll({ BuiltInComponents->Tags.AbsoluteBlend, GetBlenderTypeTag() })
 		.FilterAny(BlendedResultMask)
 		.FilterNone({ BuiltInComponents->Tags.Ignored })
@@ -853,7 +859,7 @@ void UMovieScenePiecewiseDoubleBlenderSystem::OnRun(FSystemTaskPrerequisites& In
 	{
 		FGraphEventRef Task = FEntityTaskBuilder()
 		.Read(BuiltInComponents->BlendChannelInput)
-		.ReadOptional(BuiltInComponents->WeightAndEasingResult)
+		.ReadOptional(BuiltInComponents->WeightAndEasingResult, WeightFilter)
 		.FilterAll({ BuiltInComponents->Tags.RelativeBlend, GetBlenderTypeTag() })
 		.FilterAny(BlendedResultMask)
 		.FilterNone({ BuiltInComponents->Tags.Ignored })
@@ -870,7 +876,7 @@ void UMovieScenePiecewiseDoubleBlenderSystem::OnRun(FSystemTaskPrerequisites& In
 	{
 		FGraphEventRef Task = FEntityTaskBuilder()
 		.Read(BuiltInComponents->BlendChannelInput)
-		.ReadOptional(BuiltInComponents->WeightAndEasingResult)
+		.ReadOptional(BuiltInComponents->WeightAndEasingResult, WeightFilter)
 		.FilterAll({ BuiltInComponents->Tags.AdditiveBlend, GetBlenderTypeTag() })
 		.FilterAny(BlendedResultMask)
 		.FilterNone({ BuiltInComponents->Tags.Ignored })
@@ -910,7 +916,7 @@ void UMovieScenePiecewiseDoubleBlenderSystem::OnRun(FSystemTaskPrerequisites& In
 	{
 		FGraphEventRef Task = FEntityTaskBuilder()
 		.Read(BuiltInComponents->BlendChannelInput)
-		.ReadOptional(BuiltInComponents->WeightAndEasingResult)
+		.ReadOptional(BuiltInComponents->WeightAndEasingResult, WeightFilter)
 		.FilterAll({ BuiltInComponents->Tags.AdditiveFromBaseBlend, GetBlenderTypeTag() })
 		.FilterAny(BlendedResultMask)
 		.FilterNone({ BuiltInComponents->Tags.Ignored })
@@ -1193,6 +1199,9 @@ FGraphEventRef UMovieScenePiecewiseDoubleBlenderSystem::DispatchDecomposeTask(co
 		}
 	};
 
+	// Do not use weights when ExternalBlending is present
+	FComponentTypeIDFilter WeightFilter(BuiltInComponents->Tags.ExternalBlending, false);
+
 	if (Params.Query.bConvertFromSourceEntityIDs)
 	{
 		return FEntityTaskBuilder()
@@ -1201,7 +1210,7 @@ FGraphEventRef UMovieScenePiecewiseDoubleBlenderSystem::DispatchDecomposeTask(co
 			.Read(ResultComponentType)
 			.ReadOptional(BuiltInComponents->BlendingOrder)
 			.ReadOptional(BaseValueComponentType)
-			.ReadOptional(BuiltInComponents->WeightAndEasingResult)
+			.ReadOptional(BuiltInComponents->WeightAndEasingResult, WeightFilter)
 			.FilterAll({ Params.PropertyTag, GetBlenderTypeTag() })
 			.Dispatch_PerAllocation<FChannelResultTask>(&Linker->EntityManager, FSystemTaskPrerequisites(), nullptr, Params, Output);
 	}
@@ -1213,7 +1222,7 @@ FGraphEventRef UMovieScenePiecewiseDoubleBlenderSystem::DispatchDecomposeTask(co
 			.Read(ResultComponentType)
 			.ReadOptional(BuiltInComponents->BlendingOrder)
 			.ReadOptional(BaseValueComponentType)
-			.ReadOptional(BuiltInComponents->WeightAndEasingResult)
+			.ReadOptional(BuiltInComponents->WeightAndEasingResult, WeightFilter)
 			.FilterAll({ Params.PropertyTag, GetBlenderTypeTag() })
 			.Dispatch_PerAllocation<FChannelResultTask>(&Linker->EntityManager, FSystemTaskPrerequisites(), nullptr, Params, Output);
 	}
diff --git a/Engine/Source/Runtime/MovieSceneTracks/Private/Systems/MovieScenePropertyInstantiator.cpp b/Engine/Source/Runtime/MovieSceneTracks/Private/Systems/MovieScenePropertyInstantiator.cpp
--- a/Engine/Source/Runtime/MovieSceneTracks/Private/Systems/MovieScenePropertyInstantiator.cpp
+++ b/Engine/Source/Runtime/MovieSceneTracks/Private/Systems/MovieScenePropertyInstantiator.cpp
@@ -12,6 +12,7 @@
 #include "EntitySystem/MovieSceneEntitySystemLinker.h"
 #include "EntitySystem/MovieScenePropertyBinding.h"
 #include "EntitySystem/MovieScenePropertyRegistry.h"
+#include "EntitySystem/MovieSceneInitialValueSystem.h"
 #include "ProfilingDebugging/CountersTrace.h"
 #include "Systems/MovieScenePiecewiseDoubleBlenderSystem.h"
 
@@ -76,6 +77,8 @@ UMovieScenePropertyInstantiatorSystem::UMovieScenePropertyInstantiatorSystem(con
 
 		DefineComponentProducer(GetClass(), BuiltInComponents->BlendChannelInput);
 		DefineComponentProducer(GetClass(), BuiltInComponents->SymbolicTags.CreatesEntities);
+
+		DefineImplicitPrerequisite(GetClass(), UMovieSceneInitialValueSystem::StaticClass());
 	}
 }
 
diff --git a/Engine/Source/Runtime/MovieSceneTracks/Public/EntitySystem/Interrogation/MovieSceneInterrogationExtension.h b/Engine/Source/Runtime/MovieSceneTracks/Public/EntitySystem/Interrogation/MovieSceneInterrogationExtension.h
--- a/Engine/Source/Runtime/MovieSceneTracks/Public/EntitySystem/Interrogation/MovieSceneInterrogationExtension.h
+++ b/Engine/Source/Runtime/MovieSceneTracks/Public/EntitySystem/Interrogation/MovieSceneInterrogationExtension.h
@@ -2,83 +2,4 @@
 
 #pragma once
 
-#include "UObject/WeakObjectPtrTemplates.h"
-#include "UObject/NameTypes.h"
-#include "Containers/SortedMap.h"
-#include "EntitySystem/MovieSceneEntitySystemTypes.h"
-#include "EntitySystem/MovieSceneEntitySystemLinkerExtension.h"
-#include "EntitySystem/MovieScenePropertyBinding.h"
-
-class UObject;
-
-namespace UE
-{
-namespace MovieScene
-{
-
-struct FInterrogationChannel;
-struct FInterrogationChannelInfo;
-
-struct FInterrogationChannelInfo
-{
-	/** The object that relates to the channel */
-	TWeakObjectPtr<UObject> WeakObject;
-
-	/** (Optional) property path for this channel */
-	FMovieScenePropertyBinding PropertyBinding;
-
-	/** The channel's hierarchical 'parent' - generally only used for transforms */
-	FInterrogationChannel ParentChannel;
-};
-
-
-struct FInterrogationUpdateParams
-{
-
-	/** The channel's hierarchical 'parent' - generally only used for transforms */
-	FInterrogationChannel ParentChannel;
-};
-
-struct FSparseInterrogationChannelInfo
-{
-	FInterrogationChannelInfo& Get(FInterrogationChannel Channel) { return ChannelInfo.FindOrAdd(Channel); }
-	const FInterrogationChannelInfo& Get(FInterrogationChannel Channel) const { return ChannelInfo.FindChecked(Channel); }
-
-	FInterrogationChannelInfo* Find(FInterrogationChannel Channel) { return ChannelInfo.Find(Channel); }
-	const FInterrogationChannelInfo* Find(FInterrogationChannel Channel) const { return ChannelInfo.Find(Channel); }
-
-	UObject* FindObject(FInterrogationChannel Channel) const
-	{
-		return ChannelInfo.FindRef(Channel).WeakObject.Get();
-	}
-
-	FInterrogationChannel FindParent(FInterrogationChannel Channel) const
-	{
-		return ChannelInfo.FindRef(Channel).ParentChannel;
-	}
-
-	void Empty()
-	{
-		ChannelInfo.Empty();
-	}
-
-private:
-
-	TSortedMap<FInterrogationChannel, FInterrogationChannelInfo> ChannelInfo;
-};
-
-struct IInterrogationExtension
-{
-public:
-	MOVIESCENETRACKS_API static TEntitySystemLinkerExtensionID<IInterrogationExtension> GetExtensionID();
-
-	virtual const FSparseInterrogationChannelInfo& GetSparseChannelInfo() const = 0;
-
-protected:
-
-	virtual ~IInterrogationExtension() {}
-};
-
-
-} // namespace MovieScene
-} // namespace UE
+#include "EntitySystem/MovieSceneInterrogation.h"
diff --git a/Engine/Source/Runtime/MovieSceneTracks/Public/MovieSceneTracksComponentTypes.h b/Engine/Source/Runtime/MovieSceneTracks/Public/MovieSceneTracksComponentTypes.h
--- a/Engine/Source/Runtime/MovieSceneTracks/Public/MovieSceneTracksComponentTypes.h
+++ b/Engine/Source/Runtime/MovieSceneTracks/Public/MovieSceneTracksComponentTypes.h
@@ -657,8 +657,15 @@ using FEulerTransformPropertyTraits     = TIndirectPropertyTraits<FEulerTransfor
 using FComponentTransformPropertyTraits = TDirectPropertyTraits<FIntermediate3DTransform>;
 using FRotatorPropertyTraits            = TDirectPropertyTraits<FRotator>;
 using FStringPropertyTraits             = TDirectPropertyTraits<FString>;
+
+using FBoolParameterTraits              = TDirectPropertyTraits<bool, false>;
+using FByteParameterTraits              = TDirectPropertyTraits<uint8, false>;
+using FIntegerParameterTraits           = TDirectPropertyTraits<int32, false>;
 using FFloatParameterTraits             = TIndirectPropertyTraits<float, double, false>;
 using FColorParameterTraits             = TIndirectPropertyTraits<FLinearColor, FIntermediateColor>;
+using FVector2ParameterTraits           = TIndirectPropertyTraits<FVector2f, FFloatIntermediateVector>;
+using FVector3ParameterTraits           = TIndirectPropertyTraits<FVector3f, FFloatIntermediateVector>;
+using FTransformParameterTraits         = TIndirectPropertyTraits<FEulerTransform, FIntermediate3DTransform>;
 
 struct FMovieSceneTracksComponentTypes
 {
@@ -705,6 +712,20 @@ struct FMovieSceneTracksComponentTypes
 	TComponentTypeID<FName> ColorParameterName;
 	TComponentTypeID<FName> TransformParameterName;
 
+	TComponentTypeID<FName> GenericParameterName;
+
+	struct
+	{
+		TPropertyComponents<FBoolParameterTraits>      Bool;
+		TPropertyComponents<FByteParameterTraits>      Byte;
+		TPropertyComponents<FIntegerParameterTraits>   Integer;
+		TPropertyComponents<FFloatParameterTraits>     Scalar;
+		TPropertyComponents<FVector2ParameterTraits>   Vector2;
+		TPropertyComponents<FVector3ParameterTraits>   Vector3;
+		TPropertyComponents<FColorParameterTraits>     Color;
+		TPropertyComponents<FTransformParameterTraits> Transform;
+	} Parameters;
+
 	TComponentTypeID<FMaterialParameterInfo> ScalarMaterialParameterInfo;
 	TComponentTypeID<FMaterialParameterInfo> ColorMaterialParameterInfo;
 	TComponentTypeID<FMaterialParameterInfo> VectorMaterialParameterInfo;
diff --git a/Engine/Source/Runtime/MovieSceneTracks/Public/MovieSceneTracksPropertyTypes.h b/Engine/Source/Runtime/MovieSceneTracks/Public/MovieSceneTracksPropertyTypes.h
--- a/Engine/Source/Runtime/MovieSceneTracks/Public/MovieSceneTracksPropertyTypes.h
+++ b/Engine/Source/Runtime/MovieSceneTracks/Public/MovieSceneTracksPropertyTypes.h
@@ -46,6 +46,33 @@ struct FFloatIntermediateVector
 	FFloatIntermediateVector(double InX, double InY, double InZ, double InW)
 		: X(InX), Y(InY), Z(InZ), W(InW)
 	{}
+
+	FVector2f AsVector2f() const
+	{
+		return FVector2f(static_cast<float>(X), static_cast<float>(Y));
+	}
+	FVector2d AsVector2d() const
+	{
+		return FVector2d(X, Y);
+	}
+
+	FVector3f AsVector3f() const
+	{
+		return FVector3f(static_cast<float>(X), static_cast<float>(Y), static_cast<float>(Z));
+	}
+	FVector3d AsVector3d() const
+	{
+		return FVector3d(X, Y, Z);
+	}
+
+	FVector4f AsVector4f() const
+	{
+		return FVector4f(static_cast<float>(X), static_cast<float>(Y), static_cast<float>(Z), static_cast<float>(W));
+	}
+	FVector4d AsVector4d() const
+	{
+		return FVector4d(X, Y, Z, W);
+	}
 };
 struct FDoubleIntermediateVector
 {
@@ -177,6 +204,10 @@ struct FIntermediate3DTransform
 	{
 		return FVector(S_X, S_Y, S_Z);
 	}
+	FEulerTransform AsEuler() const
+	{
+		return FEulerTransform(GetTranslation(), GetRotation(), GetScale());
+	}
 
 	MOVIESCENETRACKS_API void ApplyTo(USceneComponent* SceneComponent) const;
 	
diff --git a/Engine/Source/Runtime/MovieSceneTracks/Public/Systems/MovieSceneInitialValueSystem.h b/Engine/Source/Runtime/MovieSceneTracks/Public/Systems/MovieSceneInitialValueSystem.h
--- a/Engine/Source/Runtime/MovieSceneTracks/Public/Systems/MovieSceneInitialValueSystem.h
+++ b/Engine/Source/Runtime/MovieSceneTracks/Public/Systems/MovieSceneInitialValueSystem.h
@@ -2,28 +2,4 @@
 
 #pragma once
 
-#include "EntitySystem/MovieSceneEntityInstantiatorSystem.h"
-#include "MovieSceneInitialValueSystem.generated.h"
-
-
-/**
- * System responsible for initializing initial values for all property types
- * Will handle the presence of an FInitialValueCache extension on the linker
- */
-UCLASS(MinimalAPI)
-class UMovieSceneInitialValueSystem
-	: public UMovieSceneEntityInstantiatorSystem
-{
-public:
-
-	GENERATED_BODY()
-
-	MOVIESCENETRACKS_API UMovieSceneInitialValueSystem(const FObjectInitializer& ObjInit);
-
-private:
-
-	MOVIESCENETRACKS_API virtual void OnRun(FSystemTaskPrerequisites& InPrerequisites, FSystemSubsequentTasks& Subsequents) override;
-	MOVIESCENETRACKS_API virtual bool IsRelevantImpl(UMovieSceneEntitySystemLinker* InLinker) const override;
-	MOVIESCENETRACKS_API virtual void OnLink() override;
-	MOVIESCENETRACKS_API virtual void OnUnlink() override;
-};
+#include "EntitySystem/MovieSceneInitialValueSystem.h"
\ No newline at end of file
