Control Rig: Evaluate the rig once only during interaction

#rnx

diff --git a/Engine/Plugins/Animation/ControlRig/Source/ControlRigEditor/Private/EditMode/ControlRigEditMode.cpp b/Engine/Plugins/Animation/ControlRig/Source/ControlRigEditor/Private/EditMode/ControlRigEditMode.cpp
--- a/Engine/Plugins/Animation/ControlRig/Source/ControlRigEditor/Private/EditMode/ControlRigEditMode.cpp
+++ b/Engine/Plugins/Animation/ControlRig/Source/ControlRigEditor/Private/EditMode/ControlRigEditMode.cpp
@@ -2151,6 +2151,7 @@ bool FControlRigEditMode::InputDelta(FEditorViewportClient* InViewportClient, FV
 						}
 					}
 
+					TMap<AControlRigShapeActor*, TArray<TFunction<void()>>> TasksPerActor;
 					for (AControlRigShapeActor* ShapeActor : Pairs.Value)
 					{
 						if (ShapeActor->IsEnabled() && ShapeActor->IsSelected())
@@ -2186,11 +2187,46 @@ bool FControlRigEditMode::InputDelta(FEditorViewportClient* InViewportClient, FV
 								}
 							}
 
+							TArray<TFunction<void()>> Tasks;
 							MoveControlShape(ShapeActor, bDoTranslation, InDrag, bDoRotation, InRot, bDoScale, InScale, ComponentTransform,
-								bUseLocal, bDoLocal, InOutLocal);
+								bUseLocal, bDoLocal, &InOutLocal, Tasks);
+							TasksPerActor.Add(ShapeActor, Tasks);
 							bManipulatorMadeChange = true;
 						}
 					}
+
+					{
+						// evaluate the rig(s) before the tasks
+						const FPendingControlRigEvaluator _(this);
+					}
+
+					// process remaining tasks
+					while(!TasksPerActor.IsEmpty())
+					{
+						// place another evaluator here which is going to run after the tasks for this phase
+						const FPendingControlRigEvaluator _(this);
+
+						// run one tasks for each control - until there are no tasks lefts
+						TArray<AControlRigShapeActor*> KeysToRemove;
+						for(TPair<AControlRigShapeActor*, TArray<TFunction<void()>>>& Pair : TasksPerActor)
+						{
+							if(!Pair.Value.IsEmpty())
+							{
+								TFunction<void()> Task = Pair.Value[0];
+								Pair.Value.RemoveAt(0);
+								Task();
+							}
+							else
+							{
+								KeysToRemove.Add(Pair.Key);
+							}
+						}
+
+						for(AControlRigShapeActor* KeyToRemove : KeysToRemove)
+						{
+							TasksPerActor.Remove(KeyToRemove);
+						}
+					}
 				}
 			}
 			else if (AreRigElementSelectedAndMovable(Pairs.Key))
@@ -4893,9 +4929,8 @@ FTransform FControlRigEditMode::GetControlShapeTransform(const AControlRigShapeA
 	return FTransform::Identity;
 }
 
-void FControlRigEditMode::MoveControlShape(AControlRigShapeActor* ShapeActor, const bool bTranslation, FVector& InDrag,
-	const bool bRotation, FRotator& InRot, const bool bScale, FVector& InScale, const FTransform& ToWorldTransform,
-	bool bUseLocal, bool bCalcLocal, FTransform& InOutLocal)
+void FControlRigEditMode::MoveControlShape(AControlRigShapeActor* ShapeActor, const bool bTranslation, FVector& InDrag, const bool bRotation,
+	FRotator& InRot, const bool bScale, FVector& InScale, const FTransform& ToWorldTransform, bool bUseLocal, bool bCalcLocal, FTransform* InOutLocal, TArray<TFunction<void()>>& OutTasks)
 {
 	UControlRig* ControlRig = ShapeActor ? ShapeActor->ControlRig.Get() : nullptr;
 	if (!ensure(ControlRig))
@@ -4925,7 +4960,10 @@ void FControlRigEditMode::MoveControlShape(AControlRigShapeActor* ShapeActor, co
 		if (bRotation)
 		{
 			FQuat CurrentRotation = CurrentLocalTransform.GetRotation();
-			CurrentRotation = (CurrentRotation * InOutLocal.GetRotation());
+			if(InOutLocal)
+			{
+				CurrentRotation = (CurrentRotation * InOutLocal->GetRotation());
+			}
 			CurrentLocalTransform.SetRotation(CurrentRotation);
 			bTransformChanged = true;
 		}
@@ -4933,7 +4971,10 @@ void FControlRigEditMode::MoveControlShape(AControlRigShapeActor* ShapeActor, co
 		if (bTranslation)
 		{
 			FVector CurrentLocation = CurrentLocalTransform.GetLocation();
-			CurrentLocation = CurrentLocation + InOutLocal.GetLocation();
+			if(InOutLocal)
+			{
+				CurrentLocation = CurrentLocation + InOutLocal->GetLocation();
+			}
 			CurrentLocalTransform.SetLocation(CurrentLocation);
 			bTransformChanged = true;
 		}
@@ -4957,13 +4998,11 @@ void FControlRigEditMode::MoveControlShape(AControlRigShapeActor* ShapeActor, co
 			
 			ShapeActor->SetGlobalTransform(CurrentTransform);
 
-			ControlRig->Evaluate_AnyThread();
-
+			RigsToEvaluateDuringThisTick.AddUnique(ControlRig);
 			return;
 		}
 	}
 
-
 	//not local or doing scale.
 	FTransform CurrentTransform = GetControlShapeTransform(ShapeActor) * ToWorldTransform;
 
@@ -5002,9 +5041,9 @@ void FControlRigEditMode::MoveControlShape(AControlRigShapeActor* ShapeActor, co
 		Context.EventName = FRigUnit_BeginExecution::EventName;
 		Context.bConstraintUpdate = true;
 
-		if (bCalcLocal)
+		if (bCalcLocal && InOutLocal)
 		{
-			InOutLocal = ControlRig->GetControlLocalTransform(ShapeActor->ControlName);
+			*InOutLocal = ControlRig->GetControlLocalTransform(ShapeActor->ControlName);
 		}
 
 		const UWorld* World = ControlRig->GetWorld();
@@ -5020,24 +5059,63 @@ void FControlRigEditMode::MoveControlShape(AControlRigShapeActor* ShapeActor, co
 		{
 			ControlRig->GetHierarchy()->SetPose(ControlRig->InputPoseOnDebuggedRig);
 		}
-		
-		ControlRig->Evaluate_AnyThread();
-		//fix flips and do rotation orders only if not additive
-		const bool bFixEulerFlips = !ControlRig->IsAdditive();
-		SetControlShapeTransform(ShapeActor, NewTransform, ToWorldTransform, Context, bPrintPythonCommands, bFixEulerFlips);
-		NotifyDrivenControls(ControlRig, ShapeActor->GetElementKey(),Context);
 
-		if (ControlElement && !bIsTransientControl)
+		// for readability
+		const bool bSolveImmediately = (bUseLocal == false) && (bCalcLocal == true);
+		const bool bQueueTasks = !bSolveImmediately;
+
+		if(bQueueTasks)
+		{
+			RigsToEvaluateDuringThisTick.AddUnique(ControlRig);
+		}
+		else
 		{
 			ControlRig->Evaluate_AnyThread();
 		}
+
+		const FTransform ToWorldTransformCopy = ToWorldTransform;
 		
-		ShapeActor->SetGlobalTransform(CurrentTransform);
+		const TFunction<void()> SetControlShapeTask = [this, ShapeActor, ControlRig, ControlElement, NewTransform, ToWorldTransformCopy, Context,
+			bPrintPythonCommands, bIsTransientControl, bQueueTasks]()
+		{
+			//fix flips and do rotation orders only if not additive
+			const bool bFixEulerFlips = !ControlRig->IsAdditive();
+			SetControlShapeTransform(ShapeActor, NewTransform, ToWorldTransformCopy, Context, bPrintPythonCommands, bFixEulerFlips);
+			NotifyDrivenControls(ControlRig, ShapeActor->GetElementKey(),Context);
 
-		if (bCalcLocal)
+			if (ControlElement && !bIsTransientControl)
+			{
+				if(bQueueTasks)
+				{
+					RigsToEvaluateDuringThisTick.AddUnique(ControlRig);
+				}
+				else
+				{
+					ControlRig->Evaluate_AnyThread();
+				}
+			}
+		};
+
+		const TFunction<void()> SetGlobalAndUpdateLocalTask = [ShapeActor, CurrentTransform, ControlRig, bCalcLocal, InOutLocal]() {
+			
+			ShapeActor->SetGlobalTransform(CurrentTransform);
+
+			if (bCalcLocal && InOutLocal)
+			{
+				const FTransform NewLocal = ControlRig->GetControlLocalTransform(ShapeActor->ControlName);
+				*InOutLocal = NewLocal.GetRelativeTransform(*InOutLocal);
+			}
+		};
+
+		if(bQueueTasks)
 		{
-			FTransform NewLocal = ControlRig->GetControlLocalTransform(ShapeActor->ControlName);
-			InOutLocal = NewLocal.GetRelativeTransform(InOutLocal);
+			OutTasks.Add(SetControlShapeTask);
+			OutTasks.Add(SetGlobalAndUpdateLocalTask);
+		}
+		else
+		{
+			SetControlShapeTask();
+			SetGlobalAndUpdateLocalTask();
 		}
 	}
 }
@@ -5910,4 +5988,16 @@ void FDetailKeyFrameCacheAndHandler::UpdateIfDirty()
 	}
 }
 
+FControlRigEditMode::FPendingControlRigEvaluator::~FPendingControlRigEvaluator()
+{
+	if(--EditMode->RigEvaluationBracket == 0)
+	{
+		for(UControlRig* ControlRig : EditMode->RigsToEvaluateDuringThisTick)
+		{
+			ControlRig->Evaluate_AnyThread();
+		}
+		EditMode->RigsToEvaluateDuringThisTick.Reset();
+	}
+}
+
 #undef LOCTEXT_NAMESPACE
diff --git a/Engine/Plugins/Animation/ControlRig/Source/ControlRigEditor/Public/EditMode/ControlRigEditMode.h b/Engine/Plugins/Animation/ControlRig/Source/ControlRigEditor/Public/EditMode/ControlRigEditMode.h
--- a/Engine/Plugins/Animation/ControlRig/Source/ControlRigEditor/Public/EditMode/ControlRigEditMode.h
+++ b/Engine/Plugins/Animation/ControlRig/Source/ControlRigEditor/Public/EditMode/ControlRigEditMode.h
@@ -3,6 +3,7 @@
 #pragma once
 
 #include "CoreMinimal.h"
+#include "ControlRig.h"
 #include "InputCoreTypes.h"
 #include "IControlRigObjectBinding.h"
 #include "RigVMModel/RigVMGraph.h"
@@ -511,9 +512,6 @@ private:
 		const FTransform& InToWorldTransform, const FRigControlModifiedContext& InContext, const bool bPrintPython,
 		const bool bFixEulerFlips = true) const;
 	static FTransform GetControlShapeTransform(const AControlRigShapeActor* ShapeActor);
-	void MoveControlShape(AControlRigShapeActor* ShapeActor, const bool bTranslation, FVector& InDrag, 
-		const bool bRotation, FRotator& InRot, const bool bScale, FVector& InScale, const FTransform& ToWorldTransform,
-		bool bUseLocal, bool bCalcLocal, FTransform& InOutLocal);
 
 	void ChangeControlShapeTransform(AControlRigShapeActor* ShapeActor, const bool bTranslation, FVector& InDrag,
 		const bool bRotation, FRotator& InRot, const bool bScale, FVector& InScale, const FTransform& ToWorldTransform);
@@ -525,7 +523,11 @@ public:
 	static void NotifyDrivenControls(UControlRig* InControlRig, const FRigElementKey& InKey, const FRigControlModifiedContext& InContext);
 
 protected:
-	
+
+	void MoveControlShape(AControlRigShapeActor* ShapeActor, const bool bTranslation, FVector& InDrag, 
+		const bool bRotation, FRotator& InRot, const bool bScale, FVector& InScale, const FTransform& ToWorldTransform,
+		bool bUseLocal, bool bCalcLocal, FTransform* InOutLocal, TArray<TFunction<void()>>& OutTasks);
+
 	/** Get bindings to a runtime object */
 	//If the passed in ControlRig is nullptr we use the first Control Rig(this can happen from the BP Editors).
 	USceneComponent* GetHostingSceneComponent(const UControlRig* ControlRig = nullptr) const;
@@ -600,6 +602,26 @@ private:
 	TArray<uint32> LastHierarchyHash;
 	TArray<uint32> LastShapeLibraryHash;
 
+	/** A list of rigs we need to run during this tick */
+	TArray<UControlRig*> RigsToEvaluateDuringThisTick;
+	uint32 RigEvaluationBracket = 0;
+
+	class FPendingControlRigEvaluator
+	{
+		public:
+
+		FPendingControlRigEvaluator(FControlRigEditMode* InEditMode)
+		: EditMode(InEditMode)
+		{
+			EditMode->RigEvaluationBracket++;
+		}
+
+		~FPendingControlRigEvaluator();
+		
+	private:
+	
+		FControlRigEditMode* EditMode;
+	};
 
 	friend class FControlRigEditorModule;
 	friend class FControlRigEditor;
