Control Rig: FIx crash when modifying rig element in details panel when construction event is set
#rnx

diff --git a/Engine/Plugins/Animation/ControlRig/Source/ControlRigEditor/Private/ControlRigElementDetails.cpp b/Engine/Plugins/Animation/ControlRig/Source/ControlRigEditor/Private/ControlRigElementDetails.cpp
--- a/Engine/Plugins/Animation/ControlRig/Source/ControlRigEditor/Private/ControlRigElementDetails.cpp
+++ b/Engine/Plugins/Animation/ControlRig/Source/ControlRigEditor/Private/ControlRigElementDetails.cpp
@@ -2598,17 +2598,24 @@ FDetailWidgetRow& FRigTransformElementDetails::CreateEulerTransformValueWidgetRo
 		return bIsRelative ? RelativeTransform : AbsoluteTransform;
 	};
 
-	auto SetSingleTransform = [CurrentTransformType, GetRelativeAbsoluteTransforms, this, Hierarchy](
+	const TWeakPtr<FRigTransformElementDetails> WeakThisPtr = StaticCastWeakPtr<FRigTransformElementDetails>(AsWeak());
+	auto SetSingleTransform = [CurrentTransformType, GetRelativeAbsoluteTransforms, WeakThisPtr, Hierarchy](
 		const FRigElementKey& Key,
 		FEulerTransform InTransform,
 		bool bIsRelative,
 		bool bSetupUndoRedo)
 	{
+		if (!WeakThisPtr.IsValid())
+		{
+			return;
+		}
+		FRigTransformElementDetails* StrongThisPtr = WeakThisPtr.Pin().Get();
+		
 		const bool bCurrent = CurrentTransformType == ERigTransformElementDetailsTransform::Current; 
 		const bool bInitial = CurrentTransformType == ERigTransformElementDetailsTransform::Initial; 
 
 		bool bConstructionModeEnabled = false;
-		if (UControlRig* DebuggedRig = Cast<UControlRig>(PerElementInfos[0].GetBlueprint()->GetObjectBeingDebugged()))
+		if (UControlRig* DebuggedRig = Cast<UControlRig>(StrongThisPtr->PerElementInfos[0].GetBlueprint()->GetObjectBeingDebugged()))
 		{
 			bConstructionModeEnabled = DebuggedRig->IsConstructionModeEnabled();
 		}
@@ -2617,7 +2624,7 @@ FDetailWidgetRow& FRigTransformElementDetails::CreateEulerTransformValueWidgetRo
 		HierarchiesToUpdate.Add(Hierarchy);
 		if(!bCurrent || bConstructionModeEnabled)
 		{
-			HierarchiesToUpdate.Add(PerElementInfos[0].GetDefaultHierarchy());
+			HierarchiesToUpdate.Add(StrongThisPtr->PerElementInfos[0].GetDefaultHierarchy());
 		}
 
 		for(URigHierarchy* HierarchyToUpdate : HierarchiesToUpdate)
@@ -2772,7 +2779,7 @@ FDetailWidgetRow& FRigTransformElementDetails::CreateEulerTransformValueWidgetRo
 	TransformWidgetArgs.OnNumericValueChanged_Lambda(
 	[
 		Keys,
-		this,
+		WeakThisPtr,
 		IsComponentRelative,
 		GetSingleTransform,
 		SetSingleTransform,
@@ -2783,6 +2790,12 @@ FDetailWidgetRow& FRigTransformElementDetails::CreateEulerTransformValueWidgetRo
 		ESlateTransformSubComponent::Type SubComponent,
 		FVector::FReal InNumericValue)
 	{
+		if (!WeakThisPtr.IsValid())
+		{
+			return;
+		}
+		FRigTransformElementDetails* StrongThisPtr = WeakThisPtr.Pin().Get();
+			
 		const bool bIsRelative = IsComponentRelative((int32)Component);
 
 		for(const FRigElementKey& Key : Keys)
@@ -2793,9 +2806,9 @@ FDetailWidgetRow& FRigTransformElementDetails::CreateEulerTransformValueWidgetRo
 
 			if(!FRigControlElementDetails::Equals(Transform, PreviousTransform))
 			{
-				if(!SliderTransaction.IsValid())
+				if(!StrongThisPtr->SliderTransaction.IsValid())
 				{
-					SliderTransaction = MakeShareable(new FScopedTransaction(NSLOCTEXT("ControlRigElementDetails", "ChangeNumericValue", "Change Numeric Value")));
+					StrongThisPtr->SliderTransaction = MakeShareable(new FScopedTransaction(NSLOCTEXT("ControlRigElementDetails", "ChangeNumericValue", "Change Numeric Value")));
 					HierarchyToChange->Modify();
 				}
 							
@@ -2807,7 +2820,7 @@ FDetailWidgetRow& FRigTransformElementDetails::CreateEulerTransformValueWidgetRo
 	TransformWidgetArgs.OnNumericValueCommitted_Lambda(
 	[
 		Keys,
-		this,
+		WeakThisPtr,
 		IsComponentRelative,
 		GetSingleTransform,
 		SetSingleTransform,
@@ -2819,11 +2832,17 @@ FDetailWidgetRow& FRigTransformElementDetails::CreateEulerTransformValueWidgetRo
 		FVector::FReal InNumericValue,
 		ETextCommit::Type InCommitType)
 	{
+		if (!WeakThisPtr.IsValid())
+		{
+			return;
+		}
+		FRigTransformElementDetails* StrongThisPtr = WeakThisPtr.Pin().Get();
+		
 		const bool bIsRelative = IsComponentRelative((int32)Component);
 
 		{
 			FScopedTransaction Transaction(LOCTEXT("ChangeNumericValue", "Change Numeric Value"));
-			if(!SliderTransaction.IsValid())
+			if(!StrongThisPtr->SliderTransaction.IsValid())
 			{
 				HierarchyToChange->Modify();
 			}
@@ -2836,18 +2855,24 @@ FDetailWidgetRow& FRigTransformElementDetails::CreateEulerTransformValueWidgetRo
 			}
 		}
 
-		SliderTransaction.Reset();
+		StrongThisPtr->SliderTransaction.Reset();
 	});
 
 	TransformWidgetArgs.OnBeginSliderMovement_Lambda(
 		[
-			this
+			WeakThisPtr
 		](
 			ESlateTransformComponent::Type Component,
 			ESlateRotationRepresentation::Type Representation,
 			ESlateTransformSubComponent::Type SubComponent)
 		{
-			if (UControlRig* DebuggedRig = Cast<UControlRig>(PerElementInfos[0].GetBlueprint()->GetObjectBeingDebugged()))
+			if (!WeakThisPtr.IsValid())
+			{
+				return;
+			}
+			FRigTransformElementDetails* StrongThisPtr = WeakThisPtr.Pin().Get();
+			
+			if (UControlRig* DebuggedRig = Cast<UControlRig>(StrongThisPtr->PerElementInfos[0].GetBlueprint()->GetObjectBeingDebugged()))
 			{
 				EControlRigInteractionType Type = EControlRigInteractionType::None;
 				switch (Component)
@@ -2859,37 +2884,43 @@ FDetailWidgetRow& FRigTransformElementDetails::CreateEulerTransformValueWidgetRo
 				}
 				DebuggedRig->InteractionType = (uint8)Type;
 				DebuggedRig->ElementsBeingInteracted.Reset();
-				for (FPerElementInfo& ElementInfo : PerElementInfos)
+				for (const FPerElementInfo& ElementInfo : StrongThisPtr->PerElementInfos)
 				{
 					DebuggedRig->ElementsBeingInteracted.AddUnique(ElementInfo.Element.GetKey());
 				}
 				
 				FControlRigInteractionScope* InteractionScope = new FControlRigInteractionScope(DebuggedRig);
-				InteractionScopes.Add(InteractionScope);
+				StrongThisPtr->InteractionScopes.Add(InteractionScope);
 			}
 		});
 	TransformWidgetArgs.OnEndSliderMovement_Lambda(
 		[
-			this
+			WeakThisPtr
 		](
 			ESlateTransformComponent::Type Component,
 			ESlateRotationRepresentation::Type Representation,
 			ESlateTransformSubComponent::Type SubComponent,
 			FVector::FReal InNumericValue)
 		{
-			if (UControlRig* DebuggedRig = Cast<UControlRig>(PerElementInfos[0].GetBlueprint()->GetObjectBeingDebugged()))
+			if (!WeakThisPtr.IsValid())
+			{
+				return;
+			}
+			FRigTransformElementDetails* StrongThisPtr = WeakThisPtr.Pin().Get();
+			
+			if (UControlRig* DebuggedRig = Cast<UControlRig>(StrongThisPtr->PerElementInfos[0].GetBlueprint()->GetObjectBeingDebugged()))
 			{
 				DebuggedRig->InteractionType = (uint8)EControlRigInteractionType::None;
 				DebuggedRig->ElementsBeingInteracted.Reset();
 			}
-			for (FControlRigInteractionScope* InteractionScope : InteractionScopes)
+			for (FControlRigInteractionScope* InteractionScope : StrongThisPtr->InteractionScopes)
 			{
 				if (InteractionScope)
 				{
 					delete InteractionScope; 
 				}
 			}
-			InteractionScopes.Reset();
+			StrongThisPtr->InteractionScopes.Reset();
 		});
 
 	TransformWidgetArgs.OnCopyToClipboard_Lambda([Keys, IsComponentRelative, ConformComponentRelative, GetSingleTransform](
