Control Rig Edit Mode: only fix euler flips when rotating

#rnx

diff --git a/Engine/Plugins/Animation/ControlRig/Source/ControlRigEditor/Private/EditMode/ControlRigEditMode.cpp b/Engine/Plugins/Animation/ControlRig/Source/ControlRigEditor/Private/EditMode/ControlRigEditMode.cpp
--- a/Engine/Plugins/Animation/ControlRig/Source/ControlRigEditor/Private/EditMode/ControlRigEditMode.cpp
+++ b/Engine/Plugins/Animation/ControlRig/Source/ControlRigEditor/Private/EditMode/ControlRigEditMode.cpp
@@ -5076,10 +5076,10 @@ void FControlRigEditMode::MoveControlShape(AControlRigShapeActor* ShapeActor, co
 		const FTransform ToWorldTransformCopy = ToWorldTransform;
 		
 		const TFunction<void()> SetControlShapeTask = [this, ShapeActor, ControlRig, ControlElement, NewTransform, ToWorldTransformCopy, Context,
-			bPrintPythonCommands, bIsTransientControl, bQueueTasks]()
+			bPrintPythonCommands, bIsTransientControl, bRotation, bQueueTasks]()
 		{
 			//fix flips and do rotation orders only if not additive
-			const bool bFixEulerFlips = !ControlRig->IsAdditive();
+			const bool bFixEulerFlips = !ControlRig->IsAdditive() && bRotation;
 			SetControlShapeTransform(ShapeActor, NewTransform, ToWorldTransformCopy, Context, bPrintPythonCommands, bFixEulerFlips);
 			NotifyDrivenControls(ControlRig, ShapeActor->GetElementKey(),Context);
 
