Control RIg: Allow test data for backwards solve
#rnx

diff --git a/Engine/Plugins/Animation/ControlRig/Source/ControlRig/Private/ControlRig.cpp b/Engine/Plugins/Animation/ControlRig/Source/ControlRig/Private/ControlRig.cpp
--- a/Engine/Plugins/Animation/ControlRig/Source/ControlRig/Private/ControlRig.cpp
+++ b/Engine/Plugins/Animation/ControlRig/Source/ControlRig/Private/ControlRig.cpp
@@ -1140,6 +1140,12 @@ bool UControlRig::Execute(const FName& InEventName)
 			InputAnimAttributeSnapshot.CopyFrom(*ExternalAnimAttributeContainer);
 		}
 #endif
+
+		if (PreExecutedEvent.IsBound())
+		{
+			FControlRigBracketScope BracketScope(ExecuteBracket);
+			PreExecutedEvent.Broadcast(this, InEventName);
+		}
 		
 		if (bIsPreForwardSolve || (bIsForwardSolve && !bPreForwardSolveInQueue))
 		{
diff --git a/Engine/Plugins/Animation/ControlRig/Source/ControlRig/Private/ControlRigTestData.cpp b/Engine/Plugins/Animation/ControlRig/Source/ControlRig/Private/ControlRigTestData.cpp
--- a/Engine/Plugins/Animation/ControlRig/Source/ControlRig/Private/ControlRigTestData.cpp
+++ b/Engine/Plugins/Animation/ControlRig/Source/ControlRig/Private/ControlRigTestData.cpp
@@ -3,6 +3,8 @@
 #include "ControlRigTestData.h"
 #include "ControlRig.h"
 #include "ControlRigObjectVersion.h"
+#include "Units/Execution/RigUnit_PrepareForExecution.h"
+#include "Units/Execution/RigUnit_BeginExecution.h"
 #include "HAL/PlatformTime.h"
 #if WITH_EDITOR
 #include "AssetToolsModule.h"
@@ -173,6 +175,11 @@ void UControlRigTestData::Serialize(FArchive& Ar)
 			}
 		}
 	}
+
+	if (EventQueue.IsEmpty())
+	{
+		EventQueue = {FRigUnit_BeginExecution::EventName};
+	}
 }
 
 UControlRigTestData* UControlRigTestData::CreateNewAsset(FString InDesiredPackagePath, FString InBlueprintPathName)
@@ -249,6 +256,7 @@ bool UControlRigTestData::Record(UControlRig* InControlRig, double InRecordingDu
 	ReleaseReplay();
 	ClearDelegates(InControlRig);
 
+	EventQueue = InControlRig->EventQueue;
 	DesiredRecordingDuration = InRecordingDuration;
 	TimeAtStartOfRecording = FPlatformTime::Seconds();
 
@@ -264,25 +272,58 @@ bool UControlRigTestData::Record(UControlRig* InControlRig, double InRecordingDu
 		);
 	}
 
-	PreForwardHandle = InControlRig->OnPreForwardsSolve_AnyThread().AddLambda(
-		[this](UControlRig* InControlRig, const FName& InEventName)
+	PreEventHandle = InControlRig->OnPreExecuted_AnyThread().AddLambda(
+		[this](URigVMHost* InRigVMHost, const FName& InEventName)
 		{
+			UControlRig* ControlRig = Cast<UControlRig>(InRigVMHost);
+			if (!ControlRig)
+			{
+				return;
+			}
+
+			if (InEventName == FRigUnit_PrepareForExecution::EventName)
+			{
+				return;
+			}
+
+			
+			if (InEventName != EventQueue[0])
+			{
+				return;
+			}
+			
 			FControlRigTestDataFrame Frame;
-			Frame.Store(InControlRig);
+			Frame.Store(ControlRig);
 
 			// reapply the variable data. we are doing this to make sure that
 			// the results in the rig are the same during a recording and replay.
-			Frame.RestoreVariables(InControlRig);
+			Frame.RestoreVariables(ControlRig);
 			
 			InputFrames.Add(Frame);
 		}
 	);
 
-	PostForwardHandle = InControlRig->OnPostForwardsSolve_AnyThread().AddLambda(
-		[this](UControlRig* InControlRig, const FName& InEventName)
+	PostEventHandle = InControlRig->OnExecuted_AnyThread().AddLambda(
+		[this](URigVMHost* InRigVMHost, const FName& InEventName)
 		{
+			UControlRig* ControlRig = Cast<UControlRig>(InRigVMHost);
+			if (!ControlRig)
+			{
+				return;
+			}
+
+			if (InEventName == FRigUnit_PrepareForExecution::EventName)
+			{
+				return;
+			}
+
+			if (InEventName != EventQueue.Last())
+			{
+				return;
+			}
+			
 			FControlRigTestDataFrame Frame;
-			Frame.Store(InControlRig);
+			Frame.Store(ControlRig);
 			OutputFrames.Add(Frame);
 			LastFrameIndex = INDEX_NONE;
 			(void)MarkPackageDirty();
@@ -294,7 +335,7 @@ bool UControlRigTestData::Record(UControlRig* InControlRig, double InRecordingDu
 				DesiredRecordingDuration = 0.0;
 
 				// Once clear delegates is called, we no longer have access to this pointer
-				ClearDelegates(InControlRig);
+				ClearDelegates(ControlRig);
 			}
 		}
 	);
@@ -311,7 +352,8 @@ bool UControlRigTestData::SetupReplay(UControlRig* InControlRig, bool bGroundTru
 	{
 		return false;
 	}
-	
+
+	InControlRig->EventQueue = EventQueue;
 	bIsApplyingOutputs = bGroundTruth;
 
 	if(InputFrames.IsEmpty() || OutputFrames.IsEmpty())
@@ -330,42 +372,64 @@ bool UControlRigTestData::SetupReplay(UControlRig* InControlRig, bool bGroundTru
 		}
 	);
 
-	PreForwardHandle = InControlRig->OnPreForwardsSolve_AnyThread().AddLambda(
-		[this](UControlRig* InControlRig, const FName& InEventName)
+	PreEventHandle = InControlRig->OnPreExecuted_AnyThread().AddLambda(
+		[this](URigVMHost* InRigVMHost, const FName& InEventName)
 		{
+			UControlRig* ControlRig = Cast<UControlRig>(InRigVMHost);
+			if (!ControlRig || InEventName == FRigUnit_PrepareForExecution::EventName)
+			{
+				return;
+			}
+
+			if (InEventName != EventQueue[0])
+			{
+				return;
+			}
+			
 			// loop the animation data
-			if(InControlRig->GetAbsoluteTime() < GetTimeRange().X - SMALL_NUMBER ||
-				InControlRig->GetAbsoluteTime() > GetTimeRange().Y + SMALL_NUMBER)
+			if(ControlRig->GetAbsoluteTime() < GetTimeRange().X - SMALL_NUMBER ||
+				ControlRig->GetAbsoluteTime() > GetTimeRange().Y + SMALL_NUMBER)
 			{
-				InControlRig->SetAbsoluteAndDeltaTime(GetTimeRange().X, InControlRig->GetDeltaTime());
+				ControlRig->SetAbsoluteAndDeltaTime(GetTimeRange().X, ControlRig->GetDeltaTime());
 			}
 			
-			const int32 FrameIndex = GetFrameIndexForTime(InControlRig->GetAbsoluteTime(), true);
+			const int32 FrameIndex = GetFrameIndexForTime(ControlRig->GetAbsoluteTime(), true);
 			const FControlRigTestDataFrame& Frame = InputFrames[FrameIndex];
-			Frame.Restore(InControlRig, false);
+			Frame.Restore(ControlRig, false);
 
 			if(Frame.DeltaTime > SMALL_NUMBER)
 			{
-				InControlRig->SetDeltaTime(Frame.DeltaTime);
+				ControlRig->SetDeltaTime(Frame.DeltaTime);
 			}
 		}
 	);
 
-	PostForwardHandle = InControlRig->OnPostForwardsSolve_AnyThread().AddLambda(
-		[this](UControlRig* InControlRig, const FName& InEventName)
+	PostEventHandle = InControlRig->OnExecuted_AnyThread().AddLambda(
+		[this](URigVMHost* InRigVMHost, const FName& InEventName)
 		{
-			const FRigPose CurrentPose = InControlRig->GetHierarchy()->GetPose();
+			UControlRig* ControlRig = Cast<UControlRig>(InRigVMHost);
+			if (!ControlRig || InEventName == FRigUnit_PrepareForExecution::EventName)
+			{
+				return;
+			}
+
+			if (InEventName != EventQueue.Last())
+			{
+				return;
+			}
+			
+			const FRigPose CurrentPose = ControlRig->GetHierarchy()->GetPose();
 
 			const FControlRigTestDataFrame& Frame = OutputFrames[LastFrameIndex];
 			if(bIsApplyingOutputs)
 			{
-				Frame.Restore(InControlRig, false);
+				Frame.Restore(ControlRig, false);
 			}
 
 			const FRigPose& ExpectedPose = Frame.Pose;
 
 			// draw differences of the pose result of the rig onto the screen
-			FRigVMDrawInterface& DrawInterface = InControlRig->GetDrawInterface();
+			FRigVMDrawInterface& DrawInterface = ControlRig->GetDrawInterface();
 			for(const FRigPoseElement& ExpectedPoseElement : ExpectedPose)
 			{
 				const int32 CurrentPoseIndex = CurrentPose.GetIndex(ExpectedPoseElement.Index.GetKey());
@@ -428,15 +492,15 @@ void UControlRigTestData::ClearDelegates(UControlRig* InControlRig)
 			InControlRig->OnPreConstruction_AnyThread().Remove(PreConstructionHandle);
 			PreConstructionHandle.Reset();
 		}
-		if(PreForwardHandle.IsValid())
+		if(PreEventHandle.IsValid())
 		{
-			InControlRig->OnPreForwardsSolve_AnyThread().Remove(PreForwardHandle);
-			PreForwardHandle.Reset();
+			InControlRig->OnPreExecuted_AnyThread().Remove(PreEventHandle);
+			PreEventHandle.Reset();
 		}
-		if(PostForwardHandle.IsValid())
+		if(PostEventHandle.IsValid())
 		{
-			InControlRig->OnPostForwardsSolve_AnyThread().Remove(PostForwardHandle);
-			PostForwardHandle.Reset();
+			InControlRig->OnExecuted_AnyThread().Remove(PostEventHandle);
+			PostEventHandle.Reset();
 		}
 	}
 }
diff --git a/Engine/Plugins/Animation/ControlRig/Source/ControlRig/Private/ModularRig.cpp b/Engine/Plugins/Animation/ControlRig/Source/ControlRig/Private/ModularRig.cpp
--- a/Engine/Plugins/Animation/ControlRig/Source/ControlRig/Private/ModularRig.cpp
+++ b/Engine/Plugins/Animation/ControlRig/Source/ControlRig/Private/ModularRig.cpp
@@ -405,6 +405,12 @@ bool UModularRig::Execute_Internal(const FName& InEventName)
 		});
 
 		ExecuteQueue();
+
+		if (bAccumulateTime)
+		{
+			AbsoluteTime += DeltaTime;
+		}
+		
 		return true;
 	}
 	return false;
diff --git a/Engine/Plugins/Animation/ControlRig/Source/ControlRig/Public/ControlRigTestData.h b/Engine/Plugins/Animation/ControlRig/Source/ControlRig/Public/ControlRigTestData.h
--- a/Engine/Plugins/Animation/ControlRig/Source/ControlRig/Public/ControlRigTestData.h
+++ b/Engine/Plugins/Animation/ControlRig/Source/ControlRig/Public/ControlRigTestData.h
@@ -133,6 +133,9 @@ public:
 	UPROPERTY(EditAnywhere, BlueprintReadOnly, Category = "ControlRigTestData")
 	double Tolerance;
 
+	UPROPERTY(EditAnywhere, BlueprintReadOnly, Category = "ControlRigTestData")
+	TArray<FName> EventQueue;
+
 private:
 
 	void ClearDelegates(UControlRig* InControlRig);
@@ -143,6 +146,6 @@ private:
 	TWeakObjectPtr<UControlRig> ReplayControlRig;
 	bool bIsApplyingOutputs;
 	FDelegateHandle PreConstructionHandle;
-	FDelegateHandle PreForwardHandle;
-	FDelegateHandle PostForwardHandle;
+	FDelegateHandle PreEventHandle;
+	FDelegateHandle PostEventHandle;
 };
\ No newline at end of file
diff --git a/Engine/Plugins/Runtime/RigVM/Source/RigVM/Private/RigVMHost.cpp b/Engine/Plugins/Runtime/RigVM/Source/RigVM/Private/RigVMHost.cpp
--- a/Engine/Plugins/Runtime/RigVM/Source/RigVM/Private/RigVMHost.cpp
+++ b/Engine/Plugins/Runtime/RigVM/Source/RigVM/Private/RigVMHost.cpp
@@ -596,6 +596,12 @@ bool URigVMHost::Execute(const FName& InEventName)
 		return false;
 	}
 
+	if (ExecutedEvent.IsBound())
+	{
+		FRigVMBracketScope BracketScope(ExecuteBracket);
+		PreExecutedEvent.Broadcast(this, InEventName);
+	}
+
 	const bool bSuccess = Execute_Internal(InEventName);
 
 #if WITH_EDITOR
diff --git a/Engine/Plugins/Runtime/RigVM/Source/RigVM/Public/RigVMHost.h b/Engine/Plugins/Runtime/RigVM/Source/RigVM/Public/RigVMHost.h
--- a/Engine/Plugins/Runtime/RigVM/Source/RigVM/Public/RigVMHost.h
+++ b/Engine/Plugins/Runtime/RigVM/Source/RigVM/Public/RigVMHost.h
@@ -258,6 +258,7 @@ public:
 	DECLARE_EVENT_TwoParams(URigVM, FRigVMExecutedEvent, class URigVMHost*, const FName&);
 	FRigVMExecutedEvent& OnInitialized_AnyThread() { return InitializedEvent; }
 	FRigVMExecutedEvent& OnExecuted_AnyThread() { return ExecutedEvent; }
+	FRigVMExecutedEvent& OnPreExecuted_AnyThread() { return PreExecutedEvent; }
 
 	const FRigVMDrawInterface& GetDrawInterface() const { return DrawInterface; };
 	FRigVMDrawInterface& GetDrawInterface() { return DrawInterface; };
@@ -412,10 +413,12 @@ public:
 	/** Broadcasts a notification whenever the RigVMHost's memory is initialized. */
 	FRigVMExecutedEvent InitializedEvent;
 
-	/** Broadcasts a not
-	 * ification whenever the RigVMHost is executed / updated. */
+	/** Broadcasts a notification whenever the RigVMHost is executed / updated. */
 	FRigVMExecutedEvent ExecutedEvent;
 
+	/** Broadcasts a notification before the RigVMHost is executed / updated. */
+	FRigVMExecutedEvent PreExecutedEvent;
+
 protected: 
 
 	virtual void InitializeFromCDO();
