[Backout] - CL38260661
Original CL Desc
-----------------------------------------------------------------
Moving Physics in Control Rig to the PhysicsControl plugin, and switching over to doing it through components rather than hierarchy elements

diff --git a/Engine/Plugins/Animation/ControlRig/Source/ControlRig/Private/ControlRig.cpp b/Engine/Plugins/Animation/ControlRig/Source/ControlRig/Private/ControlRig.cpp
--- a/Engine/Plugins/Animation/ControlRig/Source/ControlRig/Private/ControlRig.cpp
+++ b/Engine/Plugins/Animation/ControlRig/Source/ControlRig/Private/ControlRig.cpp
@@ -566,10 +566,8 @@ void UControlRig::InitializeFromCDO()
 			UControlRig* CDO = GetClass()->GetDefaultObject<UControlRig>();
 			URigHierarchy* Hierarchy = GetHierarchy();
 
-			// Danny TODO this could copy the physics solvers from the CDO, but they are transient
-			// objects that only exist through construction, so I think we just need to start afresh
-			// each time. 
-			PhysicsSolvers.Empty();
+			// copy physics solvers
+			PhysicsSolvers = CDO->PhysicsSolvers;
 
 			// copy hierarchy
 			{
@@ -581,6 +579,25 @@ void UControlRig::InitializeFromCDO()
 				Hierarchy->ResetPoseToInitial(ERigElementType::All);
 			}
 
+			// update the physics solvers with new unique IDs and remap the physics solver IDs within hierarchy
+			{
+				TMap<FRigPhysicsSolverID, FRigPhysicsSolverID> PhysicsSolverMap;
+				for(FRigPhysicsSolverDescription& Solver : PhysicsSolvers)
+				{
+					const FRigPhysicsSolverID OldID = Solver.ID;
+					Solver.ID = FRigPhysicsSolverDescription::MakeID(GetPathName(), Solver.Name);
+					PhysicsSolverMap.Add(OldID, Solver.ID);
+				}
+				const TArray<FRigPhysicsElement*> PhysicsElements = Hierarchy->GetPhysicsElements();
+				for(FRigPhysicsElement* PhysicsElement : PhysicsElements)
+				{
+					if(const FRigPhysicsSolverID* NewID = PhysicsSolverMap.Find(PhysicsElement->Solver))
+					{
+						PhysicsElement->Solver = *NewID;
+					}
+				}
+			}
+
 #if WITH_EDITOR
 			// current hierarchy should always mirror CDO's hierarchy whenever a change of interest happens
 			CDO->GetHierarchy()->RegisterListeningHierarchy(Hierarchy);
@@ -1790,10 +1807,9 @@ void UControlRig::Serialize(FArchive& Ar)
 
 	Ar.UsingCustomVersion(FControlRigObjectVersion::GUID);
 
-	// Danny TODO solvers are transient and only made from the construction event
 	if (Ar.CustomVer(FControlRigObjectVersion::GUID) >= FControlRigObjectVersion::ControlRigStoresPhysicsSolvers)
 	{
-		//Ar << PhysicsSolvers;
+		Ar << PhysicsSolvers;
 		UE_RIGVM_ARCHIVETRACE_ENTRY(Ar, TEXT("PhysicsSolvers"));
 	}
 	else
@@ -3582,33 +3598,46 @@ void UControlRig::OnHierarchyTransformUndoRedo(URigHierarchy* InHierarchy, const
 	}
 }
 
-FRigPhysicsSolverBase* UControlRig::GetPhysicsSolver(const FRigPhysicsSolverID& InID)
+int32 UControlRig::NumPhysicsSolvers() const
 {
-	TSharedPtr<FRigPhysicsSolverBase>* PhysicsSolver = PhysicsSolvers.Find(InID);
-	return PhysicsSolver ? PhysicsSolver->Get() : nullptr;
+	return PhysicsSolvers.Num();
 }
 
-const FRigPhysicsSolverBase* UControlRig::GetPhysicsSolver(const FRigPhysicsSolverID& InID) const
+const FRigPhysicsSolverDescription* UControlRig::GetPhysicsSolver(int32 InIndex) const
 {
-	const TSharedPtr<FRigPhysicsSolverBase>* PhysicsSolver = PhysicsSolvers.Find(InID);
-	return PhysicsSolver ? PhysicsSolver->Get() : nullptr;
+	if(PhysicsSolvers.IsValidIndex(InIndex))
+	{
+		return &PhysicsSolvers[InIndex];
+	}
+	return nullptr;
 }
 
-const FRigPhysicsSolverBase* UControlRig::FindPhysicsSolverByName(FName InName) const
+const FRigPhysicsSolverDescription* UControlRig::FindPhysicsSolver(const FRigPhysicsSolverID& InID) const
 {
-	for (const TPair<FRigPhysicsSolverID, TSharedPtr<FRigPhysicsSolverBase>>& SolverPair : PhysicsSolvers)
+	return PhysicsSolvers.FindByPredicate([InID](const FRigPhysicsSolverDescription& Solver) -> bool
 	{
-		if (SolverPair.Value.IsValid() && SolverPair.Value->Name == InName)
-		{
-			return SolverPair.Value.Get();
-		}
-	}
-	return nullptr;
+		return Solver.ID == InID;
+	});
+}
+
+const FRigPhysicsSolverDescription* UControlRig::FindPhysicsSolverByName(const FName& InName) const
+{
+	return PhysicsSolvers.FindByPredicate([InName](const FRigPhysicsSolverDescription& Solver) -> bool
+	{
+		return Solver.Name == InName;
+	});
 }
 
-FRigPhysicsSolverID UControlRig::RegisterPhysicsSolver(TSharedPtr<FRigPhysicsSolverBase> InPhysicsWorld, FName InName)
+FRigPhysicsSolverID UControlRig::AddPhysicsSolver(FName InName, bool bSetupUndo, bool bPrintPythonCommand)
 {
-	if(InName.IsNone() || !InPhysicsWorld.IsValid())
+#if WITH_EDITOR
+	if(CVarControlRigHierarchyEnablePhysics.GetValueOnAnyThread() == false)
+	{
+		return FRigPhysicsSolverID();
+	}
+#endif
+
+	if(InName.IsNone())
 	{
 		return FRigPhysicsSolverID();
 	}
@@ -3618,15 +3647,11 @@ FRigPhysicsSolverID UControlRig::RegisterPhysicsSolver(TSharedPtr<FRigPhysicsSol
 		return FindPhysicsSolverByName(InName) == nullptr;
 	});
 
-	const FString CompletePath = FString::Printf(TEXT("%s|%s"), *GetPathName(), *Name.ToString());
-	FRigPhysicsSolverID SolverID(FGuid::NewDeterministicGuid(CompletePath));
-
-	PhysicsSolvers.FindOrAdd(SolverID, InPhysicsWorld);
-
-	InPhysicsWorld->Name = Name;
-	InPhysicsWorld->SolverID = SolverID;
-
-	return SolverID;
+	FRigPhysicsSolverDescription Solver;
+	Solver.ID = FRigPhysicsSolverDescription::MakeID(GetPathName(), Name);
+	Solver.Name = Name;
+	PhysicsSolvers.Add(Solver);
+	return Solver.ID;
 }
 
 UControlRig::FPoseScope::FPoseScope(UControlRig* InControlRig, ERigElementType InFilter, const TArray<FRigElementKey>& InElements, const ERigTransformType::Type InTransformType)
diff --git a/Engine/Plugins/Animation/ControlRig/Source/ControlRig/Private/Rigs/RigHierarchy.cpp b/Engine/Plugins/Animation/ControlRig/Source/ControlRig/Private/Rigs/RigHierarchy.cpp
--- a/Engine/Plugins/Animation/ControlRig/Source/ControlRig/Private/Rigs/RigHierarchy.cpp
+++ b/Engine/Plugins/Animation/ControlRig/Source/ControlRig/Private/Rigs/RigHierarchy.cpp
@@ -52,6 +52,7 @@ static TAutoConsoleVariable<int32> CVarControlRigHierarchyTracePrecision(TEXT("C
 static TAutoConsoleVariable<int32> CVarControlRigHierarchyTraceOnSpawn(TEXT("ControlRig.Hierarchy.TraceOnSpawn"), 0, TEXT("sets the number of frames to trace when a new hierarchy is spawned"));
 TAutoConsoleVariable<bool> CVarControlRigHierarchyEnableRotationOrder(TEXT("ControlRig.Hierarchy.EnableRotationOrder"), true, TEXT("enables the rotation order for controls"));
 TAutoConsoleVariable<bool> CVarControlRigHierarchyEnableModules(TEXT("ControlRig.Hierarchy.Modules"), true, TEXT("enables the modular rigging functionality"));
+TAutoConsoleVariable<bool> CVarControlRigHierarchyEnablePhysics(TEXT("ControlRig.Hierarchy.Physics"), false, TEXT("enables physics support for the rig hierarchy"));
 static int32 sRigHierarchyLastTrace = INDEX_NONE;
 
 // A console command to trace a single frame / single execution for a control rig anim node / control rig component
@@ -807,6 +808,7 @@ void URigHierarchy::CopyHierarchy(URigHierarchy* InHierarchy)
 			sizeof(FRigNullElement),
 			sizeof(FRigControlElement),
 			sizeof(FRigCurveElement),
+			sizeof(FRigPhysicsElement),
 			sizeof(FRigReferenceElement),
 			sizeof(FRigConnectorElement),
 			sizeof(FRigSocketElement),
@@ -1678,6 +1680,24 @@ TArray<FRigElementKey> URigHierarchy::RestoreConnectorsFromStates(TArray<FRigCon
 	return Keys;
 }
 
+const FRigPhysicsSolverDescription* URigHierarchy::FindPhysicsSolver(const FRigPhysicsSolverID& InID) const
+{
+	if(const UControlRig* ControlRig = Cast<UControlRig>(GetOuter()))
+	{
+		return ControlRig->FindPhysicsSolver(InID);
+	}
+	return nullptr;
+}
+
+const FRigPhysicsSolverDescription* URigHierarchy::FindPhysicsSolverByName(const FName& InName) const
+{
+	if(const UControlRig* ControlRig = Cast<UControlRig>(GetOuter()))
+	{
+		return ControlRig->FindPhysicsSolverByName(InName);
+	}
+	return nullptr;
+}
+
 TArray<FName> URigHierarchy::GetMetadataNames(FRigElementKey InItem) const
 {
 	TArray<FName> Names;
@@ -5873,6 +5893,15 @@ FRigBaseElement* URigHierarchy::MakeElement(ERigElementType InElementType, int32
 			Element = NewElement<FRigCurveElement>(InCount);
 			break;
 		}
+		case ERigElementType::Physics:
+		{
+			if(OutStructureSize)
+			{
+				*OutStructureSize = sizeof(FRigPhysicsElement);
+			}
+			Element = NewElement<FRigPhysicsElement>(InCount);
+			break;
+		}
 		case ERigElementType::Reference:
 		{
 			if(OutStructureSize)
@@ -5970,6 +5999,18 @@ void URigHierarchy::DestroyElement(FRigBaseElement*& InElement)
 			}
 			break;
 		}
+		case ERigElementType::Physics:
+		{
+			FRigPhysicsElement* ExistingElements = Cast<FRigPhysicsElement>(InElement);
+			for(int32 Index=0;Index<Count;Index++)
+			{
+				TGuardValue<const FRigBaseElement*> DestroyGuard(ElementBeingDestroyed, &ExistingElements[Index]);
+				DestroyComponents(&ExistingElements[Index]);
+				DeallocateElementStorage(&ExistingElements[Index]);
+				ExistingElements[Index].~FRigPhysicsElement(); 
+			}
+			break;
+		}
 		case ERigElementType::Reference:
 		{
 			FRigReferenceElement* ExistingElements = Cast<FRigReferenceElement>(InElement);
diff --git a/Engine/Plugins/Animation/ControlRig/Source/ControlRig/Private/Rigs/RigHierarchyController.cpp b/Engine/Plugins/Animation/ControlRig/Source/ControlRig/Private/Rigs/RigHierarchyController.cpp
--- a/Engine/Plugins/Animation/ControlRig/Source/ControlRig/Private/Rigs/RigHierarchyController.cpp
+++ b/Engine/Plugins/Animation/ControlRig/Source/ControlRig/Private/Rigs/RigHierarchyController.cpp
@@ -522,6 +522,81 @@ FRigElementKey URigHierarchyController::AddCurve(FName InName, float InValue, bo
 	return NewElement->Key;
 }
 
+FRigElementKey URigHierarchyController::AddPhysicsElement(FName InName, FRigElementKey InParent, FRigPhysicsSolverID InSolver, 
+                                                          FRigPhysicsSettings InSettings, FTransform InLocalTransform, bool bSetupUndo, bool bPrintPythonCommand)
+{
+	if(!IsValid())
+	{
+		return FRigElementKey();
+	}
+
+#if WITH_EDITOR
+	if(CVarControlRigHierarchyEnablePhysics.GetValueOnAnyThread() == false)
+	{
+		return FRigElementKey();
+	}
+#endif
+
+	URigHierarchy* Hierarchy = GetHierarchy();
+
+	if(!InSolver.IsValid())
+	{
+		ReportError(TEXT("Physics Solver Guid is not valid"));
+		return FRigElementKey();
+	}
+
+	if(Hierarchy->FindPhysicsSolver(InSolver) == nullptr)
+	{
+		ReportErrorf(TEXT("Physics Solver with guid '%s' cannot be found."), *InSolver.ToString());
+		return FRigElementKey();
+	}
+
+#if WITH_EDITOR
+	TSharedPtr<FScopedTransaction> TransactionPtr;
+	if(bSetupUndo)
+	{
+		TransactionPtr = MakeShared<FScopedTransaction>(NSLOCTEXT("RigHierarchyController", "Add Physics Element", "Add Physics Element"));
+		Hierarchy->Modify();
+	}
+#endif
+
+	FRigPhysicsElement* NewElement = MakeElement<FRigPhysicsElement>();
+	{
+		TGuardValue<bool> DisableCacheValidityChecks(Hierarchy->bEnableCacheValidityCheck, false);		
+		NewElement->Key.Type = ERigElementType::Physics;
+		NewElement->Key.Name = GetSafeNewName(InName, NewElement->Key.Type);
+		NewElement->Solver = InSolver;
+		NewElement->Settings = InSettings;
+		AddElement(NewElement, Hierarchy->Get(Hierarchy->GetIndex(InParent)), true, InName);
+
+		Hierarchy->SetTransform(NewElement, InLocalTransform, ERigTransformType::InitialLocal, true, false);
+		NewElement->GetTransform().Current = NewElement->GetTransform().Initial;
+		NewElement->GetDirtyState().Current = NewElement->GetDirtyState().Initial;
+	}
+	
+#if WITH_EDITOR
+	TransactionPtr.Reset();
+
+	if (bPrintPythonCommand && !bSuspendPythonPrinting)
+	{
+		UBlueprint* Blueprint = GetTypedOuter<UBlueprint>();
+		if (Blueprint)
+		{
+			TArray<FString> Commands = GetAddPhysicsElementPythonCommands(NewElement);
+			for (const FString& Command : Commands)
+			{
+				RigVMPythonUtils::Print(Blueprint->GetFName().ToString(),
+					FString::Printf(TEXT("%s"), *Command));
+			}
+		}
+	}
+#endif
+
+	GetHierarchy()->EnsureCacheValidity();
+
+	return NewElement->Key;
+}
+
 FRigElementKey URigHierarchyController::AddReference(FName InName, FRigElementKey InParent,
 	FRigReferenceGetWorldTransformDelegate InDelegate, bool bSetupUndo)
 {
@@ -1508,6 +1583,12 @@ FString URigHierarchyController::ExportToText(TArray<FRigElementKey> InKeys) con
 				FRigCurveElement::StaticStruct()->ExportText(PerElementData.Content, Element, &DefaultElement, nullptr, PPF_None, nullptr);
 				break;
 			}
+			case ERigElementType::Physics:
+			{
+				FRigPhysicsElement DefaultElement;
+				FRigPhysicsElement::StaticStruct()->ExportText(PerElementData.Content, Element, &DefaultElement, nullptr, PPF_None, nullptr);
+				break;
+			}
 			case ERigElementType::Reference:
 			{
 				FRigReferenceElement DefaultElement;
@@ -1658,6 +1739,12 @@ TArray<FRigElementKey> URigHierarchyController::ImportFromText(FString InContent
 				FRigCurveElement::StaticStruct()->ImportText(*PerElementData.Content, NewElement, nullptr, EPropertyPortFlags::PPF_None, &ErrorPipe, FRigCurveElement::StaticStruct()->GetName(), true);
 				break;
 			}
+			case ERigElementType::Physics:
+			{
+				NewElement = MakeElement<FRigPhysicsElement>(true);
+				FRigPhysicsElement::StaticStruct()->ImportText(*PerElementData.Content, NewElement, nullptr, EPropertyPortFlags::PPF_None, &ErrorPipe, FRigPhysicsElement::StaticStruct()->GetName(), true);
+				break;
+			}
 			case ERigElementType::Reference:
 			{
 				NewElement = MakeElement<FRigReferenceElement>(true);
@@ -1997,6 +2084,10 @@ TArray<FString> URigHierarchyController::GetAddElementPythonCommands(FRigBaseEle
 	{
 		return GetAddCurvePythonCommands(CurveElement);
 	}
+	else if(FRigPhysicsElement* PhysicsElement = Cast<FRigPhysicsElement>(Element))
+	{
+		return GetAddPhysicsElementPythonCommands(PhysicsElement);
+	}
 	else if(FRigReferenceElement* ReferenceElement = Cast<FRigReferenceElement>(Element))
 	{
 		ensure(false);
@@ -2107,6 +2198,42 @@ TArray<FString> URigHierarchyController::GetAddCurvePythonCommands(FRigCurveElem
 		Hierarchy->GetCurveValue(Curve))};
 }
 
+TArray<FString> URigHierarchyController::GetAddPhysicsElementPythonCommands(FRigPhysicsElement* PhysicsElement) const
+{
+	TArray<FString> Commands;
+	FString TransformStr = RigVMPythonUtils::TransformToPythonString(PhysicsElement->GetTransform().Initial.Local.Get());
+	
+	FString ParentKeyStr = "''";
+	if (PhysicsElement->ParentElement)
+	{
+		ParentKeyStr = PhysicsElement->ParentElement->GetKey().ToPythonString();
+	}
+
+	FString PhysicsElementNamePythonized = RigVMPythonUtils::PythonizeName(PhysicsElement->GetName());
+	FRigPhysicsSettings& Settings = PhysicsElement->Settings;
+	FString SettingsStr;
+	{
+		SettingsStr =FString::Printf(TEXT("physics_settings%s"),
+			*PhysicsElementNamePythonized);
+			
+		Commands.Add(FString::Printf(TEXT("physics_settings%s = unreal.RigPhysicsSettings()"),
+			*PhysicsElementNamePythonized));
+
+		Commands.Add(FString::Printf(TEXT("control_settings_%s.mass = %f"),
+			*PhysicsElementNamePythonized,
+			Settings.Mass));	
+	}
+	
+	// FRigElementKey AddPhysicsElement(FName InName, FRigElementKey InParent, FRigPhysicsSettings InSettings, FTransform InLocalTransform, bool bSetupUndo = false);
+	Commands.Add(FString::Printf(TEXT("hierarchy_controller.add_physics_element('%s', %s, %s, %s)"),
+		*PhysicsElement->GetName(),
+		*ParentKeyStr,
+		*SettingsStr,
+		*TransformStr));
+
+	return Commands;
+}
+
 TArray<FString> URigHierarchyController::GetAddConnectorPythonCommands(FRigConnectorElement* Connector) const
 {
 	TArray<FString> Commands;
diff --git a/Engine/Plugins/Animation/ControlRig/Source/ControlRig/Private/Rigs/RigHierarchyElements.cpp b/Engine/Plugins/Animation/ControlRig/Source/ControlRig/Private/Rigs/RigHierarchyElements.cpp
--- a/Engine/Plugins/Animation/ControlRig/Source/ControlRig/Private/Rigs/RigHierarchyElements.cpp
+++ b/Engine/Plugins/Animation/ControlRig/Source/ControlRig/Private/Rigs/RigHierarchyElements.cpp
@@ -2055,6 +2055,110 @@ void FRigCurveElement::CopyFrom(const FRigBaseElement* InOther)
 	}
 }
 
+////////////////////////////////////////////////////////////////////////////////
+// FRigPhysicsSolver
+////////////////////////////////////////////////////////////////////////////////
+
+void FRigPhysicsSolverDescription::Serialize(FArchive& Ar)
+{
+	if (Ar.IsSaving() || Ar.IsObjectReferenceCollector() || Ar.IsCountingMemory())
+	{
+		Save(Ar);
+	}
+	else if (Ar.IsLoading())
+	{
+		Load(Ar);
+	}
+}
+
+void FRigPhysicsSolverDescription::Save(FArchive& Ar)
+{
+	Ar << ID;
+	Ar << Name;
+}
+
+void FRigPhysicsSolverDescription::Load(FArchive& Ar)
+{
+	Ar << ID;
+	Ar << Name;
+}
+
+FGuid FRigPhysicsSolverDescription::MakeGuid(const FString& InObjectPath, const FName& InSolverName)
+{
+	const FString CompletePath = FString::Printf(TEXT("%s|%s"), *InObjectPath, *InSolverName.ToString());
+	return FGuid::NewDeterministicGuid(CompletePath);
+}
+
+FRigPhysicsSolverID FRigPhysicsSolverDescription::MakeID(const FString& InObjectPath, const FName& InSolverName)
+{
+	return FRigPhysicsSolverID(MakeGuid(InObjectPath, InSolverName));
+}
+
+void FRigPhysicsSolverDescription::CopyFrom(const FRigPhysicsSolverDescription* InOther)
+{
+	if(InOther)
+	{
+		ID = InOther->ID;
+		Name = InOther->Name;
+	}
+}
+
+////////////////////////////////////////////////////////////////////////////////
+// FRigPhysicsSettings
+////////////////////////////////////////////////////////////////////////////////
+
+FRigPhysicsSettings::FRigPhysicsSettings()
+	: Mass(1.f)
+{
+}
+
+void FRigPhysicsSettings::Save(FArchive& Ar)
+{
+	Ar << Mass;
+}
+
+void FRigPhysicsSettings::Load(FArchive& Ar)
+{
+	Ar << Mass;
+}
+
+////////////////////////////////////////////////////////////////////////////////
+// FRigPhysicsElement
+////////////////////////////////////////////////////////////////////////////////
+
+const FRigBaseElement::EElementIndex FRigPhysicsElement::ElementTypeIndex = PhysicsElement;
+
+void FRigPhysicsElement::Save(FArchive& Ar, const FRigHierarchySerializationSettings& InSettings)
+{
+	Super::Save(Ar, InSettings);
+
+	if(InSettings.SerializationPhase == FRigHierarchySerializationSettings::StaticData)
+	{
+		Ar << Solver;
+		Settings.Save(Ar);
+	}
+}
+
+void FRigPhysicsElement::Load(FArchive& Ar, const FRigHierarchySerializationSettings& InSettings)
+{
+	Super::Load(Ar, InSettings);
+
+	if(InSettings.SerializationPhase == FRigHierarchySerializationSettings::StaticData)
+	{
+		Ar << Solver;
+		Settings.Load(Ar);
+	}
+}
+
+void FRigPhysicsElement::CopyFrom(const FRigBaseElement* InOther)
+{
+	Super::CopyFrom(InOther);
+	
+	const FRigPhysicsElement* Source = CastChecked<FRigPhysicsElement>(InOther);
+	Solver = Source->Solver;
+	Settings = Source->Settings;
+}
+
 ////////////////////////////////////////////////////////////////////////////////
 // FRigReferenceElement
 ////////////////////////////////////////////////////////////////////////////////
diff --git a/Engine/Plugins/Animation/ControlRig/Source/ControlRig/Private/Units/Execution/RigUnit_Physics.cpp b/Engine/Plugins/Animation/ControlRig/Source/ControlRig/Private/Units/Execution/RigUnit_Physics.cpp
new file mode 100644
--- /dev/null
+++ b/Engine/Plugins/Animation/ControlRig/Source/ControlRig/Private/Units/Execution/RigUnit_Physics.cpp
@@ -0,0 +1,49 @@
+// Copyright Epic Games, Inc. All Rights Reserved.
+
+#include "Units/Execution/RigUnit_Physics.h"
+#include "Rigs/RigHierarchyController.h"
+#include "Units/RigUnitContext.h"
+#include "ControlRig.h"
+
+#include UE_INLINE_GENERATED_CPP_BY_NAME(RigUnit_Physics)
+
+FRigUnit_HierarchyAddPhysicsSolver_Execute()
+{
+	FString ErrorMessage;
+	if(!FRigUnit_DynamicHierarchyBase::IsValidToRunInContext(ExecuteContext, true, &ErrorMessage))
+	{
+		if(!ErrorMessage.IsEmpty())
+		{
+			UE_CONTROLRIG_RIGUNIT_REPORT_ERROR(TEXT("%s"), *ErrorMessage);
+		}
+		return;
+	}
+
+	Solver = FRigPhysicsSolverID();
+
+	if(UControlRig* ControlRig = Cast<UControlRig>(ExecuteContext.Hierarchy->GetOuter()))
+	{
+		Solver = ControlRig->AddPhysicsSolver(Name, false, false);
+	}
+}
+
+FRigUnit_HierarchyAddPhysicsJoint_Execute()
+{
+	FString ErrorMessage;
+	if(!FRigUnit_DynamicHierarchyBase::IsValidToRunInContext(ExecuteContext, true, &ErrorMessage))
+	{
+		if(!ErrorMessage.IsEmpty())
+		{
+			UE_CONTROLRIG_RIGUNIT_REPORT_ERROR(TEXT("%s"), *ErrorMessage);
+		}
+		return;
+	}
+
+	Item.Reset();
+
+	if(URigHierarchyController* Controller = ExecuteContext.Hierarchy->GetController(true))
+	{
+		FRigHierarchyControllerInstructionBracket InstructionBracket(Controller, ExecuteContext.GetInstructionIndex());
+		Item = Controller->AddPhysicsElement(Name, Parent, Solver, Settings, Transform, false, false);
+	}
+}
diff --git a/Engine/Plugins/Animation/ControlRig/Source/ControlRig/Public/ControlRig.h b/Engine/Plugins/Animation/ControlRig/Source/ControlRig/Public/ControlRig.h
--- a/Engine/Plugins/Animation/ControlRig/Source/ControlRig/Public/ControlRig.h
+++ b/Engine/Plugins/Animation/ControlRig/Source/ControlRig/Public/ControlRig.h
@@ -21,7 +21,6 @@
 #include "Animation/MeshDeformerInstance.h"
 #include "Animation/MeshDeformerProducer.h"
 #include "Rigs/RigModuleDefines.h"
-#include "Rigs/RigPhysics.h"
 
 #if WITH_EDITOR
 #include "RigVMModel/RigVMPin.h"
@@ -45,7 +44,6 @@ class UTransformableControlHandle;
 struct FReferenceSkeleton;
 struct FRigUnit;
 struct FRigControl;
-struct FRigPhysicsSolverBase;
 
 CONTROLRIG_API DECLARE_LOG_CATEGORY_EXTERN(LogControlRig, Log, All);
 
@@ -716,37 +714,50 @@ protected:
 private:
 	float DebugBoneRadiusMultiplier;
 
-	// Physics Solvers/Worlds
-	TMap<FRigPhysicsSolverID, TSharedPtr<FRigPhysicsSolverBase>> PhysicsSolvers;
+	// Physics Solvers
+	TArray<FRigPhysicsSolverDescription> PhysicsSolvers;
 
 public:
 
 	/**
-	 * Returns the physics world given a solver ID
-	 * @param InID The id identifying the physics solver
-	 * @return The physics solver (may be null)
+	 * Returns the number of physics solvers
+	 * @return The number of physics solvers
 	 */
-	FRigPhysicsSolverBase* GetPhysicsSolver(const FRigPhysicsSolverID& InID);
+	int32 NumPhysicsSolvers() const; 
 
 	/**
-	 * Returns the physics world given a solver ID
+	 * Returns a physics solver by index (or nullptr)
+	 * @param InIndex The index of the physics solver to return
+	 * @return The physics solver
+	 */
+	const FRigPhysicsSolverDescription* GetPhysicsSolver(int32 InIndex) const; 
+
+	/**
+	 * Finds a new physics solver given its guid
 	 * @param InID The id identifying the physics solver
-	 * @return The physics solver (may be null)
+	 * @return The physics solver
 	 */
-	const FRigPhysicsSolverBase* GetPhysicsSolver(const FRigPhysicsSolverID& InID) const;
+	const FRigPhysicsSolverDescription* FindPhysicsSolver(const FRigPhysicsSolverID& InID) const; 
 
 	/**
-	 * Finds a physics solver given its name
+	 * Finds a new physics solver given its name
 	 * @param InName The name identifying the physics solver in the scope of this control rig
 	 * @return The physics solver
 	 */
-	const FRigPhysicsSolverBase* FindPhysicsSolverByName(FName InName) const; 
+	const FRigPhysicsSolverDescription* FindPhysicsSolverByName(const FName& InName) const; 
 
 	/**
-	 * Registers and stores (taking ownership) of a physics world, returning an ID for it. The world will be 
-	 * responsible for the low-level storage and simulation of objects.
+	 * Adds a physics solver to the hierarchy
+	 * @param InName The suggested name of the new physics solver - will eventually be corrected by the namespace
+	 * @param bSetupUndo If set to true the stack will record the change for undo / redo
+	 * @param bPrintPythonCommand If set to true a python command equivalent to this call will be printed out
+	 * @return The ID for the newly created physics solver.
 	 */
-	FRigPhysicsSolverID RegisterPhysicsSolver(TSharedPtr<FRigPhysicsSolverBase> PhysicsWorld, FName Name);
+	UFUNCTION(BlueprintCallable, Category = ControlRig)
+	FRigPhysicsSolverID AddPhysicsSolver(
+		FName InName,
+		bool bSetupUndo = false,
+		bool bPrintPythonCommand = false);
 
 #if WITH_EDITOR	
 
diff --git a/Engine/Plugins/Animation/ControlRig/Source/ControlRig/Public/Rigs/RigHierarchy.h b/Engine/Plugins/Animation/ControlRig/Source/ControlRig/Public/Rigs/RigHierarchy.h
--- a/Engine/Plugins/Animation/ControlRig/Source/ControlRig/Public/Rigs/RigHierarchy.h
+++ b/Engine/Plugins/Animation/ControlRig/Source/ControlRig/Public/Rigs/RigHierarchy.h
@@ -34,6 +34,7 @@ DECLARE_MULTICAST_DELEGATE_ThreeParams(FRigHierarchyMetadataTagChangedDelegate,
 
 extern CONTROLRIG_API TAutoConsoleVariable<bool> CVarControlRigHierarchyEnableRotationOrder;
 extern CONTROLRIG_API TAutoConsoleVariable<bool> CVarControlRigHierarchyEnableModules;
+extern CONTROLRIG_API TAutoConsoleVariable<bool> CVarControlRigHierarchyEnablePhysics;
 
 UENUM()
 enum ERigTransformStackEntryType : int
@@ -1086,6 +1087,25 @@ public:
 		return GetKeysOfType<FRigCurveElement>(false);
 	}
 
+	/**
+	 * Returns all Physics elements
+	 * @param bTraverse Returns the elements in order of a depth first traversal
+	 */
+	TArray<FRigPhysicsElement*> GetPhysicsElements(bool bTraverse = false) const
+	{
+		return GetElementsOfType<FRigPhysicsElement>(bTraverse);
+	}
+
+	/**
+	 * Returns all Physics elements
+	 * @param bTraverse Returns the elements in order of a depth first traversal
+	 */
+	UFUNCTION(BlueprintCallable, Category = URigHierarchy, meta = (DisplayName = "Get Physics Keys", ScriptName = "GetPhysicsKeys"))
+    TArray<FRigElementKey> GetPhysicsKeys(bool bTraverse = true) const
+	{
+		return GetKeysOfType<FRigPhysicsElement>(bTraverse);
+	}
+
 	/**
 	 * Returns all references
 	 * @param bTraverse Returns the elements in order of a depth first traversal
@@ -1190,6 +1210,20 @@ public:
 		}, false);
 	}
 
+	/**
+	 * Finds a new physics solver given its guid
+	 * @param InID The id identifying the physics solver
+	 * @return The physics solver
+	 */
+	const FRigPhysicsSolverDescription* FindPhysicsSolver(const FRigPhysicsSolverID& InID) const; 
+
+	/**
+	 * Finds a new physics solver given its name
+	 * @param InName The name identifying the physics solver in the scope of this hierarchy
+	 * @return The physics solver
+	 */
+	const FRigPhysicsSolverDescription* FindPhysicsSolverByName(const FName& InName) const; 
+
 	/**
 	 * Returns the name of metadata for a given element
 	 * @param InItem The element key to return the metadata keys for
diff --git a/Engine/Plugins/Animation/ControlRig/Source/ControlRig/Public/Rigs/RigHierarchyController.h b/Engine/Plugins/Animation/ControlRig/Source/ControlRig/Public/Rigs/RigHierarchyController.h
--- a/Engine/Plugins/Animation/ControlRig/Source/ControlRig/Public/Rigs/RigHierarchyController.h
+++ b/Engine/Plugins/Animation/ControlRig/Source/ControlRig/Public/Rigs/RigHierarchyController.h
@@ -200,6 +200,27 @@ public:
 		bool bPrintPythonCommand = false
         );
 
+	/**
+	* Adds a physics element to the hierarchy
+	* @param InName The suggested name of the new physics element - will eventually be corrected by the namespace
+	* @param InParent The (optional) parent of the new physics element. If you don't need a parent, pass FRigElementKey()
+	* @param InSolver The guid identifying the solver to use
+	* @param InSettings All of the physics element's settings
+	* @param InLocalTransform The transform for the new physics element - in the space of the provided parent
+	* @param bSetupUndo If set to true the stack will record the change for undo / redo
+	* @param bPrintPythonCommand If set to true a python command equivalent to this call will be printed out
+	* @return The key for the newly created physics element.
+	*/
+	UFUNCTION(BlueprintCallable, Category = URigHierarchyController)
+    FRigElementKey AddPhysicsElement(
+    	FName InName,
+    	FRigElementKey InParent,
+    	FRigPhysicsSolverID InSolver,
+        FRigPhysicsSettings InSettings,
+    	FTransform InLocalTransform,
+    	bool bSetupUndo = false,
+		bool bPrintPythonCommand = false);
+
 	/**
 	* Adds an reference to the hierarchy
 	* @param InName The suggested name of the new reference - will eventually be corrected by the namespace
@@ -712,6 +733,8 @@ public:
 
 	TArray<FString> GetAddCurvePythonCommands(FRigCurveElement* Curve) const;
 
+	TArray<FString> GetAddPhysicsElementPythonCommands(FRigPhysicsElement* PhysicsElement) const;
+
 	TArray<FString> GetAddConnectorPythonCommands(FRigConnectorElement* Connector) const;
 
 	TArray<FString> GetAddSocketPythonCommands(FRigSocketElement* Socket) const;
diff --git a/Engine/Plugins/Animation/ControlRig/Source/ControlRig/Public/Rigs/RigHierarchyDefines.h b/Engine/Plugins/Animation/ControlRig/Source/ControlRig/Public/Rigs/RigHierarchyDefines.h
--- a/Engine/Plugins/Animation/ControlRig/Source/ControlRig/Public/Rigs/RigHierarchyDefines.h
+++ b/Engine/Plugins/Animation/ControlRig/Source/ControlRig/Public/Rigs/RigHierarchyDefines.h
@@ -29,7 +29,7 @@ enum class ERigElementType : uint8
 	Space = Null UMETA(Hidden),
 	Control = 0x004,
 	Curve = 0x008,
-	Physics = 0x010 UMETA(Hidden),
+	Physics = 0x010, 
 	Reference = 0x020,
 	Connector = 0x040,
 	Socket = 0x080,
diff --git a/Engine/Plugins/Animation/ControlRig/Source/ControlRig/Public/Rigs/RigHierarchyElements.h b/Engine/Plugins/Animation/ControlRig/Source/ControlRig/Public/Rigs/RigHierarchyElements.h
--- a/Engine/Plugins/Animation/ControlRig/Source/ControlRig/Public/Rigs/RigHierarchyElements.h
+++ b/Engine/Plugins/Animation/ControlRig/Source/ControlRig/Public/Rigs/RigHierarchyElements.h
@@ -1814,6 +1814,168 @@ private:
 	friend class FControlRigHierarchyRelinkElementStorage;
 };
 
+USTRUCT(BlueprintType)
+struct CONTROLRIG_API FRigPhysicsSolverID
+{
+public:
+	
+	GENERATED_BODY()
+	
+	FRigPhysicsSolverID()
+		: Guid()
+	{
+	}
+
+	explicit FRigPhysicsSolverID(const FGuid& InGuid)
+		: Guid(InGuid)
+	{
+	}
+
+	UPROPERTY(BlueprintReadOnly, EditAnywhere, Category = Physics, meta=(ShowOnlyInnerProperties))
+	FGuid Guid;
+
+	bool IsValid() const
+	{
+		return Guid.IsValid();
+	}
+
+	FString ToString() const
+	{
+		return Guid.ToString();
+	}
+
+	bool operator == (const FRigPhysicsSolverID& InOther) const
+	{
+		return Guid == InOther.Guid;
+	}
+
+	bool operator != (const FRigPhysicsSolverID& InOther) const
+	{
+		return !(*this == InOther);
+	}
+
+	friend FArchive& operator <<(FArchive& Ar, FRigPhysicsSolverID& ID)
+	{
+		Ar << ID.Guid;
+		return Ar;
+	}
+
+	friend uint32 GetTypeHash(const FRigPhysicsSolverID& InID)
+	{
+		return GetTypeHash(InID.Guid);
+	}
+};
+
+USTRUCT(BlueprintType)
+struct CONTROLRIG_API FRigPhysicsSolverDescription
+{
+public:
+	
+	GENERATED_BODY()
+
+	UPROPERTY(BlueprintReadOnly, EditAnywhere, Category = Physics, meta=(ShowOnlyInnerProperties))
+	FRigPhysicsSolverID ID;
+
+	UPROPERTY(BlueprintReadOnly, EditAnywhere, Category = Physics, meta=(ShowOnlyInnerProperties))
+	FName Name;
+
+	FRigPhysicsSolverDescription()
+	: ID()
+	, Name(NAME_None)
+	{}
+	
+	FRigPhysicsSolverDescription(const FRigPhysicsSolverDescription& InOther)
+	{
+		*this = InOther;
+	}
+	
+	FRigPhysicsSolverDescription& operator=(const FRigPhysicsSolverDescription& InOther)
+	{
+		CopyFrom(&InOther);
+		return *this;
+	}
+
+	void Serialize(FArchive& Ar);
+	void Save(FArchive& Ar);
+	void Load(FArchive& Ar);
+	friend FArchive& operator<<(FArchive& Ar, FRigPhysicsSolverDescription& P)
+	{
+		P.Serialize(Ar);
+		return Ar;
+	}
+
+	static FGuid MakeGuid(const FString& InObjectPath, const FName& InSolverName);
+	static FRigPhysicsSolverID MakeID(const FString& InObjectPath, const FName& InSolverName);
+
+private:
+
+	void CopyFrom(const FRigPhysicsSolverDescription* InOther);
+};
+
+USTRUCT(BlueprintType)
+struct CONTROLRIG_API FRigPhysicsSettings
+{
+	GENERATED_BODY()
+
+	FRigPhysicsSettings();
+
+	void Save(FArchive& Ar);
+	void Load(FArchive& Ar);
+
+
+	UPROPERTY(BlueprintReadWrite, EditAnywhere, Category = Control)
+	float Mass;
+};
+
+USTRUCT(BlueprintType)
+struct CONTROLRIG_API FRigPhysicsElement : public FRigSingleParentElement
+{
+public:
+	
+	GENERATED_BODY()
+	DECLARE_RIG_ELEMENT_METHODS(FRigPhysicsElement)
+
+	static const EElementIndex ElementTypeIndex;
+
+	UPROPERTY(BlueprintReadOnly, EditAnywhere, Category = Physics, meta=(ShowOnlyInnerProperties))
+	FRigPhysicsSolverID Solver;
+
+	UPROPERTY(BlueprintReadOnly, EditAnywhere, Category = Physics, meta=(ShowOnlyInnerProperties))
+	FRigPhysicsSettings Settings;
+	
+	FRigPhysicsElement()
+        : FRigPhysicsElement(nullptr)
+	{ }
+	
+	FRigPhysicsElement(const FRigPhysicsElement& InOther)
+	{
+		*this = InOther;
+	}
+	
+	FRigPhysicsElement& operator=(const FRigPhysicsElement& InOther)
+	{
+		Super::operator=(InOther);
+		Solver = InOther.Solver;
+		Settings = InOther.Settings;
+		return *this;
+	}
+	
+	virtual ~FRigPhysicsElement() override {}
+
+	virtual void Save(FArchive& A, const FRigHierarchySerializationSettings& InSettings) override;
+	virtual void Load(FArchive& Ar, const FRigHierarchySerializationSettings& InSettings) override;
+
+private:
+	explicit FRigPhysicsElement(URigHierarchy* InOwner)
+		: FRigSingleParentElement(InOwner, ERigElementType::Physics)
+	{ }
+	
+	virtual void CopyFrom(const FRigBaseElement* InOther) override;
+
+	friend class URigHierarchy;
+	friend struct FRigBaseElement;
+};
+
 USTRUCT(BlueprintType)
 struct CONTROLRIG_API FRigReferenceElement final : public FRigSingleParentElement
 {
diff --git a/Engine/Plugins/Animation/ControlRig/Source/ControlRig/Public/Rigs/RigPhysics.h b/Engine/Plugins/Animation/ControlRig/Source/ControlRig/Public/Rigs/RigPhysics.h
deleted file mode 100644
diff --git a/Engine/Plugins/Animation/ControlRig/Source/ControlRig/Public/Units/Execution/RigUnit_Physics.h b/Engine/Plugins/Animation/ControlRig/Source/ControlRig/Public/Units/Execution/RigUnit_Physics.h
new file mode 100644
--- /dev/null
+++ b/Engine/Plugins/Animation/ControlRig/Source/ControlRig/Public/Units/Execution/RigUnit_Physics.h
@@ -0,0 +1,80 @@
+// Copyright Epic Games, Inc. All Rights Reserved.
+
+#pragma once
+
+#include "RigUnit_DynamicHierarchy.h"
+#include "RigUnit_Physics.generated.h"
+
+/**
+ * Adds a new physics solver to the hierarchy
+ * Note: This node only runs as part of the construction event.
+ */
+USTRUCT(meta=(DisplayName="Spawn Physics Solver", Keywords="Construction,Create,New,Simulation", Varying))
+struct CONTROLRIG_API FRigUnit_HierarchyAddPhysicsSolver : public FRigUnit_DynamicHierarchyBaseMutable
+{
+	GENERATED_BODY()
+
+	FRigUnit_HierarchyAddPhysicsSolver()
+	{
+		Name = TEXT("Solver");
+		Solver = FRigPhysicsSolverID();
+	}
+
+	/*
+	 * The name of the new solver to add
+	 */
+	UPROPERTY(meta = (Input))
+	FName Name;
+
+	/*
+	 * The identifier of the solver
+	 */
+	UPROPERTY(meta = (Output))
+	FRigPhysicsSolverID Solver;
+
+	RIGVM_METHOD()
+	virtual void Execute() override;
+};
+
+/**
+ * Adds a new physics joint to the hierarchy
+ * Note: This node only runs as part of the construction event.
+ */
+USTRUCT(meta=(DisplayName="Spawn Physics Joint", Keywords="Construction,Create,New,Joint,Skeleton", Varying))
+struct CONTROLRIG_API FRigUnit_HierarchyAddPhysicsJoint : public FRigUnit_HierarchyAddElement
+{
+	GENERATED_BODY()
+
+	FRigUnit_HierarchyAddPhysicsJoint()
+	{
+		Name = TEXT("NewPhysicsJoint");
+		Parent = FRigElementKey(NAME_None, ERigElementType::Bone);
+		Transform = FTransform::Identity;
+		Solver = FRigPhysicsSolverID();
+	}
+
+	virtual ERigElementType GetElementTypeToSpawn() const override { return ERigElementType::Physics; }
+
+
+	RIGVM_METHOD()
+	virtual void Execute() override;
+
+	/*
+	 * The initial global transform of the spawned element
+	 */
+	UPROPERTY(meta = (Input))
+	FTransform Transform;
+
+	/*
+	 * The solver to relate this new physics element to
+	 */
+	UPROPERTY(meta = (Input))
+	FRigPhysicsSolverID Solver;
+
+	/*
+	 * The settings of the new physics element
+	 */
+	UPROPERTY(meta = (Input))
+	FRigPhysicsSettings Settings;
+};
+
diff --git a/Engine/Plugins/Animation/ControlRig/Source/ControlRigDeveloper/Private/Graph/ControlRigGraph.cpp b/Engine/Plugins/Animation/ControlRig/Source/ControlRigDeveloper/Private/Graph/ControlRigGraph.cpp
--- a/Engine/Plugins/Animation/ControlRig/Source/ControlRigDeveloper/Private/Graph/ControlRigGraph.cpp
+++ b/Engine/Plugins/Animation/ControlRig/Source/ControlRigDeveloper/Private/Graph/ControlRigGraph.cpp
@@ -92,6 +92,7 @@ void UControlRigGraph::CacheNameLists(URigHierarchy* InHierarchy, const FRigVMDr
 		TArray<TSharedPtr<FRigVMStringWithTag>>& NullNameList = ElementNameLists.FindChecked(ERigElementType::Null);
 		TArray<TSharedPtr<FRigVMStringWithTag>>& ControlNameList = ElementNameLists.FindChecked(ERigElementType::Control);
 		TArray<TSharedPtr<FRigVMStringWithTag>>& CurveNameList = ElementNameLists.FindChecked(ERigElementType::Curve);
+		TArray<TSharedPtr<FRigVMStringWithTag>>& PhysicsNameList = ElementNameLists.FindChecked(ERigElementType::Physics);
 		TArray<TSharedPtr<FRigVMStringWithTag>>& ReferenceNameList = ElementNameLists.FindChecked(ERigElementType::Reference);
 		TArray<TSharedPtr<FRigVMStringWithTag>>& ConnectorNameList = ElementNameLists.FindChecked(ERigElementType::Connector);
 		TArray<TSharedPtr<FRigVMStringWithTag>>& SocketNameList = ElementNameLists.FindChecked(ERigElementType::Socket);
@@ -102,6 +103,7 @@ void UControlRigGraph::CacheNameLists(URigHierarchy* InHierarchy, const FRigVMDr
 		CacheNameListForHierarchy<FRigControlElement>(ControlRig, InHierarchy, ControlNameList, false);
 		CacheNameListForHierarchy<FRigControlElement>(ControlRig, InHierarchy, ControlNameListWithoutAnimationChannels, true);
 		CacheNameListForHierarchy<FRigCurveElement>(ControlRig, InHierarchy, CurveNameList, false);
+		CacheNameListForHierarchy<FRigPhysicsElement>(ControlRig, InHierarchy, PhysicsNameList, false);
 		CacheNameListForHierarchy<FRigReferenceElement>(ControlRig, InHierarchy, ReferenceNameList, false);
 		CacheNameListForHierarchy<FRigConnectorElement>(ControlRig, InHierarchy, ConnectorNameList, false);
 		CacheNameListForHierarchy<FRigSocketElement>(ControlRig, InHierarchy, SocketNameList, false);
diff --git a/Engine/Plugins/Animation/ControlRig/Source/ControlRigEditor/Private/ControlRigElementDetails.cpp b/Engine/Plugins/Animation/ControlRig/Source/ControlRigEditor/Private/ControlRigElementDetails.cpp
--- a/Engine/Plugins/Animation/ControlRig/Source/ControlRigEditor/Private/ControlRigElementDetails.cpp
+++ b/Engine/Plugins/Animation/ControlRig/Source/ControlRigEditor/Private/ControlRigElementDetails.cpp
@@ -6476,4 +6476,98 @@ void FRigConnectionRuleDetails::OnRuleContentChanged()
 	StructPropertyHandle->GetParentHandle()->NotifyPostChange(EPropertyChangeType::ValueSet);
 }
 
+void FRigPhysicsElementDetails::CustomizeDetails(IDetailLayoutBuilder& DetailBuilder)
+{
+	FRigTransformElementDetails::CustomizeDetails(DetailBuilder);
+	CustomizeSettings(DetailBuilder);
+	CustomizeTransform(DetailBuilder);
+	CustomizeComponents(DetailBuilder);
+	CustomizeMetadata(DetailBuilder);
+}
+
+void FRigPhysicsElementDetails::CustomizeSettings(IDetailLayoutBuilder& DetailBuilder)
+{
+	if(PerElementInfos.IsEmpty())
+	{
+		return;
+	}
+
+	if(IsAnyElementNotOfType(ERigElementType::Physics))
+	{
+		return;
+	}
+
+	const TSharedPtr<IPropertyHandle> SolverHandle = DetailBuilder.GetProperty(GET_MEMBER_NAME_CHECKED(FRigPhysicsElement, Solver));
+	DetailBuilder.HideProperty(SolverHandle);
+
+	const TSharedPtr<IPropertyHandle> SettingsHandle = DetailBuilder.GetProperty(GET_MEMBER_NAME_CHECKED(FRigPhysicsElement, Settings));
+	DetailBuilder.HideProperty(SettingsHandle);
+
+	IDetailCategoryBuilder& PhysicsCategory = DetailBuilder.EditCategory(TEXT("Physics"), LOCTEXT("Physics", "Physics"));
+
+	FDetailWidgetRow& SolverNameRow = PhysicsCategory.AddCustomRow(LOCTEXT("SolverName", "Solver Name"));
+	SolverNameRow.IsEnabled(false);
+	SolverNameRow.NameContent()
+	[
+		SNew(STextBlock)
+		.Text(LOCTEXT("SolverName", "Solver Name"))
+		.Font(IDetailLayoutBuilder::GetDetailFont())
+	];
+	SolverNameRow.ValueContent()
+	[
+		SNew(SEditableText)
+		.Text(this, &FRigPhysicsElementDetails::GetSolverNameText)
+		.Font(IDetailLayoutBuilder::GetDetailFont())
+	];
+
+	PhysicsCategory
+	.AddProperty(SolverHandle)
+	.DisplayName(LOCTEXT("SolverGuid", "Solver Guid"))
+	.IsEnabled(false);
+
+	const bool bIsProcedural = IsAnyElementProcedural();
+
+	PhysicsCategory
+		.AddProperty(SettingsHandle->GetChildHandle(GET_MEMBER_NAME_CHECKED(FRigPhysicsSettings, Mass)))
+		.IsEnabled(!bIsProcedural);
+}
+
+FText FRigPhysicsElementDetails::GetSolverNameText() const
+{
+	if(PerElementInfos.Num() > 0)
+	{
+		if(const URigHierarchy* Hierarchy = PerElementInfos[0].GetHierarchy())
+		{
+			TOptional<FText> FirstValue;
+			
+			for(const FPerElementInfo& Info : PerElementInfos)
+			{
+				if(const FRigPhysicsElement* PhysicsElement = Info.Element.Get<FRigPhysicsElement>())
+				{
+					if(const FRigPhysicsSolverDescription* Solver = Hierarchy->FindPhysicsSolver(PhysicsElement->Solver))
+					{
+						const FText SolverName = FText::FromName(Solver->Name);
+						if(FirstValue.IsSet())
+						{
+							if(!FirstValue.GetValue().EqualTo(SolverName))
+							{
+								return ControlRigDetailsMultipleValues;	
+							}
+						}
+						else
+						{
+							FirstValue = SolverName;
+						}
+					}
+				}
+			}
+
+			if(FirstValue.IsSet())
+			{
+				return FirstValue.GetValue();
+			}
+		}
+	}
+	return FText();
+}
 #undef LOCTEXT_NAMESPACE
diff --git a/Engine/Plugins/Animation/ControlRig/Source/ControlRigEditor/Private/ControlRigElementDetails.h b/Engine/Plugins/Animation/ControlRig/Source/ControlRigEditor/Private/ControlRigElementDetails.h
--- a/Engine/Plugins/Animation/ControlRig/Source/ControlRigEditor/Private/ControlRigElementDetails.h
+++ b/Engine/Plugins/Animation/ControlRig/Source/ControlRigEditor/Private/ControlRigElementDetails.h
@@ -921,3 +921,22 @@ protected:
 	TAttribute<bool> EnabledAttribute;
 };
 
+class FRigPhysicsElementDetails : public FRigTransformElementDetails
+{
+public:
+
+	// Makes a new instance of this detail layout class for a specific detail view requesting it
+	static TSharedRef<IDetailCustomization> MakeInstance()
+	{
+		return MakeShareable(new FRigPhysicsElementDetails);
+	}
+
+	/** IDetailCustomization interface */
+	virtual void CustomizeDetails(IDetailLayoutBuilder& DetailBuilder) override;
+
+	void CustomizeSettings(IDetailLayoutBuilder& DetailBuilder);
+
+private:
+
+	FText GetSolverNameText() const;
+};
diff --git a/Engine/Plugins/Animation/ControlRig/Source/ControlRigEditor/Private/EditMode/SControlRigOutliner.cpp b/Engine/Plugins/Animation/ControlRig/Source/ControlRigEditor/Private/EditMode/SControlRigOutliner.cpp
--- a/Engine/Plugins/Animation/ControlRig/Source/ControlRigEditor/Private/EditMode/SControlRigOutliner.cpp
+++ b/Engine/Plugins/Animation/ControlRig/Source/ControlRigEditor/Private/EditMode/SControlRigOutliner.cpp
@@ -521,7 +521,15 @@ bool SMultiRigHierarchyTreeView::AddElement(UControlRig* InControlRig, const FRi
 			}
 			case ERigElementType::Physics:
 			{
-				return false;
+				if (!Settings.bShowPhysics)
+				{
+					return false;
+				}
+				if(CVarControlRigHierarchyEnablePhysics.GetValueOnAnyThread() == false)
+				{
+					return false;
+				}
+				break;
 			}
 			case ERigElementType::Reference:
 			{
diff --git a/Engine/Plugins/Animation/ControlRig/Source/ControlRigEditor/Private/Editor/ControlRigWrapperObject.cpp b/Engine/Plugins/Animation/ControlRig/Source/ControlRigEditor/Private/Editor/ControlRigWrapperObject.cpp
--- a/Engine/Plugins/Animation/ControlRig/Source/ControlRigEditor/Private/Editor/ControlRigWrapperObject.cpp
+++ b/Engine/Plugins/Animation/ControlRig/Source/ControlRigEditor/Private/Editor/ControlRigWrapperObject.cpp
@@ -45,6 +45,10 @@ UClass* UControlRigWrapperObject::GetClassForStruct(UScriptStruct* InStruct, boo
 			{
 				PropertyEditorModule.RegisterCustomClassLayout(WrapperClassName, FOnGetDetailCustomizationInstance::CreateStatic(&FRigSocketElementDetails::MakeInstance));
 			}
+			else if(InStruct == FRigPhysicsElement::StaticStruct())
+			{
+				PropertyEditorModule.RegisterCustomClassLayout(WrapperClassName, FOnGetDetailCustomizationInstance::CreateStatic(&FRigPhysicsElementDetails::MakeInstance));
+			}
 		}
 	}
 #endif
diff --git a/Engine/Plugins/Animation/ControlRig/Source/ControlRigEditor/Private/Editor/SRigHierarchyTreeView.cpp b/Engine/Plugins/Animation/ControlRig/Source/ControlRigEditor/Private/Editor/SRigHierarchyTreeView.cpp
--- a/Engine/Plugins/Animation/ControlRig/Source/ControlRigEditor/Private/Editor/SRigHierarchyTreeView.cpp
+++ b/Engine/Plugins/Animation/ControlRig/Source/ControlRigEditor/Private/Editor/SRigHierarchyTreeView.cpp
@@ -563,7 +563,15 @@ bool SRigHierarchyTreeView::AddElement(const FRigBaseElement* InElement)
 		}
 		case ERigElementType::Physics:
 		{
-			return false;
+			if(!Settings.bShowPhysics)
+			{
+				return false;
+			}
+			if(CVarControlRigHierarchyEnablePhysics.GetValueOnAnyThread() == false)
+			{
+				return false;
+			}
+			break;
 		}
 		case ERigElementType::Reference:
 		{
diff --git a/Engine/Plugins/Experimental/PhysicsControl/PhysicsControl.uplugin b/Engine/Plugins/Experimental/PhysicsControl/PhysicsControl.uplugin
--- a/Engine/Plugins/Experimental/PhysicsControl/PhysicsControl.uplugin
+++ b/Engine/Plugins/Experimental/PhysicsControl/PhysicsControl.uplugin
@@ -15,18 +15,7 @@
 	"IsBetaVersion": false,
 	"IsExperimentalVersion": true,
 	"Installed": false,
-	"Plugins": [
-		{
-			"Name": "ControlRig",
-			"Enabled": true
-		},
-	],
 	"Modules": [
-		{
-			"Name": "ControlRigPhysics",
-			"Type": "Runtime",
-			"LoadingPhase": "Default"
-		},
 		{
 			"Name": "PhysicsControl",
 			"Type": "Runtime",
diff --git a/Engine/Plugins/Experimental/PhysicsControl/Source/ControlRigPhysics/ControlRigPhysics.Build.cs b/Engine/Plugins/Experimental/PhysicsControl/Source/ControlRigPhysics/ControlRigPhysics.Build.cs
deleted file mode 100644
diff --git a/Engine/Plugins/Experimental/PhysicsControl/Source/ControlRigPhysics/Private/ControlRigPhysicsModule.cpp b/Engine/Plugins/Experimental/PhysicsControl/Source/ControlRigPhysics/Private/ControlRigPhysicsModule.cpp
deleted file mode 100644
diff --git a/Engine/Plugins/Experimental/PhysicsControl/Source/ControlRigPhysics/Private/RigPhysicsComponents.cpp b/Engine/Plugins/Experimental/PhysicsControl/Source/ControlRigPhysics/Private/RigPhysicsComponents.cpp
deleted file mode 100644
diff --git a/Engine/Plugins/Experimental/PhysicsControl/Source/ControlRigPhysics/Private/RigPhysicsSolver.cpp b/Engine/Plugins/Experimental/PhysicsControl/Source/ControlRigPhysics/Private/RigPhysicsSolver.cpp
deleted file mode 100644
diff --git a/Engine/Plugins/Experimental/PhysicsControl/Source/ControlRigPhysics/Private/RigPhysicsSolver.h b/Engine/Plugins/Experimental/PhysicsControl/Source/ControlRigPhysics/Private/RigPhysicsSolver.h
deleted file mode 100644
diff --git a/Engine/Plugins/Experimental/PhysicsControl/Source/ControlRigPhysics/Private/RigPhysicsSolverSimulationSpace.cpp b/Engine/Plugins/Experimental/PhysicsControl/Source/ControlRigPhysics/Private/RigPhysicsSolverSimulationSpace.cpp
deleted file mode 100644
diff --git a/Engine/Plugins/Experimental/PhysicsControl/Source/ControlRigPhysics/Private/RigPhysicsSolverVisualization.cpp b/Engine/Plugins/Experimental/PhysicsControl/Source/ControlRigPhysics/Private/RigPhysicsSolverVisualization.cpp
deleted file mode 100644
diff --git a/Engine/Plugins/Experimental/PhysicsControl/Source/ControlRigPhysics/Private/RigUnit_Physics.cpp b/Engine/Plugins/Experimental/PhysicsControl/Source/ControlRigPhysics/Private/RigUnit_Physics.cpp
deleted file mode 100644
diff --git a/Engine/Plugins/Experimental/PhysicsControl/Source/ControlRigPhysics/Public/ControlRigPhysicsModule.h b/Engine/Plugins/Experimental/PhysicsControl/Source/ControlRigPhysics/Public/ControlRigPhysicsModule.h
deleted file mode 100644
diff --git a/Engine/Plugins/Experimental/PhysicsControl/Source/ControlRigPhysics/Public/RigPhysicsComponents.h b/Engine/Plugins/Experimental/PhysicsControl/Source/ControlRigPhysics/Public/RigPhysicsComponents.h
deleted file mode 100644
diff --git a/Engine/Plugins/Experimental/PhysicsControl/Source/ControlRigPhysics/Public/RigUnit_Physics.h b/Engine/Plugins/Experimental/PhysicsControl/Source/ControlRigPhysics/Public/RigUnit_Physics.h
deleted file mode 100644
diff --git a/Engine/Plugins/Experimental/PhysicsControl/Source/PhysicsControl/Private/RigidBodyWithControl.cpp b/Engine/Plugins/Experimental/PhysicsControl/Source/PhysicsControl/Private/RigidBodyWithControl.cpp
--- a/Engine/Plugins/Experimental/PhysicsControl/Source/PhysicsControl/Private/RigidBodyWithControl.cpp
+++ b/Engine/Plugins/Experimental/PhysicsControl/Source/PhysicsControl/Private/RigidBodyWithControl.cpp
@@ -720,47 +720,36 @@ void FAnimNode_RigidBodyWithControl::ApplyControl(FRigidBodyControlRecord& Contr
 	}
 }
 
-//======================================================================================================================
-void UpdateBodyFromModifierData(
-	ImmediatePhysics::FActorHandle*    ActorHandle, 
-	ImmediatePhysics::FSimulation*     PhysicsSimulation,
-	const FPhysicsControlModifierData& ModifierData,
-	const FVector&                     SimSpaceGravity)
-{
-	if (ModifierData.MovementType != EPhysicsMovementType::Default)
-	{
-		// Note that there's an early out if there's no change needed, so this should be OK.
-		PhysicsSimulation->SetIsKinematic(
-			ActorHandle,
-			ModifierData.MovementType != EPhysicsMovementType::Simulated);
-	}
-
-	// Note that the actual kinematic targets will be set separately, since they need to be set
-	// for all kinematics whether or not they were under a modifier.
-
-	// Scale gravity
-	float GravityMultiplier = ModifierData.GravityMultiplier;
-	if (GravityMultiplier != 1 && ActorHandle->IsGravityEnabled())
-	{
-		float Mass = (float) ActorHandle->GetMass();
-		FVector AntiGravityForce = SimSpaceGravity * (-Mass * (1.0f - GravityMultiplier));
-		ActorHandle->AddForce(AntiGravityForce);
-	}
-
-	// Set collision
-	PhysicsSimulation->SetHasCollision(
-		ActorHandle, CollisionEnabledHasPhysics(ModifierData.CollisionType));
-}
-
 //======================================================================================================================
 void FAnimNode_RigidBodyWithControl::ApplyModifier(
 	const FRigidBodyModifierRecord& BodyModifierRecord, 
 	const FVector&                  SimSpaceGravity)
 {
-	if (BodyModifierRecord.ActorHandle && PhysicsSimulation)
+	if (BodyModifierRecord.ActorHandle)
 	{
-		UpdateBodyFromModifierData(
-			BodyModifierRecord.ActorHandle, PhysicsSimulation, BodyModifierRecord.ModifierData, SimSpaceGravity);
+		if (BodyModifierRecord.ModifierData.MovementType != EPhysicsMovementType::Default)
+		{
+			// Note that there's an early out if there's no change needed, so this should be OK.
+			PhysicsSimulation->SetIsKinematic(
+				BodyModifierRecord.ActorHandle,
+				BodyModifierRecord.ModifierData.MovementType != EPhysicsMovementType::Simulated);
+		}
+
+		// Note that the actual kinematic targets will be set separately, since they need to be set
+		// for all kinematics whether or not they were under a modifier.
+
+		// Scale gravity
+		float GravityMultiplier = BodyModifierRecord.ModifierData.GravityMultiplier;
+		if (GravityMultiplier != 1 && BodyModifierRecord.ActorHandle->IsGravityEnabled())
+		{
+			float Mass = (float) BodyModifierRecord.ActorHandle->GetMass();
+			FVector AntiGravityForce = SimSpaceGravity * (-Mass * (1.0f - GravityMultiplier));
+			BodyModifierRecord.ActorHandle->AddForce(AntiGravityForce);
+		}
+
+		// Set collision
+		PhysicsSimulation->SetHasCollision(
+			BodyModifierRecord.ActorHandle, CollisionEnabledHasPhysics(BodyModifierRecord.ModifierData.CollisionType));
 	}
 }
 
diff --git a/Engine/Plugins/Experimental/PhysicsControl/Source/PhysicsControl/Public/AnimNode_RigidBodyWithControl.h b/Engine/Plugins/Experimental/PhysicsControl/Source/PhysicsControl/Public/AnimNode_RigidBodyWithControl.h
--- a/Engine/Plugins/Experimental/PhysicsControl/Source/PhysicsControl/Public/AnimNode_RigidBodyWithControl.h
+++ b/Engine/Plugins/Experimental/PhysicsControl/Source/PhysicsControl/Public/AnimNode_RigidBodyWithControl.h
@@ -50,16 +50,6 @@ enum class MapConstraintsBehaviorType : uint8
 	None UMETA(DisplayName = "Nothing", ToolTip = "Nothing"),
 };
 
-/**
- * A helper function that lets any immediate simulation user update their bodies in the same way that
- * RigidBodyWithControl does.
- */
-PHYSICSCONTROL_API void UpdateBodyFromModifierData(
-	ImmediatePhysics::FActorHandle* ActorHandle,
-	ImmediatePhysics::FSimulation* PhysicsSimulation,
-	const FPhysicsControlModifierData& ModifierData,
-	const FVector& SimSpaceGravity);
-
 /**
  * Controller that simulates physics based on the physics asset of the skeletal mesh component
  */
diff --git a/Engine/Plugins/Runtime/RigVM/Source/RigVM/Private/RigVMCore/RigVMDrawInterface.cpp b/Engine/Plugins/Runtime/RigVM/Source/RigVM/Private/RigVMCore/RigVMDrawInterface.cpp
--- a/Engine/Plugins/Runtime/RigVM/Source/RigVM/Private/RigVMCore/RigVMDrawInterface.cpp
+++ b/Engine/Plugins/Runtime/RigVM/Source/RigVM/Private/RigVMCore/RigVMDrawInterface.cpp
@@ -259,12 +259,13 @@ void FRigVMDrawInterface::DrawArc(const FTransform& WorldOffset, const FTransfor
 	V = Rotation.RotateVector(V);
 	Instruction.Positions.Add(V);
 	float StepAngle = (MaximumAngle - MinimumAngle) / float(Count);
-	if (FMath::Abs<float>(MaximumAngle - MinimumAngle) >= TWO_PI)
+	if (FMath::Abs<float>(MaximumAngle - MinimumAngle) >= PI * 2.f - SMALL_NUMBER)
 	{
-		StepAngle = TWO_PI / float(Count);
+		StepAngle = (PI * 2.f) / float(Count);
+		Count++;
 	}
 	Rotation = FQuat(FVector(0.f, 0.f, 1.f), StepAngle);
-	for(int32 Index=0;Index<Count;Index++)
+	for(int32 Index=1;Index<Count;Index++)
 	{
 		V = Rotation.RotateVector(V);
 		Instruction.Positions.Add(V);
@@ -279,53 +280,7 @@ void FRigVMDrawInterface::DrawCircle(const FTransform& WorldOffset, const FTrans
 	DrawArc(WorldOffset, Transform, Radius, 0, TWO_PI, Color, Thickness, Detail);
 }
 
-void FRigVMDrawInterface::DrawSphere(const FTransform& WorldOffset, const FTransform& Transform, float Radius, 
-	const FLinearColor& Color, float Thickness, int32 Detail)
-{
-	// Circles are drawn in the XY plane
-	FTransform R1(FQuat(FVector(1.f, 0.f, 0.f), HALF_PI));
-	FTransform R2(FQuat(FVector(0.f, 1.f, 0.f), HALF_PI));
-	DrawCircle(WorldOffset, Transform, Radius, Color, Thickness, Detail);
-	DrawCircle(WorldOffset, R1 * Transform, Radius, Color, Thickness, Detail);
-	DrawCircle(WorldOffset, R2 * Transform, Radius, Color, Thickness, Detail);
-}
-
-void FRigVMDrawInterface::DrawHemisphere(const FTransform& WorldOffset, const FTransform& Transform, float Radius,
-	const FLinearColor& Color, float Thickness, int32 Detail)
-{
-	// Circles are drawn in the XY plane
-	FTransform R1(FQuat(FVector(1.f, 0.f, 0.f), HALF_PI));
-	FTransform R2(FQuat(FVector(0.f, 1.f, 0.f), -HALF_PI));
-	DrawCircle(WorldOffset, Transform, Radius, Color, Thickness, Detail);
-	DrawArc(WorldOffset, R1 * Transform, Radius, 0, PI, Color, Thickness, Detail);
-	DrawArc(WorldOffset, R2 * Transform, Radius, -HALF_PI, HALF_PI, Color, Thickness, Detail);
-}
-
-void FRigVMDrawInterface::DrawCapsule(const FTransform& WorldOffset, const FTransform& Transform,
-	float Radius, float Length, const FLinearColor& Color, float Thickness, int32 Detail)
-{
-	FTransform T1(FVector(0, 0, Length * 0.5));
-	FTransform T2(FQuat(FVector(0.f, 1.f, 0.f), PI), FVector(0, 0, -Length * 0.5));
-
-	DrawHemisphere(WorldOffset, T1 * Transform, Radius, Color, Thickness, Detail);
-	DrawHemisphere(WorldOffset, T2 * Transform, Radius, Color, Thickness, Detail);
-
-	DrawLine(WorldOffset, 
-		Transform.TransformPosition(FVector(Radius, 0, Length * 0.5)), 
-		Transform.TransformPosition(FVector(Radius, 0, -Length * 0.5)), Color, Thickness);
-	DrawLine(WorldOffset,
-		Transform.TransformPosition(FVector(-Radius, 0, Length * 0.5)),
-		Transform.TransformPosition(FVector(-Radius, 0, -Length * 0.5)), Color, Thickness);
-	DrawLine(WorldOffset,
-		Transform.TransformPosition(FVector(0, Radius, Length * 0.5)),
-		Transform.TransformPosition(FVector(0, Radius, -Length * 0.5)), Color, Thickness);
-	DrawLine(WorldOffset,
-		Transform.TransformPosition(FVector(0, -Radius, Length * 0.5)),
-		Transform.TransformPosition(FVector(0, -Radius, -Length * 0.5)), Color, Thickness);
-}
-
-
-void FRigVMDrawInterface::DrawCone(const FTransform& WorldOffset, const FTransform& ConeOffset, float Angle1, float Angle2, uint32 NumSides, bool bDrawSideLines, const FLinearColor& SideLineColor, FMaterialRenderProxy* const MaterialRenderProxy, float SideLineThickness)
+void FRigVMDrawInterface::DrawCone(const FTransform& WorldOffset, const FTransform& ConeOffset, float Angle1, float Angle2, uint32 NumSides, bool bDrawSideLines, const FLinearColor& SideLineColor, FMaterialRenderProxy* const MaterialRenderProxy)
 {
 	if (!IsEnabled())
 	{
@@ -356,7 +311,7 @@ void FRigVMDrawInterface::DrawCone(const FTransform& WorldOffset, const FTransfo
 			float Fraction = (float)i / (float)(4);
 			float Azi = 2.f * PI * Fraction;
 			FVector ConeVert = CalcConeVert(Angle1, Angle2, Azi);
-			DrawLine(WorldOffset, ConeOffset.GetLocation(), ConeOffset.TransformPosition(ConeVert), SideLineColor, SideLineThickness);
+			DrawLine(WorldOffset, ConeOffset.GetLocation(), ConeOffset.TransformPosition(ConeVert), SideLineColor, 1.f);
 		}
 	}
 }
diff --git a/Engine/Plugins/Runtime/RigVM/Source/RigVM/Public/RigVMCore/RigVMDrawInterface.h b/Engine/Plugins/Runtime/RigVM/Source/RigVM/Public/RigVMCore/RigVMDrawInterface.h
--- a/Engine/Plugins/Runtime/RigVM/Source/RigVM/Public/RigVMCore/RigVMDrawInterface.h
+++ b/Engine/Plugins/Runtime/RigVM/Source/RigVM/Public/RigVMCore/RigVMDrawInterface.h
@@ -20,19 +20,14 @@ public:
 	void DrawLines(const FTransform& WorldOffset, const TArrayView<const FVector>& Positions, const FLinearColor& Color, float Thickness = 0.f);
 	void DrawLineStrip(const FTransform& WorldOffset, const TArrayView<const FVector>& Positions, const FLinearColor& Color, float Thickness = 0.f);
 	void DrawBox(const FTransform& WorldOffset, const FTransform& Transform, const FLinearColor& Color, float Thickness = 0.f);
-	void DrawSphere(const FTransform& WorldOffset, const FTransform& Transform, float Radius, const FLinearColor& Color, float Thickness = 0.f, int32 Detail = 12);
-	// Hemispheres are drawn so that they extend up along +ve Z
-	void DrawHemisphere(const FTransform& WorldOffset, const FTransform& Transform, float Radius, const FLinearColor& Color, float Thickness = 0.f, int32 Detail = 12);
-	void DrawCapsule(const FTransform& WorldOffset, const FTransform& Transform, float Radius, float Length, const FLinearColor& Color, float Thickness = 0.f, int32 Detail = 12);
 	void DrawAxes(const FTransform& WorldOffset, const FTransform& Transform, float Size, float Thickness = 0.f);
 	void DrawAxes(const FTransform& WorldOffset, const TArrayView<const FTransform>& Transforms, float Size, float Thickness = 0.f);
 	void DrawAxes(const FTransform& WorldOffset, const FTransform& Transform, const FLinearColor& InColor, float Size, float Thickness = 0.f);
 	void DrawAxes(const FTransform& WorldOffset, const TArrayView<const FTransform>& Transforms, const FLinearColor& InColor, float Size, float Thickness = 0.f);
 	void DrawRectangle(const FTransform& WorldOffset, const FTransform& Transform, float Size, const FLinearColor& Color, float Thickness);
 	void DrawArc(const FTransform& WorldOffset, const FTransform& Transform, float Radius, float MinimumAngle, float MaximumAngle, const FLinearColor& Color, float Thickness, int32 Detail);
-	// Circles are drawn in the XY plane
 	void DrawCircle(const FTransform& WorldOffset, const FTransform& Transform, float Radius, const FLinearColor& Color, float Thickness, int32 Detail);
-	void DrawCone(const FTransform& WorldOffset, const FTransform& ConeOffset, float Angle1, float Angle2, uint32 NumSides, bool bDrawSideLines, const FLinearColor& SideLineColor, FMaterialRenderProxy* const MaterialRenderProxy, float SideLineThickness = 1.0f);
+	void DrawCone(const FTransform& WorldOffset, const FTransform& ConeOffset, float Angle1, float Angle2, uint32 NumSides, bool bDrawSideLines, const FLinearColor& SideLineColor, FMaterialRenderProxy* const MaterialRenderProxy);
 	void DrawArrow(const FTransform& WorldOffset, const FVector& Direction, const FVector& Side, const FLinearColor& InColor, float Thickness);
 	void DrawPlane(const FTransform& WorldOffset, const FVector2D& Scale, const FLinearColor& MeshColor, bool bDrawLines, const FLinearColor& LineColor, FMaterialRenderProxy* const MaterialRenderProxy);
 
