CIS FIX merge

diff --git a/Engine/Plugins/Animation/ControlRig/Source/ControlRigEditor/Private/Sequencer/AnimLayers/AnimLayers.cpp b/Engine/Plugins/Animation/ControlRig/Source/ControlRigEditor/Private/Sequencer/AnimLayers/AnimLayers.cpp
--- a/Engine/Plugins/Animation/ControlRig/Source/ControlRigEditor/Private/Sequencer/AnimLayers/AnimLayers.cpp
+++ b/Engine/Plugins/Animation/ControlRig/Source/ControlRigEditor/Private/Sequencer/AnimLayers/AnimLayers.cpp
@@ -1864,10 +1864,10 @@ static void MergeControlRigSections(UMovieSceneControlRigParameterSection* BaseS
 		return (CRSection->GetControlNameMask(ControlElement->GetFName()) == false &&
 				CRSection->GetBlendType().IsValid() && CRSection->GetBlendType() == EMovieSceneBlendType::Override);
 	};
+	if (!BaseSection || !Section)
 	{
-		return (CRSection->GetControlNameMask(ControlElement->GetFName()) == false &&
-				CRSection->GetBlendType().IsValid() && CRSection->GetBlendType() == EMovieSceneBlendType::Override);
-	};
+		return;
+	}
 	TArrayView<FMovieSceneFloatChannel*> BaseFloatChannels = BaseSection->GetChannelProxy().GetChannels<FMovieSceneFloatChannel>();
 	if (BaseFloatChannels.Num() > 0)
 	{
