Control Rig: Refactoring redirector to allow one to many redirection

#rnx

diff --git a/Engine/Plugins/Animation/ControlRig/Source/ControlRig/Private/ControlRig.cpp b/Engine/Plugins/Animation/ControlRig/Source/ControlRig/Private/ControlRig.cpp
--- a/Engine/Plugins/Animation/ControlRig/Source/ControlRig/Private/ControlRig.cpp
+++ b/Engine/Plugins/Animation/ControlRig/Source/ControlRig/Private/ControlRig.cpp
@@ -662,34 +662,39 @@ bool UControlRig::AllConnectorsAreResolved(FString* OutFailureReason, FRigElemen
 		const TArray<FRigModuleConnector> Connectors = GetRigModuleSettings().ExposedConnectors;
 
 		// collect the connection map
-		TMap<FRigElementKey, FRigElementKey> ConnectionMap;
+		TMap<FRigElementKey, FRigElementKeyRedirector::FKeyArray> ConnectionMap;
 		for(const FRigModuleConnector& Connector : Connectors)
 		{
 			if (Connector.Settings.bOptional)
 			{
 				continue;
 			}
-			
+
+			FRigElementKeyRedirector::FKeyArray TargetKeys;
 			const FRigElementKey ConnectorKey(*Connector.Name, ERigElementType::Connector);
-			if(const FCachedRigElement* Cache = ElementKeyRedirector.Find(ConnectorKey))
+			if(const FRigElementKeyRedirector::FCachedKeyArray* Cache = ElementKeyRedirector.Find(ConnectorKey))
 			{
-				if(const_cast<FCachedRigElement*>(Cache)->UpdateCache(Hierarchy))
+				for(const FCachedRigElement& CachedRigElement : *Cache)
 				{
-					ConnectionMap.Add(ConnectorKey, Cache->GetKey());
-				}
-				else
-				{
-					if(OutFailureReason)
+					if(const_cast<FCachedRigElement*>(&CachedRigElement)->UpdateCache(Hierarchy))
 					{
-						static constexpr TCHAR Format[] = TEXT("Connector '%s%s' has invalid target '%s'.");
-						*OutFailureReason = FString::Printf(Format, *GetRigModuleNameSpace(), *Connector.Name, *Cache->GetKey().ToString());
+						TargetKeys.Add(CachedRigElement.GetKey());
 					}
-					if(OutConnector)
+					else
 					{
-						*OutConnector = ConnectorKey;
+						if(OutFailureReason)
+						{
+							static constexpr TCHAR Format[] = TEXT("Connector '%s%s' has invalid target '%s'.");
+							*OutFailureReason = FString::Printf(Format, *GetRigModuleNameSpace(), *Connector.Name, *CachedRigElement.GetKey().ToString());
+						}
+						if(OutConnector)
+						{
+							*OutConnector = ConnectorKey;
+						}
+						return false;
 					}
-					return false;
 				}
+				ConnectionMap.Add(ConnectorKey, TargetKeys);
 			}
 			else
 			{
diff --git a/Engine/Plugins/Animation/ControlRig/Source/ControlRig/Private/ModularRig.cpp b/Engine/Plugins/Animation/ControlRig/Source/ControlRig/Private/ModularRig.cpp
--- a/Engine/Plugins/Animation/ControlRig/Source/ControlRig/Private/ModularRig.cpp
+++ b/Engine/Plugins/Animation/ControlRig/Source/ControlRig/Private/ModularRig.cpp
@@ -731,7 +731,7 @@ void UModularRig::RunPostConstructionEvent()
 }
 
 FRigModuleInstance* UModularRig::AddModuleInstance(const FName& InModuleName, TSubclassOf<UControlRig> InModuleClass, const FRigModuleInstance* InParent,
-                                                   const TMap<FRigElementKey, FRigElementKey>& InConnectionMap, const TMap<FName, FString>& InVariableDefaultValues) 
+                                                   const FRigElementKeyRedirector::FKeyMap& InConnectionMap, const TMap<FName, FString>& InVariableDefaultValues) 
 {
 	// Make sure there are no name clashes
 	if (InParent)
@@ -913,7 +913,10 @@ void UModularRig::RecomputeShortestDisplayPathCache() const
 	for(const FModularRigSingleConnection& Connection : ModularRigModel.Connections)
 	{
 		AllKeys.AddUnique(Connection.Connector);
-		AllKeys.AddUnique(Connection.Target);
+		for(const FRigElementKey& Target : Connection.Targets)
+		{
+			AllKeys.AddUnique(Target);
+		}
 	}
 
 	auto GetNameForElement = [Hierarchy](const FRigElementKey& InElementKey)
diff --git a/Engine/Plugins/Animation/ControlRig/Source/ControlRig/Private/ModularRigController.cpp b/Engine/Plugins/Animation/ControlRig/Source/ControlRig/Private/ModularRigController.cpp
--- a/Engine/Plugins/Animation/ControlRig/Source/ControlRig/Private/ModularRigController.cpp
+++ b/Engine/Plugins/Animation/ControlRig/Source/ControlRig/Private/ModularRigController.cpp
@@ -304,7 +304,7 @@ bool UModularRigController::ConnectConnectorToElement(const FRigElementKey& InCo
 		DisconnectConnector_Internal(InConnectorKey, false, &PreviousConnections, bSetupUndo);
 	}
 	
-	Model->Connections.AddConnection(InConnectorKey, InTargetKey);
+	Model->Connections.AddConnection(InConnectorKey, {InTargetKey});
 
 	// restore previous connections if possible
 	for(const TPair<FRigElementKey, FRigElementKey>& PreviousConnection : PreviousConnections)
@@ -520,24 +520,29 @@ TArray<FRigElementKey> UModularRigController::DisconnectCyclicConnectors(bool bS
 	for (const FModularRigSingleConnection& Connection : Model->Connections)
 	{
 		const FString ConnectorModulePath = Hierarchy->GetModulePath(Connection.Connector);
-		const FString TargetModulePath = Hierarchy->GetModulePath(Connection.Target);
 
-		// targets in the base hierarchy are always allowed
-		if(TargetModulePath.IsEmpty())
+		for(const FRigElementKey& Target : Connection.Targets)
 		{
-			continue;
-		}
+			const FString TargetModulePath = Hierarchy->GetModulePath(Target);
 
-		const FRigModuleReference* ConnectorModule = Model->FindModule(ConnectorModulePath);
-		const FRigModuleReference* TargetModule = Model->FindModule(TargetModulePath);
-		if(ConnectorModule == nullptr || TargetModule == nullptr)
-		{
-			continue;
-		}
+			// targets in the base hierarchy are always allowed
+			if(TargetModulePath.IsEmpty())
+			{
+				continue;
+			}
 
-		if(!Model->IsModuleParentedTo(ConnectorModule, TargetModule))
-		{
-			ConnectorsToDisconnect.Add(Connection.Connector);
+			const FRigModuleReference* ConnectorModule = Model->FindModule(ConnectorModulePath);
+			const FRigModuleReference* TargetModule = Model->FindModule(TargetModulePath);
+			if(ConnectorModule == nullptr || TargetModule == nullptr)
+			{
+				continue;
+			}
+
+			if(!Model->IsModuleParentedTo(ConnectorModule, TargetModule))
+			{
+				ConnectorsToDisconnect.Add(Connection.Connector);
+				break;
+			}
 		}
 	}
 
@@ -666,7 +671,7 @@ bool UModularRigController::AutoConnectSecondaryConnectors(const TArray<FRigElem
 				bool bFoundMatch = false;
 				if (RuleResults.GetMatches().Num() == 1)
 				{
-					Model->Connections.AddConnection(ConnectorKey, RuleResults.GetMatches()[0].GetKey());
+					Model->Connections.AddConnection(ConnectorKey, {RuleResults.GetMatches()[0].GetKey()});
 					Notify(EModularRigNotification::ConnectionChanged, Module);
 					bFoundMatch = true;
 				}
@@ -676,7 +681,7 @@ bool UModularRigController::AutoConnectSecondaryConnectors(const TArray<FRigElem
 					{
 						if (Result.GetState() == ERigElementResolveState::DefaultTarget)
 						{
-							Model->Connections.AddConnection(ConnectorKey, Result.GetKey());
+							Model->Connections.AddConnection(ConnectorKey, {Result.GetKey()});
 							Notify(EModularRigNotification::ConnectionChanged, Module);
 							bFoundMatch = true;
 							break;
@@ -1114,11 +1119,15 @@ bool UModularRigController::DeleteModule(const FString& InModulePath, bool bSetu
 				ToRemove.Add(Connection.Connector);
 			}
 
-			FString TargetModulePath, TargetName;
-			(void)URigHierarchy::SplitNameSpace(Connection.Target.Name.ToString(), &TargetModulePath, &TargetName);
-			if (TargetModulePath == InModulePath)
+			for(const FRigElementKey& Target : Connection.Targets)
 			{
-				ToRemove.Add(Connection.Connector);
+				FString TargetModulePath, TargetName;
+				(void)URigHierarchy::SplitNameSpace(Target.Name.ToString(), &TargetModulePath, &TargetName);
+				if (TargetModulePath == InModulePath)
+				{
+					ToRemove.Add(Connection.Connector);
+					break;
+				}
 			}
 		}
 		for (FRigElementKey& KeyToRemove : ToRemove)
@@ -1219,9 +1228,12 @@ FString UModularRigController::RenameModule(const FString& InModulePath, const F
 			{
 				Connection.Connector.Name = *FString::Printf(TEXT("%s%s"), *NewNamespace, *Connection.Connector.Name.ToString().RightChop(OldNamespace.Len()));
 			}
-			if (Connection.Target.Name.ToString().StartsWith(OldNamespace, ESearchCase::IgnoreCase))
+			for(FRigElementKey& Target : Connection.Targets)
 			{
-				Connection.Target.Name = *FString::Printf(TEXT("%s%s"), *NewNamespace, *Connection.Target.Name.ToString().RightChop(OldNamespace.Len()));
+				if (Target.Name.ToString().StartsWith(OldNamespace, ESearchCase::IgnoreCase))
+				{
+					Target.Name = *FString::Printf(TEXT("%s%s"), *NewNamespace, *Target.Name.ToString().RightChop(OldNamespace.Len()));
+				}
 			}
 		}
 		Model->Connections.UpdateFromConnectionList();
@@ -1351,10 +1363,14 @@ FString UModularRigController::ReparentModule(const FString& InModulePath, const
 			{
 				Connection.Connector.Name = *FString::Printf(TEXT("%s%s"), *NewPath, *Connection.Connector.Name.ToString().RightChop(OldPath.Len()));
 			}
-			const TPair<FRigName, FRigName> TargetNamespace = URigHierarchy::SplitNameSpace(Connection.Target.Name);
-			if (TargetNamespace.Key.ToString().StartsWith(OldPath, ESearchCase::IgnoreCase))
+
+			for(FRigElementKey& Target : Connection.Targets)
 			{
-				Connection.Target.Name = *FString::Printf(TEXT("%s%s"), *NewPath, *Connection.Target.Name.ToString().RightChop(OldPath.Len()));
+				const TPair<FRigName, FRigName> TargetNamespace = URigHierarchy::SplitNameSpace(Target.Name);
+				if (TargetNamespace.Key.ToString().StartsWith(OldPath, ESearchCase::IgnoreCase))
+				{
+					Target.Name = *FString::Printf(TEXT("%s%s"), *NewPath, *Target.Name.ToString().RightChop(OldPath.Len()));
+				}
 			}
 		}
 		Model->Connections.UpdateFromConnectionList();
@@ -1425,7 +1441,7 @@ FString UModularRigController::MirrorModule(const FString& InModulePath, const F
 	}
 
 	// Before any changes, gather all the information we need from the OriginalModule, as the pointer might become invalid afterwards
-	const TMap<FRigElementKey, FRigElementKey> OriginalConnectionMap = Model->Connections.GetModuleConnectionMap(InModulePath);
+	const FRigElementKeyRedirector::FKeyMap OriginalConnectionMap = Model->Connections.GetModuleConnectionMap(InModulePath);
 	const TMap<FName, FString> OriginalBindings = OriginalModule->Bindings;
 	const TSubclassOf<UControlRig> OriginalClass = OriginalModule->Class.Get();
 	const FString OriginalParentPath = OriginalModule->ParentPath;
@@ -1440,15 +1456,18 @@ FString UModularRigController::MirrorModule(const FString& InModulePath, const F
 		return FString();
 	}
 
-	for (const TPair<FRigElementKey, FRigElementKey>& Pair : OriginalConnectionMap)
+	for (const  FRigElementKeyRedirector::FKeyPair& Pair : OriginalConnectionMap)
 	{
-		FString OriginalTargetPath = Pair.Value.Name.ToString();
-		FString NewTargetPath = OriginalTargetPath.Replace(*InSettings.SearchString, *InSettings.ReplaceString, ESearchCase::CaseSensitive);
-		FRigElementKey NewTargetKey(*NewTargetPath, Pair.Value.Type);
+		for(const FRigElementKey& Target : Pair.Value)
+		{
+			FString OriginalTargetPath = Target.Name.ToString();
+			FString NewTargetPath = OriginalTargetPath.Replace(*InSettings.SearchString, *InSettings.ReplaceString, ESearchCase::CaseSensitive);
+			FRigElementKey NewTargetKey(*NewTargetPath, Target.Type);
 
-		FString NewConnectorPath = URigHierarchy::JoinNameSpace(NewModulePath, Pair.Key.Name.ToString());
-		FRigElementKey NewConnectorKey(*NewConnectorPath, ERigElementType::Connector);
-		ConnectConnectorToElement(NewConnectorKey, NewTargetKey, bSetupUndo, false, false);
+			FString NewConnectorPath = URigHierarchy::JoinNameSpace(NewModulePath, Pair.Key.Name.ToString());
+			FRigElementKey NewConnectorKey(*NewConnectorPath, ERigElementType::Connector);
+			ConnectConnectorToElement(NewConnectorKey, NewTargetKey, bSetupUndo, false, false);
+		}
 
 		// Path might change after connecting
 		NewModulePath = NewModule->GetPath();
@@ -1674,10 +1693,14 @@ bool UModularRigController::SwapModuleClass(const FString& InModulePath, TSubcla
 					continue;
 				}
 
-				FText ErrorMessage;
-				if (!CanConnectConnectorToElement(Connection.Connector, Connection.Target, ErrorMessage))
+				for(const FRigElementKey& Target : Connection.Targets)
 				{
-					ConnectionsToRemove.Add(Connection.Connector);
+					FText ErrorMessage;
+					if (!CanConnectConnectorToElement(Connection.Connector, Target, ErrorMessage))
+					{
+						ConnectionsToRemove.Add(Connection.Connector);
+						break;
+					}
 				}
 			}
 		}
diff --git a/Engine/Plugins/Animation/ControlRig/Source/ControlRig/Private/ModularRigModel.cpp b/Engine/Plugins/Animation/ControlRig/Source/ControlRig/Private/ModularRigModel.cpp
--- a/Engine/Plugins/Animation/ControlRig/Source/ControlRig/Private/ModularRigModel.cpp
+++ b/Engine/Plugins/Animation/ControlRig/Source/ControlRig/Private/ModularRigModel.cpp
@@ -76,9 +76,16 @@ TArray<const FRigConnectorElement*> FRigModuleReference::FindConnectors(const UR
 	return Connectors;
 }
 
-TMap<FRigElementKey, FRigElementKey> FModularRigConnections::GetModuleConnectionMap(const FString& InModulePath) const
+FRigElementKeyRedirector::FKeyArray FModularRigSingleConnection::GetTargetArray() const
 {
-	TMap<FRigElementKey, FRigElementKey> Result;
+	FRigElementKeyRedirector::FKeyArray Result;
+	Result.Append(Targets);
+	return Result;
+}
+
+FRigElementKeyRedirector::FKeyMap FModularRigConnections::GetModuleConnectionMap(const FString& InModulePath) const
+{
+	FRigElementKeyRedirector::FKeyMap Result;
 	for (const FModularRigSingleConnection& Connection : ConnectionList)
 	{
 		FString Path, Name;
@@ -87,12 +94,28 @@ TMap<FRigElementKey, FRigElementKey> FModularRigConnections::GetModuleConnection
 		// Exactly the same path (do not return connectors from child modules)
 		if (Path == InModulePath)
 		{
-			Result.Add(FRigElementKey(*Name, ERigElementType::Connector), Connection.Target);
+			Result.Add(FRigElementKey(*Name, ERigElementType::Connector), Connection.GetTargetArray());
 		}
 	}
 	return Result;
 }
 
+void FModularRigConnections::PatchOnLoad()
+{
+	for(FModularRigSingleConnection& Connection : ConnectionList)
+	{
+		if(Connection.Targets.IsEmpty() || !Connection.Targets[0].IsValid())
+		{
+			if(Connection.Target_DEPRECATED.IsValid())
+			{
+				Connection.Targets.Reset();
+				Connection.Targets.Add(Connection.Target_DEPRECATED);
+				Connection.Target_DEPRECATED.Reset();
+			}
+		}
+	}
+}
+
 void FModularRigModel::PatchModelsOnLoad()
 {
 	if (Connections.IsEmpty())
@@ -104,12 +127,13 @@ void FModularRigModel::PatchModelsOnLoad()
 			{
 				const FString ConnectorPath = FString::Printf(TEXT("%s%s"), *ModuleNamespace, *Connection.Key.Name.ToString());
 				FRigElementKey ConnectorKey(*ConnectorPath, ERigElementType::Connector);
-				Connections.AddConnection(ConnectorKey, Connection.Value);
+				Connections.AddConnection(ConnectorKey, {Connection.Value});
 			}
 			return true;
 		});
 	}
 
+	Connections.PatchOnLoad();
 	Connections.UpdateFromConnectionList();
 }
 
diff --git a/Engine/Plugins/Animation/ControlRig/Source/ControlRig/Private/Rigs/RigConnectionRules.cpp b/Engine/Plugins/Animation/ControlRig/Source/ControlRig/Private/Rigs/RigConnectionRules.cpp
--- a/Engine/Plugins/Animation/ControlRig/Source/ControlRig/Private/Rigs/RigConnectionRules.cpp
+++ b/Engine/Plugins/Animation/ControlRig/Source/ControlRig/Private/Rigs/RigConnectionRules.cpp
@@ -168,9 +168,12 @@ const FRigTransformElement* FRigConnectionRuleInput::ResolveConnector(const FRig
 {
 	check(Redirector);
 
-	if(const FCachedRigElement* Target = Redirector->Find(InConnector->GetKey()))
+	if(const FRigElementKeyRedirector::FCachedKeyArray* Cache = Redirector->Find(InConnector->GetKey()))
 	{
-		return Cast<FRigTransformElement>(Target->GetElement());
+		if(!Cache->IsEmpty())
+		{
+			return Cast<FRigTransformElement>(Cache->operator[](0).GetElement());
+		}
 	}
 
 	if(OutErrorMessage)
diff --git a/Engine/Plugins/Animation/ControlRig/Source/ControlRig/Private/Rigs/RigHierarchy.cpp b/Engine/Plugins/Animation/ControlRig/Source/ControlRig/Private/Rigs/RigHierarchy.cpp
--- a/Engine/Plugins/Animation/ControlRig/Source/ControlRig/Private/Rigs/RigHierarchy.cpp
+++ b/Engine/Plugins/Animation/ControlRig/Source/ControlRig/Private/Rigs/RigHierarchy.cpp
@@ -3463,9 +3463,12 @@ const FRigElementKey& URigHierarchy::GetResolvedTarget(const FRigElementKey& InC
 	{
 		if(ElementKeyRedirector)
 		{
-			if(const FCachedRigElement* Target = ElementKeyRedirector->Find(InConnectorKey))
+			if(const FRigElementKeyRedirector::FCachedKeyArray* Targets = ElementKeyRedirector->Find(InConnectorKey))
 			{
-				return Target->GetKey();
+				if(!Targets->IsEmpty())
+				{
+					return (*Targets)[0].GetKey();
+				}
 			}
 		}
 	}
diff --git a/Engine/Plugins/Animation/ControlRig/Source/ControlRig/Private/Rigs/RigHierarchyCache.cpp b/Engine/Plugins/Animation/ControlRig/Source/ControlRig/Private/Rigs/RigHierarchyCache.cpp
--- a/Engine/Plugins/Animation/ControlRig/Source/ControlRig/Private/Rigs/RigHierarchyCache.cpp
+++ b/Engine/Plugins/Animation/ControlRig/Source/ControlRig/Private/Rigs/RigHierarchyCache.cpp
@@ -149,20 +149,34 @@ bool FCachedRigComponent::IsIdentical(const FRigElementKey& InKey, const FName&
 // FRigElementKeyRedirector
 ////////////////////////////////////////////////////////////////////////////////
 
-FRigElementKeyRedirector::FRigElementKeyRedirector(const TMap<FRigElementKey, FRigElementKey>& InMap, const URigHierarchy* InHierarchy)
+FRigElementKeyRedirector::FRigElementKeyRedirector(const FKeyMap& InMap, const URigHierarchy* InHierarchy)
 {
 	check(InHierarchy);
 	InternalKeyToExternalKey.Reserve(InMap.Num());
 	ExternalKeys.Reserve(InMap.Num());
 
 	Hash = 0;
-	for(const TPair<FRigElementKey, FRigElementKey>& Pair : InMap)
+	for(const FKeyPair& Pair : InMap)
 	{
 		check(Pair.Key.IsValid());
 		Add(Pair.Key, Pair.Value, InHierarchy);
 	}
 }
 
+FRigElementKeyRedirector::FRigElementKeyRedirector(const TMap<FRigElementKey, FRigElementKeyCollection>& InMap, const URigHierarchy* InHierarchy)
+{
+	check(InHierarchy);
+	InternalKeyToExternalKey.Reserve(InMap.Num());
+	ExternalKeys.Reserve(InMap.Num());
+
+	Hash = 0;
+	for(const TPair<FRigElementKey, FRigElementKeyCollection>& Pair : InMap)
+	{
+		check(Pair.Key.IsValid());
+		Add(Pair.Key, Convert(Pair.Value.Keys), InHierarchy);
+	}
+}
+
 FRigElementKeyRedirector::FRigElementKeyRedirector(const FRigElementKeyRedirector& InOther, const URigHierarchy* InHierarchy)
 {
 	check(InHierarchy);
@@ -170,10 +184,10 @@ FRigElementKeyRedirector::FRigElementKeyRedirector(const FRigElementKeyRedirecto
 	ExternalKeys.Reserve(InOther.ExternalKeys.Num());
 
 	Hash = 0;
-	for(const TPair<FRigElementKey, FCachedRigElement>& Pair : InOther.InternalKeyToExternalKey)
+	for(const FCachedKeyPair& Pair : InOther.InternalKeyToExternalKey)
 	{
 		check(Pair.Key.IsValid());
-		Add(Pair.Key, Pair.Value.GetKey(), InHierarchy);
+		Add(Pair.Key, Convert(Pair.Value), InHierarchy);
 	}
 }
 
@@ -187,35 +201,68 @@ FRigElementKeyRedirector::FRigElementKeyRedirector(const FModularRigConnections&
 	for(const FModularRigSingleConnection& Connection : InOther)
 	{
 		check(Connection.Connector.IsValid());
-		check(Connection.Target.IsValid());
-		Add(Connection.Connector, Connection.Target, InHierarchy);
+		check(IsValid(Connection.Targets));
+		Add(Connection.Connector, Convert(Connection.Targets), InHierarchy);
 	}
 }
 
 const FRigElementKey* FRigElementKeyRedirector::FindReverse(const FRigElementKey& InKey) const
 {
-	for(const TPair<FRigElementKey, FCachedRigElement>& Pair : InternalKeyToExternalKey)
+	for(const FCachedKeyPair& Pair : InternalKeyToExternalKey)
 	{
-		if(Pair.Value.GetKey() == InKey)
+		for(const FCachedRigElement& Target : Pair.Value)
 		{
-			return &Pair.Key;
+			if(Target.GetKey() == InKey)
+			{
+				return &Pair.Key;
+			}
 		}
 	}
 	return nullptr;
 }
 
-void FRigElementKeyRedirector::Add(const FRigElementKey& InSource, const FRigElementKey& InTarget, const URigHierarchy* InHierarchy)
+void FRigElementKeyRedirector::Add(const FRigElementKey& InSource, const FKeyArray& InTargets, const URigHierarchy* InHierarchy)
 {
-	if(!InSource.IsValid() || InSource == InTarget)
+	if(!InSource.IsValid() || InTargets.IsEmpty() || InTargets.Contains(InSource))
 	{
 		return;
 	}
 
-	if (InTarget.IsValid())
+	if (IsValid(InTargets))
+	{
+		InternalKeyToExternalKey.Add(InSource, Convert(InTargets, InHierarchy, true));
+	}
+	
+	ExternalKeys.Add(InSource, InTargets);
+	Hash = HashCombine(Hash, HashCombine(GetTypeHash(InSource), GetTypeHash(InTargets)));
+}
+
+FRigElementKeyRedirector::FKeyArray FRigElementKeyRedirector::Convert(const FCachedKeyArray& InCachedKeys)
+{
+	FKeyArray Result;
+	Result.Reserve(InCachedKeys.Num());
+	for(int32 Index = 0; Index < InCachedKeys.Num(); Index++)
 	{
-		InternalKeyToExternalKey.Add(InSource, FCachedRigElement(InTarget, InHierarchy, true));
+		Result.Add(InCachedKeys[Index].GetKey());
 	}
-	ExternalKeys.Add(InSource, InTarget);
-	Hash = HashCombine(Hash, HashCombine(GetTypeHash(InSource), GetTypeHash(InTarget)));
+	return Result;
 }
 
+FRigElementKeyRedirector::FKeyArray FRigElementKeyRedirector::Convert(const TArray<FRigElementKey>& InKeys)
+{
+	FKeyArray Result;
+	Result.Append(InKeys);
+	return Result;
+}
+
+FRigElementKeyRedirector::FCachedKeyArray FRigElementKeyRedirector::Convert(const FKeyArray& InKeys, const URigHierarchy* InHierarchy,
+	bool bForceStoreKey)
+{
+	FCachedKeyArray Result;
+	Result.Reserve(InKeys.Num());
+	for(const FRigElementKey& Key : InKeys)
+	{
+		Result.Emplace(Key, InHierarchy, bForceStoreKey);
+	}
+	return Result;
+}
\ No newline at end of file
diff --git a/Engine/Plugins/Animation/ControlRig/Source/ControlRig/Private/Rigs/RigHierarchyController.cpp b/Engine/Plugins/Animation/ControlRig/Source/ControlRig/Private/Rigs/RigHierarchyController.cpp
--- a/Engine/Plugins/Animation/ControlRig/Source/ControlRig/Private/Rigs/RigHierarchyController.cpp
+++ b/Engine/Plugins/Animation/ControlRig/Source/ControlRig/Private/Rigs/RigHierarchyController.cpp
@@ -107,11 +107,34 @@ bool URigHierarchyController::SelectElement(FRigElementKey InKey, bool bSelect,
 	FRigElementKey Key = InKey;
 	if(Hierarchy->ElementKeyRedirector)
 	{
-		if(const FCachedRigElement* Cache = Hierarchy->ElementKeyRedirector->Find(Key))
+		if(const FRigElementKeyRedirector::FCachedKeyArray* Cache = Hierarchy->ElementKeyRedirector->Find(Key))
 		{
-			if(const_cast<FCachedRigElement*>(Cache)->UpdateCache(Hierarchy))
+			if(!Cache->IsEmpty())
 			{
-				Key = Cache->GetKey();
+				FRigElementKeyRedirector::FKeyArray Keys;
+				for(const FCachedRigElement& CachedRigElement : *Cache)
+				{
+					if(const_cast<FCachedRigElement*>(&CachedRigElement)->UpdateCache(Hierarchy))
+					{
+						Keys.Add(CachedRigElement.GetKey());
+					}
+				}
+
+				if(Keys.Num() == 1)
+				{
+					Key = Keys[0];
+				}
+				else
+				{
+					for(const FRigElementKey& RedirectedKey : Keys)
+					{
+						if(!SelectElement(RedirectedKey, bSelect, false))
+						{
+							return false;
+						}
+					}
+					return true;
+				}
 			}
 		}
 	}
@@ -2435,7 +2458,7 @@ int32 URigHierarchyController::AddElement(FRigBaseElement* InElementToAdd, FRigB
 
 				if(Hierarchy->ElementKeyRedirector)
 				{
-					Hierarchy->ElementKeyRedirector->Add(FRigElementKey(DesiredName.GetFName(), InElementToAdd->Key.Type), InElementToAdd->Key, Hierarchy);
+					Hierarchy->ElementKeyRedirector->Add(FRigElementKey(DesiredName.GetFName(), InElementToAdd->Key.Type), {InElementToAdd->Key}, Hierarchy);
 				}
 			}
 
diff --git a/Engine/Plugins/Animation/ControlRig/Source/ControlRig/Public/ModularRig.h b/Engine/Plugins/Animation/ControlRig/Source/ControlRig/Public/ModularRig.h
--- a/Engine/Plugins/Animation/ControlRig/Source/ControlRig/Public/ModularRig.h
+++ b/Engine/Plugins/Animation/ControlRig/Source/ControlRig/Public/ModularRig.h
@@ -271,7 +271,7 @@ protected:
 private:
 
 	/** Adds a module to the rig*/
-	FRigModuleInstance* AddModuleInstance(const FName& InModuleName, TSubclassOf<UControlRig> InModuleClass, const FRigModuleInstance* InParent, const TMap<FRigElementKey, FRigElementKey>& InConnectionMap, const TMap<FName, FString>& InVariableDefaultValues);
+	FRigModuleInstance* AddModuleInstance(const FName& InModuleName, TSubclassOf<UControlRig> InModuleClass, const FRigModuleInstance* InParent, const FRigElementKeyRedirector::FKeyMap& InConnectionMap, const TMap<FName, FString>& InVariableDefaultValues);
 
 	/** Updates the module's variable bindings */
 	bool SetModuleVariableBindings(const FString& InModulePath, const TMap<FName, FString>& InVariableBindings);
diff --git a/Engine/Plugins/Animation/ControlRig/Source/ControlRig/Public/ModularRigModel.h b/Engine/Plugins/Animation/ControlRig/Source/ControlRig/Public/ModularRigModel.h
--- a/Engine/Plugins/Animation/ControlRig/Source/ControlRig/Public/ModularRigModel.h
+++ b/Engine/Plugins/Animation/ControlRig/Source/ControlRig/Public/ModularRigModel.h
@@ -126,16 +126,24 @@ struct FModularRigSingleConnection
 	GENERATED_BODY()
 
 	FModularRigSingleConnection()
-		: Connector(FRigElementKey()), Target(FRigElementKey()) {}
+		: Connector(FRigElementKey()), Targets({FRigElementKey()}) {}
 	
 	FModularRigSingleConnection(const FRigElementKey& InConnector, const FRigElementKey& InTarget)
-		: Connector(InConnector), Target(InTarget) {}
+		: Connector(InConnector), Targets({InTarget}) {}
+
+	FModularRigSingleConnection(const FRigElementKey& InConnector, const FRigElementKeyRedirector::FKeyArray& InTargets)
+		: Connector(InConnector), Targets(InTargets) {}
 
 	UPROPERTY()
 	FRigElementKey Connector;
 
 	UPROPERTY()
-	FRigElementKey Target;
+	FRigElementKey Target_DEPRECATED;
+
+	UPROPERTY()
+	TArray<FRigElementKey> Targets;
+
+	FRigElementKeyRedirector::FKeyArray GetTargetArray() const;
 };
 
 USTRUCT(BlueprintType)
@@ -172,18 +180,24 @@ public:
 		ReverseConnectionMap.Reset();
 		for (const FModularRigSingleConnection& Connection : ConnectionList)
 		{
-			TArray<FRigElementKey>& Connectors = ReverseConnectionMap.FindOrAdd(Connection.Target);
-			Connectors.AddUnique(Connection.Connector);
+			for(const FRigElementKey& Target : Connection.Targets)
+			{
+				TArray<FRigElementKey>& Connectors = ReverseConnectionMap.FindOrAdd(Target);
+				Connectors.AddUnique(Connection.Connector);
+			}
 		}
 	}
 
-	void AddConnection(const FRigElementKey& Connector, const FRigElementKey& Target)
+	void AddConnection(const FRigElementKey& Connector, const FRigElementKeyRedirector::FKeyArray& Targets)
 	{
 		// Remove any existing connection
 		RemoveConnection(Connector);
 
-		ConnectionList.Add(FModularRigSingleConnection(Connector, Target));
-		ReverseConnectionMap.FindOrAdd(Target).AddUnique(Connector);
+		ConnectionList.Add(FModularRigSingleConnection(Connector, Targets));
+		for(const FRigElementKey& Target : Targets)
+		{
+			ReverseConnectionMap.FindOrAdd(Target).AddUnique(Connector);
+		}
 	}
 
 	void RemoveConnection(const FRigElementKey& Connector)
@@ -191,15 +205,19 @@ public:
 		int32 ExistingIndex = FindConnectionIndex(Connector);
 		if (ConnectionList.IsValidIndex(ExistingIndex))
 		{
-			if (TArray<FRigElementKey>* Connectors = ReverseConnectionMap.Find(ConnectionList[ExistingIndex].Target))
+			const TArray<FRigElementKey>& Targets = ConnectionList[ExistingIndex].Targets;
+			for(const FRigElementKey& Target : Targets)
 			{
-				*Connectors = Connectors->FilterByPredicate([Connector](const FRigElementKey& TargetConnector)
-				{
-					return TargetConnector != Connector;
-				});
-				if (Connectors->IsEmpty())
+				if (TArray<FRigElementKey>* Connectors = ReverseConnectionMap.Find(Target))
 				{
-					ReverseConnectionMap.Remove(ConnectionList[ExistingIndex].Target);
+					*Connectors = Connectors->FilterByPredicate([Connector](const FRigElementKey& TargetConnector)
+					{
+						return TargetConnector != Connector;
+					});
+					if (Connectors->IsEmpty())
+					{
+						ReverseConnectionMap.Remove(Target);
+					}
 				}
 			}
 			ConnectionList.RemoveAt(ExistingIndex);
@@ -219,12 +237,25 @@ public:
 		int32 Index = FindConnectionIndex(InConnectorKey);
 		if (ConnectionList.IsValidIndex(Index))
 		{
-			return ConnectionList[Index].Target;
+			check(ConnectionList[Index].Targets.Num() == 1);
+			return ConnectionList[Index].Targets[0];
 		}
 		static const FRigElementKey EmptyKey;
 		return EmptyKey;
 	}
 
+	TArray<FRigElementKey> FindTargetsFromConnector(const FRigElementKey& InConnectorKey) const
+	{
+		int32 Index = FindConnectionIndex(InConnectorKey);
+		if (ConnectionList.IsValidIndex(Index))
+		{
+			check(!ConnectionList[Index].Targets.IsEmpty());
+			return ConnectionList[Index].Targets;
+		}
+		static const TArray<FRigElementKey> EmptyKeys;
+		return EmptyKeys;
+	}
+
 	const TArray<FRigElementKey>& FindConnectorsFromTarget(const FRigElementKey& InTargetKey) const
 	{
 		if(const TArray<FRigElementKey>* Connectors = ReverseConnectionMap.Find(InTargetKey))
@@ -251,7 +282,9 @@ public:
 	}
 
 	/** Gets the connection map for a single module, where the connectors are identified without its namespace*/
-	TMap<FRigElementKey, FRigElementKey> GetModuleConnectionMap(const FString& InModulePath) const;
+	FRigElementKeyRedirector::FKeyMap GetModuleConnectionMap(const FString& InModulePath) const;
+
+	void PatchOnLoad();
 };
 
 // A management struct containing all modules in the rig
diff --git a/Engine/Plugins/Animation/ControlRig/Source/ControlRig/Public/Rigs/RigHierarchy.h b/Engine/Plugins/Animation/ControlRig/Source/ControlRig/Public/Rigs/RigHierarchy.h
--- a/Engine/Plugins/Animation/ControlRig/Source/ControlRig/Public/Rigs/RigHierarchy.h
+++ b/Engine/Plugins/Animation/ControlRig/Source/ControlRig/Public/Rigs/RigHierarchy.h
@@ -433,12 +433,15 @@ public:
 	{
 		if(ElementKeyRedirector)
 		{
-			if(FCachedRigElement* CachedRigElement = ElementKeyRedirector->Find(InKey))
+			if(FRigElementKeyRedirector::FCachedKeyArray* CachedRigElements = ElementKeyRedirector->Find(InKey))
 			{
-				if(CachedRigElement->UpdateCache(this))
+				if(CachedRigElements->Num() == 1)
 				{
-					return CachedRigElement->GetIndex();
-				}
+					if((*CachedRigElements)[0].UpdateCache(this))
+					{
+						return (*CachedRigElements)[0].GetIndex();
+					}
+			}
 				return INDEX_NONE;
 			}
 		}
@@ -3003,7 +3006,6 @@ public:
 	/**
 	 * Sets a curve's value given its index
 	 * @param InElementIndex The index of the element to set the value for
-	 * @param InValue The value to set on the curve
 	 * @param bSetupUndo If true the transform stack will be setup for undo / redo
 	 */
 	UFUNCTION(BlueprintCallable, Category = URigHierarchy)
@@ -3185,7 +3187,7 @@ public:
 
 	/**
 	 * Returns the child elements of a given element key
-	 * @param InKey The key of the element to retrieve the children for
+	 * @param InElement The element to retrieve the children for
 	 * @param bRecursive If set to true grand-children will also be returned etc
 	 * @return Returns the child elements
 	 */
diff --git a/Engine/Plugins/Animation/ControlRig/Source/ControlRig/Public/Rigs/RigHierarchyCache.h b/Engine/Plugins/Animation/ControlRig/Source/ControlRig/Public/Rigs/RigHierarchyCache.h
--- a/Engine/Plugins/Animation/ControlRig/Source/ControlRig/Public/Rigs/RigHierarchyCache.h
+++ b/Engine/Plugins/Animation/ControlRig/Source/ControlRig/Public/Rigs/RigHierarchyCache.h
@@ -376,23 +376,52 @@ public:
 		, Hash(UINT32_MAX)
 	{}
 
-	FRigElementKeyRedirector(const TMap<FRigElementKey, FRigElementKey>& InMap, const URigHierarchy* InHierarchy);
+	typedef TArray<FRigElementKey, TInlineAllocator<1>> FKeyArray;
+	typedef TMap<FRigElementKey, FKeyArray> FKeyMap;
+	typedef TPair<FRigElementKey, FKeyArray> FKeyPair;
+	typedef TArray<FCachedRigElement, TInlineAllocator<1>> FCachedKeyArray;
+	typedef TMap<FRigElementKey, FCachedKeyArray> FCachedKeyMap;
+	typedef TPair<FRigElementKey, FCachedKeyArray> FCachedKeyPair;
+
+	FRigElementKeyRedirector(const FKeyMap& InMap, const URigHierarchy* InHierarchy);
+	FRigElementKeyRedirector(const TMap<FRigElementKey, FRigElementKeyCollection>& InMap, const URigHierarchy* InHierarchy);
 	FRigElementKeyRedirector(const FRigElementKeyRedirector& InOther, const URigHierarchy* InHierarchy);
 	FRigElementKeyRedirector(const FModularRigConnections& InOther, const URigHierarchy* InHierarchy);
 
 	bool Contains(const FRigElementKey& InKey) const { return InternalKeyToExternalKey.Contains(InKey); }
-	const FCachedRigElement* Find(const FRigElementKey& InKey) const { return InternalKeyToExternalKey.Find(InKey); }
-	const FRigElementKey* FindExternalKey(const FRigElementKey& InKey) const { return ExternalKeys.Find(InKey); }
-	FCachedRigElement* Find(const FRigElementKey& InKey) { return InternalKeyToExternalKey.Find(InKey); }
+	const FCachedKeyArray* Find(const FRigElementKey& InKey) const { return InternalKeyToExternalKey.Find(InKey); }
+	const FKeyArray* FindExternalKey(const FRigElementKey& InKey) const { return ExternalKeys.Find(InKey); }
+	FCachedKeyArray* Find(const FRigElementKey& InKey) { return InternalKeyToExternalKey.Find(InKey); }
 	const FRigElementKey* FindReverse(const FRigElementKey& InKey) const;
 	uint32 GetHash() const { return Hash; }
 	
 private:
 
-	void Add(const FRigElementKey& InSource, const FRigElementKey& InTarget, const URigHierarchy* InHierarchy);
+	void Add(const FRigElementKey& InSource, const FKeyArray& InTargets, const URigHierarchy* InHierarchy);
 
-	TMap<FRigElementKey, FCachedRigElement> InternalKeyToExternalKey;
-	TMap<FRigElementKey, FRigElementKey> ExternalKeys;
+	static FKeyArray Convert(const FCachedKeyArray& InCachedKeys);
+	static FKeyArray Convert(const TArray<FRigElementKey>& InKeys);
+	static FCachedKeyArray Convert(const FKeyArray& InKeys, const URigHierarchy* InHierarchy, bool bForceStoreKey = false);
+
+	template<typename T = FKeyArray>
+	static bool IsValid(const T& InKeys)
+	{
+		if(InKeys.IsEmpty())
+		{
+			return false;
+		}
+		for(const FRigElementKey& Key : InKeys)
+		{
+			if(!Key.IsValid())
+			{
+				return false;
+			}
+		}
+		return true;
+	}
+
+	FCachedKeyMap InternalKeyToExternalKey;
+	FKeyMap ExternalKeys;
 	uint32 Hash;
 
 	friend class URigHierarchy;
@@ -400,4 +429,4 @@ private:
 	friend class FRigModuleInstanceDetails;
 	friend class UModularRig;
 	friend class FControlRigSchematicModel;
-};
\ No newline at end of file
+};
diff --git a/Engine/Plugins/Animation/ControlRig/Source/ControlRigDeveloper/Private/ControlRigBlueprint.cpp b/Engine/Plugins/Animation/ControlRig/Source/ControlRigDeveloper/Private/ControlRigBlueprint.cpp
--- a/Engine/Plugins/Animation/ControlRig/Source/ControlRigDeveloper/Private/ControlRigBlueprint.cpp
+++ b/Engine/Plugins/Animation/ControlRig/Source/ControlRigDeveloper/Private/ControlRigBlueprint.cpp
@@ -1153,12 +1153,12 @@ bool UControlRigBlueprint::ResolveConnector(const FRigElementKey& DraggedKey, co
 	}
 	if(TargetKey.IsValid())
 	{
-		FRigElementKey& ExistingTargetKey = ConnectionMap.FindOrAdd(DraggedKey);
-		if(ExistingTargetKey == TargetKey)
+		FRigElementKeyCollection& ExistingTargetKeys = ArrayConnectionMap.FindOrAdd(DraggedKey);
+		if(ExistingTargetKeys.Keys.Num() == 1 && ExistingTargetKeys.Keys.Contains(TargetKey))
 		{
 			return false;
 		}
-		ExistingTargetKey = TargetKey;
+		ExistingTargetKeys.Keys = {TargetKey};
 
 		if (IsModularRig())
 		{
@@ -1170,7 +1170,7 @@ bool UControlRigBlueprint::ResolveConnector(const FRigElementKey& DraggedKey, co
 		}
 		else
 		{
-			ConnectionMap.FindOrAdd(DraggedKey) = TargetKey;
+			ArrayConnectionMap.FindOrAdd(DraggedKey) = ExistingTargetKeys;
 		}
 	}
 	else
@@ -1185,7 +1185,7 @@ bool UControlRigBlueprint::ResolveConnector(const FRigElementKey& DraggedKey, co
 		}
 		else
 		{
-			ConnectionMap.Remove(DraggedKey);
+			ArrayConnectionMap.Remove(DraggedKey);
 		}
 	}
 
@@ -1213,11 +1213,11 @@ void UControlRigBlueprint::UpdateConnectionMapFromModel()
 {
 	if (IsModularRig())
 	{
-		ConnectionMap.Reset();
+		ArrayConnectionMap.Reset();
 
 		for (const FModularRigSingleConnection& Connection : ModularRigModel.Connections)
 		{
-			ConnectionMap.Add(Connection.Connector, Connection.Target);
+			ArrayConnectionMap.Add(Connection.Connector, {Connection.Targets});
 		}
 	}
 }
@@ -1366,7 +1366,16 @@ void UControlRigBlueprint::PostLoad()
 	}
 #endif
 
+	if(ArrayConnectionMap.IsEmpty() && !ConnectionMap_DEPRECATED.IsEmpty())
+	{
+		for(const TPair<FRigElementKey, FRigElementKey>& Pair : ConnectionMap_DEPRECATED)
+		{
+			ArrayConnectionMap.Add(Pair.Key, FRigElementKeyCollection({Pair.Value}));
+		}
+	}
+
 	ModularRigModel.PatchModelsOnLoad();
+
 	UpdateModularDependencyDelegates();
 }
 
@@ -1447,8 +1456,8 @@ void UControlRigBlueprint::UpdateConnectionMapAfterRename(const FString& InOldNa
 	const FString OldNameSpace = InOldNameSpace + UModularRig::NamespaceSeparator;
 	const FString NewNameSpace = RigModuleSettings.Identifier.Name + UModularRig::NamespaceSeparator;
 	
-	TMap<FRigElementKey, FRigElementKey> FixedConnectionMap;
-	for(const TPair<FRigElementKey, FRigElementKey>& Pair : ConnectionMap)
+	TMap<FRigElementKey, FRigElementKeyCollection> FixedConnectionMap;
+	for(const TPair<FRigElementKey, FRigElementKeyCollection>& Pair : ArrayConnectionMap)
 	{
 		auto FixUpConnectionMap = [OldNameSpace, NewNameSpace](const FRigElementKey& InKey) -> FRigElementKey
 		{
@@ -1461,11 +1470,15 @@ void UControlRigBlueprint::UpdateConnectionMapAfterRename(const FString& InOldNa
 		};
 
 		const FRigElementKey Key = FixUpConnectionMap(Pair.Key);
-		const FRigElementKey Value = FixUpConnectionMap(Pair.Value);
-		FixedConnectionMap.FindOrAdd(Key) = Value;
+		FRigElementKeyCollection Values;
+		for(const FRigElementKey& OldValue : Pair.Value)
+		{
+			Values.Keys.Add(FixUpConnectionMap(OldValue));
+		}
+		FixedConnectionMap.FindOrAdd(Key) = Values;
 	}
 
-	Swap(ConnectionMap, FixedConnectionMap);
+	Swap(ArrayConnectionMap, FixedConnectionMap);
 }
 
 UClass* UControlRigBlueprint::GetRigVMEdGraphNodeClass() const
@@ -1597,7 +1610,7 @@ void UControlRigBlueprint::PostTransacted(const FTransactionObjectEvent& Transac
 			PropagateDrawInstructionsFromBPToInstances();
 		}
 
-		if (PropertiesChanged.Contains(GET_MEMBER_NAME_CHECKED(UControlRigBlueprint, ConnectionMap)))
+		if (PropertiesChanged.Contains(GET_MEMBER_NAME_CHECKED(UControlRigBlueprint, ArrayConnectionMap)))
 		{
 			PropagateHierarchyFromBPToInstances();
 		}
@@ -2304,7 +2317,7 @@ void UControlRigBlueprint::UpdateElementKeyRedirector(UControlRig* InControlRig)
 {
 	InControlRig->HierarchySettings = HierarchySettings;
 	InControlRig->RigModuleSettings = RigModuleSettings;
-	InControlRig->ElementKeyRedirector = FRigElementKeyRedirector(ConnectionMap, InControlRig->GetHierarchy());
+	InControlRig->ElementKeyRedirector = FRigElementKeyRedirector(ArrayConnectionMap, InControlRig->GetHierarchy());
 }
 
 void UControlRigBlueprint::PropagatePoseFromInstanceToBP(UControlRig* InControlRig) const
@@ -2645,7 +2658,7 @@ void UControlRigBlueprint::RefreshModuleConnectors(const FRigModuleReference* In
 				{
 					Hierarchy->Modify();
 					(void)Controller->RemoveElement(Connector);
-					ConnectionMap.Remove(Connector);
+					ArrayConnectionMap.Remove(Connector);
 				}
 			}
 
@@ -2709,10 +2722,10 @@ void UControlRigBlueprint::HandleHierarchyModified(ERigHierarchyNotification InN
 			Influences.OnKeyRenamed(OldKey, InElement->GetKey());
 			if (IsControlRigModule() && InElement->IsTypeOf(ERigElementType::Connector))
 			{
-				if (FRigElementKey* Target = ConnectionMap.Find(OldKey))
+				if (FRigElementKeyCollection* Targets = ArrayConnectionMap.Find(OldKey))
 				{
-					ConnectionMap.FindOrAdd(InElement->GetKey(), *Target);
-					ConnectionMap.Remove(OldKey);
+					ArrayConnectionMap.FindOrAdd(InElement->GetKey(), *Targets);
+					ArrayConnectionMap.Remove(OldKey);
 				}
 			}
 			PropagateHierarchyFromBPToInstances();
@@ -2810,7 +2823,7 @@ void UControlRigBlueprint::HandleRigModulesModified(EModularRigNotification InNo
 						struct ConnectionInfo
 						{
 							FString NewPath;
-							FRigElementKey TargetConnection;
+							FRigElementKeyCollection TargetConnections;
 							FRigConnectorSettings Settings;
 						};
 						FString OldNamespace;
@@ -2839,9 +2852,9 @@ void UControlRigBlueprint::HandleRigModulesModified(EModularRigNotification InNo
 								ConnectionInfo& Info = RenamedConnectors.FindOrAdd(Connector);
 								Info.NewPath = OldConnectorName.Replace(*OldNamespace, *NewNamespace);
 								Info.Settings = CastChecked<FRigConnectorElement>(Controller->GetHierarchy()->FindChecked(Connector))->Settings;
-								if (FRigElementKey* TargetKey = ConnectionMap.Find(Connector))
+								if (FRigElementKeyCollection* TargetKeys = ArrayConnectionMap.Find(Connector))
 								{
-									Info.TargetConnection = *TargetKey;
+									Info.TargetConnections = *TargetKeys;
 								}
 							}
 						}
diff --git a/Engine/Plugins/Animation/ControlRig/Source/ControlRigDeveloper/Public/ControlRigBlueprint.h b/Engine/Plugins/Animation/ControlRig/Source/ControlRigDeveloper/Public/ControlRigBlueprint.h
--- a/Engine/Plugins/Animation/ControlRig/Source/ControlRigDeveloper/Public/ControlRigBlueprint.h
+++ b/Engine/Plugins/Animation/ControlRig/Source/ControlRigDeveloper/Public/ControlRigBlueprint.h
@@ -203,7 +203,10 @@ public:
 	TArray<FModuleReferenceData> ModuleReferenceData;
 
 	UPROPERTY()
-	TMap<FRigElementKey, FRigElementKey> ConnectionMap;
+	TMap<FRigElementKey, FRigElementKey> ConnectionMap_DEPRECATED;
+
+	UPROPERTY()
+	TMap<FRigElementKey, FRigElementKeyCollection> ArrayConnectionMap;
 
 	UFUNCTION(BlueprintPure, Category = "Control Rig Blueprint")
 	TArray<FModuleReferenceData> FindReferencesToModule() const;
diff --git a/Engine/Plugins/Animation/ControlRig/Source/ControlRigDeveloper/Public/Graph/ControlRigGraph.h b/Engine/Plugins/Animation/ControlRig/Source/ControlRigDeveloper/Public/Graph/ControlRigGraph.h
--- a/Engine/Plugins/Animation/ControlRig/Source/ControlRigDeveloper/Public/Graph/ControlRigGraph.h
+++ b/Engine/Plugins/Animation/ControlRig/Source/ControlRigDeveloper/Public/Graph/ControlRigGraph.h
@@ -107,9 +107,12 @@ private:
 						// look up the resolved name
 						if(Element->GetType() == ERigElementType::Connector)
 						{
-							if(const FCachedRigElement* Cache = InControlRig->GetElementKeyRedirector().Find(Element->GetKey()))
+							if(const FRigElementKeyRedirector::FCachedKeyArray* Cache = InControlRig->GetElementKeyRedirector().Find(Element->GetKey()))
 							{
-								Tag = FRigVMStringTag(Cache->GetKey().Name, Color);
+								if(!Cache->IsEmpty())
+								{
+									Tag = FRigVMStringTag((*Cache)[0].GetKey().Name, Color);
+								}
 							}
 						}
 					}
diff --git a/Engine/Plugins/Animation/ControlRig/Source/ControlRigEditor/Private/ControlRigModuleDetails.cpp b/Engine/Plugins/Animation/ControlRig/Source/ControlRigEditor/Private/ControlRigModuleDetails.cpp
--- a/Engine/Plugins/Animation/ControlRig/Source/ControlRigEditor/Private/ControlRigModuleDetails.cpp
+++ b/Engine/Plugins/Animation/ControlRig/Source/ControlRigEditor/Private/ControlRigModuleDetails.cpp
@@ -377,28 +377,30 @@ void FRigModuleInstanceDetails::CustomizeDetails(IDetailLayoutBuilder& DetailBui
 				});
 				TreeDelegates.OnGetSelection.BindLambda([this, Connector]() -> TArray<FRigElementKey>
 				{
-					FRigElementKey Target;
+					TOptional<FRigElementKey> Target;
 					FRigElementKeyRedirector Redirector = PerModuleInfos[0].GetModularRig()->GetElementKeyRedirector();
 					for (int32 ModuleIndex=0; ModuleIndex<PerModuleInfos.Num(); ++ModuleIndex)
 					{
 						FString ConnectorPath = FString::Printf(TEXT("%s:%s"), *PerModuleInfos[ModuleIndex].Path, *Connector.Name);
 						FRigElementKey ConnectorKey(*ConnectorPath, ERigElementType::Connector);
-						if (const FRigElementKey* Key = Redirector.FindExternalKey(ConnectorKey))
+						if (const FRigElementKeyRedirector::FKeyArray* Keys = Redirector.FindExternalKey(ConnectorKey))
 						{
-							if (ModuleIndex == 0)
+							for(const FRigElementKey& Key : *Keys)
 							{
-								Target = *Key;
-							}
-							else if (Target != *Key)
-							{
-								Target.Name = *ControlRigModuleDetailsMultipleValues.ToString();
-								return {Target};
+								if (!Target.IsSet())
+								{
+									Target = Key;
+								}
+								else if (Target.GetValue() != Key)
+								{
+									Target.GetValue().Name = *ControlRigModuleDetailsMultipleValues.ToString();
+									return {Target.GetValue()};
+								}
 							}
 						}
 						else
 						{
-							Target.Name = *ControlRigModuleDetailsMultipleValues.ToString();
-							return {Target};
+							return {FRigElementKey(*ControlRigModuleDetailsMultipleValues.ToString(), ERigElementType::None)};
 						}
 					}
 					return {};
@@ -591,9 +593,11 @@ void FRigModuleInstanceDetails::CustomizeDetails(IDetailLayoutBuilder& DetailBui
 										{
 											FString ConnectorPath = FString::Printf(TEXT("%s:%s"), *PerModuleInfos[0].Path, *Connector.Name);
 											FRigElementKey ConnectorKey(*ConnectorPath, ERigElementType::Connector);
-											if (const FRigElementKey* TargetKey = ModularRig->GetElementKeyRedirector().FindExternalKey(ConnectorKey))
+											if (const FRigElementKeyRedirector::FKeyArray* TargetKeys = ModularRig->GetElementKeyRedirector().FindExternalKey(ConnectorKey))
 											{
-												ModularRig->GetHierarchy()->GetController()->SelectElement(*TargetKey, true, true);
+												TArray<FRigElementKey> KeysToSelect;
+												KeysToSelect.Append(*TargetKeys);
+												ModularRig->GetHierarchy()->GetController()->SetSelection(KeysToSelect, true);
 											}
 										}
 										return FReply::Handled();
@@ -619,27 +623,36 @@ void FRigModuleInstanceDetails::CustomizeDetails(IDetailLayoutBuilder& DetailBui
 				FRigElementKey CurrentTargetKey;
 				if (PerModuleInfos.Num() >= 1)
 				{
+					TOptional<FRigElementKey> TargetKey;
 					for (int32 ModuleIndex=0; ModuleIndex<PerModuleInfos.Num(); ++ModuleIndex)
 					{
-						FRigElementKey TargetKey;
 						FString ConnectorPath = FString::Printf(TEXT("%s:%s"), *PerModuleInfos[ModuleIndex].Path, *Connector.Name);
 						FRigElementKey ConnectorKey(*ConnectorPath, ERigElementType::Connector);
-						if (const FRigElementKey* Key = PerModuleInfos[ModuleIndex].GetModularRig()->GetElementKeyRedirector().FindExternalKey(ConnectorKey))
-						{
-							TargetKey = *Key;
-						}
-						if (ModuleIndex == 0)
-						{
-							CurrentTargetKey = TargetKey;
-						}
-						else
+						if (const FRigElementKeyRedirector::FKeyArray* Keys = PerModuleInfos[ModuleIndex].GetModularRig()->GetElementKeyRedirector().FindExternalKey(ConnectorKey))
 						{
-							if (CurrentTargetKey != TargetKey)
+							for(const FRigElementKey& Key : *Keys)
 							{
-								CurrentTargetKey.Name = *ControlRigModuleDetailsMultipleValues.ToString();
+								if (!TargetKey.IsSet())
+								{
+									TargetKey = Key;
+								}
+								else
+								{
+									if (TargetKey.GetValue() != Key)
+									{
+										CurrentTargetKey.Name = *ControlRigModuleDetailsMultipleValues.ToString();
+										break;
+									}
+								}
 							}
+							break;
 						}
 					}
+
+					if(TargetKey.IsSet())
+					{
+						CurrentTargetKey = TargetKey.GetValue();
+					}
 				}
 				TPair<const FSlateBrush*, FSlateColor> IconAndColor = SRigHierarchyItem::GetBrushForElementType(PerModuleInfos[0].GetModularRig()->GetHierarchy(), CurrentTargetKey);
 				PopulateConnectorCurrentTarget(ButtonBox, IconAndColor.Key, IconAndColor.Value, FText::FromName(CurrentTargetKey.Name));
diff --git a/Engine/Plugins/Animation/ControlRig/Source/ControlRigEditor/Private/Editor/ControlRigEditor.cpp b/Engine/Plugins/Animation/ControlRig/Source/ControlRigEditor/Private/Editor/ControlRigEditor.cpp
--- a/Engine/Plugins/Animation/ControlRig/Source/ControlRigEditor/Private/Editor/ControlRigEditor.cpp
+++ b/Engine/Plugins/Animation/ControlRig/Source/ControlRigEditor/Private/Editor/ControlRigEditor.cpp
@@ -3152,7 +3152,7 @@ void FControlRigEditor::FilterDraggedKeys(TArray<FRigElementKey>& Keys, bool bRe
 		{
 			for(const FModularRigSingleConnection& Connection : ControlRigBlueprint->ModularRigModel.Connections)
 			{
-				if(Connection.Target == Key)
+				if(Connection.Targets.Contains(Key))
 				{
 					Key = Connection.Connector;
 					break;
diff --git a/Engine/Plugins/Animation/ControlRig/Source/ControlRigEditor/Private/Editor/ControlRigSchematicModel.cpp b/Engine/Plugins/Animation/ControlRig/Source/ControlRigEditor/Private/Editor/ControlRigSchematicModel.cpp
--- a/Engine/Plugins/Animation/ControlRig/Source/ControlRigEditor/Private/Editor/ControlRigSchematicModel.cpp
+++ b/Engine/Plugins/Animation/ControlRig/Source/ControlRigEditor/Private/Editor/ControlRigSchematicModel.cpp
@@ -722,12 +722,15 @@ const FSlateBrush* FControlRigSchematicModel::GetBrushForKey(const FRigElementKe
 				int32 Count = 0;
 				for(const FModularRigSingleConnection& Connection : ControlRigBlueprint->ModularRigModel.Connections)
 				{
-					if(Connection.Target == InKey)
+					for(const FRigElementKey& Target : Connection.Targets)
 					{
-						Count++;
-						if(Count == 2)
+						if(Target == InKey)
 						{
-							break;
+							Count++;
+							if(Count == 2)
+							{
+								break;
+							}
 						}
 					}
 				}
@@ -796,12 +799,15 @@ const FSlateBrush* FControlRigSchematicModel::GetBrushForKey(const FRigElementKe
 						int32 Count = 0;
 						for(const FModularRigSingleConnection& Connection : ControlRigBlueprint->ModularRigModel.Connections)
 						{
-							if(Connection.Target == InKey)
+							for(const FRigElementKey& Target : Connection.Targets)
 							{
-								Count++;
-								if(Count == 2)
+								if(Target == InKey)
 								{
-									return ResolvedMultipleSocketBrush;
+									Count++;
+									if(Count == 2)
+									{
+										return ResolvedMultipleSocketBrush;
+									}
 								}
 							}
 						}
@@ -857,12 +863,15 @@ const FSlateBrush* FControlRigSchematicModel::GetBrushForNode(const FSchematicGr
 				int32 Count = 0;
 				for(const FModularRigSingleConnection& Connection : ControlRigBlueprint->ModularRigModel.Connections)
 				{
-					if(Connection.Target == Node->GetKey())
+					for(const FRigElementKey& Target : Connection.Targets)
 					{
-						Count++;
-						if(Count == 2)
+						if(Target == Node->GetKey())
 						{
-							return OutlineDoubleBrush;
+							Count++;
+							if(Count == 2)
+							{
+								return OutlineDoubleBrush;
+							}
 						}
 					}
 				}
diff --git a/Engine/Plugins/Animation/ControlRig/Source/ControlRigEditor/Private/Editor/SModularRigTreeView.cpp b/Engine/Plugins/Animation/ControlRig/Source/ControlRigEditor/Private/Editor/SModularRigTreeView.cpp
--- a/Engine/Plugins/Animation/ControlRig/Source/ControlRigEditor/Private/Editor/SModularRigTreeView.cpp
+++ b/Engine/Plugins/Animation/ControlRig/Source/ControlRigEditor/Private/Editor/SModularRigTreeView.cpp
@@ -465,9 +465,11 @@ TSharedRef<SWidget> SModularRigModelItem::GenerateWidgetForColumn(const FName& C
 		TreeDelegates.OnGetSelection.BindLambda([this]() -> TArray<FRigElementKey>
 		{
 			const FRigElementKeyRedirector Redirector = Delegates.GetModularRig()->GetElementKeyRedirector();
-			if (const FRigElementKey* Key = Redirector.FindExternalKey(ConnectorKey))
+			if (const FRigElementKeyRedirector::FKeyArray* Cache = Redirector.FindExternalKey(ConnectorKey))
 			{
-				return {*Key};
+				TArray<FRigElementKey> Keys;
+				Keys.Append(*Cache);
+				return Keys;
 			}
 			return {};
 		});
@@ -510,9 +512,12 @@ TSharedRef<SWidget> SModularRigModelItem::GenerateWidgetForColumn(const FName& C
 		const UModularRig* ModularRig = Delegates.GetModularRig();
 		const FRigElementKeyRedirector& Redirector = ModularRig->GetElementKeyRedirector();
 		FRigElementKey CurrentTargetKey;
-		if (const FRigElementKey* Key = Redirector.FindExternalKey(ConnectorKey))
+		if (const FRigElementKeyRedirector::FKeyArray* Keys = Redirector.FindExternalKey(ConnectorKey))
 		{
-			CurrentTargetKey = *Key;
+			if(!Keys->IsEmpty())
+			{
+				CurrentTargetKey = (*Keys)[0];
+			}
 		}
 		TPair<const FSlateBrush*, FSlateColor> IconAndColor = SRigHierarchyItem::GetBrushForElementType(ModularRig->GetHierarchy(), CurrentTargetKey);
 
@@ -618,9 +623,12 @@ TSharedRef<SWidget> SModularRigModelItem::GenerateWidgetForColumn(const FName& C
 				if (const UModularRig* ModularRig = Delegates.GetModularRig())
 				{
 					const FRigElementKeyRedirector& Redirector = ModularRig->GetElementKeyRedirector();
-					if (const FRigElementKey* TargetKey = Redirector.FindExternalKey(ConnectorKey))
+					if (const FRigElementKeyRedirector::FKeyArray* TargetKeys = Redirector.FindExternalKey(ConnectorKey))
 					{
-						ModularRig->GetHierarchy()->GetController()->SelectElement(*TargetKey, true, true);
+						for(const FRigElementKey& TargetKey : *TargetKeys)
+						{
+							ModularRig->GetHierarchy()->GetController()->SelectElement(TargetKey, true, true);
+						}
 					}
 				}
 				return FReply::Handled();
diff --git a/Engine/Plugins/Animation/ControlRig/Source/ControlRigEditor/Private/Editor/SRigHierarchyTreeView.cpp b/Engine/Plugins/Animation/ControlRig/Source/ControlRigEditor/Private/Editor/SRigHierarchyTreeView.cpp
--- a/Engine/Plugins/Animation/ControlRig/Source/ControlRigEditor/Private/Editor/SRigHierarchyTreeView.cpp
+++ b/Engine/Plugins/Animation/ControlRig/Source/ControlRigEditor/Private/Editor/SRigHierarchyTreeView.cpp
@@ -593,70 +593,87 @@ bool SRigHierarchyTreeView::AddElement(const FRigBaseElement* InElement)
 				if(UControlRig* ControlRig = Hierarchy->GetTypedOuter<UControlRig>())
 				{
 					FRigElementKeyRedirector& Redirector = ControlRig->GetElementKeyRedirector();
-					if(const FCachedRigElement* Cache = Redirector.Find(InElement->GetKey()))
+					if(const FRigElementKeyRedirector::FCachedKeyArray* Cache = Redirector.Find(InElement->GetKey()))
 					{
-						if(const_cast<FCachedRigElement*>(Cache)->UpdateCache(Hierarchy))
+						bool bValid = true;
+						TArray<TSharedPtr<FRigTreeElement>> TargetElementPtrs;
+						for(const FCachedRigElement& CachedRigElement : *Cache)
 						{
-							if(const TSharedPtr<FRigTreeElement>* TargetElementPtr = ElementMap.Find(Cache->GetKey()))
+							if(!const_cast<FCachedRigElement*>(&CachedRigElement)->UpdateCache(Hierarchy))
 							{
-								const FRigElementKey ConnectorKey = InElement->GetKey();
+								bValid = false;
+							}
+							if(const TSharedPtr<FRigTreeElement>* TargetElementPtr = ElementMap.Find(CachedRigElement.GetKey()))
+							{
+								if(TargetElementPtr->IsValid())
+								{
+									TargetElementPtrs.Add(*TargetElementPtr);
+								}
+							}
+						}
+						
+						if(bValid && !TargetElementPtrs.IsEmpty())
+						{
+							const FRigElementKey ConnectorKey = InElement->GetKey();
 
-								SRigHierarchyTagWidget::FArguments TagArguments;
+							SRigHierarchyTagWidget::FArguments TagArguments;
 
-								static const FLinearColor BackgroundColor = FColor::FromHex(TEXT("#26BBFF"));
-								static const FLinearColor TextColor = FColor::FromHex(TEXT("#0F0F0F"));
-								static const FLinearColor IconColor = FColor::FromHex(TEXT("#1A1A1A"));
+							static const FLinearColor BackgroundColor = FColor::FromHex(TEXT("#26BBFF"));
+							static const FLinearColor TextColor = FColor::FromHex(TEXT("#0F0F0F"));
+							static const FLinearColor IconColor = FColor::FromHex(TEXT("#1A1A1A"));
 
-								static const FSlateBrush* PrimaryBrush = FControlRigEditorStyle::Get().GetBrush("ControlRig.ConnectorPrimary");
-								static const FSlateBrush* SecondaryBrush = FControlRigEditorStyle::Get().GetBrush("ControlRig.ConnectorSecondary");
-								static const FSlateBrush* OptionalBrush = FControlRigEditorStyle::Get().GetBrush("ControlRig.ConnectorOptional");
+							static const FSlateBrush* PrimaryBrush = FControlRigEditorStyle::Get().GetBrush("ControlRig.ConnectorPrimary");
+							static const FSlateBrush* SecondaryBrush = FControlRigEditorStyle::Get().GetBrush("ControlRig.ConnectorSecondary");
+							static const FSlateBrush* OptionalBrush = FControlRigEditorStyle::Get().GetBrush("ControlRig.ConnectorOptional");
 
-								const FSlateBrush* IconBrush = PrimaryBrush;
-								if(const FRigConnectorElement* ConnectorElement = Cast<FRigConnectorElement>(InElement))
+							const FSlateBrush* IconBrush = PrimaryBrush;
+							if(const FRigConnectorElement* ConnectorElement = Cast<FRigConnectorElement>(InElement))
+							{
+								if(ConnectorElement->Settings.Type == EConnectorType::Secondary)
 								{
-									if(ConnectorElement->Settings.Type == EConnectorType::Secondary)
-									{
-										IconBrush = ConnectorElement->Settings.bOptional ? OptionalBrush : SecondaryBrush;
-									}
+									IconBrush = ConnectorElement->Settings.bOptional ? OptionalBrush : SecondaryBrush;
 								}
+							}
 
-								FName Name = ConnectorKey.Name;
-								if (GetRigTreeDelegates().GetDisplaySettings().bUseShortName)
-								{
-									Name = *Hierarchy->GetDisplayNameForUI(ConnectorKey).ToString();
-								}
-								TagArguments.Text(FText::FromName(Name));
-								TagArguments.TooltipText(FText::FromName(ConnectorKey.Name));
-								TagArguments.Color(BackgroundColor);
-								TagArguments.IconColor(IconColor);
-								TagArguments.TextColor(TextColor);
-								TagArguments.Icon(IconBrush);
-								TagArguments.IconSize(FVector2d(16.f, 16.f));
-								TagArguments.AllowDragDrop(true);
-								FString Identifier;
-								FRigElementKey::StaticStruct()->ExportText(Identifier, &ConnectorKey, nullptr, nullptr, PPF_None, nullptr);
-								TagArguments.Identifier(Identifier);
-
-								TagArguments.OnClicked_Lambda([ConnectorKey, this]()
+							FName Name = ConnectorKey.Name;
+							if (GetRigTreeDelegates().GetDisplaySettings().bUseShortName)
+							{
+								Name = *Hierarchy->GetDisplayNameForUI(ConnectorKey).ToString();
+							}
+							TagArguments.Text(FText::FromName(Name));
+							TagArguments.TooltipText(FText::FromName(ConnectorKey.Name));
+							TagArguments.Color(BackgroundColor);
+							TagArguments.IconColor(IconColor);
+							TagArguments.TextColor(TextColor);
+							TagArguments.Icon(IconBrush);
+							TagArguments.IconSize(FVector2d(16.f, 16.f));
+							TagArguments.AllowDragDrop(true);
+							FString Identifier;
+							FRigElementKey::StaticStruct()->ExportText(Identifier, &ConnectorKey, nullptr, nullptr, PPF_None, nullptr);
+							TagArguments.Identifier(Identifier);
+
+							TagArguments.OnClicked_Lambda([ConnectorKey, this]()
+							{
+								Delegates.RequestDetailsInspection(ConnectorKey);
+							});
+
+							if (!ControlRig->IsModularRig())
+							{
+								TagArguments.OnRenamed_Lambda([ConnectorKey, this](const FText& InNewName, ETextCommit::Type InCommitType)
 								{
-									Delegates.RequestDetailsInspection(ConnectorKey);
+									Delegates.HandleRenameElement(ConnectorKey, InNewName.ToString());
 								});
-
-								if (!ControlRig->IsModularRig())
+								TagArguments.OnVerifyRename_Lambda([ConnectorKey, this](const FText& InText, FText& OutError)
 								{
-									TagArguments.OnRenamed_Lambda([ConnectorKey, this](const FText& InNewName, ETextCommit::Type InCommitType)
-									{
-										Delegates.HandleRenameElement(ConnectorKey, InNewName.ToString());
-									});
-									TagArguments.OnVerifyRename_Lambda([ConnectorKey, this](const FText& InText, FText& OutError)
-									{
-										return Delegates.HandleVerifyElementNameChanged(ConnectorKey, InText.ToString(), OutError);
-									});
-								}
+									return Delegates.HandleVerifyElementNameChanged(ConnectorKey, InText.ToString(), OutError);
+								});
+							}
 
-								TargetElementPtr->Get()->Tags.Add(TagArguments);
-								return true;
+							for(TSharedPtr<FRigTreeElement> TargetElementPtr : TargetElementPtrs)
+							{
+								TargetElementPtr->Tags.Add(TagArguments);
 							}
+							return true;
 						}
 					}
 				}
