Sequencer: Control Rig: Keying Float channels would double the value do to recent ECS regression.

diff --git a/Engine/Plugins/Animation/ControlRig/Source/ControlRig/Private/Sequencer/MovieSceneControlRigParameterSection.cpp b/Engine/Plugins/Animation/ControlRig/Source/ControlRig/Private/Sequencer/MovieSceneControlRigParameterSection.cpp
--- a/Engine/Plugins/Animation/ControlRig/Source/ControlRig/Private/Sequencer/MovieSceneControlRigParameterSection.cpp
+++ b/Engine/Plugins/Animation/ControlRig/Source/ControlRig/Private/Sequencer/MovieSceneControlRigParameterSection.cpp
@@ -82,7 +82,6 @@ struct FParameterFloatChannelEditorData
 
 		UMovieSceneTrack* Track = SectionToKey->GetTypedOuter<UMovieSceneTrack>();
 
-		/*
 		if (Track)
 		{
 			FMovieSceneEvaluationTrack EvalTrack = CastChecked<IMovieSceneTrackTemplateProducer>(Track)->GenerateTrackTemplate(Track);
@@ -103,7 +102,6 @@ struct FParameterFloatChannelEditorData
 			}
 			OutValue = Val;
 		}
-		*/
 		OutWeight = MovieSceneHelpers::CalculateWeightForBlending(SectionToKey, KeyTime);
 	}
 
@@ -240,7 +238,6 @@ struct FParameterVectorChannelEditorData
 			return;
 		}
 
-		/*
 		UMovieSceneTrack* Track = SectionToKey->GetTypedOuter<UMovieSceneTrack>();
 
 		if (Track)
@@ -303,7 +300,6 @@ struct FParameterVectorChannelEditorData
 			break;
 			}
 		}
-		*/
 		OutWeight = MovieSceneHelpers::CalculateWeightForBlending(SectionToKey, KeyTime);
 	}
 	FText							GroupName;
@@ -712,7 +708,6 @@ struct FParameterTransformChannelEditorData
 		float& OutValue, float& OutWeight)
 	{
 		UMovieSceneTrack* Track = SectionToKey->GetTypedOuter<UMovieSceneTrack>();
-		/*
 		FMovieSceneEvaluationTrack EvalTrack = CastChecked<UMovieSceneControlRigParameterTrack>(Track)->GenerateTrackTemplate(Track);
 		FMovieSceneInterrogationData InterrogationData;
 		RootTemplate.CopyActuators(InterrogationData.GetAccumulator());
@@ -766,7 +761,6 @@ struct FParameterTransformChannelEditorData
 			break;
 
 		}
-		*/
 		OutWeight = MovieSceneHelpers::CalculateWeightForBlending(SectionToKey, KeyTime);
 	}
 		
