Anim Layers: Get Override working again with legacy, was mistakenly removed.

diff --git a/Engine/Plugins/Animation/ControlRig/Source/ControlRig/Private/Sequencer/MovieSceneControlRigParameterTrack.cpp b/Engine/Plugins/Animation/ControlRig/Source/ControlRig/Private/Sequencer/MovieSceneControlRigParameterTrack.cpp
--- a/Engine/Plugins/Animation/ControlRig/Source/ControlRig/Private/Sequencer/MovieSceneControlRigParameterTrack.cpp
+++ b/Engine/Plugins/Animation/ControlRig/Source/ControlRig/Private/Sequencer/MovieSceneControlRigParameterTrack.cpp
@@ -72,10 +72,7 @@ UMovieSceneControlRigParameterTrack::UMovieSceneControlRigParameterTrack(const F
 	SupportedBlendTypes = FMovieSceneBlendTypeField::None();
 	SupportedBlendTypes.Add(EMovieSceneBlendType::Additive);
 	SupportedBlendTypes.Add(EMovieSceneBlendType::Absolute);
-	if (ShouldUseLegacyTemplate() == false)
-	{
-		SupportedBlendTypes.Add(EMovieSceneBlendType::Override);
-	}
+	SupportedBlendTypes.Add(EMovieSceneBlendType::Override);
 }
 
 void UMovieSceneControlRigParameterTrack::BeginDestroy()
