Control Rig Snapper: Make sure channels array is set to be empty when snapping with settings or it will snap everything.

diff --git a/Engine/Plugins/Animation/ControlRig/Source/ControlRigEditor/Private/Tools/ControlRigSnapper.cpp b/Engine/Plugins/Animation/ControlRig/Source/ControlRigEditor/Private/Tools/ControlRigSnapper.cpp
--- a/Engine/Plugins/Animation/ControlRig/Source/ControlRigEditor/Private/Tools/ControlRigSnapper.cpp
+++ b/Engine/Plugins/Animation/ControlRig/Source/ControlRigEditor/Private/Tools/ControlRigSnapper.cpp
@@ -236,9 +236,10 @@ struct FGuidAndActor
 		}
 
 		// fill channels to key
-		EMovieSceneTransformChannel Channels;
+		EMovieSceneTransformChannel Channels = EMovieSceneTransformChannel::AllTransform;
 		if (SnapSettings)
 		{
+			Channels = EMovieSceneTransformChannel::None;
 			if (SnapSettings->bSnapPosition)
 			{
 				EnumAddFlags(Channels, EMovieSceneTransformChannel::Translation);
@@ -251,6 +252,10 @@ struct FGuidAndActor
 			{
 				EnumAddFlags(Channels, EMovieSceneTransformChannel::Scale);
 			}
+			if (Channels == EMovieSceneTransformChannel::None)
+			{
+				return false;
+			}
 		}
 
 		// compute local transforms
@@ -749,6 +754,9 @@ bool FControlRigSnapper::SnapIt(FFrameNumber StartFrame, FFrameNumber EndFrame,c
 
 
 
+
+
+
 #undef LOCTEXT_NAMESPACE
 
 
