Control Rig Edit Mode: request shape creation when setting sequencer

#rnx

diff --git a/Engine/Plugins/Animation/ControlRig/Source/ControlRigEditor/Private/EditMode/ControlRigEditMode.cpp b/Engine/Plugins/Animation/ControlRig/Source/ControlRigEditor/Private/EditMode/ControlRigEditMode.cpp
--- a/Engine/Plugins/Animation/ControlRig/Source/ControlRigEditor/Private/EditMode/ControlRigEditMode.cpp
+++ b/Engine/Plugins/Animation/ControlRig/Source/ControlRigEditor/Private/EditMode/ControlRigEditMode.cpp
@@ -295,6 +295,11 @@ bool FControlRigEditMode::SetSequencer(TWeakPtr<ISequencer> InSequencer)
 		{
 			FControlRigEditModeToolkit::Outliner->SetEditMode(*this);
 		}
+
+		if (WeakSequencer.IsValid() && !RuntimeControlRigs.IsEmpty())
+		{
+			RequestToRecreateControlShapeActors();
+		}
 	}
 	return false;
 }
