Force lower case FNames for FKControlRig control-name mapping

diff --git a/Engine/Plugins/Animation/ControlRig/Source/ControlRig/Private/Rigs/FKControlRig.cpp b/Engine/Plugins/Animation/ControlRig/Source/ControlRig/Private/Rigs/FKControlRig.cpp
--- a/Engine/Plugins/Animation/ControlRig/Source/ControlRig/Private/Rigs/FKControlRig.cpp
+++ b/Engine/Plugins/Animation/ControlRig/Source/ControlRig/Private/Rigs/FKControlRig.cpp
@@ -37,10 +37,12 @@ FName UFKControlRig::GetControlName(const FName& InName, const ERigElementType&
 			}
 			else
 			{
-				static thread_local FStringBuilderBase ScratchString;		
+				static thread_local FString ScratchString;
 
 				// ToString performs ScratchString.Reset() internally
 				InName.ToString(ScratchString);
+				// Force lower-case names
+				ScratchString.ToLowerInline();
 
 				if (InType == ERigElementType::Curve)
 				{
@@ -94,7 +96,11 @@ FName UFKControlRig::GetControlTargetName(const FName& InName, const ERigElement
 				StartPostFix = ScratchString.Find(ControlPostFix, ESearchCase::CaseSensitive);
 			}
 
-			const FStringView ControlTargetString(*ScratchString, StartPostFix != INDEX_NONE ? StartPostFix : ScratchString.Len());
+			FString ControlTargetString = FString::ConstructFromPtrSize(*ScratchString, StartPostFix != INDEX_NONE ? StartPostFix : ScratchString.Len());
+
+			// Force lower-case names
+			ControlTargetString.ToLowerInline();
+
 			return NameToTargetMapping.Add(InName, FName(ControlTargetString));
 		}
 	}
