Control Rig: Avoid to establish cyclic tick dependencies on Control Rig Component

#rnx

diff --git a/Engine/Plugins/Animation/ControlRig/Source/ControlRig/Private/ControlRigComponent.cpp b/Engine/Plugins/Animation/ControlRig/Source/ControlRig/Private/ControlRigComponent.cpp
--- a/Engine/Plugins/Animation/ControlRig/Source/ControlRig/Private/ControlRigComponent.cpp
+++ b/Engine/Plugins/Animation/ControlRig/Source/ControlRig/Private/ControlRigComponent.cpp
@@ -411,14 +411,37 @@ void UControlRigComponent::Initialize()
 
 			// input elements should tick before ControlRig updates
 			// output elements should tick after ControlRig updates
+#if WITH_EDITOR
+			FString TickDependencyErrorMessage;
+			FString* TickDependencyErrorMessagePtr = &TickDependencyErrorMessage;
+#else
+			static FString* TickDependencyErrorMessagePtr = nullptr;
+#endif
 			if (MappedElement.Direction == EControlRigComponentMapDirection::Output)
 			{
-				MappedElement.SceneComponent->AddTickPrerequisiteComponent(this);
+				if(AllowsTickDependency(MappedElement.SceneComponent, this, TickDependencyErrorMessagePtr))
+				{
+					MappedElement.SceneComponent->AddTickPrerequisiteComponent(this);
+				}
 			}
 			else
 			{
-				AddTickPrerequisiteComponent(MappedElement.SceneComponent);
+				if(AllowsTickDependency(this, MappedElement.SceneComponent, TickDependencyErrorMessagePtr))
+				{
+					AddTickPrerequisiteComponent(MappedElement.SceneComponent);
+				}
+			}
+
+#if WITH_EDITOR
+			if(!TickDependencyErrorMessage.IsEmpty())
+			{
+				// only throw the error message if the user has picked a particular component to set up the dependency for
+				if(!MappedElement.ComponentReference.ComponentProperty.IsNone() || MappedElement.ComponentReference.OtherActor.IsValid())
+				{
+					UE_LOG(LogControlRig, Warning, TEXT("%s"), *TickDependencyErrorMessage);
+				}
 			}
+#endif
 
 			MappedComponents.Add(MappedElement.SceneComponent) = MappedElement.Direction;
 			
@@ -2030,6 +2053,200 @@ void UControlRigComponent::ReportError(const FString& InMessage)
 #endif
 }
 
+bool UControlRigComponent::HasTickDependency(const UObject* InDependent, const UObject* InDependency, FString* OutErrorMessage)
+{
+	TSet<const UObject*> Visited;
+	return HasTickDependency_Impl(InDependent, InDependency, Visited, OutErrorMessage);
+}
+
+bool UControlRigComponent::HasTickDependency_Impl(const UObject* InDependent, const UObject* InDependency, TSet<const UObject*>& InOutVisited, FString* OutErrorMessage)
+{
+	if(InDependent == InDependency)
+	{
+		return true;
+	}
+
+	// avoid cycles / infinite recursion
+	if(InOutVisited.Contains(InDependent))
+	{
+		return false;
+	}
+	InOutVisited.Add(InDependent);
+	
+	if(const USceneComponent* DependencyComponent = Cast<USceneComponent>(InDependency))
+	{
+		if(const USceneComponent* DependentComponent = Cast<USceneComponent>(InDependent))
+		{
+			if(const AActor* DependencyActor = DependencyComponent->GetOwner())
+			{
+				if(const AActor* DependentActor = DependentComponent->GetOwner())
+				{
+					if(const USceneComponent* DependentParentComponent = DependentComponent->GetAttachParent())
+					{
+						if(HasTickDependency_Impl(DependentParentComponent, InDependency, InOutVisited, OutErrorMessage))
+						{
+#if WITH_EDITOR
+							if(DependencyActor == DependentActor)
+							{
+								if(DependentParentComponent == InDependency)
+								{
+									if(OutErrorMessage)
+									{
+										if(!OutErrorMessage->IsEmpty())
+										{
+											*OutErrorMessage += TEXT("\n");
+										}
+										static constexpr TCHAR Format[] = TEXT("The dependency '%s' is a parent component of dependent '%s'.");
+										*OutErrorMessage += FString::Printf(Format, *InDependency->GetPathName(), *InDependent->GetPathName());
+									}
+								}
+							}
+#endif
+							return true;
+						}
+					}
+				}
+			}
+		}
+		else if(const AActor* DependencyActor = Cast<AActor>(InDependency))
+		{
+			if(DependencyActor->GetComponents().Contains(DependentComponent))
+			{
+#if WITH_EDITOR
+				if(OutErrorMessage)
+				{
+					if(!OutErrorMessage->IsEmpty())
+					{
+						*OutErrorMessage += TEXT("\n");
+					}
+					static constexpr TCHAR Format[] = TEXT("The actor '%s' owns dependent component '%s'.");
+					*OutErrorMessage += FString::Printf(Format, *DependencyActor->GetPathName(), *InDependent->GetPathName());
+				}
+#endif
+				return true;
+			}
+		}
+	}
+
+	// iterate all existing tick dependency components
+	if(const UActorComponent* DependentActorComponent = Cast<UActorComponent>(InDependent))
+	{
+		const TArray<struct FTickPrerequisite>& Prerequisites = DependentActorComponent->PrimaryComponentTick.GetPrerequisites();
+		for(const struct FTickPrerequisite& Prerequisite : Prerequisites)
+		{
+			if(Prerequisite.PrerequisiteObject.IsValid())
+			{
+				if(const UObject* TickDependency = Prerequisite.PrerequisiteObject.Get())
+				{
+					if(TickDependency == InDependency)
+					{
+#if WITH_EDITOR
+						if(OutErrorMessage)
+						{
+							if(!OutErrorMessage->IsEmpty())
+							{
+								*OutErrorMessage += TEXT("\n");
+							}
+							static constexpr TCHAR Format[] = TEXT("The dependency '%s' is a prerequisite of dependent'%s'.");
+							*OutErrorMessage += FString::Printf(Format, *InDependency->GetPathName(), *InDependent->GetPathName());
+						}
+#endif
+						return true;
+					}
+					
+					if(HasTickDependency_Impl(TickDependency, InDependency, InOutVisited, OutErrorMessage))
+					{
+#if WITH_EDITOR
+						if(OutErrorMessage)
+						{
+							if(!OutErrorMessage->IsEmpty())
+							{
+								*OutErrorMessage += TEXT("\n");
+							}
+							static constexpr TCHAR Format[] = TEXT("The dependency '%s' is a indirect prerequisite of dependent '%s'.");
+							*OutErrorMessage += FString::Printf(Format, *InDependency->GetPathName(), *InDependent->GetPathName());
+						}
+#endif
+						return true;
+					}
+				}
+			}
+		}
+	}
+	
+	return false;
+}
+
+bool UControlRigComponent::AllowsTickDependency(const UObject* InDependent, const UObject* InDependency, FString* OutErrorMessage)
+{
+	if(InDependent == InDependency)
+	{
+#if WITH_EDITOR
+		if(OutErrorMessage)
+		{
+			if(!OutErrorMessage->IsEmpty())
+			{
+				*OutErrorMessage += TEXT("\n");
+			}
+			static constexpr TCHAR Format[] = TEXT("The dependency and dependent '%s' are identical.");
+			*OutErrorMessage += FString::Printf(Format, *InDependent->GetPathName());
+		}
+#endif
+		return false;
+	}
+	
+	// don't allow to depend on an inner
+	if(InDependency->IsInOuter(InDependent))
+	{
+#if WITH_EDITOR
+		if(OutErrorMessage)
+		{
+			if(!OutErrorMessage->IsEmpty())
+			{
+				*OutErrorMessage += TEXT("\n");
+			}
+			static constexpr TCHAR Format[] = TEXT("The dependency '%s' is within the dependent '%s'.");
+			*OutErrorMessage += FString::Printf(Format, *InDependency->GetPathName(), *InDependent->GetPathName());
+		}
+#endif
+		return false;
+	}
+	
+	if(HasTickDependency(InDependent, InDependency, OutErrorMessage))
+	{
+#if WITH_EDITOR
+		if(OutErrorMessage)
+		{
+			if(!OutErrorMessage->IsEmpty())
+			{
+				*OutErrorMessage += TEXT("\n");
+			}
+			static constexpr TCHAR Format[] = TEXT("The dependent '%s' already depends on dependency '%s'.");
+			*OutErrorMessage += FString::Printf(Format, *InDependent->GetPathName(), *InDependency->GetPathName());
+		}
+#endif
+		return false;
+	}
+	
+	if(HasTickDependency(InDependency, InDependent, OutErrorMessage))
+	{
+#if WITH_EDITOR
+		if(OutErrorMessage)
+		{
+			if(!OutErrorMessage->IsEmpty())
+			{
+				*OutErrorMessage += TEXT("\n");
+			}
+			static constexpr TCHAR Format[] = TEXT("The dependency '%s' depends on dependent '%s'. This would cause a cycle.");
+			*OutErrorMessage += FString::Printf(Format, *InDependency->GetPathName(), *InDependent->GetPathName());
+		}
+#endif
+		return false;
+	}
+
+	return true;
+}
+
 FControlRigSceneProxy::FControlRigSceneProxy(const UControlRigComponent* InComponent)
 : FPrimitiveSceneProxy(InComponent)
 , ControlRigComponent(InComponent)
diff --git a/Engine/Plugins/Animation/ControlRig/Source/ControlRig/Public/ControlRigComponent.h b/Engine/Plugins/Animation/ControlRig/Source/ControlRig/Public/ControlRigComponent.h
--- a/Engine/Plugins/Animation/ControlRig/Source/ControlRig/Public/ControlRigComponent.h
+++ b/Engine/Plugins/Animation/ControlRig/Source/ControlRig/Public/ControlRigComponent.h
@@ -654,6 +654,15 @@ private:
 	bool EnsureCalledOutsideOfBracket(const TCHAR* InCallingFunctionName = nullptr);
 	void ReportError(const FString& InMessage);
 
+public:
+	
+	static bool HasTickDependency(const UObject* InDependent, const UObject* InDependency, FString* OutErrorMessage = nullptr);
+	static bool AllowsTickDependency(const UObject* InDependent, const UObject* InDependency, FString* OutErrorMessage = nullptr);
+
+private:
+	
+	static bool HasTickDependency_Impl(const UObject* InDependent, const UObject* InDependency, TSet<const UObject*>& InOutVisited, FString* OutErrorMessage);
+
 #if WITH_EDITOR
 	static TMap<FString, TSharedPtr<SNotificationItem>> EditorNotifications;
 #endif
@@ -701,4 +710,4 @@ public:
 private:
 
 	const UControlRigComponent* ControlRigComponent;
-};
\ No newline at end of file
+};
