Fixed use of soon-to-be-required unused return values.

diff --git a/Engine/Plugins/Animation/ControlRig/Source/ControlRigEditor/Private/ControlRigModuleDetails.cpp b/Engine/Plugins/Animation/ControlRig/Source/ControlRigEditor/Private/ControlRigModuleDetails.cpp
--- a/Engine/Plugins/Animation/ControlRig/Source/ControlRigEditor/Private/ControlRigModuleDetails.cpp
+++ b/Engine/Plugins/Animation/ControlRig/Source/ControlRigEditor/Private/ControlRigModuleDetails.cpp
@@ -349,9 +349,9 @@ void FRigModuleInstanceDetails::CustomizeDetails(IDetailLayoutBuilder& DetailBui
 									else
 									{
 										const FModularRigResolveResult& ConnectorMatches = RuleManager->FindMatches(ConnectorElement, Module, ModularRig->GetElementKeyRedirector());
-										Matches.FilterByPredicate([ConnectorMatches](const FRigElementResolveResult& Match)
+										Matches.RemoveAll([&ConnectorMatches](const FRigElementResolveResult& Match)
 										{
-											return ConnectorMatches.ContainsMatch(Match.GetKey());
+											return !ConnectorMatches.ContainsMatch(Match.GetKey());
 										});
 									}
 								}
diff --git a/Engine/Plugins/Experimental/EditorDataStorageFeatures/Source/TedsAssetData/Private/CB/TedsAssetDataCBDataSource.cpp b/Engine/Plugins/Experimental/EditorDataStorageFeatures/Source/TedsAssetData/Private/CB/TedsAssetDataCBDataSource.cpp
--- a/Engine/Plugins/Experimental/EditorDataStorageFeatures/Source/TedsAssetData/Private/CB/TedsAssetDataCBDataSource.cpp
+++ b/Engine/Plugins/Experimental/EditorDataStorageFeatures/Source/TedsAssetData/Private/CB/TedsAssetDataCBDataSource.cpp
@@ -451,7 +451,6 @@ void FTedsAssetDataCBDataSource::PrepopulateTagsMetadataCache()
 
 		TArray<UClass*> ChildrenClasses;
 
-		ClassesWithAssetDefinition.Contains(TSoftClassPtr<UObject>(FSoftObjectPath(TEXT("/Script/Engine.UserDefinedStruct"))));
 		for (const TSoftClassPtr<UObject>& SoftClass : ClassesWithAssetDefinition)
 		{
 			// Check 
diff --git a/Engine/Plugins/Experimental/Mover/Source/Mover/Private/MovementModeStateMachine.cpp b/Engine/Plugins/Experimental/Mover/Source/Mover/Private/MovementModeStateMachine.cpp
--- a/Engine/Plugins/Experimental/Mover/Source/Mover/Private/MovementModeStateMachine.cpp
+++ b/Engine/Plugins/Experimental/Mover/Source/Mover/Private/MovementModeStateMachine.cpp
@@ -82,7 +82,7 @@ void UMovementModeStateMachine::ClearAllMovementModes()
 
 void UMovementModeStateMachine::SetDefaultMode(FName NewDefaultModeName)
 {
-	Modes.FindChecked(NewDefaultModeName);
+	check(Modes.Contains(NewDefaultModeName));
 
 	DefaultModeName = NewDefaultModeName;
 }
diff --git a/Engine/Plugins/Experimental/Water/Source/Editor/Private/WaterBrushManager.cpp b/Engine/Plugins/Experimental/Water/Source/Editor/Private/WaterBrushManager.cpp
--- a/Engine/Plugins/Experimental/Water/Source/Editor/Private/WaterBrushManager.cpp
+++ b/Engine/Plugins/Experimental/Water/Source/Editor/Private/WaterBrushManager.cpp
@@ -560,7 +560,6 @@ bool AWaterBrushManager::SetupRiverSplineRenderMIDs(const FBrushActorRenderConte
 		}
 		else
 		{
-			RiverSplineMIDs.IsValidIndex(MIDIndex);
 			RiverSplineMIDs[MIDIndex] = FWaterUtils::GetOrCreateTransientMID(RiverSplineMIDs[MIDIndex], TEXT("RiverSplineMID"), RenderRiverSplineDepthMaterial);
 			UMaterialInstanceDynamic* TempMID = RiverSplineMIDs[MIDIndex];
 			if (TempMID != nullptr)
diff --git a/Engine/Plugins/Runtime/HairStrands/Source/HairStrandsCore/Private/GroomManager.cpp b/Engine/Plugins/Runtime/HairStrands/Source/HairStrandsCore/Private/GroomManager.cpp
--- a/Engine/Plugins/Runtime/HairStrands/Source/HairStrandsCore/Private/GroomManager.cpp
+++ b/Engine/Plugins/Runtime/HairStrands/Source/HairStrandsCore/Private/GroomManager.cpp
@@ -371,7 +371,7 @@ struct FHairGeometryCache
 					++DebugData.InstanceCount;
 					DebugData.CacheType    |= 1u << uint32(InCacheType);
 					DebugData.GeometryType |= 1u << uint32(InGeometryType);
-					DebugData.SectionBits.BitwiseOR(DebugData.SectionBits, ToBitArray(InGeom.Sections), EBitwiseOperatorFlags::MaxSize);
+					DebugData.SectionBits = TBitArray<>::BitwiseOR(DebugData.SectionBits, ToBitArray(InGeom.Sections), EBitwiseOperatorFlags::MaxSize);
 					return;
 				}
 			}
