Adding support to Control Rig (core) for physics - the implementation of which will be in a separate plugin/change.

Full Cook Preflight:  https://horde.devtools.epicgames.com/job/6746e553c260f817c6b46d6f

#rnx

diff --git a/Engine/Plugins/Animation/ControlRig/Source/ControlRig/Private/ControlRig.cpp b/Engine/Plugins/Animation/ControlRig/Source/ControlRig/Private/ControlRig.cpp
--- a/Engine/Plugins/Animation/ControlRig/Source/ControlRig/Private/ControlRig.cpp
+++ b/Engine/Plugins/Animation/ControlRig/Source/ControlRig/Private/ControlRig.cpp
@@ -566,8 +566,10 @@ void UControlRig::InitializeFromCDO()
 			UControlRig* CDO = GetClass()->GetDefaultObject<UControlRig>();
 			URigHierarchy* Hierarchy = GetHierarchy();
 
-			// copy physics solvers
-			PhysicsSolvers = CDO->PhysicsSolvers;
+			// Danny TODO this could copy the physics solvers from the CDO, but they are transient
+			// objects that only exist through construction, so I think we just need to start afresh
+			// each time. 
+			PhysicsSolvers.Empty();
 
 			// copy hierarchy
 			{
@@ -579,25 +581,6 @@ void UControlRig::InitializeFromCDO()
 				Hierarchy->ResetPoseToInitial(ERigElementType::All);
 			}
 
-			// update the physics solvers with new unique IDs and remap the physics solver IDs within hierarchy
-			{
-				TMap<FRigPhysicsSolverID, FRigPhysicsSolverID> PhysicsSolverMap;
-				for(FRigPhysicsSolverDescription& Solver : PhysicsSolvers)
-				{
-					const FRigPhysicsSolverID OldID = Solver.ID;
-					Solver.ID = FRigPhysicsSolverDescription::MakeID(GetPathName(), Solver.Name);
-					PhysicsSolverMap.Add(OldID, Solver.ID);
-				}
-				const TArray<FRigPhysicsElement*> PhysicsElements = Hierarchy->GetPhysicsElements();
-				for(FRigPhysicsElement* PhysicsElement : PhysicsElements)
-				{
-					if(const FRigPhysicsSolverID* NewID = PhysicsSolverMap.Find(PhysicsElement->Solver))
-					{
-						PhysicsElement->Solver = *NewID;
-					}
-				}
-			}
-
 #if WITH_EDITOR
 			// current hierarchy should always mirror CDO's hierarchy whenever a change of interest happens
 			CDO->GetHierarchy()->RegisterListeningHierarchy(Hierarchy);
@@ -1812,13 +1795,19 @@ void UControlRig::Serialize(FArchive& Ar)
 
 	Ar.UsingCustomVersion(FControlRigObjectVersion::GUID);
 
-	if (Ar.CustomVer(FControlRigObjectVersion::GUID) >= FControlRigObjectVersion::ControlRigStoresPhysicsSolvers)
-	{
-		Ar << PhysicsSolvers;
-		UE_RIGVM_ARCHIVETRACE_ENTRY(Ar, TEXT("PhysicsSolvers"));
-	}
-	else
+	// We are no longer storing the physics solver descriptions as part of the rig, and the struct
+	// definitions have been removed. In addition, the array would have been empty when serialized
+	// with the affected version, so we just need to absorb the data into a temporary array.
+	if (Ar.IsLoading())
 	{
+		const int32 Ver = Ar.CustomVer(FControlRigObjectVersion::GUID);
+		if (Ver >= FControlRigObjectVersion::ControlRigStoresPhysicsSolvers && 
+			Ver < FControlRigObjectVersion::ControlRigDoesntStorePhysicsSolvers)
+		{
+			TArray<FRigPhysicsSolverDescription_Deprecated> TempSolvers;
+			Ar << TempSolvers;
+			UE_RIGVM_ARCHIVETRACE_ENTRY(Ar, TEXT("PhysicsSolverDescriptions"));
+		}
 		PhysicsSolvers.Reset();
 	}
 }
@@ -3603,46 +3592,33 @@ void UControlRig::OnHierarchyTransformUndoRedo(URigHierarchy* InHierarchy, const
 	}
 }
 
-int32 UControlRig::NumPhysicsSolvers() const
+FRigPhysicsSolverBase* UControlRig::GetPhysicsSolver(const FRigPhysicsSolverID& InID)
 {
-	return PhysicsSolvers.Num();
-}
-
-const FRigPhysicsSolverDescription* UControlRig::GetPhysicsSolver(int32 InIndex) const
-{
-	if(PhysicsSolvers.IsValidIndex(InIndex))
-	{
-		return &PhysicsSolvers[InIndex];
-	}
-	return nullptr;
+	TSharedPtr<FRigPhysicsSolverBase>* PhysicsSolver = PhysicsSolvers.Find(InID);
+	return PhysicsSolver ? PhysicsSolver->Get() : nullptr;
 }
 
-const FRigPhysicsSolverDescription* UControlRig::FindPhysicsSolver(const FRigPhysicsSolverID& InID) const
+const FRigPhysicsSolverBase* UControlRig::GetPhysicsSolver(const FRigPhysicsSolverID& InID) const
 {
-	return PhysicsSolvers.FindByPredicate([InID](const FRigPhysicsSolverDescription& Solver) -> bool
-	{
-		return Solver.ID == InID;
-	});
+	const TSharedPtr<FRigPhysicsSolverBase>* PhysicsSolver = PhysicsSolvers.Find(InID);
+	return PhysicsSolver ? PhysicsSolver->Get() : nullptr;
 }
 
-const FRigPhysicsSolverDescription* UControlRig::FindPhysicsSolverByName(const FName& InName) const
+const FRigPhysicsSolverBase* UControlRig::FindPhysicsSolverByName(FName InName) const
 {
-	return PhysicsSolvers.FindByPredicate([InName](const FRigPhysicsSolverDescription& Solver) -> bool
+	for (const TPair<FRigPhysicsSolverID, TSharedPtr<FRigPhysicsSolverBase>>& SolverPair : PhysicsSolvers)
 	{
-		return Solver.Name == InName;
-	});
+		if (SolverPair.Value.IsValid() && SolverPair.Value->Name == InName)
+		{
+			return SolverPair.Value.Get();
+		}
+	}
+	return nullptr;
 }
 
-FRigPhysicsSolverID UControlRig::AddPhysicsSolver(FName InName, bool bSetupUndo, bool bPrintPythonCommand)
+FRigPhysicsSolverID UControlRig::RegisterPhysicsSolver(TSharedPtr<FRigPhysicsSolverBase> InPhysicsWorld, FName InName)
 {
-#if WITH_EDITOR
-	if(CVarControlRigHierarchyEnablePhysics.GetValueOnAnyThread() == false)
-	{
-		return FRigPhysicsSolverID();
-	}
-#endif
-
-	if(InName.IsNone())
+	if(InName.IsNone() || !InPhysicsWorld.IsValid())
 	{
 		return FRigPhysicsSolverID();
 	}
@@ -3652,11 +3628,15 @@ FRigPhysicsSolverID UControlRig::AddPhysicsSolver(FName InName, bool bSetupUndo,
 		return FindPhysicsSolverByName(InName) == nullptr;
 	});
 
-	FRigPhysicsSolverDescription Solver;
-	Solver.ID = FRigPhysicsSolverDescription::MakeID(GetPathName(), Name);
-	Solver.Name = Name;
-	PhysicsSolvers.Add(Solver);
-	return Solver.ID;
+	const FString CompletePath = FString::Printf(TEXT("%s|%s"), *GetPathName(), *Name.ToString());
+	FRigPhysicsSolverID SolverID(FGuid::NewDeterministicGuid(CompletePath));
+
+	PhysicsSolvers.FindOrAdd(SolverID, InPhysicsWorld);
+
+	InPhysicsWorld->Name = Name;
+	InPhysicsWorld->SolverID = SolverID;
+
+	return SolverID;
 }
 
 UControlRig::FPoseScope::FPoseScope(UControlRig* InControlRig, ERigElementType InFilter, const TArray<FRigElementKey>& InElements, const ERigTransformType::Type InTransformType)
diff --git a/Engine/Plugins/Animation/ControlRig/Source/ControlRig/Private/Rigs/RigHierarchy.cpp b/Engine/Plugins/Animation/ControlRig/Source/ControlRig/Private/Rigs/RigHierarchy.cpp
--- a/Engine/Plugins/Animation/ControlRig/Source/ControlRig/Private/Rigs/RigHierarchy.cpp
+++ b/Engine/Plugins/Animation/ControlRig/Source/ControlRig/Private/Rigs/RigHierarchy.cpp
@@ -52,7 +52,6 @@ static TAutoConsoleVariable<int32> CVarControlRigHierarchyTracePrecision(TEXT("C
 static TAutoConsoleVariable<int32> CVarControlRigHierarchyTraceOnSpawn(TEXT("ControlRig.Hierarchy.TraceOnSpawn"), 0, TEXT("sets the number of frames to trace when a new hierarchy is spawned"));
 TAutoConsoleVariable<bool> CVarControlRigHierarchyEnableRotationOrder(TEXT("ControlRig.Hierarchy.EnableRotationOrder"), true, TEXT("enables the rotation order for controls"));
 TAutoConsoleVariable<bool> CVarControlRigHierarchyEnableModules(TEXT("ControlRig.Hierarchy.Modules"), true, TEXT("enables the modular rigging functionality"));
-TAutoConsoleVariable<bool> CVarControlRigHierarchyEnablePhysics(TEXT("ControlRig.Hierarchy.Physics"), false, TEXT("enables physics support for the rig hierarchy"));
 static int32 sRigHierarchyLastTrace = INDEX_NONE;
 
 // A console command to trace a single frame / single execution for a control rig anim node / control rig component
@@ -1680,24 +1679,6 @@ TArray<FRigElementKey> URigHierarchy::RestoreConnectorsFromStates(TArray<FRigCon
 	return Keys;
 }
 
-const FRigPhysicsSolverDescription* URigHierarchy::FindPhysicsSolver(const FRigPhysicsSolverID& InID) const
-{
-	if(const UControlRig* ControlRig = Cast<UControlRig>(GetOuter()))
-	{
-		return ControlRig->FindPhysicsSolver(InID);
-	}
-	return nullptr;
-}
-
-const FRigPhysicsSolverDescription* URigHierarchy::FindPhysicsSolverByName(const FName& InName) const
-{
-	if(const UControlRig* ControlRig = Cast<UControlRig>(GetOuter()))
-	{
-		return ControlRig->FindPhysicsSolverByName(InName);
-	}
-	return nullptr;
-}
-
 TArray<FName> URigHierarchy::GetMetadataNames(FRigElementKey InItem) const
 {
 	TArray<FName> Names;
@@ -5898,7 +5879,7 @@ FRigBaseElement* URigHierarchy::MakeElement(ERigElementType InElementType, int32
 		}
 		case ERigElementType::Physics:
 		{
-			if(OutStructureSize)
+			if (OutStructureSize)
 			{
 				*OutStructureSize = sizeof(FRigPhysicsElement);
 			}
@@ -6005,12 +5986,12 @@ void URigHierarchy::DestroyElement(FRigBaseElement*& InElement)
 		case ERigElementType::Physics:
 		{
 			FRigPhysicsElement* ExistingElements = Cast<FRigPhysicsElement>(InElement);
-			for(int32 Index=0;Index<Count;Index++)
+			for (int32 Index = 0; Index < Count; Index++)
 			{
 				TGuardValue<const FRigBaseElement*> DestroyGuard(ElementBeingDestroyed, &ExistingElements[Index]);
 				DestroyComponents(&ExistingElements[Index]);
 				DeallocateElementStorage(&ExistingElements[Index]);
-				ExistingElements[Index].~FRigPhysicsElement(); 
+				ExistingElements[Index].~FRigPhysicsElement();
 			}
 			break;
 		}
diff --git a/Engine/Plugins/Animation/ControlRig/Source/ControlRig/Private/Rigs/RigHierarchyController.cpp b/Engine/Plugins/Animation/ControlRig/Source/ControlRig/Private/Rigs/RigHierarchyController.cpp
--- a/Engine/Plugins/Animation/ControlRig/Source/ControlRig/Private/Rigs/RigHierarchyController.cpp
+++ b/Engine/Plugins/Animation/ControlRig/Source/ControlRig/Private/Rigs/RigHierarchyController.cpp
@@ -545,81 +545,6 @@ FRigElementKey URigHierarchyController::AddCurve(FName InName, float InValue, bo
 	return NewElement->Key;
 }
 
-FRigElementKey URigHierarchyController::AddPhysicsElement(FName InName, FRigElementKey InParent, FRigPhysicsSolverID InSolver, 
-                                                          FRigPhysicsSettings InSettings, FTransform InLocalTransform, bool bSetupUndo, bool bPrintPythonCommand)
-{
-	if(!IsValid())
-	{
-		return FRigElementKey();
-	}
-
-#if WITH_EDITOR
-	if(CVarControlRigHierarchyEnablePhysics.GetValueOnAnyThread() == false)
-	{
-		return FRigElementKey();
-	}
-#endif
-
-	URigHierarchy* Hierarchy = GetHierarchy();
-
-	if(!InSolver.IsValid())
-	{
-		ReportError(TEXT("Physics Solver Guid is not valid"));
-		return FRigElementKey();
-	}
-
-	if(Hierarchy->FindPhysicsSolver(InSolver) == nullptr)
-	{
-		ReportErrorf(TEXT("Physics Solver with guid '%s' cannot be found."), *InSolver.ToString());
-		return FRigElementKey();
-	}
-
-#if WITH_EDITOR
-	TSharedPtr<FScopedTransaction> TransactionPtr;
-	if(bSetupUndo)
-	{
-		TransactionPtr = MakeShared<FScopedTransaction>(NSLOCTEXT("RigHierarchyController", "Add Physics Element", "Add Physics Element"));
-		Hierarchy->Modify();
-	}
-#endif
-
-	FRigPhysicsElement* NewElement = MakeElement<FRigPhysicsElement>();
-	{
-		TGuardValue<bool> DisableCacheValidityChecks(Hierarchy->bEnableCacheValidityCheck, false);		
-		NewElement->Key.Type = ERigElementType::Physics;
-		NewElement->Key.Name = GetSafeNewName(InName, NewElement->Key.Type);
-		NewElement->Solver = InSolver;
-		NewElement->Settings = InSettings;
-		AddElement(NewElement, Hierarchy->Get(Hierarchy->GetIndex(InParent)), true, InName);
-
-		Hierarchy->SetTransform(NewElement, InLocalTransform, ERigTransformType::InitialLocal, true, false);
-		NewElement->GetTransform().Current = NewElement->GetTransform().Initial;
-		NewElement->GetDirtyState().Current = NewElement->GetDirtyState().Initial;
-	}
-	
-#if WITH_EDITOR
-	TransactionPtr.Reset();
-
-	if (bPrintPythonCommand && !bSuspendPythonPrinting)
-	{
-		UBlueprint* Blueprint = GetTypedOuter<UBlueprint>();
-		if (Blueprint)
-		{
-			TArray<FString> Commands = GetAddPhysicsElementPythonCommands(NewElement);
-			for (const FString& Command : Commands)
-			{
-				RigVMPythonUtils::Print(Blueprint->GetFName().ToString(),
-					FString::Printf(TEXT("%s"), *Command));
-			}
-		}
-	}
-#endif
-
-	GetHierarchy()->EnsureCacheValidity();
-
-	return NewElement->Key;
-}
-
 FRigElementKey URigHierarchyController::AddReference(FName InName, FRigElementKey InParent,
 	FRigReferenceGetWorldTransformDelegate InDelegate, bool bSetupUndo)
 {
@@ -2107,10 +2032,6 @@ TArray<FString> URigHierarchyController::GetAddElementPythonCommands(FRigBaseEle
 	{
 		return GetAddCurvePythonCommands(CurveElement);
 	}
-	else if(FRigPhysicsElement* PhysicsElement = Cast<FRigPhysicsElement>(Element))
-	{
-		return GetAddPhysicsElementPythonCommands(PhysicsElement);
-	}
 	else if(FRigReferenceElement* ReferenceElement = Cast<FRigReferenceElement>(Element))
 	{
 		ensure(false);
@@ -2221,42 +2142,6 @@ TArray<FString> URigHierarchyController::GetAddCurvePythonCommands(FRigCurveElem
 		Hierarchy->GetCurveValue(Curve))};
 }
 
-TArray<FString> URigHierarchyController::GetAddPhysicsElementPythonCommands(FRigPhysicsElement* PhysicsElement) const
-{
-	TArray<FString> Commands;
-	FString TransformStr = RigVMPythonUtils::TransformToPythonString(PhysicsElement->GetTransform().Initial.Local.Get());
-	
-	FString ParentKeyStr = "''";
-	if (PhysicsElement->ParentElement)
-	{
-		ParentKeyStr = PhysicsElement->ParentElement->GetKey().ToPythonString();
-	}
-
-	FString PhysicsElementNamePythonized = RigVMPythonUtils::PythonizeName(PhysicsElement->GetName());
-	FRigPhysicsSettings& Settings = PhysicsElement->Settings;
-	FString SettingsStr;
-	{
-		SettingsStr =FString::Printf(TEXT("physics_settings%s"),
-			*PhysicsElementNamePythonized);
-			
-		Commands.Add(FString::Printf(TEXT("physics_settings%s = unreal.RigPhysicsSettings()"),
-			*PhysicsElementNamePythonized));
-
-		Commands.Add(FString::Printf(TEXT("control_settings_%s.mass = %f"),
-			*PhysicsElementNamePythonized,
-			Settings.Mass));	
-	}
-	
-	// FRigElementKey AddPhysicsElement(FName InName, FRigElementKey InParent, FRigPhysicsSettings InSettings, FTransform InLocalTransform, bool bSetupUndo = false);
-	Commands.Add(FString::Printf(TEXT("hierarchy_controller.add_physics_element('%s', %s, %s, %s)"),
-		*PhysicsElement->GetName(),
-		*ParentKeyStr,
-		*SettingsStr,
-		*TransformStr));
-
-	return Commands;
-}
-
 TArray<FString> URigHierarchyController::GetAddConnectorPythonCommands(FRigConnectorElement* Connector) const
 {
 	TArray<FString> Commands;
diff --git a/Engine/Plugins/Animation/ControlRig/Source/ControlRig/Private/Rigs/RigHierarchyElements.cpp b/Engine/Plugins/Animation/ControlRig/Source/ControlRig/Private/Rigs/RigHierarchyElements.cpp
--- a/Engine/Plugins/Animation/ControlRig/Source/ControlRig/Private/Rigs/RigHierarchyElements.cpp
+++ b/Engine/Plugins/Animation/ControlRig/Source/ControlRig/Private/Rigs/RigHierarchyElements.cpp
@@ -1357,6 +1357,13 @@ void FRigBoneElement::CopyFrom(const FRigBaseElement* InOther)
 
 const FRigBaseElement::EElementIndex FRigNullElement::ElementTypeIndex = NullElement;
 
+////////////////////////////////////////////////////////////////////////////////
+// FRigPhysicsElement
+////////////////////////////////////////////////////////////////////////////////
+
+const FRigPhysicsElement::EElementIndex FRigPhysicsElement::ElementTypeIndex = PhysicsElement;
+
+
 ////////////////////////////////////////////////////////////////////////////////
 // FRigControlSettings
 ////////////////////////////////////////////////////////////////////////////////
@@ -2055,110 +2062,6 @@ void FRigCurveElement::CopyFrom(const FRigBaseElement* InOther)
 	}
 }
 
-////////////////////////////////////////////////////////////////////////////////
-// FRigPhysicsSolver
-////////////////////////////////////////////////////////////////////////////////
-
-void FRigPhysicsSolverDescription::Serialize(FArchive& Ar)
-{
-	if (Ar.IsSaving() || Ar.IsObjectReferenceCollector() || Ar.IsCountingMemory())
-	{
-		Save(Ar);
-	}
-	else if (Ar.IsLoading())
-	{
-		Load(Ar);
-	}
-}
-
-void FRigPhysicsSolverDescription::Save(FArchive& Ar)
-{
-	Ar << ID;
-	Ar << Name;
-}
-
-void FRigPhysicsSolverDescription::Load(FArchive& Ar)
-{
-	Ar << ID;
-	Ar << Name;
-}
-
-FGuid FRigPhysicsSolverDescription::MakeGuid(const FString& InObjectPath, const FName& InSolverName)
-{
-	const FString CompletePath = FString::Printf(TEXT("%s|%s"), *InObjectPath, *InSolverName.ToString());
-	return FGuid::NewDeterministicGuid(CompletePath);
-}
-
-FRigPhysicsSolverID FRigPhysicsSolverDescription::MakeID(const FString& InObjectPath, const FName& InSolverName)
-{
-	return FRigPhysicsSolverID(MakeGuid(InObjectPath, InSolverName));
-}
-
-void FRigPhysicsSolverDescription::CopyFrom(const FRigPhysicsSolverDescription* InOther)
-{
-	if(InOther)
-	{
-		ID = InOther->ID;
-		Name = InOther->Name;
-	}
-}
-
-////////////////////////////////////////////////////////////////////////////////
-// FRigPhysicsSettings
-////////////////////////////////////////////////////////////////////////////////
-
-FRigPhysicsSettings::FRigPhysicsSettings()
-	: Mass(1.f)
-{
-}
-
-void FRigPhysicsSettings::Save(FArchive& Ar)
-{
-	Ar << Mass;
-}
-
-void FRigPhysicsSettings::Load(FArchive& Ar)
-{
-	Ar << Mass;
-}
-
-////////////////////////////////////////////////////////////////////////////////
-// FRigPhysicsElement
-////////////////////////////////////////////////////////////////////////////////
-
-const FRigBaseElement::EElementIndex FRigPhysicsElement::ElementTypeIndex = PhysicsElement;
-
-void FRigPhysicsElement::Save(FArchive& Ar, const FRigHierarchySerializationSettings& InSettings)
-{
-	Super::Save(Ar, InSettings);
-
-	if(InSettings.SerializationPhase == FRigHierarchySerializationSettings::StaticData)
-	{
-		Ar << Solver;
-		Settings.Save(Ar);
-	}
-}
-
-void FRigPhysicsElement::Load(FArchive& Ar, const FRigHierarchySerializationSettings& InSettings)
-{
-	Super::Load(Ar, InSettings);
-
-	if(InSettings.SerializationPhase == FRigHierarchySerializationSettings::StaticData)
-	{
-		Ar << Solver;
-		Settings.Load(Ar);
-	}
-}
-
-void FRigPhysicsElement::CopyFrom(const FRigBaseElement* InOther)
-{
-	Super::CopyFrom(InOther);
-	
-	const FRigPhysicsElement* Source = CastChecked<FRigPhysicsElement>(InOther);
-	Solver = Source->Solver;
-	Settings = Source->Settings;
-}
-
 ////////////////////////////////////////////////////////////////////////////////
 // FRigReferenceElement
 ////////////////////////////////////////////////////////////////////////////////
diff --git a/Engine/Plugins/Animation/ControlRig/Source/ControlRig/Private/Rigs/RigPhysics.cpp b/Engine/Plugins/Animation/ControlRig/Source/ControlRig/Private/Rigs/RigPhysics.cpp
new file mode 100644
--- /dev/null
+++ b/Engine/Plugins/Animation/ControlRig/Source/ControlRig/Private/Rigs/RigPhysics.cpp
@@ -0,0 +1,15 @@
+// Copyright Epic Games, Inc. All Rights Reserved.
+
+#include "Rigs/RigPhysics.h"
+
+#include UE_INLINE_GENERATED_CPP_BY_NAME(RigPhysics)
+
+////////////////////////////////////////////////////////////////////////////////
+// FRigPhysicsSolverDescription_Deprecated
+////////////////////////////////////////////////////////////////////////////////
+
+void FRigPhysicsSolverDescription_Deprecated::Serialize(FArchive& Ar)
+{
+	Ar << ID;
+	Ar << Name;
+}
diff --git a/Engine/Plugins/Animation/ControlRig/Source/ControlRig/Private/Units/Execution/RigUnit_Physics.cpp b/Engine/Plugins/Animation/ControlRig/Source/ControlRig/Private/Units/Execution/RigUnit_Physics.cpp
--- a/Engine/Plugins/Animation/ControlRig/Source/ControlRig/Private/Units/Execution/RigUnit_Physics.cpp
+++ b/Engine/Plugins/Animation/ControlRig/Source/ControlRig/Private/Units/Execution/RigUnit_Physics.cpp
@@ -1,49 +1,27 @@
 // Copyright Epic Games, Inc. All Rights Reserved.
 
 #include "Units/Execution/RigUnit_Physics.h"
-#include "Rigs/RigHierarchyController.h"
-#include "Units/RigUnitContext.h"
-#include "ControlRig.h"
 
 #include UE_INLINE_GENERATED_CPP_BY_NAME(RigUnit_Physics)
 
 FRigUnit_HierarchyAddPhysicsSolver_Execute()
 {
-	FString ErrorMessage;
-	if(!FRigUnit_DynamicHierarchyBase::IsValidToRunInContext(ExecuteContext, true, &ErrorMessage))
-	{
-		if(!ErrorMessage.IsEmpty())
-		{
-			UE_CONTROLRIG_RIGUNIT_REPORT_ERROR(TEXT("%s"), *ErrorMessage);
-		}
-		return;
-	}
-
-	Solver = FRigPhysicsSolverID();
+	UE_CONTROLRIG_RIGUNIT_REPORT_ERROR(TEXT("This node has been deprecated."));
+}
 
-	if(UControlRig* ControlRig = Cast<UControlRig>(ExecuteContext.Hierarchy->GetOuter()))
-	{
-		Solver = ControlRig->AddPhysicsSolver(Name, false, false);
-	}
+FRigVMStructUpgradeInfo FRigUnit_HierarchyAddPhysicsSolver::GetUpgradeInfo() const
+{
+	// nothing to upgrade to.
+	return FRigVMStructUpgradeInfo();
 }
 
 FRigUnit_HierarchyAddPhysicsJoint_Execute()
 {
-	FString ErrorMessage;
-	if(!FRigUnit_DynamicHierarchyBase::IsValidToRunInContext(ExecuteContext, true, &ErrorMessage))
-	{
-		if(!ErrorMessage.IsEmpty())
-		{
-			UE_CONTROLRIG_RIGUNIT_REPORT_ERROR(TEXT("%s"), *ErrorMessage);
-		}
-		return;
-	}
-
-	Item.Reset();
+	UE_CONTROLRIG_RIGUNIT_REPORT_ERROR(TEXT("This node has been deprecated."));
+}
 
-	if(URigHierarchyController* Controller = ExecuteContext.Hierarchy->GetController(true))
-	{
-		FRigHierarchyControllerInstructionBracket InstructionBracket(Controller, ExecuteContext.GetInstructionIndex());
-		Item = Controller->AddPhysicsElement(Name, Parent, Solver, Settings, Transform, false, false);
-	}
+FRigVMStructUpgradeInfo FRigUnit_HierarchyAddPhysicsJoint::GetUpgradeInfo() const
+{
+	// nothing to upgrade to.
+	return FRigVMStructUpgradeInfo();
 }
diff --git a/Engine/Plugins/Animation/ControlRig/Source/ControlRig/Public/ControlRig.h b/Engine/Plugins/Animation/ControlRig/Source/ControlRig/Public/ControlRig.h
--- a/Engine/Plugins/Animation/ControlRig/Source/ControlRig/Public/ControlRig.h
+++ b/Engine/Plugins/Animation/ControlRig/Source/ControlRig/Public/ControlRig.h
@@ -21,6 +21,7 @@
 #include "Animation/MeshDeformerInstance.h"
 #include "Animation/MeshDeformerProducer.h"
 #include "Rigs/RigModuleDefines.h"
+#include "Rigs/RigPhysics.h"
 
 #if WITH_EDITOR
 #include "RigVMModel/RigVMPin.h"
@@ -44,6 +45,7 @@ class UTransformableControlHandle;
 struct FReferenceSkeleton;
 struct FRigUnit;
 struct FRigControl;
+struct FRigPhysicsSolverBase;
 
 CONTROLRIG_API DECLARE_LOG_CATEGORY_EXTERN(LogControlRig, Log, All);
 
@@ -714,50 +716,37 @@ protected:
 private:
 	float DebugBoneRadiusMultiplier;
 
-	// Physics Solvers
-	TArray<FRigPhysicsSolverDescription> PhysicsSolvers;
+	// Physics Solvers/Worlds
+	TMap<FRigPhysicsSolverID, TSharedPtr<FRigPhysicsSolverBase>> PhysicsSolvers;
 
 public:
 
 	/**
-	 * Returns the number of physics solvers
-	 * @return The number of physics solvers
-	 */
-	int32 NumPhysicsSolvers() const; 
-
-	/**
-	 * Returns a physics solver by index (or nullptr)
-	 * @param InIndex The index of the physics solver to return
-	 * @return The physics solver
+	 * Returns the physics world given a solver ID
+	 * @param InID The id identifying the physics solver
+	 * @return The physics solver (may be null)
 	 */
-	const FRigPhysicsSolverDescription* GetPhysicsSolver(int32 InIndex) const; 
+	FRigPhysicsSolverBase* GetPhysicsSolver(const FRigPhysicsSolverID& InID);
 
 	/**
-	 * Finds a new physics solver given its guid
+	 * Returns the physics world given a solver ID
 	 * @param InID The id identifying the physics solver
-	 * @return The physics solver
+	 * @return The physics solver (may be null)
 	 */
-	const FRigPhysicsSolverDescription* FindPhysicsSolver(const FRigPhysicsSolverID& InID) const; 
+	const FRigPhysicsSolverBase* GetPhysicsSolver(const FRigPhysicsSolverID& InID) const;
 
 	/**
-	 * Finds a new physics solver given its name
+	 * Finds a physics solver given its name
 	 * @param InName The name identifying the physics solver in the scope of this control rig
 	 * @return The physics solver
 	 */
-	const FRigPhysicsSolverDescription* FindPhysicsSolverByName(const FName& InName) const; 
+	const FRigPhysicsSolverBase* FindPhysicsSolverByName(FName InName) const; 
 
 	/**
-	 * Adds a physics solver to the hierarchy
-	 * @param InName The suggested name of the new physics solver - will eventually be corrected by the namespace
-	 * @param bSetupUndo If set to true the stack will record the change for undo / redo
-	 * @param bPrintPythonCommand If set to true a python command equivalent to this call will be printed out
-	 * @return The ID for the newly created physics solver.
+	 * Registers and stores (taking ownership) of a physics world, returning an ID for it. The world will be 
+	 * responsible for the low-level storage and simulation of objects.
 	 */
-	UFUNCTION(BlueprintCallable, Category = ControlRig)
-	FRigPhysicsSolverID AddPhysicsSolver(
-		FName InName,
-		bool bSetupUndo = false,
-		bool bPrintPythonCommand = false);
+	FRigPhysicsSolverID RegisterPhysicsSolver(TSharedPtr<FRigPhysicsSolverBase> PhysicsWorld, FName Name);
 
 #if WITH_EDITOR	
 
diff --git a/Engine/Plugins/Animation/ControlRig/Source/ControlRig/Public/ControlRigObjectVersion.h b/Engine/Plugins/Animation/ControlRig/Source/ControlRig/Public/ControlRigObjectVersion.h
--- a/Engine/Plugins/Animation/ControlRig/Source/ControlRig/Public/ControlRigObjectVersion.h
+++ b/Engine/Plugins/Animation/ControlRig/Source/ControlRig/Public/ControlRigObjectVersion.h
@@ -116,6 +116,9 @@ struct CONTROLRIG_API FControlRigObjectVersion
 		// Improve transform compactness when serializing the hierarchy
 		RigHierarchyCompactTransformSerialization,
 
+		// Physics solvers stored on hierarchy
+		ControlRigDoesntStorePhysicsSolvers,
+
 		// -----<new versions can be added above this line>-------------------------------------------------
 		VersionPlusOne,
 		LatestVersion = VersionPlusOne - 1
diff --git a/Engine/Plugins/Animation/ControlRig/Source/ControlRig/Public/Rigs/RigHierarchy.h b/Engine/Plugins/Animation/ControlRig/Source/ControlRig/Public/Rigs/RigHierarchy.h
--- a/Engine/Plugins/Animation/ControlRig/Source/ControlRig/Public/Rigs/RigHierarchy.h
+++ b/Engine/Plugins/Animation/ControlRig/Source/ControlRig/Public/Rigs/RigHierarchy.h
@@ -34,7 +34,6 @@ DECLARE_MULTICAST_DELEGATE_ThreeParams(FRigHierarchyMetadataTagChangedDelegate,
 
 extern CONTROLRIG_API TAutoConsoleVariable<bool> CVarControlRigHierarchyEnableRotationOrder;
 extern CONTROLRIG_API TAutoConsoleVariable<bool> CVarControlRigHierarchyEnableModules;
-extern CONTROLRIG_API TAutoConsoleVariable<bool> CVarControlRigHierarchyEnablePhysics;
 
 UENUM()
 enum ERigTransformStackEntryType : int
@@ -1090,25 +1089,6 @@ public:
 		return GetKeysOfType<FRigCurveElement>(false);
 	}
 
-	/**
-	 * Returns all Physics elements
-	 * @param bTraverse Returns the elements in order of a depth first traversal
-	 */
-	TArray<FRigPhysicsElement*> GetPhysicsElements(bool bTraverse = false) const
-	{
-		return GetElementsOfType<FRigPhysicsElement>(bTraverse);
-	}
-
-	/**
-	 * Returns all Physics elements
-	 * @param bTraverse Returns the elements in order of a depth first traversal
-	 */
-	UFUNCTION(BlueprintCallable, Category = URigHierarchy, meta = (DisplayName = "Get Physics Keys", ScriptName = "GetPhysicsKeys"))
-    TArray<FRigElementKey> GetPhysicsKeys(bool bTraverse = true) const
-	{
-		return GetKeysOfType<FRigPhysicsElement>(bTraverse);
-	}
-
 	/**
 	 * Returns all references
 	 * @param bTraverse Returns the elements in order of a depth first traversal
@@ -1213,20 +1193,6 @@ public:
 		}, false);
 	}
 
-	/**
-	 * Finds a new physics solver given its guid
-	 * @param InID The id identifying the physics solver
-	 * @return The physics solver
-	 */
-	const FRigPhysicsSolverDescription* FindPhysicsSolver(const FRigPhysicsSolverID& InID) const; 
-
-	/**
-	 * Finds a new physics solver given its name
-	 * @param InName The name identifying the physics solver in the scope of this hierarchy
-	 * @return The physics solver
-	 */
-	const FRigPhysicsSolverDescription* FindPhysicsSolverByName(const FName& InName) const; 
-
 	/**
 	 * Returns the name of metadata for a given element
 	 * @param InItem The element key to return the metadata keys for
diff --git a/Engine/Plugins/Animation/ControlRig/Source/ControlRig/Public/Rigs/RigHierarchyController.h b/Engine/Plugins/Animation/ControlRig/Source/ControlRig/Public/Rigs/RigHierarchyController.h
--- a/Engine/Plugins/Animation/ControlRig/Source/ControlRig/Public/Rigs/RigHierarchyController.h
+++ b/Engine/Plugins/Animation/ControlRig/Source/ControlRig/Public/Rigs/RigHierarchyController.h
@@ -200,27 +200,6 @@ public:
 		bool bPrintPythonCommand = false
         );
 
-	/**
-	* Adds a physics element to the hierarchy
-	* @param InName The suggested name of the new physics element - will eventually be corrected by the namespace
-	* @param InParent The (optional) parent of the new physics element. If you don't need a parent, pass FRigElementKey()
-	* @param InSolver The guid identifying the solver to use
-	* @param InSettings All of the physics element's settings
-	* @param InLocalTransform The transform for the new physics element - in the space of the provided parent
-	* @param bSetupUndo If set to true the stack will record the change for undo / redo
-	* @param bPrintPythonCommand If set to true a python command equivalent to this call will be printed out
-	* @return The key for the newly created physics element.
-	*/
-	UFUNCTION(BlueprintCallable, Category = URigHierarchyController)
-    FRigElementKey AddPhysicsElement(
-    	FName InName,
-    	FRigElementKey InParent,
-    	FRigPhysicsSolverID InSolver,
-        FRigPhysicsSettings InSettings,
-    	FTransform InLocalTransform,
-    	bool bSetupUndo = false,
-		bool bPrintPythonCommand = false);
-
 	/**
 	* Adds an reference to the hierarchy
 	* @param InName The suggested name of the new reference - will eventually be corrected by the namespace
@@ -733,8 +712,6 @@ public:
 
 	TArray<FString> GetAddCurvePythonCommands(FRigCurveElement* Curve) const;
 
-	TArray<FString> GetAddPhysicsElementPythonCommands(FRigPhysicsElement* PhysicsElement) const;
-
 	TArray<FString> GetAddConnectorPythonCommands(FRigConnectorElement* Connector) const;
 
 	TArray<FString> GetAddSocketPythonCommands(FRigSocketElement* Socket) const;
diff --git a/Engine/Plugins/Animation/ControlRig/Source/ControlRig/Public/Rigs/RigHierarchyDefines.h b/Engine/Plugins/Animation/ControlRig/Source/ControlRig/Public/Rigs/RigHierarchyDefines.h
--- a/Engine/Plugins/Animation/ControlRig/Source/ControlRig/Public/Rigs/RigHierarchyDefines.h
+++ b/Engine/Plugins/Animation/ControlRig/Source/ControlRig/Public/Rigs/RigHierarchyDefines.h
@@ -29,7 +29,7 @@ enum class ERigElementType : uint8
 	Space = Null UMETA(Hidden),
 	Control = 0x004,
 	Curve = 0x008,
-	Physics = 0x010, 
+	Physics = 0x010 UMETA(Hidden),
 	Reference = 0x020,
 	Connector = 0x040,
 	Socket = 0x080,
diff --git a/Engine/Plugins/Animation/ControlRig/Source/ControlRig/Public/Rigs/RigHierarchyElements.h b/Engine/Plugins/Animation/ControlRig/Source/ControlRig/Public/Rigs/RigHierarchyElements.h
--- a/Engine/Plugins/Animation/ControlRig/Source/ControlRig/Public/Rigs/RigHierarchyElements.h
+++ b/Engine/Plugins/Animation/ControlRig/Source/ControlRig/Public/Rigs/RigHierarchyElements.h
@@ -1362,6 +1362,29 @@ private:
 	friend struct FRigBaseElement;
 };
 
+USTRUCT(BlueprintType)
+struct CONTROLRIG_API FRigPhysicsElement final : public FRigBaseElement
+{
+public:
+
+	GENERATED_BODY()
+	DECLARE_RIG_ELEMENT_METHODS(FRigPhysicsElement)
+
+	static const EElementIndex ElementTypeIndex;
+
+	FRigPhysicsElement() : FRigPhysicsElement(nullptr) {}
+
+private:
+	explicit FRigPhysicsElement(URigHierarchy* InOwner)
+		: FRigBaseElement(InOwner, ERigElementType::Physics)
+	{}
+
+	friend class URigHierarchy;
+	friend struct FRigBaseElement;
+};
+
+
+
 USTRUCT(BlueprintType)
 struct FRigControlElementCustomization
 {
@@ -1814,168 +1837,6 @@ private:
 	friend class FControlRigHierarchyRelinkElementStorage;
 };
 
-USTRUCT(BlueprintType)
-struct CONTROLRIG_API FRigPhysicsSolverID
-{
-public:
-	
-	GENERATED_BODY()
-	
-	FRigPhysicsSolverID()
-		: Guid()
-	{
-	}
-
-	explicit FRigPhysicsSolverID(const FGuid& InGuid)
-		: Guid(InGuid)
-	{
-	}
-
-	UPROPERTY(BlueprintReadOnly, EditAnywhere, Category = Physics, meta=(ShowOnlyInnerProperties))
-	FGuid Guid;
-
-	bool IsValid() const
-	{
-		return Guid.IsValid();
-	}
-
-	FString ToString() const
-	{
-		return Guid.ToString();
-	}
-
-	bool operator == (const FRigPhysicsSolverID& InOther) const
-	{
-		return Guid == InOther.Guid;
-	}
-
-	bool operator != (const FRigPhysicsSolverID& InOther) const
-	{
-		return !(*this == InOther);
-	}
-
-	friend FArchive& operator <<(FArchive& Ar, FRigPhysicsSolverID& ID)
-	{
-		Ar << ID.Guid;
-		return Ar;
-	}
-
-	friend uint32 GetTypeHash(const FRigPhysicsSolverID& InID)
-	{
-		return GetTypeHash(InID.Guid);
-	}
-};
-
-USTRUCT(BlueprintType)
-struct CONTROLRIG_API FRigPhysicsSolverDescription
-{
-public:
-	
-	GENERATED_BODY()
-
-	UPROPERTY(BlueprintReadOnly, EditAnywhere, Category = Physics, meta=(ShowOnlyInnerProperties))
-	FRigPhysicsSolverID ID;
-
-	UPROPERTY(BlueprintReadOnly, EditAnywhere, Category = Physics, meta=(ShowOnlyInnerProperties))
-	FName Name;
-
-	FRigPhysicsSolverDescription()
-	: ID()
-	, Name(NAME_None)
-	{}
-	
-	FRigPhysicsSolverDescription(const FRigPhysicsSolverDescription& InOther)
-	{
-		*this = InOther;
-	}
-	
-	FRigPhysicsSolverDescription& operator=(const FRigPhysicsSolverDescription& InOther)
-	{
-		CopyFrom(&InOther);
-		return *this;
-	}
-
-	void Serialize(FArchive& Ar);
-	void Save(FArchive& Ar);
-	void Load(FArchive& Ar);
-	friend FArchive& operator<<(FArchive& Ar, FRigPhysicsSolverDescription& P)
-	{
-		P.Serialize(Ar);
-		return Ar;
-	}
-
-	static FGuid MakeGuid(const FString& InObjectPath, const FName& InSolverName);
-	static FRigPhysicsSolverID MakeID(const FString& InObjectPath, const FName& InSolverName);
-
-private:
-
-	void CopyFrom(const FRigPhysicsSolverDescription* InOther);
-};
-
-USTRUCT(BlueprintType)
-struct CONTROLRIG_API FRigPhysicsSettings
-{
-	GENERATED_BODY()
-
-	FRigPhysicsSettings();
-
-	void Save(FArchive& Ar);
-	void Load(FArchive& Ar);
-
-
-	UPROPERTY(BlueprintReadWrite, EditAnywhere, Category = Control)
-	float Mass;
-};
-
-USTRUCT(BlueprintType)
-struct CONTROLRIG_API FRigPhysicsElement : public FRigSingleParentElement
-{
-public:
-	
-	GENERATED_BODY()
-	DECLARE_RIG_ELEMENT_METHODS(FRigPhysicsElement)
-
-	static const EElementIndex ElementTypeIndex;
-
-	UPROPERTY(BlueprintReadOnly, EditAnywhere, Category = Physics, meta=(ShowOnlyInnerProperties))
-	FRigPhysicsSolverID Solver;
-
-	UPROPERTY(BlueprintReadOnly, EditAnywhere, Category = Physics, meta=(ShowOnlyInnerProperties))
-	FRigPhysicsSettings Settings;
-	
-	FRigPhysicsElement()
-        : FRigPhysicsElement(nullptr)
-	{ }
-	
-	FRigPhysicsElement(const FRigPhysicsElement& InOther)
-	{
-		*this = InOther;
-	}
-	
-	FRigPhysicsElement& operator=(const FRigPhysicsElement& InOther)
-	{
-		Super::operator=(InOther);
-		Solver = InOther.Solver;
-		Settings = InOther.Settings;
-		return *this;
-	}
-	
-	virtual ~FRigPhysicsElement() override {}
-
-	virtual void Save(FArchive& A, const FRigHierarchySerializationSettings& InSettings) override;
-	virtual void Load(FArchive& Ar, const FRigHierarchySerializationSettings& InSettings) override;
-
-private:
-	explicit FRigPhysicsElement(URigHierarchy* InOwner)
-		: FRigSingleParentElement(InOwner, ERigElementType::Physics)
-	{ }
-	
-	virtual void CopyFrom(const FRigBaseElement* InOther) override;
-
-	friend class URigHierarchy;
-	friend struct FRigBaseElement;
-};
-
 USTRUCT(BlueprintType)
 struct CONTROLRIG_API FRigReferenceElement final : public FRigSingleParentElement
 {
diff --git a/Engine/Plugins/Animation/ControlRig/Source/ControlRig/Public/Rigs/RigPhysics.h b/Engine/Plugins/Animation/ControlRig/Source/ControlRig/Public/Rigs/RigPhysics.h
new file mode 100644
--- /dev/null
+++ b/Engine/Plugins/Animation/ControlRig/Source/ControlRig/Public/Rigs/RigPhysics.h
@@ -0,0 +1,92 @@
+// Copyright Epic Games, Inc. All Rights Reserved.
+
+#pragma once
+
+#include "CoreMinimal.h"
+#include "RigHierarchyComponents.h"
+#include "RigPhysics.generated.h"
+
+// This simply exposes FGuid to blueprint, and thus custom UI at some point.
+USTRUCT(BlueprintType)
+struct CONTROLRIG_API FRigPhysicsSolverID
+{
+public:
+	
+	GENERATED_BODY()
+	
+	UPROPERTY(BlueprintReadOnly, EditAnywhere, Category = Physics, meta=(ShowOnlyInnerProperties))
+	FGuid Guid;
+
+	FRigPhysicsSolverID() : Guid() { }
+	explicit FRigPhysicsSolverID(const FGuid& InGuid) : Guid(InGuid) { }
+
+	bool IsValid() const
+	{
+		return Guid.IsValid();
+	}
+
+	FString ToString() const
+	{
+		return Guid.ToString();
+	}
+
+	bool operator == (const FRigPhysicsSolverID& InOther) const
+	{
+		return Guid == InOther.Guid;
+	}
+
+	bool operator != (const FRigPhysicsSolverID& InOther) const
+	{
+		return !(*this == InOther);
+	}
+
+	friend FArchive& operator <<(FArchive& Ar, FRigPhysicsSolverID& ID)
+	{
+		Ar << ID.Guid;
+		return Ar;
+	}
+
+	friend uint32 GetTypeHash(const FRigPhysicsSolverID& InID)
+	{
+		return GetTypeHash(InID.Guid);
+	}
+};
+
+// Danny TODO This is just inheriting from the base component to acquire the casting functions - is there a better way?
+USTRUCT(BlueprintType)
+struct CONTROLRIG_API FRigPhysicsSolverBase : public FRigBaseComponent
+{
+	GENERATED_BODY()
+	DECLARE_RIG_COMPONENT_METHODS(FRigPhysicsSolverBase)
+
+	virtual ~FRigPhysicsSolverBase() {}
+
+	FName               Name;
+	FRigPhysicsSolverID SolverID;
+};
+
+USTRUCT()
+struct CONTROLRIG_API FRigPhysicsSolverDescription_Deprecated
+{
+public:
+	
+	GENERATED_BODY()
+
+	UPROPERTY()
+	FRigPhysicsSolverID ID;
+
+	UPROPERTY()
+	FName Name;
+
+	FRigPhysicsSolverDescription_Deprecated()
+	: ID()
+	, Name(NAME_None)
+	{}
+	
+	void Serialize(FArchive& Ar);
+	friend FArchive& operator<<(FArchive& Ar, FRigPhysicsSolverDescription_Deprecated& P)
+	{
+		P.Serialize(Ar);
+		return Ar;
+	}
+};
\ No newline at end of file
diff --git a/Engine/Plugins/Animation/ControlRig/Source/ControlRig/Public/Units/Execution/RigUnit_Physics.h b/Engine/Plugins/Animation/ControlRig/Source/ControlRig/Public/Units/Execution/RigUnit_Physics.h
--- a/Engine/Plugins/Animation/ControlRig/Source/ControlRig/Public/Units/Execution/RigUnit_Physics.h
+++ b/Engine/Plugins/Animation/ControlRig/Source/ControlRig/Public/Units/Execution/RigUnit_Physics.h
@@ -2,14 +2,11 @@
 
 #pragma once
 
+#include "Rigs/RigPhysics.h"
 #include "RigUnit_DynamicHierarchy.h"
 #include "RigUnit_Physics.generated.h"
 
-/**
- * Adds a new physics solver to the hierarchy
- * Note: This node only runs as part of the construction event.
- */
-USTRUCT(meta=(DisplayName="Spawn Physics Solver", Keywords="Construction,Create,New,Simulation", Varying))
+USTRUCT(meta=(DisplayName="Spawn Physics Solver", Keywords="Construction,Create,New,Simulation", Varying, Deprecated = "5.6"))
 struct CONTROLRIG_API FRigUnit_HierarchyAddPhysicsSolver : public FRigUnit_DynamicHierarchyBaseMutable
 {
 	GENERATED_BODY()
@@ -20,27 +17,20 @@ struct CONTROLRIG_API FRigUnit_HierarchyAddPhysicsSolver : public FRigUnit_Dynam
 		Solver = FRigPhysicsSolverID();
 	}
 
-	/*
-	 * The name of the new solver to add
-	 */
 	UPROPERTY(meta = (Input))
 	FName Name;
 
-	/*
-	 * The identifier of the solver
-	 */
 	UPROPERTY(meta = (Output))
 	FRigPhysicsSolverID Solver;
 
 	RIGVM_METHOD()
 	virtual void Execute() override;
+
+	RIGVM_METHOD()
+	virtual FRigVMStructUpgradeInfo GetUpgradeInfo() const override;
 };
 
-/**
- * Adds a new physics joint to the hierarchy
- * Note: This node only runs as part of the construction event.
- */
-USTRUCT(meta=(DisplayName="Spawn Physics Joint", Keywords="Construction,Create,New,Joint,Skeleton", Varying))
+USTRUCT(meta=(DisplayName="Spawn Physics Joint", Keywords="Construction,Create,New,Joint,Skeleton", Varying, Deprecated = "5.6"))
 struct CONTROLRIG_API FRigUnit_HierarchyAddPhysicsJoint : public FRigUnit_HierarchyAddElement
 {
 	GENERATED_BODY()
@@ -59,22 +49,13 @@ struct CONTROLRIG_API FRigUnit_HierarchyAddPhysicsJoint : public FRigUnit_Hierar
 	RIGVM_METHOD()
 	virtual void Execute() override;
 
-	/*
-	 * The initial global transform of the spawned element
-	 */
 	UPROPERTY(meta = (Input))
 	FTransform Transform;
 
-	/*
-	 * The solver to relate this new physics element to
-	 */
 	UPROPERTY(meta = (Input))
 	FRigPhysicsSolverID Solver;
 
-	/*
-	 * The settings of the new physics element
-	 */
-	UPROPERTY(meta = (Input))
-	FRigPhysicsSettings Settings;
+	RIGVM_METHOD()
+	virtual FRigVMStructUpgradeInfo GetUpgradeInfo() const override;
 };
 
diff --git a/Engine/Plugins/Animation/ControlRig/Source/ControlRigEditor/Private/ControlRigElementDetails.cpp b/Engine/Plugins/Animation/ControlRig/Source/ControlRigEditor/Private/ControlRigElementDetails.cpp
--- a/Engine/Plugins/Animation/ControlRig/Source/ControlRigEditor/Private/ControlRigElementDetails.cpp
+++ b/Engine/Plugins/Animation/ControlRig/Source/ControlRigEditor/Private/ControlRigElementDetails.cpp
@@ -6476,98 +6476,4 @@ void FRigConnectionRuleDetails::OnRuleContentChanged()
 	StructPropertyHandle->GetParentHandle()->NotifyPostChange(EPropertyChangeType::ValueSet);
 }
 
-void FRigPhysicsElementDetails::CustomizeDetails(IDetailLayoutBuilder& DetailBuilder)
-{
-	FRigTransformElementDetails::CustomizeDetails(DetailBuilder);
-	CustomizeSettings(DetailBuilder);
-	CustomizeTransform(DetailBuilder);
-	CustomizeComponents(DetailBuilder);
-	CustomizeMetadata(DetailBuilder);
-}
-
-void FRigPhysicsElementDetails::CustomizeSettings(IDetailLayoutBuilder& DetailBuilder)
-{
-	if(PerElementInfos.IsEmpty())
-	{
-		return;
-	}
-
-	if(IsAnyElementNotOfType(ERigElementType::Physics))
-	{
-		return;
-	}
-
-	const TSharedPtr<IPropertyHandle> SolverHandle = DetailBuilder.GetProperty(GET_MEMBER_NAME_CHECKED(FRigPhysicsElement, Solver));
-	DetailBuilder.HideProperty(SolverHandle);
-
-	const TSharedPtr<IPropertyHandle> SettingsHandle = DetailBuilder.GetProperty(GET_MEMBER_NAME_CHECKED(FRigPhysicsElement, Settings));
-	DetailBuilder.HideProperty(SettingsHandle);
-
-	IDetailCategoryBuilder& PhysicsCategory = DetailBuilder.EditCategory(TEXT("Physics"), LOCTEXT("Physics", "Physics"));
-
-	FDetailWidgetRow& SolverNameRow = PhysicsCategory.AddCustomRow(LOCTEXT("SolverName", "Solver Name"));
-	SolverNameRow.IsEnabled(false);
-	SolverNameRow.NameContent()
-	[
-		SNew(STextBlock)
-		.Text(LOCTEXT("SolverName", "Solver Name"))
-		.Font(IDetailLayoutBuilder::GetDetailFont())
-	];
-	SolverNameRow.ValueContent()
-	[
-		SNew(SEditableText)
-		.Text(this, &FRigPhysicsElementDetails::GetSolverNameText)
-		.Font(IDetailLayoutBuilder::GetDetailFont())
-	];
-
-	PhysicsCategory
-	.AddProperty(SolverHandle)
-	.DisplayName(LOCTEXT("SolverGuid", "Solver Guid"))
-	.IsEnabled(false);
-
-	const bool bIsProcedural = IsAnyElementProcedural();
-
-	PhysicsCategory
-		.AddProperty(SettingsHandle->GetChildHandle(GET_MEMBER_NAME_CHECKED(FRigPhysicsSettings, Mass)))
-		.IsEnabled(!bIsProcedural);
-}
-
-FText FRigPhysicsElementDetails::GetSolverNameText() const
-{
-	if(PerElementInfos.Num() > 0)
-	{
-		if(const URigHierarchy* Hierarchy = PerElementInfos[0].GetHierarchy())
-		{
-			TOptional<FText> FirstValue;
-			
-			for(const FPerElementInfo& Info : PerElementInfos)
-			{
-				if(const FRigPhysicsElement* PhysicsElement = Info.Element.Get<FRigPhysicsElement>())
-				{
-					if(const FRigPhysicsSolverDescription* Solver = Hierarchy->FindPhysicsSolver(PhysicsElement->Solver))
-					{
-						const FText SolverName = FText::FromName(Solver->Name);
-						if(FirstValue.IsSet())
-						{
-							if(!FirstValue.GetValue().EqualTo(SolverName))
-							{
-								return ControlRigDetailsMultipleValues;	
-							}
-						}
-						else
-						{
-							FirstValue = SolverName;
-						}
-					}
-				}
-			}
-
-			if(FirstValue.IsSet())
-			{
-				return FirstValue.GetValue();
-			}
-		}
-	}
-	return FText();
-}
 #undef LOCTEXT_NAMESPACE
diff --git a/Engine/Plugins/Animation/ControlRig/Source/ControlRigEditor/Private/ControlRigElementDetails.h b/Engine/Plugins/Animation/ControlRig/Source/ControlRigEditor/Private/ControlRigElementDetails.h
--- a/Engine/Plugins/Animation/ControlRig/Source/ControlRigEditor/Private/ControlRigElementDetails.h
+++ b/Engine/Plugins/Animation/ControlRig/Source/ControlRigEditor/Private/ControlRigElementDetails.h
@@ -921,22 +921,3 @@ protected:
 	TAttribute<bool> EnabledAttribute;
 };
 
-class FRigPhysicsElementDetails : public FRigTransformElementDetails
-{
-public:
-
-	// Makes a new instance of this detail layout class for a specific detail view requesting it
-	static TSharedRef<IDetailCustomization> MakeInstance()
-	{
-		return MakeShareable(new FRigPhysicsElementDetails);
-	}
-
-	/** IDetailCustomization interface */
-	virtual void CustomizeDetails(IDetailLayoutBuilder& DetailBuilder) override;
-
-	void CustomizeSettings(IDetailLayoutBuilder& DetailBuilder);
-
-private:
-
-	FText GetSolverNameText() const;
-};
diff --git a/Engine/Plugins/Animation/ControlRig/Source/ControlRigEditor/Private/EditMode/SControlRigOutliner.cpp b/Engine/Plugins/Animation/ControlRig/Source/ControlRigEditor/Private/EditMode/SControlRigOutliner.cpp
--- a/Engine/Plugins/Animation/ControlRig/Source/ControlRigEditor/Private/EditMode/SControlRigOutliner.cpp
+++ b/Engine/Plugins/Animation/ControlRig/Source/ControlRigEditor/Private/EditMode/SControlRigOutliner.cpp
@@ -521,15 +521,7 @@ bool SMultiRigHierarchyTreeView::AddElement(UControlRig* InControlRig, const FRi
 			}
 			case ERigElementType::Physics:
 			{
-				if (!Settings.bShowPhysics)
-				{
-					return false;
-				}
-				if(CVarControlRigHierarchyEnablePhysics.GetValueOnAnyThread() == false)
-				{
-					return false;
-				}
-				break;
+				return false;
 			}
 			case ERigElementType::Reference:
 			{
diff --git a/Engine/Plugins/Animation/ControlRig/Source/ControlRigEditor/Private/Editor/ControlRigWrapperObject.cpp b/Engine/Plugins/Animation/ControlRig/Source/ControlRigEditor/Private/Editor/ControlRigWrapperObject.cpp
--- a/Engine/Plugins/Animation/ControlRig/Source/ControlRigEditor/Private/Editor/ControlRigWrapperObject.cpp
+++ b/Engine/Plugins/Animation/ControlRig/Source/ControlRigEditor/Private/Editor/ControlRigWrapperObject.cpp
@@ -45,10 +45,6 @@ UClass* UControlRigWrapperObject::GetClassForStruct(UScriptStruct* InStruct, boo
 			{
 				PropertyEditorModule.RegisterCustomClassLayout(WrapperClassName, FOnGetDetailCustomizationInstance::CreateStatic(&FRigSocketElementDetails::MakeInstance));
 			}
-			else if(InStruct == FRigPhysicsElement::StaticStruct())
-			{
-				PropertyEditorModule.RegisterCustomClassLayout(WrapperClassName, FOnGetDetailCustomizationInstance::CreateStatic(&FRigPhysicsElementDetails::MakeInstance));
-			}
 		}
 	}
 #endif
diff --git a/Engine/Plugins/Animation/ControlRig/Source/ControlRigEditor/Private/Editor/SRigHierarchyTreeView.cpp b/Engine/Plugins/Animation/ControlRig/Source/ControlRigEditor/Private/Editor/SRigHierarchyTreeView.cpp
--- a/Engine/Plugins/Animation/ControlRig/Source/ControlRigEditor/Private/Editor/SRigHierarchyTreeView.cpp
+++ b/Engine/Plugins/Animation/ControlRig/Source/ControlRigEditor/Private/Editor/SRigHierarchyTreeView.cpp
@@ -563,15 +563,7 @@ bool SRigHierarchyTreeView::AddElement(const FRigBaseElement* InElement)
 		}
 		case ERigElementType::Physics:
 		{
-			if(!Settings.bShowPhysics)
-			{
-				return false;
-			}
-			if(CVarControlRigHierarchyEnablePhysics.GetValueOnAnyThread() == false)
-			{
-				return false;
-			}
-			break;
+			return false;
 		}
 		case ERigElementType::Reference:
 		{
diff --git a/Engine/Plugins/Runtime/RigVM/Source/RigVM/Private/RigVMCore/RigVMDrawInterface.cpp b/Engine/Plugins/Runtime/RigVM/Source/RigVM/Private/RigVMCore/RigVMDrawInterface.cpp
--- a/Engine/Plugins/Runtime/RigVM/Source/RigVM/Private/RigVMCore/RigVMDrawInterface.cpp
+++ b/Engine/Plugins/Runtime/RigVM/Source/RigVM/Private/RigVMCore/RigVMDrawInterface.cpp
@@ -259,13 +259,12 @@ void FRigVMDrawInterface::DrawArc(const FTransform& WorldOffset, const FTransfor
 	V = Rotation.RotateVector(V);
 	Instruction.Positions.Add(V);
 	float StepAngle = (MaximumAngle - MinimumAngle) / float(Count);
-	if (FMath::Abs<float>(MaximumAngle - MinimumAngle) >= PI * 2.f - SMALL_NUMBER)
+	if (FMath::Abs<float>(MaximumAngle - MinimumAngle) >= TWO_PI)
 	{
-		StepAngle = (PI * 2.f) / float(Count);
-		Count++;
+		StepAngle = TWO_PI / float(Count);
 	}
 	Rotation = FQuat(FVector(0.f, 0.f, 1.f), StepAngle);
-	for(int32 Index=1;Index<Count;Index++)
+	for(int32 Index=0;Index<Count;Index++)
 	{
 		V = Rotation.RotateVector(V);
 		Instruction.Positions.Add(V);
@@ -280,7 +279,53 @@ void FRigVMDrawInterface::DrawCircle(const FTransform& WorldOffset, const FTrans
 	DrawArc(WorldOffset, Transform, Radius, 0, TWO_PI, Color, Thickness, Detail);
 }
 
-void FRigVMDrawInterface::DrawCone(const FTransform& WorldOffset, const FTransform& ConeOffset, float Angle1, float Angle2, uint32 NumSides, bool bDrawSideLines, const FLinearColor& SideLineColor, FMaterialRenderProxy* const MaterialRenderProxy)
+void FRigVMDrawInterface::DrawSphere(const FTransform& WorldOffset, const FTransform& Transform, float Radius, 
+	const FLinearColor& Color, float Thickness, int32 Detail)
+{
+	// Circles are drawn in the XY plane
+	FTransform R1(FQuat(FVector(1.f, 0.f, 0.f), HALF_PI));
+	FTransform R2(FQuat(FVector(0.f, 1.f, 0.f), HALF_PI));
+	DrawCircle(WorldOffset, Transform, Radius, Color, Thickness, Detail);
+	DrawCircle(WorldOffset, R1 * Transform, Radius, Color, Thickness, Detail);
+	DrawCircle(WorldOffset, R2 * Transform, Radius, Color, Thickness, Detail);
+}
+
+void FRigVMDrawInterface::DrawHemisphere(const FTransform& WorldOffset, const FTransform& Transform, float Radius,
+	const FLinearColor& Color, float Thickness, int32 Detail)
+{
+	// Circles are drawn in the XY plane
+	FTransform R1(FQuat(FVector(1.f, 0.f, 0.f), HALF_PI));
+	FTransform R2(FQuat(FVector(0.f, 1.f, 0.f), -HALF_PI));
+	DrawCircle(WorldOffset, Transform, Radius, Color, Thickness, Detail);
+	DrawArc(WorldOffset, R1 * Transform, Radius, 0, PI, Color, Thickness, Detail);
+	DrawArc(WorldOffset, R2 * Transform, Radius, -HALF_PI, HALF_PI, Color, Thickness, Detail);
+}
+
+void FRigVMDrawInterface::DrawCapsule(const FTransform& WorldOffset, const FTransform& Transform,
+	float Radius, float Length, const FLinearColor& Color, float Thickness, int32 Detail)
+{
+	FTransform T1(FVector(0, 0, Length * 0.5));
+	FTransform T2(FQuat(FVector(0.f, 1.f, 0.f), PI), FVector(0, 0, -Length * 0.5));
+
+	DrawHemisphere(WorldOffset, T1 * Transform, Radius, Color, Thickness, Detail);
+	DrawHemisphere(WorldOffset, T2 * Transform, Radius, Color, Thickness, Detail);
+
+	DrawLine(WorldOffset, 
+		Transform.TransformPosition(FVector(Radius, 0, Length * 0.5)), 
+		Transform.TransformPosition(FVector(Radius, 0, -Length * 0.5)), Color, Thickness);
+	DrawLine(WorldOffset,
+		Transform.TransformPosition(FVector(-Radius, 0, Length * 0.5)),
+		Transform.TransformPosition(FVector(-Radius, 0, -Length * 0.5)), Color, Thickness);
+	DrawLine(WorldOffset,
+		Transform.TransformPosition(FVector(0, Radius, Length * 0.5)),
+		Transform.TransformPosition(FVector(0, Radius, -Length * 0.5)), Color, Thickness);
+	DrawLine(WorldOffset,
+		Transform.TransformPosition(FVector(0, -Radius, Length * 0.5)),
+		Transform.TransformPosition(FVector(0, -Radius, -Length * 0.5)), Color, Thickness);
+}
+
+
+void FRigVMDrawInterface::DrawCone(const FTransform& WorldOffset, const FTransform& ConeOffset, float Angle1, float Angle2, uint32 NumSides, bool bDrawSideLines, const FLinearColor& SideLineColor, FMaterialRenderProxy* const MaterialRenderProxy, float SideLineThickness)
 {
 	if (!IsEnabled())
 	{
@@ -311,7 +356,7 @@ void FRigVMDrawInterface::DrawCone(const FTransform& WorldOffset, const FTransfo
 			float Fraction = (float)i / (float)(4);
 			float Azi = 2.f * PI * Fraction;
 			FVector ConeVert = CalcConeVert(Angle1, Angle2, Azi);
-			DrawLine(WorldOffset, ConeOffset.GetLocation(), ConeOffset.TransformPosition(ConeVert), SideLineColor, 1.f);
+			DrawLine(WorldOffset, ConeOffset.GetLocation(), ConeOffset.TransformPosition(ConeVert), SideLineColor, SideLineThickness);
 		}
 	}
 }
diff --git a/Engine/Plugins/Runtime/RigVM/Source/RigVM/Public/RigVMCore/RigVMDrawInterface.h b/Engine/Plugins/Runtime/RigVM/Source/RigVM/Public/RigVMCore/RigVMDrawInterface.h
--- a/Engine/Plugins/Runtime/RigVM/Source/RigVM/Public/RigVMCore/RigVMDrawInterface.h
+++ b/Engine/Plugins/Runtime/RigVM/Source/RigVM/Public/RigVMCore/RigVMDrawInterface.h
@@ -20,14 +20,19 @@ public:
 	void DrawLines(const FTransform& WorldOffset, const TArrayView<const FVector>& Positions, const FLinearColor& Color, float Thickness = 0.f);
 	void DrawLineStrip(const FTransform& WorldOffset, const TArrayView<const FVector>& Positions, const FLinearColor& Color, float Thickness = 0.f);
 	void DrawBox(const FTransform& WorldOffset, const FTransform& Transform, const FLinearColor& Color, float Thickness = 0.f);
+	void DrawSphere(const FTransform& WorldOffset, const FTransform& Transform, float Radius, const FLinearColor& Color, float Thickness = 0.f, int32 Detail = 12);
+	// Hemispheres are drawn so that they extend up along +ve Z
+	void DrawHemisphere(const FTransform& WorldOffset, const FTransform& Transform, float Radius, const FLinearColor& Color, float Thickness = 0.f, int32 Detail = 12);
+	void DrawCapsule(const FTransform& WorldOffset, const FTransform& Transform, float Radius, float Length, const FLinearColor& Color, float Thickness = 0.f, int32 Detail = 12);
 	void DrawAxes(const FTransform& WorldOffset, const FTransform& Transform, float Size, float Thickness = 0.f);
 	void DrawAxes(const FTransform& WorldOffset, const TArrayView<const FTransform>& Transforms, float Size, float Thickness = 0.f);
 	void DrawAxes(const FTransform& WorldOffset, const FTransform& Transform, const FLinearColor& InColor, float Size, float Thickness = 0.f);
 	void DrawAxes(const FTransform& WorldOffset, const TArrayView<const FTransform>& Transforms, const FLinearColor& InColor, float Size, float Thickness = 0.f);
 	void DrawRectangle(const FTransform& WorldOffset, const FTransform& Transform, float Size, const FLinearColor& Color, float Thickness);
 	void DrawArc(const FTransform& WorldOffset, const FTransform& Transform, float Radius, float MinimumAngle, float MaximumAngle, const FLinearColor& Color, float Thickness, int32 Detail);
+	// Circles are drawn in the XY plane
 	void DrawCircle(const FTransform& WorldOffset, const FTransform& Transform, float Radius, const FLinearColor& Color, float Thickness, int32 Detail);
-	void DrawCone(const FTransform& WorldOffset, const FTransform& ConeOffset, float Angle1, float Angle2, uint32 NumSides, bool bDrawSideLines, const FLinearColor& SideLineColor, FMaterialRenderProxy* const MaterialRenderProxy);
+	void DrawCone(const FTransform& WorldOffset, const FTransform& ConeOffset, float Angle1, float Angle2, uint32 NumSides, bool bDrawSideLines, const FLinearColor& SideLineColor, FMaterialRenderProxy* const MaterialRenderProxy, float SideLineThickness = 1.0f);
 	void DrawArrow(const FTransform& WorldOffset, const FVector& Direction, const FVector& Side, const FLinearColor& InColor, float Thickness);
 	void DrawPlane(const FTransform& WorldOffset, const FVector2D& Scale, const FLinearColor& MeshColor, bool bDrawLines, const FLinearColor& LineColor, FMaterialRenderProxy* const MaterialRenderProxy);
 
