Control Rig: ArrayConnector further work + tests

Cook preflight: https://horde.devtools.epicgames.com/job/67484b033d44f38e74b2eef0

#rnx

diff --git a/Engine/Plugins/Animation/ControlRig/Source/ControlRig/Private/ControlRig.cpp b/Engine/Plugins/Animation/ControlRig/Source/ControlRig/Private/ControlRig.cpp
--- a/Engine/Plugins/Animation/ControlRig/Source/ControlRig/Private/ControlRig.cpp
+++ b/Engine/Plugins/Animation/ControlRig/Source/ControlRig/Private/ControlRig.cpp
@@ -144,13 +144,19 @@ void UControlRig::BeginDestroy()
 #if WITH_EDITOR
 	if (!HasAnyFlags(RF_ClassDefaultObject))
 	{
-		if(UControlRig* CDO = GetClass()->GetDefaultObject<UControlRig>())
+		if(const UClass* Class = GetClass())
 		{
-			if (!IsGarbageOrDestroyed(CDO))
+			if(IsValid(Class))
 			{
-				if (CDO->GetHierarchy())
+				if(UControlRig* CDO = Class->GetDefaultObject<UControlRig>())
 				{
-					CDO->GetHierarchy()->UnregisterListeningHierarchy(GetHierarchy());
+					if (!IsGarbageOrDestroyed(CDO))
+					{
+						if (CDO->GetHierarchy())
+						{
+							CDO->GetHierarchy()->UnregisterListeningHierarchy(GetHierarchy());
+						}
+					}
 				}
 			}
 		}
diff --git a/Engine/Plugins/Animation/ControlRig/Source/ControlRig/Private/ModularRigController.cpp b/Engine/Plugins/Animation/ControlRig/Source/ControlRig/Private/ModularRigController.cpp
--- a/Engine/Plugins/Animation/ControlRig/Source/ControlRig/Private/ModularRigController.cpp
+++ b/Engine/Plugins/Animation/ControlRig/Source/ControlRig/Private/ModularRigController.cpp
@@ -145,6 +145,12 @@ TArray<FRigElementKey> UModularRigController::GetConnectorsForModule(FString InM
 }
 
 bool UModularRigController::CanConnectConnectorToElement(const FRigElementKey& InConnectorKey, const FRigElementKey& InTargetKey, FText& OutErrorMessage)
+{
+	return CanConnectConnectorToElements(InConnectorKey, {InTargetKey}, OutErrorMessage);
+}
+
+bool UModularRigController::CanConnectConnectorToElements(const FRigElementKey& InConnectorKey, const TArray<FRigElementKey>& InTargetKeys,
+	FText& OutErrorMessage)
 {
 	FString ConnectorModulePath, ConnectorName;
 	if (!URigHierarchy::SplitNameSpace(InConnectorKey.Name.ToString(), &ConnectorModulePath, &ConnectorName))
@@ -178,39 +184,48 @@ bool UModularRigController::CanConnectConnectorToElement(const FRigElementKey& I
 		return false;
 	}
 
-	if (!InTargetKey.IsValid())
+	if (InTargetKeys.IsEmpty())
 	{
-		OutErrorMessage = FText::FromString(FString::Printf(TEXT("Invalid target %s in class %s"), *InTargetKey.ToString(), *Module->Class->GetPathName()));
+		OutErrorMessage = FText::FromString(TEXT("Invalid empty target array"));
 		return false;
 	}
 
-	if (InTargetKey == InConnectorKey)
+	for(const FRigElementKey& InTargetKey : InTargetKeys)
 	{
-		OutErrorMessage = FText::FromString(FString::Printf(TEXT("Cannot resolve connector %s to itself in class %s"), *InTargetKey.ToString(), *Module->Class->GetPathName()));
-		return false;
-	}
+		if (!InTargetKey.IsValid())
+		{
+			OutErrorMessage = FText::FromString(FString::Printf(TEXT("Invalid target %s in class %s"), *InTargetKey.ToString(), *Module->Class->GetPathName()));
+			return false;
+		}
 
-	FRigElementKey CurrentTarget = Model->Connections.FindTargetFromConnector(InConnectorKey);
-	if (CurrentTarget.IsValid() && InTargetKey == CurrentTarget)
-	{
-		return true; // Nothing to do
-	}
+		if (InTargetKey == InConnectorKey)
+		{
+			OutErrorMessage = FText::FromString(FString::Printf(TEXT("Cannot resolve connector %s to itself in class %s"), *InTargetKey.ToString(), *Module->Class->GetPathName()));
+			return false;
+		}
 
-	if (!ModuleConnector->IsPrimary())
-	{
-		const FRigModuleConnector* PrimaryModuleConnector = RigCDO->GetRigModuleSettings().ExposedConnectors.FindByPredicate(
-		[](const FRigModuleConnector& Connector)
+		FRigElementKey CurrentTarget = Model->Connections.FindTargetFromConnector(InConnectorKey);
+		if (CurrentTarget.IsValid() && InTargetKey == CurrentTarget)
 		{
-			return Connector.IsPrimary();
-		});
+			return true; // Nothing to do
+		}
 
-		const FString PrimaryConnectorPath = FString::Printf(TEXT("%s:%s"), *ConnectorModulePath, *PrimaryModuleConnector->Name);
-		const FRigElementKey PrimaryConnectorKey(*PrimaryConnectorPath, ERigElementType::Connector);
-		const FRigElementKey PrimaryTarget = Model->Connections.FindTargetFromConnector(PrimaryConnectorKey);
-		if (!PrimaryTarget.IsValid())
+		if (!ModuleConnector->IsPrimary())
 		{
-			OutErrorMessage = FText::FromString(FString::Printf(TEXT("Cannot resolve connector %s because primary connector is not resolved"), *InConnectorKey.ToString()));
-			return false;
+			const FRigModuleConnector* PrimaryModuleConnector = RigCDO->GetRigModuleSettings().ExposedConnectors.FindByPredicate(
+			[](const FRigModuleConnector& Connector)
+			{
+				return Connector.IsPrimary();
+			});
+
+			const FString PrimaryConnectorPath = FString::Printf(TEXT("%s:%s"), *ConnectorModulePath, *PrimaryModuleConnector->Name);
+			const FRigElementKey PrimaryConnectorKey(*PrimaryConnectorPath, ERigElementType::Connector);
+			const FRigElementKey PrimaryTarget = Model->Connections.FindTargetFromConnector(PrimaryConnectorKey);
+			if (!PrimaryTarget.IsValid())
+			{
+				OutErrorMessage = FText::FromString(FString::Printf(TEXT("Cannot resolve connector %s because primary connector is not resolved"), *InConnectorKey.ToString()));
+				return false;
+			}
 		}
 	}
 
@@ -225,14 +240,14 @@ bool UModularRigController::CanConnectConnectorToElement(const FRigElementKey& I
 			OutErrorMessage = FText::FromString(FString::Printf(TEXT("Could not find debugged modular rig in %s"), *Blueprint->GetPathName()));
 			return false;
 		}
-	
+
 		URigHierarchy* Hierarchy = ModularRig->GetHierarchy();
 		if (!Hierarchy)
 		{
 			OutErrorMessage = FText::FromString(FString::Printf(TEXT("Could not find hierarchy in %s"), *ModularRig->GetPathName()));
 			return false;
 		}
-	
+
 		const FRigConnectorElement* Connector = Cast<FRigConnectorElement>(Hierarchy->Find(InConnectorKey));
 		if (!Connector)
 		{
@@ -240,6 +255,15 @@ bool UModularRigController::CanConnectConnectorToElement(const FRigElementKey& I
 			return false;
 		}
 
+		if(!Connector->IsArrayConnector())
+		{
+			if(InTargetKeys.Num() > 1)
+			{
+				OutErrorMessage = FText::FromString(FString::Printf(TEXT("Connector %s can only be resolved to one target - it is not an array connector."), *InConnectorKey.ToString()));
+				return false;
+			}
+		}
+
 		UModularRigRuleManager* RuleManager = Hierarchy->GetRuleManager();
 		if (!RuleManager)
 		{
@@ -249,23 +273,47 @@ bool UModularRigController::CanConnectConnectorToElement(const FRigElementKey& I
 
 		const FRigModuleInstance* ModuleInstance = ModularRig->FindModule(Module->GetPath());
 		FModularRigResolveResult RuleResults = RuleManager->FindMatches(Connector, ModuleInstance, ModularRig->GetElementKeyRedirector());
-		if (!RuleResults.ContainsMatch(InTargetKey))
+		for(const FRigElementKey& InTargetKey : InTargetKeys)
 		{
-			OutErrorMessage = FText::FromString(FString::Printf(TEXT("The target %s is not a valid match for connector %s"), *InTargetKey.ToString(), *InConnectorKey.ToString()));
-			return false;
+			if (!RuleResults.ContainsMatch(InTargetKey))
+			{
+				OutErrorMessage = FText::FromString(FString::Printf(TEXT("The target %s is not a valid match for connector %s"), *InTargetKey.ToString(), *InConnectorKey.ToString()));
+				return false;
+			}
 		}
 	}
 #endif
-
 	return true;
 }
 
 bool UModularRigController::ConnectConnectorToElement(const FRigElementKey& InConnectorKey, const FRigElementKey& InTargetKey, bool bSetupUndo, bool bAutoResolveOtherConnectors, bool bCheckValidConnection)
 {
+	return ConnectConnectorToElements(InConnectorKey, {InTargetKey}, bSetupUndo, bAutoResolveOtherConnectors, bCheckValidConnection);
+}
+
+bool UModularRigController::ConnectConnectorToElements(const FRigElementKey& InConnectorKey, const TArray<FRigElementKey>& InTargetKeys,
+	bool bSetupUndo, bool bAutoResolveOtherConnectors, bool bCheckValidConnection)
+{
+	const FRigElementKeyRedirector::FKeyArray TargetKeys(InTargetKeys);
+	if(TargetKeys.IsEmpty())
+	{
+		UE_LOG(LogControlRig, Error, TEXT("Could not connect %s: Target Keys array is empty."), *InConnectorKey.ToString());
+		return false;
+	}
+	
 	FText ErrorMessage;
-	if (bCheckValidConnection && !CanConnectConnectorToElement(InConnectorKey, InTargetKey, ErrorMessage))
+	if (bCheckValidConnection && !CanConnectConnectorToElements(InConnectorKey, InTargetKeys, ErrorMessage))
 	{
-		UE_LOG(LogControlRig, Error, TEXT("Could not connect %s to %s: %s"), *InConnectorKey.ToString(), *InTargetKey.ToString(), *ErrorMessage.ToString());
+		FString TargetKeyString = TargetKeys[0].ToString();
+		if(TargetKeys.Num() > 1)
+		{
+			for(int32 Index = 1; Index < TargetKeys.Num(); Index++)
+			{
+				TargetKeyString = TargetKeyString + TEXT(", ") + TargetKeys[Index].ToString();				
+			}
+			TargetKeyString = TEXT("{") + TargetKeyString = TEXT("}");
+		}
+		UE_LOG(LogControlRig, Error, TEXT("Could not connect %s to %s: %s"), *InConnectorKey.ToString(), *TargetKeyString, *ErrorMessage.ToString());
 		return false;
 	}
 	
@@ -278,13 +326,27 @@ bool UModularRigController::ConnectConnectorToElement(const FRigElementKey& InCo
 	UBlueprint* Blueprint = Cast<UBlueprint>(GetOuter());
 
 #if WITH_EDITOR
-	FName TargetModulePathName = NAME_None;
+	TOptional<FName> TargetModulePathName;
 	UModularRig* ModularRig = GetDebuggedModularRig();
 	if(ModularRig)
 	{
 		if (URigHierarchy* Hierarchy = ModularRig->GetHierarchy())
 		{
-			TargetModulePathName = Hierarchy->GetModulePathFName(InTargetKey);
+			for(const FRigElementKey& TargetKey : TargetKeys)
+			{
+				if(TargetModulePathName.IsSet())
+				{
+					if(TargetModulePathName.GetValue() != Hierarchy->GetModulePathFName(TargetKey))
+					{
+						TargetModulePathName.Reset();
+						break;
+					}
+				}
+				else
+				{
+					TargetModulePathName = Hierarchy->GetModulePathFName(TargetKey);
+				}
+			}
 		}
 	}
 	
@@ -297,24 +359,24 @@ bool UModularRigController::ConnectConnectorToElement(const FRigElementKey& InCo
 #endif 
 
 	// First disconnect before connecting to anything else. This might disconnect other secondary/optional connectors.
-	TMap<FRigElementKey, FRigElementKey> PreviousConnections;
+	TMap<FRigElementKey, TArray<FRigElementKey>> PreviousConnections;
 	if (CurrentTarget.IsValid())
 	{
 		const TGuardValue<bool> DisableAutomaticReparenting(bAutomaticReparenting, false);
 		DisconnectConnector_Internal(InConnectorKey, false, &PreviousConnections, bSetupUndo);
 	}
-	
-	Model->Connections.AddConnection(InConnectorKey, {InTargetKey});
+
+	Model->Connections.AddConnection(InConnectorKey, TargetKeys);
 
 	// restore previous connections if possible
-	for(const TPair<FRigElementKey, FRigElementKey>& PreviousConnection : PreviousConnections)
+	for(const TPair<FRigElementKey, TArray<FRigElementKey>>& PreviousConnection : PreviousConnections)
 	{
 		if(!Model->Connections.HasConnection(PreviousConnection.Key))
 		{
 			FText ErrorMessageForPreviousConnection;
-			if(CanConnectConnectorToElement(PreviousConnection.Key, PreviousConnection.Value, ErrorMessageForPreviousConnection))
+			if(CanConnectConnectorToElements(PreviousConnection.Key, PreviousConnection.Value, ErrorMessageForPreviousConnection))
 			{
-				(void)ConnectConnectorToElement(PreviousConnection.Key, PreviousConnection.Value, bSetupUndo, false, false);
+				(void)ConnectConnectorToElements(PreviousConnection.Key, PreviousConnection.Value, bSetupUndo, false, false);
 			}
 		}
 	}
@@ -341,9 +403,9 @@ bool UModularRigController::ConnectConnectorToElement(const FRigElementKey& InCo
 					{
 						if(Connector->IsPrimary())
 						{
-							if(!TargetModulePathName.IsNone())
+							if(!TargetModulePathName.Get(NAME_None).IsNone())
 							{
-								const FString NewModulePath = ReparentModule(Module->GetPath(), TargetModulePathName.ToString(), bSetupUndo);
+								const FString NewModulePath = ReparentModule(Module->GetPath(), TargetModulePathName.GetValue().ToString(), bSetupUndo);
 								if(!NewModulePath.IsEmpty())
 								{
 									Module = FindModule(NewModulePath);
@@ -377,7 +439,7 @@ bool UModularRigController::DisconnectConnector(const FRigElementKey& InConnecto
 };
 
 bool UModularRigController::DisconnectConnector_Internal(const FRigElementKey& InConnectorKey, bool bDisconnectSubModules,
-	TMap<FRigElementKey, FRigElementKey>* OutRemovedConnections, bool bSetupUndo)
+	TMap<FRigElementKey, TArray<FRigElementKey>>* OutRemovedConnections, bool bSetupUndo)
 {
 	FString ConnectorModulePath, ConnectorName;
 	if (!URigHierarchy::SplitNameSpace(InConnectorKey.Name.ToString(), &ConnectorModulePath, &ConnectorName))
@@ -429,7 +491,7 @@ bool UModularRigController::DisconnectConnector_Internal(const FRigElementKey& I
 
 	if(OutRemovedConnections)
 	{
-		OutRemovedConnections->Add(InConnectorKey, Model->Connections.FindTargetFromConnector(InConnectorKey));
+		OutRemovedConnections->Add(InConnectorKey, Model->Connections.FindTargetsFromConnector(InConnectorKey));
 	}
 	Model->Connections.RemoveConnection(InConnectorKey);
 
@@ -448,7 +510,7 @@ bool UModularRigController::DisconnectConnector_Internal(const FRigElementKey& I
 		{
 			if(OutRemovedConnections)
 			{
-				OutRemovedConnections->Add(ToRemove, Model->Connections.FindTargetFromConnector(ToRemove));
+				OutRemovedConnections->Add(ToRemove, Model->Connections.FindTargetsFromConnector(ToRemove));
 			}
 			Model->Connections.RemoveConnection(ToRemove);
 		}
@@ -470,7 +532,7 @@ bool UModularRigController::DisconnectConnector_Internal(const FRigElementKey& I
 		{
 			if(OutRemovedConnections)
 			{
-				OutRemovedConnections->Add(ToRemove, Model->Connections.FindTargetFromConnector(ToRemove));
+				OutRemovedConnections->Add(ToRemove, Model->Connections.FindTargetsFromConnector(ToRemove));
 			}
 			Model->Connections.RemoveConnection(ToRemove);
 		}
diff --git a/Engine/Plugins/Animation/ControlRig/Source/ControlRig/Private/Rigs/RigHierarchy.cpp b/Engine/Plugins/Animation/ControlRig/Source/ControlRig/Private/Rigs/RigHierarchy.cpp
--- a/Engine/Plugins/Animation/ControlRig/Source/ControlRig/Private/Rigs/RigHierarchy.cpp
+++ b/Engine/Plugins/Animation/ControlRig/Source/ControlRig/Private/Rigs/RigHierarchy.cpp
@@ -3475,6 +3475,30 @@ const FRigElementKey& URigHierarchy::GetResolvedTarget(const FRigElementKey& InC
 	return InConnectorKey;
 }
 
+TArray<FRigElementKey> URigHierarchy::GetResolvedTargets(const FRigElementKey& InConnectorKey) const
+{
+	if (InConnectorKey.Type == ERigElementType::Connector)
+	{
+		if(ElementKeyRedirector)
+		{
+			if(const FRigElementKeyRedirector::FCachedKeyArray* Targets = ElementKeyRedirector->Find(InConnectorKey))
+			{
+				if(!Targets->IsEmpty())
+				{
+					TArray<FRigElementKey> TargetKeys;
+					TargetKeys.Reserve(Targets->Num());
+					for(const FCachedRigElement& Cache : *Targets)
+					{
+						TargetKeys.Add(Cache.GetKey());
+					}
+					return TargetKeys;
+				}
+			}
+		}
+	}
+	return {InConnectorKey};
+}
+
 bool URigHierarchy::Undo()
 {
 #if WITH_EDITOR
@@ -5230,7 +5254,10 @@ void URigHierarchy::SetConnectorSettings(FRigConnectorElement* InConnectorElemen
 		Modify();
 	}
 
-	InConnectorElement->Settings = InSettings;
+	FRigConnectorSettings Settings = InSettings;
+	Settings.bIsArray = InConnectorElement->IsPrimary() ? false : InSettings.bIsArray;
+
+	InConnectorElement->Settings = Settings;
 	Notify(ERigHierarchyNotification::ConnectorSettingChanged, InConnectorElement);
 	
 #if WITH_EDITOR
diff --git a/Engine/Plugins/Animation/ControlRig/Source/ControlRig/Private/Rigs/RigHierarchyController.cpp b/Engine/Plugins/Animation/ControlRig/Source/ControlRig/Private/Rigs/RigHierarchyController.cpp
--- a/Engine/Plugins/Animation/ControlRig/Source/ControlRig/Private/Rigs/RigHierarchyController.cpp
+++ b/Engine/Plugins/Animation/ControlRig/Source/ControlRig/Private/Rigs/RigHierarchyController.cpp
@@ -673,6 +673,8 @@ FRigElementKey URigHierarchyController::AddConnector(FName InName, FRigConnector
 	// only allow to add one primary connector
 	if(InSettings.Type == EConnectorType::Primary)
 	{
+		InSettings.bIsArray = false;
+		
 		const TArray<FRigConnectorElement*>& Connectors = Hierarchy->GetConnectors();
 		for(const FRigConnectorElement* Connector : Connectors)
 		{
diff --git a/Engine/Plugins/Animation/ControlRig/Source/ControlRig/Private/Rigs/RigHierarchyElements.cpp b/Engine/Plugins/Animation/ControlRig/Source/ControlRig/Private/Rigs/RigHierarchyElements.cpp
--- a/Engine/Plugins/Animation/ControlRig/Source/ControlRig/Private/Rigs/RigHierarchyElements.cpp
+++ b/Engine/Plugins/Animation/ControlRig/Source/ControlRig/Private/Rigs/RigHierarchyElements.cpp
@@ -2212,6 +2212,7 @@ void FRigReferenceElement::CopyPose(FRigBaseElement* InOther, bool bCurrent, boo
 FRigConnectorSettings::FRigConnectorSettings()
 	: Type(EConnectorType::Primary)
 	, bOptional(false)
+	, bIsArray(false)
 {
 }
 
@@ -2227,12 +2228,9 @@ void FRigConnectorSettings::Save(FArchive& Ar)
 	Ar.UsingCustomVersion(FControlRigObjectVersion::GUID);
 
 	Ar << Description;
-
-	if (Ar.CustomVer(FControlRigObjectVersion::GUID) >= FControlRigObjectVersion::ConnectorsWithType)
-	{
-		Ar << Type;
-		Ar << bOptional;
-	}
+	Ar << Type;
+	Ar << bOptional;
+	Ar << bIsArray;
 
 	int32 NumRules = Rules.Num();
 	Ar << NumRules;
@@ -2255,6 +2253,15 @@ void FRigConnectorSettings::Load(FArchive& Ar)
 		Ar << bOptional;
 	}
 
+	if (Ar.CustomVer(FControlRigObjectVersion::GUID) >= FControlRigObjectVersion::RigHierarchyArrayConnectors)
+	{
+		Ar << bIsArray;
+	}
+	else
+	{
+		bIsArray = false;
+	}
+
 	int32 NumRules = 0;
 	Ar << NumRules;
 	Rules.SetNumZeroed(NumRules);
@@ -2278,6 +2285,10 @@ bool FRigConnectorSettings::operator==(const FRigConnectorSettings& InOther) con
 	{
 		return false;
 	}
+	if(bIsArray != InOther.bIsArray)
+	{
+		return false;
+	}
 	if(Rules.Num() != InOther.Rules.Num())
 	{
 		return false;
@@ -2306,6 +2317,7 @@ uint32 GetTypeHash(const FRigConnectorSettings& Settings)
 {
 	uint32 Hash = HashCombine(GetTypeHash(Settings.Type), Settings.GetRulesHash());
 	Hash = HashCombine(Hash, GetTypeHash(Settings.bOptional));
+	Hash = HashCombine(Hash, GetTypeHash(Settings.bIsArray));
 	return Hash;
 }
 
diff --git a/Engine/Plugins/Animation/ControlRig/Source/ControlRig/Public/ControlRigObjectVersion.h b/Engine/Plugins/Animation/ControlRig/Source/ControlRig/Public/ControlRigObjectVersion.h
--- a/Engine/Plugins/Animation/ControlRig/Source/ControlRig/Public/ControlRigObjectVersion.h
+++ b/Engine/Plugins/Animation/ControlRig/Source/ControlRig/Public/ControlRigObjectVersion.h
@@ -116,6 +116,9 @@ struct CONTROLRIG_API FControlRigObjectVersion
 		// Improve transform compactness when serializing the hierarchy
 		RigHierarchyCompactTransformSerialization,
 
+		// Connectors to support arrays
+		RigHierarchyArrayConnectors,
+
 		// -----<new versions can be added above this line>-------------------------------------------------
 		VersionPlusOne,
 		LatestVersion = VersionPlusOne - 1
diff --git a/Engine/Plugins/Animation/ControlRig/Source/ControlRig/Public/ModularRigController.h b/Engine/Plugins/Animation/ControlRig/Source/ControlRig/Public/ModularRigController.h
--- a/Engine/Plugins/Animation/ControlRig/Source/ControlRig/Public/ModularRigController.h
+++ b/Engine/Plugins/Animation/ControlRig/Source/ControlRig/Public/ModularRigController.h
@@ -31,12 +31,18 @@ class CONTROLRIG_API UModularRigController : public UObject
 	UFUNCTION(BlueprintCallable, Category = "Control Rig | Modules")
 	bool CanConnectConnectorToElement(const FRigElementKey& InConnectorKey, const FRigElementKey& InTargetKey, FText& OutErrorMessage);
 
+	UFUNCTION(BlueprintCallable, Category = "Control Rig | Modules")
+	bool CanConnectConnectorToElements(const FRigElementKey& InConnectorKey, const TArray<FRigElementKey>& InTargetKeys, FText& OutErrorMessage);
+
 	UFUNCTION(BlueprintCallable, Category = "Control Rig | Modules")
 	bool ConnectConnectorToElement(const FRigElementKey& InConnectorKey, const FRigElementKey& InTargetKey, bool bSetupUndo = true, bool bAutoResolveOtherConnectors = true, bool bCheckValidConnection = true);
 
+	UFUNCTION(BlueprintCallable, Category = "Control Rig | Modules")
+	bool ConnectConnectorToElements(const FRigElementKey& InConnectorKey, const TArray<FRigElementKey>& InTargetKeys, bool bSetupUndo = true, bool bAutoResolveOtherConnectors = true, bool bCheckValidConnection = true);
+
 	UFUNCTION(BlueprintCallable, Category = "Control Rig | Modules")
 	bool DisconnectConnector(const FRigElementKey& InConnectorKey, bool bDisconnectSubModules = false, bool bSetupUndo = true);
-	bool DisconnectConnector_Internal(const FRigElementKey& InConnectorKey, bool bDisconnectSubModules = false, TMap<FRigElementKey, FRigElementKey>*OutRemovedConnections = nullptr, bool bSetupUndo = true);
+	bool DisconnectConnector_Internal(const FRigElementKey& InConnectorKey, bool bDisconnectSubModules = false, TMap<FRigElementKey, TArray<FRigElementKey>>*OutRemovedConnections = nullptr, bool bSetupUndo = true);
 
 	UFUNCTION(BlueprintCallable, Category = "Control Rig | Modules")
 	TArray<FRigElementKey> DisconnectCyclicConnectors(bool bSetupUndo = true);
diff --git a/Engine/Plugins/Animation/ControlRig/Source/ControlRig/Public/Rigs/RigHierarchy.h b/Engine/Plugins/Animation/ControlRig/Source/ControlRig/Public/Rigs/RigHierarchy.h
--- a/Engine/Plugins/Animation/ControlRig/Source/ControlRig/Public/Rigs/RigHierarchy.h
+++ b/Engine/Plugins/Animation/ControlRig/Source/ControlRig/Public/Rigs/RigHierarchy.h
@@ -3625,10 +3625,15 @@ public:
 	void Traverse(TFunction<void(FRigBaseElement*, bool& /* continue */)> PerElementFunction, bool bTowardsChildren = true) const;
 
 	/**
-	 * Returns the currently resolved target for given connector key
+	 * Returns the first currently resolved target for given connector key
 	 */
 	const FRigElementKey& GetResolvedTarget(const FRigElementKey& InConnectorKey) const;
 
+	/**
+	 * Returns all currently resolved targets for given connector key
+	 */
+	TArray<FRigElementKey> GetResolvedTargets(const FRigElementKey& InConnectorKey) const;
+
 	/**
 	 * Performs undo for one transform change
 	 */
diff --git a/Engine/Plugins/Animation/ControlRig/Source/ControlRig/Public/Rigs/RigHierarchyElements.h b/Engine/Plugins/Animation/ControlRig/Source/ControlRig/Public/Rigs/RigHierarchyElements.h
--- a/Engine/Plugins/Animation/ControlRig/Source/ControlRig/Public/Rigs/RigHierarchyElements.h
+++ b/Engine/Plugins/Animation/ControlRig/Source/ControlRig/Public/Rigs/RigHierarchyElements.h
@@ -2043,6 +2043,10 @@ struct CONTROLRIG_API FRigConnectorSettings
 	UPROPERTY(BlueprintReadWrite, EditAnywhere, Category = Connector)
 	bool bOptional;
 
+	// by enabling this the connector will be able to connect to more than one target
+	UPROPERTY(BlueprintReadWrite, EditAnywhere, Category = Connector)
+	bool bIsArray;
+
 	UPROPERTY(BlueprintReadWrite, EditAnywhere, Category = Connector)
 	TArray<FRigConnectionRuleStash> Rules;
 
@@ -2116,6 +2120,7 @@ struct CONTROLRIG_API FRigConnectorElement final : public FRigBaseElement
 	bool IsPrimary() const { return Settings.Type == EConnectorType::Primary; }
 	bool IsSecondary() const { return Settings.Type == EConnectorType::Secondary; }
 	bool IsOptional() const { return IsSecondary() && Settings.bOptional; }
+	bool IsArrayConnector() const { return IsSecondary() && Settings.bIsArray; }
 
 private:
 	explicit FRigConnectorElement(URigHierarchy* InOwner)
diff --git a/Engine/Plugins/Animation/ControlRig/Source/ControlRigEditor/Private/ControlRigElementDetails.cpp b/Engine/Plugins/Animation/ControlRig/Source/ControlRigEditor/Private/ControlRigElementDetails.cpp
--- a/Engine/Plugins/Animation/ControlRig/Source/ControlRigEditor/Private/ControlRigElementDetails.cpp
+++ b/Engine/Plugins/Animation/ControlRig/Source/ControlRigEditor/Private/ControlRigElementDetails.cpp
@@ -6110,8 +6110,9 @@ void FRigConnectorElementDetails::CustomizeSettings(IDetailLayoutBuilder& Detail
 
 	IDetailCategoryBuilder& SettingsCategory = DetailBuilder.EditCategory(TEXT("Settings"), LOCTEXT("Settings", "Settings"));
 
+	ConnectorTypeHandle = SettingsHandle->GetChildHandle(GET_MEMBER_NAME_CHECKED(FRigConnectorSettings, Type));
 	SettingsCategory
-		.AddProperty(SettingsHandle->GetChildHandle(GET_MEMBER_NAME_CHECKED(FRigConnectorSettings, Type)))
+		.AddProperty(ConnectorTypeHandle)
 		.IsEnabled(false);
 
 	SettingsCategory
@@ -6145,6 +6146,12 @@ void FRigConnectorElementDetails::CustomizeSettings(IDetailLayoutBuilder& Detail
 		}
 	}
 
+	IsArrayHandle = SettingsHandle->GetChildHandle(GET_MEMBER_NAME_CHECKED(FRigConnectorSettings, bIsArray));
+	SettingsCategory
+		.AddProperty(IsArrayHandle)
+		.IsEnabled(!IsAnyConnectorImported())
+		.Visibility(TAttribute<EVisibility>::CreateSP(this, &FRigConnectorElementDetails::GetArrayVisibility));
+
 	if(!bHideRules)
 	{
 		SettingsCategory
@@ -6153,6 +6160,52 @@ void FRigConnectorElementDetails::CustomizeSettings(IDetailLayoutBuilder& Detail
 	}
 }
 
+TOptional<EConnectorType> FRigConnectorElementDetails::GetConnectorType() const
+{
+	TOptional<EConnectorType> Result;
+	for (const FPerElementInfo& Info : PerElementInfos)
+	{
+		if(const FRigConnectorElement* Connector = Info.GetElement<FRigConnectorElement>())
+		{
+			if(!Result.IsSet())
+			{
+				Result = Connector->Settings.Type;
+			}
+			else if(Result.GetValue() != Connector->Settings.Type)
+			{
+				return {};
+			}
+		}
+	}
+	return Result;
+}
+
+TOptional<bool> FRigConnectorElementDetails::GetIsConnectorArray() const
+{
+	TOptional<bool> Result;
+	for (const FPerElementInfo& Info : PerElementInfos)
+	{
+		if(const FRigConnectorElement* Connector = Info.GetElement<FRigConnectorElement>())
+		{
+			if(!Result.IsSet())
+			{
+				Result = Connector->IsArrayConnector();
+			}
+			else if(Result.GetValue() != Connector->IsArrayConnector())
+			{
+				return {};
+			}
+		}
+	}
+	return Result;
+}
+
+EVisibility FRigConnectorElementDetails::GetArrayVisibility() const
+{
+	const TOptional<EConnectorType> ConnectorType = GetConnectorType();
+	return ConnectorType.Get(EConnectorType::Primary) == EConnectorType::Primary ? EVisibility::Collapsed : EVisibility::Visible;
+}
+
 void FRigSocketElementDetails::CustomizeDetails(IDetailLayoutBuilder& DetailBuilder)
 {
 	FRigTransformElementDetails::CustomizeDetails(DetailBuilder);
diff --git a/Engine/Plugins/Animation/ControlRig/Source/ControlRigEditor/Private/ControlRigElementDetails.h b/Engine/Plugins/Animation/ControlRig/Source/ControlRigEditor/Private/ControlRigElementDetails.h
--- a/Engine/Plugins/Animation/ControlRig/Source/ControlRigEditor/Private/ControlRigElementDetails.h
+++ b/Engine/Plugins/Animation/ControlRig/Source/ControlRigEditor/Private/ControlRigElementDetails.h
@@ -867,6 +867,15 @@ public:
 	virtual void CustomizeDetails(IDetailLayoutBuilder& DetailBuilder) override;
 
 	void CustomizeSettings(IDetailLayoutBuilder& DetailBuilder);
+
+private:
+
+	TOptional<EConnectorType> GetConnectorType() const;
+	TOptional<bool> GetIsConnectorArray() const;
+	EVisibility GetArrayVisibility() const;
+
+	TSharedPtr<IPropertyHandle> ConnectorTypeHandle;
+	TSharedPtr<IPropertyHandle> IsArrayHandle;
 };
 
 class FRigSocketElementDetails : public FRigNullElementDetails
