Control Rig: Follow to breakage from ArrayConnector further work + tests

#rnx

diff --git a/Engine/Plugins/Animation/ControlRig/Source/ControlRig/Private/Rigs/RigHierarchy.cpp b/Engine/Plugins/Animation/ControlRig/Source/ControlRig/Private/Rigs/RigHierarchy.cpp
--- a/Engine/Plugins/Animation/ControlRig/Source/ControlRig/Private/Rigs/RigHierarchy.cpp
+++ b/Engine/Plugins/Animation/ControlRig/Source/ControlRig/Private/Rigs/RigHierarchy.cpp
@@ -421,6 +421,9 @@ void URigHierarchy::Load(FArchive& Ar)
 
 		MemoryReader = MakeUnique<FRigHierarchyMemoryReader>(bStoreCompressedBytes ? UncompressedBytes : CompressedBytes, UniqueNames, Ar.IsPersistent());
 		ArchiveForElements = MemoryReader.Get();
+		
+		// force to use the versions from within the main archive
+		ArchiveForElements->SetCustomVersions(Ar.GetCustomVersions());
 	}
 
 	int32 ElementCount = 0;
