Control RIg: Store metadata in the test data assets
#rnx

diff --git a/Engine/Plugins/Animation/ControlRig/Source/ControlRig/Private/ControlRigTestData.cpp b/Engine/Plugins/Animation/ControlRig/Source/ControlRig/Private/ControlRigTestData.cpp
--- a/Engine/Plugins/Animation/ControlRig/Source/ControlRig/Private/ControlRigTestData.cpp
+++ b/Engine/Plugins/Animation/ControlRig/Source/ControlRig/Private/ControlRigTestData.cpp
@@ -46,6 +46,14 @@ bool FControlRigTestDataFrame::Store(UControlRig* InControlRig, bool bInitial)
 		Variables.Add(VariableData);
 	}
 
+	MetadataMap = InControlRig->GetHierarchy()->CopyMetadata();
+	TArray<uint8> UncompressedBytes;
+	FMemoryWriter ArchiveWriter(UncompressedBytes);
+	ArchiveWriter.UsingCustomVersion(FControlRigObjectVersion::GUID);
+
+	ArchiveWriter << MetadataMap;
+	Metadata = UncompressedBytes;
+
 	return true;
 }
 
@@ -70,8 +78,10 @@ bool FControlRigTestDataFrame::Restore(UControlRig* InControlRig, bool bInitial)
 	}
 
 	Hierarchy->SetPose(Pose, bInitial ? ERigTransformType::InitialLocal : ERigTransformType::CurrentLocal);
+	bool bSuccess = RestoreVariables(InControlRig);
+	bSuccess &= RestoreMetadata(InControlRig);
 
-	return RestoreVariables(InControlRig);
+	return bSuccess;
 }
 
 bool FControlRigTestDataFrame::RestoreVariables(UControlRig* InControlRig) const
@@ -147,12 +157,70 @@ bool FControlRigTestDataFrame::RestoreVariables(UControlRig* InControlRig) const
 	return true;
 }
 
+bool FControlRigTestDataFrame::RestoreMetadata(UControlRig* InControlRig) const
+{
+	bool bSuccess = true;
+
+	if (MetadataMap.IsEmpty())
+	{
+		if (Metadata.IsEmpty())
+		{
+			bTestMetadata = false;
+		}
+		else
+		{
+			FMemoryReader ArchiveReader(Metadata);
+			ArchiveReader.SetCustomVersions(ArchiveCustomVersions);
+			ArchiveReader << MetadataMap;
+		}
+	}
+
+	if (bTestMetadata)
+	{
+		bSuccess &= InControlRig->GetHierarchy()->SetMetadata(MetadataMap);
+	}
+		
+	return bSuccess;
+}
+
+void UControlRigTestData::BeginDestroy()
+{
+	UObject::BeginDestroy();
+
+	TArray<FControlRigTestDataFrame> Frames;
+	Frames.Add(Initial);
+	Frames.Append(InputFrames);
+	Frames.Append(OutputFrames);
+
+	for (FControlRigTestDataFrame& Frame : Frames)
+	{
+		for (TPair<FRigElementKey, URigHierarchy::FMetadataStorage>& Pair : Frame.MetadataMap)
+		{
+			Pair.Value.Reset();
+		}
+		Frame.MetadataMap.Reset();
+	}
+}
+
 void UControlRigTestData::Serialize(FArchive& Ar)
 {
-	Ar.UsingCustomVersion(FControlRigObjectVersion::GUID);	
+	Ar.UsingCustomVersion(FControlRigObjectVersion::GUID);
 	UObject::Serialize(Ar);
 	LastFrameIndex = INDEX_NONE;
 
+	Initial.ArchiveCustomVersions = Ar.GetCustomVersions();
+	Initial.bTestMetadata = !Initial.Metadata.IsEmpty();
+	for (int32 IO=0; IO<2; IO++)
+	{
+		TArray<FControlRigTestDataFrame>& Frames = (IO == 0) ? InputFrames : OutputFrames;
+		for (FControlRigTestDataFrame& Frame : Frames)
+		{
+			Frame.ArchiveCustomVersions = Ar.GetCustomVersions();
+			Frame.bTestMetadata = !Frame.Metadata.IsEmpty();
+		}
+	}
+	
+	
 	// If pose is older than RigPoseWithParentKey, set the active parent of all poses to invalid key
 	if (GetLinkerCustomVersion(FControlRigObjectVersion::GUID) < FControlRigObjectVersion::RigPoseWithParentKey)
 	{
@@ -260,17 +328,12 @@ bool UControlRigTestData::Record(UControlRig* InControlRig, double InRecordingDu
 	DesiredRecordingDuration = InRecordingDuration;
 	TimeAtStartOfRecording = FPlatformTime::Seconds();
 
-	// if this is the first frame
-	if(InputFrames.IsEmpty())
-	{
-		InControlRig->RequestInit();
-		PreConstructionHandle = InControlRig->OnPreConstruction_AnyThread().AddLambda(
-			[this](UControlRig* InControlRig, const FName& InEventName)
-			{
-				Initial.Store(InControlRig, true);
-			}
-		);
-	}
+	PreConstructionHandle = InControlRig->OnPreConstruction_AnyThread().AddLambda(
+		[this](UControlRig* InControlRig, const FName& InEventName)
+		{
+			Initial.Store(InControlRig, true);
+		}
+	);
 
 	PreEventHandle = InControlRig->OnPreExecuted_AnyThread().AddLambda(
 		[this](URigVMHost* InRigVMHost, const FName& InEventName)
@@ -299,7 +362,7 @@ bool UControlRigTestData::Record(UControlRig* InControlRig, double InRecordingDu
 			// the results in the rig are the same during a recording and replay.
 			Frame.RestoreVariables(ControlRig);
 			
-			InputFrames.Add(Frame);
+			InputFrames.Add(MoveTemp(Frame));
 		}
 	);
 
@@ -324,7 +387,7 @@ bool UControlRigTestData::Record(UControlRig* InControlRig, double InRecordingDu
 			
 			FControlRigTestDataFrame Frame;
 			Frame.Store(ControlRig);
-			OutputFrames.Add(Frame);
+			OutputFrames.Add(MoveTemp(Frame));
 			LastFrameIndex = INDEX_NONE;
 			(void)MarkPackageDirty();
 
@@ -340,6 +403,12 @@ bool UControlRigTestData::Record(UControlRig* InControlRig, double InRecordingDu
 		}
 	);
 
+	// if this is the first frame
+	if (InputFrames.IsEmpty())
+	{
+		InControlRig->RequestInit();
+	}
+
 	return true;
 }
 
@@ -364,7 +433,6 @@ bool UControlRigTestData::SetupReplay(UControlRig* InControlRig, bool bGroundTru
 	// reset the control rig's absolute time
 	InControlRig->SetAbsoluteAndDeltaTime(InputFrames[0].AbsoluteTime, InputFrames[0].DeltaTime);
 	
-	InControlRig->RequestInit();
 	PreConstructionHandle = InControlRig->OnPreConstruction_AnyThread().AddLambda(
 		[this](UControlRig* InControlRig, const FName& InEventName)
 		{
@@ -420,7 +488,62 @@ bool UControlRigTestData::SetupReplay(UControlRig* InControlRig, bool bGroundTru
 			
 			const FRigPose CurrentPose = ControlRig->GetHierarchy()->GetPose();
 
-			const FControlRigTestDataFrame& Frame = OutputFrames[LastFrameIndex];
+			const int32 FrameIndex = GetFrameIndexForTime(ControlRig->GetAbsoluteTime(), false);
+			const FControlRigTestDataFrame& Frame = OutputFrames[FrameIndex];
+
+			if (Frame.bTestMetadata)
+			{
+				const TMap<FRigElementKey, URigHierarchy::FMetadataStorage>& ExpectedMetadata = Frame.MetadataMap;
+			   URigHierarchy* Hierarchy = ControlRig->GetHierarchy();
+#if WITH_EDITOR
+			   FRigVMLog* Log = ControlRig->GetLog();
+			   if (Log)
+			   {
+				   Hierarchy->ForEach([ExpectedMetadata, Hierarchy, Log, FrameIndex](const FRigBaseElement* Element)
+				   {
+					   TArray<FName> MetadataNames = Hierarchy->GetMetadataNames(Element->GetKey());
+					   const URigHierarchy::FMetadataStorage* ExpectedElementMetadata = ExpectedMetadata.Find(Element->GetKey());
+					   if (ExpectedElementMetadata)
+					   {
+						   if (ExpectedElementMetadata->MetadataMap.Num() != MetadataNames.Num())
+						   {
+							   const FString Message = FString::Printf(TEXT("Metadata mismatch in element %s: number of metadata elements does not match"), *Element->GetKey().ToString());
+							   Log->Report(EMessageSeverity::Error, TEXT("TestData"), FrameIndex, Message);
+						   }
+						   for(TTuple<FName, FRigBaseMetadata*> Pair : ExpectedElementMetadata->MetadataMap)
+						   {
+							   if(FRigBaseMetadata* Value = Hierarchy->FindMetadataForElement(Element, Pair.Key, Pair.Value->GetType()))
+							   {
+								   const FProperty* Property = Value->GetValueProperty();
+								   check(Value->IsValid());
+								   check(Pair.Value->IsValid());
+								   if(!Property->Identical(Value->GetValueData(), Pair.Value->GetValueData()))
+								   {
+									   FString Expected, Received;
+									   Property->ExportText_Direct(Expected, Pair.Value->GetValueData(), nullptr, nullptr, PPF_None, nullptr);
+									   Property->ExportText_Direct(Received, Value->GetValueData(), nullptr, nullptr, PPF_None, nullptr);
+									   const FString Message = FString::Printf(TEXT("Metadata mismatch in element %s: property %s expected %s, but found %s"), *Element->GetKey().ToString(), *Pair.Key.ToString(), *Expected, *Received);
+									   Log->Report(EMessageSeverity::Error, TEXT("TestData"), FrameIndex, Message);
+								   }
+							   }
+							   else
+							   {
+								   const FString Message = FString::Printf(TEXT("Metadata mismatch in element %s: property %s expected, but was not found"), *Element->GetKey().ToString(), *Pair.Key.ToString());
+								   Log->Report(EMessageSeverity::Error, TEXT("TestData"), FrameIndex, Message);
+							   }
+						   }
+					   }
+					   else if (Hierarchy->HasMetadata(Element))
+					   {
+						   const FString Message = FString::Printf(TEXT("Metadata mismatch in element %s: number of metadata elements does not match"), *Element->GetKey().ToString());
+						   Log->Report(EMessageSeverity::Error, TEXT("TestData"), FrameIndex, Message);
+					   }
+					   return true;
+				   });
+			   }
+#endif
+			}
+			
 			if(bIsApplyingOutputs)
 			{
 				Frame.Restore(ControlRig, false);
@@ -452,6 +575,8 @@ bool UControlRigTestData::SetupReplay(UControlRig* InControlRig, bool bGroundTru
 		}
 	);
 
+	InControlRig->RequestInit();
+
 	ReplayControlRig = InControlRig;
 	return true;
 }
diff --git a/Engine/Plugins/Animation/ControlRig/Source/ControlRig/Private/Rigs/RigHierarchy.cpp b/Engine/Plugins/Animation/ControlRig/Source/ControlRig/Private/Rigs/RigHierarchy.cpp
--- a/Engine/Plugins/Animation/ControlRig/Source/ControlRig/Private/Rigs/RigHierarchy.cpp
+++ b/Engine/Plugins/Animation/ControlRig/Source/ControlRig/Private/Rigs/RigHierarchy.cpp
@@ -1686,7 +1686,7 @@ TArray<FName> URigHierarchy::GetMetadataNames(FRigElementKey InItem) const
 	TArray<FName> Names;
 	if (const FRigBaseElement* Element = Find(InItem))
 	{
-		if (Element->MetadataStorageIndex != INDEX_NONE)
+		if (ElementMetadata.IsValidIndex(Element->MetadataStorageIndex))
 		{
 			ElementMetadata[Element->MetadataStorageIndex].MetadataMap.GetKeys(Names);
 		}
@@ -7025,6 +7025,57 @@ void URigHierarchy::PropagateMetadata(const FRigBaseElement* InElement, const FN
 #endif
 }
 
+TMap<FRigElementKey, URigHierarchy::FMetadataStorage> URigHierarchy::CopyMetadata() const
+{
+	TMap<FRigElementKey, FMetadataStorage> ElementMetadataToSave;
+	for (const FRigBaseElement* Element: Elements)
+	{
+		if (ElementMetadata.IsValidIndex(Element->MetadataStorageIndex))
+		{
+			const FMetadataStorage& SourceStorage = ElementMetadata[Element->MetadataStorageIndex];
+			FMetadataStorage& TargetStorage = ElementMetadataToSave.FindOrAdd(Element->Key);
+			for (const TPair<FName, FRigBaseMetadata*>& Pair : SourceStorage.MetadataMap)
+			{
+				FRigBaseMetadata* NewValue = FRigBaseMetadata::MakeMetadata(Pair.Key, Pair.Value->GetType());
+				Pair.Value->GetValueProperty()->CopyCompleteValue(NewValue->GetValueData(), Pair.Value->GetValueData());
+				TargetStorage.MetadataMap.Add(Pair.Key, NewValue);
+			}
+		}
+	}
+	return ElementMetadataToSave;
+}
+
+bool URigHierarchy::SetMetadata(const TMap<FRigElementKey, URigHierarchy::FMetadataStorage>& InMetadata)
+{
+	bool bResult = true;
+	{
+		TGuardValue<FRigHierarchyMetadataChangedDelegate> MetadataChangedDelegateGuard(MetadataChangedDelegate, FRigHierarchyMetadataChangedDelegate());
+		RemoveAllMetadata();
+	
+		for (const TPair<FRigElementKey, FMetadataStorage>& Entry: InMetadata)
+		{
+			if (FRigBaseElement* Element = Find(Entry.Key))
+			{
+				for (const TTuple<FName, FRigBaseMetadata*>& MetadataElement : Entry.Value.MetadataMap)
+				{
+					if (FRigBaseMetadata* TargetMetadata = GetMetadataForElement(Element, MetadataElement.Key, MetadataElement.Value->Type, false))
+					{
+						TargetMetadata->GetValueProperty()->CopyCompleteValue(TargetMetadata->GetValueData(), MetadataElement.Value->GetValueData());
+					}
+				}
+			}
+			else
+			{
+				bResult = false;
+			}
+		}
+	}
+
+	OnMetadataChanged(FRigElementKey(ERigElementType::All), NAME_None);
+
+	return bResult;
+}
+
 void URigHierarchy::OnMetadataChanged(const FRigElementKey& InKey, const FName& InName)
 {
 	MetadataVersion++;
@@ -7142,6 +7193,10 @@ const FRigBaseMetadata* URigHierarchy::FindMetadataForElement(const FRigBaseElem
 	return const_cast<URigHierarchy*>(this)->FindMetadataForElement(InElement, InName, InType);
 }
 
+bool URigHierarchy::HasMetadata(const FRigBaseElement* InElement) const
+{
+	return ElementMetadata.IsValidIndex(InElement->MetadataStorageIndex) && !ElementMetadata[InElement->MetadataStorageIndex].MetadataMap.IsEmpty();
+}
 
 bool URigHierarchy::RemoveMetadataForElement(FRigBaseElement* InElement, const FName& InName)
 {
@@ -7210,6 +7265,24 @@ bool URigHierarchy::RemoveAllMetadataForElement(FRigBaseElement* InElement)
 	return true; 
 }
 
+bool URigHierarchy::RemoveAllMetadata()
+{
+	bool bSuccess = true;
+
+	{
+		TGuardValue<FRigHierarchyMetadataChangedDelegate> MetadataChangedDelegateGuard(MetadataChangedDelegate, FRigHierarchyMetadataChangedDelegate());
+		ForEach([this, &bSuccess](FRigBaseElement* Element)
+		{
+			bSuccess &= RemoveAllMetadataForElement(Element);
+			return true;
+		});
+	}
+
+	OnMetadataChanged(FRigElementKey(ERigElementType::All), NAME_None);
+	
+	return bSuccess;
+}
+
 
 void URigHierarchy::CopyAllMetadataFromElement(FRigBaseElement* InTargetElement, const FRigBaseElement* InSourceElement)
 {
diff --git a/Engine/Plugins/Animation/ControlRig/Source/ControlRig/Public/ControlRigTestData.h b/Engine/Plugins/Animation/ControlRig/Source/ControlRig/Public/ControlRigTestData.h
--- a/Engine/Plugins/Animation/ControlRig/Source/ControlRig/Public/ControlRigTestData.h
+++ b/Engine/Plugins/Animation/ControlRig/Source/ControlRig/Public/ControlRigTestData.h
@@ -3,6 +3,7 @@
 #pragma once
 
 #include "CoreMinimal.h"
+#include "Rigs/RigHierarchy.h"
 #include "Rigs/RigHierarchyPose.h"
 
 #include "ControlRigTestData.generated.h"
@@ -43,11 +44,13 @@ struct CONTROLRIG_API FControlRigTestDataFrame
 		DeltaTime = 0.0;
 		Variables.Reset();
 		Pose.Reset();
+		bTestMetadata = true;
 	}
 
 	bool Store(UControlRig* InControlRig, bool bInitial = false);
 	bool Restore(UControlRig* InControlRig, bool bInitial = false) const;
 	bool RestoreVariables(UControlRig* InControlRig) const;
+	bool RestoreMetadata(UControlRig* InControlRig) const;
 
 	UPROPERTY(VisibleAnywhere, BlueprintReadOnly, Category = "ControlRigTestDataFrame")
 	double AbsoluteTime;
@@ -60,6 +63,13 @@ struct CONTROLRIG_API FControlRigTestDataFrame
 
 	UPROPERTY(VisibleAnywhere, BlueprintReadOnly, Category = "ControlRigTestDataFrame")
 	FRigPose Pose;
+
+	UPROPERTY(VisibleAnywhere, BlueprintReadOnly, Category = "ControlRigTestDataFrame")
+	TArray<uint8> Metadata;
+
+	FCustomVersionContainer ArchiveCustomVersions;
+	mutable TMap<FRigElementKey, URigHierarchy::FMetadataStorage> MetadataMap;
+	mutable bool bTestMetadata;
 };
 
 UENUM()
@@ -86,6 +96,8 @@ public:
 		, bIsApplyingOutputs(false)
 	{}
 
+	virtual void BeginDestroy() override;
+
 	virtual void Serialize(FArchive& Ar) override;
 
 	UFUNCTION(BlueprintCallable, Category = "ControlRigTestData")
diff --git a/Engine/Plugins/Animation/ControlRig/Source/ControlRig/Public/Rigs/RigHierarchy.h b/Engine/Plugins/Animation/ControlRig/Source/ControlRig/Public/Rigs/RigHierarchy.h
--- a/Engine/Plugins/Animation/ControlRig/Source/ControlRig/Public/Rigs/RigHierarchy.h
+++ b/Engine/Plugins/Animation/ControlRig/Source/ControlRig/Public/Rigs/RigHierarchy.h
@@ -4461,6 +4461,8 @@ private:
 			return Ar;
 		}
 	};
+	friend class UControlRigTestData;
+	friend struct FControlRigTestDataFrame;
 
 	// Managed lookup from Element Key to Index
 	TMap<FRigElementKey, int32> ElementIndexLookup;
@@ -4981,6 +4983,10 @@ public:
 	
 	void PropagateMetadata(const FRigElementKey& InKey, const FName& InName, bool bNotify = true);
 	void PropagateMetadata(const FRigBaseElement* InElement, const FName& InName, bool bNotify = true);
+
+
+	TMap<FRigElementKey, URigHierarchy::FMetadataStorage> CopyMetadata() const;
+	bool SetMetadata(const TMap<FRigElementKey, URigHierarchy::FMetadataStorage>& InMetadata);
 	
 private:
 	
@@ -4998,12 +5004,17 @@ private:
 	    */ 
 	FRigBaseMetadata* FindMetadataForElement(const FRigBaseElement* InElement, const FName& InName, ERigMetadataType InType);
 	const FRigBaseMetadata* FindMetadataForElement(const FRigBaseElement* InElement, const FName& InName, ERigMetadataType InType) const;
+
+	/** Attempts to find if the element has any metadata.
+		*/ 
+	bool HasMetadata(const FRigBaseElement* InElement) const;
 	
 	/** Removes the named meta data for the given element, regardless of type. If the element doesn't exist, or it doesn't have any
 	    metadata of the given name, this function does nothing and returns \c false.
 		*/
 	bool RemoveMetadataForElement(FRigBaseElement* InElement, const FName& InName);
 	bool RemoveAllMetadataForElement(FRigBaseElement* InElement);
+	bool RemoveAllMetadata();
 	
 	void CopyAllMetadataFromElement(FRigBaseElement* InTargetElement, const FRigBaseElement* InSourceElement);
 	
diff --git a/Engine/Plugins/Animation/ControlRig/Source/ControlRig/Public/Rigs/RigHierarchyMetadata.h b/Engine/Plugins/Animation/ControlRig/Source/ControlRig/Public/Rigs/RigHierarchyMetadata.h
--- a/Engine/Plugins/Animation/ControlRig/Source/ControlRig/Public/Rigs/RigHierarchyMetadata.h
+++ b/Engine/Plugins/Animation/ControlRig/Source/ControlRig/Public/Rigs/RigHierarchyMetadata.h
@@ -164,6 +164,8 @@ protected:
 	friend struct FRigBaseElement;
 	friend class URigHierarchy;
 	friend class URigHierarchyController;
+	friend class UControlRigTestData;
+	friend struct FControlRigTestDataFrame;
 };
 
 USTRUCT()
diff --git a/Engine/Plugins/Animation/ControlRig/Source/ControlRigEditor/Private/ControlRigElementDetails.cpp b/Engine/Plugins/Animation/ControlRig/Source/ControlRigEditor/Private/ControlRigElementDetails.cpp
--- a/Engine/Plugins/Animation/ControlRig/Source/ControlRigEditor/Private/ControlRigElementDetails.cpp
+++ b/Engine/Plugins/Animation/ControlRig/Source/ControlRigEditor/Private/ControlRigElementDetails.cpp
@@ -27,6 +27,7 @@
 #include "StructViewerModule.h"
 #include "Widgets/SRigVMGraphPinEnumPicker.h"
 #include "IStructureDataProvider.h"
+#include "Editor/ControlRigEditor.h"
 
 #define LOCTEXT_NAMESPACE "ControlRigElementDetails"
 
@@ -1094,7 +1095,12 @@ void FRigBaseElementDetails::PendingDelete()
 	{
 		for (const FPerElementInfo& Info : PerElementInfos)
 		{
-			if (URigHierarchy* Hierarchy = Info.IsValid() ? Info.GetHierarchy() : nullptr)
+			// We do not check Info.IsValid here, because even if the element
+			// doesn't exist anymore in the hierarchy, we still want to detach the
+			// metadata handle from the hierarchy
+
+			
+			if (URigHierarchy* Hierarchy = Info.GetHierarchy())
 			{
 				if (Hierarchy->OnMetadataChanged().Remove(MetadataHandle))
 				{
@@ -1703,7 +1709,30 @@ void FRigBaseElementDetails::CustomizeMetadata(IDetailLayoutBuilder& DetailBuild
 		{
 			return;
 		}
-			
+
+		if (UControlRigBlueprint* Blueprint = PerElementInfos[0].GetBlueprint())
+		{
+			UAssetEditorSubsystem* AssetEditorSubsystem = GEditor->GetEditorSubsystem<UAssetEditorSubsystem>();
+			if (IAssetEditorInstance* EditorInstance = AssetEditorSubsystem->FindEditorForAsset(Blueprint, false))
+			{
+				if (FControlRigEditor* Editor = static_cast<FControlRigEditor*>(EditorInstance))
+				{
+					if (Editor->GetTestAssetPlaybackMode() != EControlRigTestDataPlaybackMode::Live)
+					{
+						Hierarchy->OnMetadataChanged().Remove(MetadataHandle);
+						MetadataHandle.Reset();
+						return;
+					}
+				}
+			}
+		}
+
+		if (MetadataHandle.IsValid())
+		{
+			Hierarchy->OnMetadataChanged().Remove(MetadataHandle);
+			MetadataHandle.Reset();
+		}
+
 		TWeakPtr<IPropertyUtilities> WeakPropertyUtilities =  DetailBuilder.GetPropertyUtilities().ToWeakPtr();
 		MetadataHandle = Hierarchy->OnMetadataChanged().AddLambda([this, WeakPropertyUtilities](const FRigElementKey& InKey, const FName&)
 		{
diff --git a/Engine/Plugins/Animation/ControlRig/Source/ControlRigEditor/Private/Editor/ControlRigEditor.cpp b/Engine/Plugins/Animation/ControlRig/Source/ControlRigEditor/Private/Editor/ControlRigEditor.cpp
--- a/Engine/Plugins/Animation/ControlRig/Source/ControlRigEditor/Private/Editor/ControlRigEditor.cpp
+++ b/Engine/Plugins/Animation/ControlRig/Source/ControlRigEditor/Private/Editor/ControlRigEditor.cpp
@@ -668,6 +668,15 @@ void FControlRigEditor::ToggleTestData()
 	}
 }
 
+EControlRigTestDataPlaybackMode FControlRigEditor::GetTestAssetPlaybackMode() const
+{
+	if (TestDataStrongPtr.IsValid())
+	{
+		return TestDataStrongPtr->GetPlaybackMode();
+	}
+	return EControlRigTestDataPlaybackMode::Live;
+}
+
 void FControlRigEditor::FillToolbar(FToolBarBuilder& ToolbarBuilder, bool bEndSection)
 {
 	FRigVMEditor::FillToolbar(ToolbarBuilder, false);
diff --git a/Engine/Plugins/Animation/ControlRig/Source/ControlRigEditor/Private/Editor/ControlRigEditor.h b/Engine/Plugins/Animation/ControlRig/Source/ControlRigEditor/Private/Editor/ControlRigEditor.h
--- a/Engine/Plugins/Animation/ControlRig/Source/ControlRigEditor/Private/Editor/ControlRigEditor.h
+++ b/Engine/Plugins/Animation/ControlRig/Source/ControlRigEditor/Private/Editor/ControlRigEditor.h
@@ -301,6 +301,9 @@ private:
 	bool RecordTestData(double InRecordingDuration);
 	void ToggleTestData();
 
+public:
+	EControlRigTestDataPlaybackMode GetTestAssetPlaybackMode() const;
+
 protected:
 
 	/** Persona toolkit used to support skeletal mesh preview */
