Control Rig: Fix rotator controls not showing correct values in details panel
#rnx

diff --git a/Engine/Plugins/Animation/ControlRig/Source/ControlRig/Public/Rigs/RigHierarchy.h b/Engine/Plugins/Animation/ControlRig/Source/ControlRig/Public/Rigs/RigHierarchy.h
--- a/Engine/Plugins/Animation/ControlRig/Source/ControlRig/Public/Rigs/RigHierarchy.h
+++ b/Engine/Plugins/Animation/ControlRig/Source/ControlRig/Public/Rigs/RigHierarchy.h
@@ -2659,23 +2659,33 @@ public:
 
 	void SetControlPreferredEulerAngles(FRigControlElement* InControlElement, const FTransform& InTransform, bool bIsInitial = false)
 	{
-		FEulerTransform EulerTransform(InTransform);
-		if (InControlElement && InControlElement->Settings.ControlType == ERigControlType::Transform)
-		{
-			FVector EulerAngle(EulerTransform.Rotation.Roll, EulerTransform.Rotation.Pitch, EulerTransform.Rotation.Yaw);
-			SetControlSpecifiedEulerAngle(InControlElement, EulerAngle, bIsInitial);
-		}
-		else if (InControlElement && InControlElement->Settings.ControlType == ERigControlType::TransformNoScale)
-		{
-			const FTransformNoScale NoScale = EulerTransform.ToFTransform();
-			FVector EulerAngle(EulerTransform.Rotation.Roll, EulerTransform.Rotation.Pitch, EulerTransform.Rotation.Yaw);
-			SetControlSpecifiedEulerAngle(InControlElement, EulerAngle, bIsInitial);
-		}
-		else if (InControlElement && InControlElement->Settings.ControlType == ERigControlType::EulerTransform)
+		if (InControlElement)
 		{
-			FVector EulerAngle(EulerTransform.Rotation.Roll, EulerTransform.Rotation.Pitch, EulerTransform.Rotation.Yaw);
-			FQuat Quat = GetControlQuaternion(InControlElement, EulerAngle);
-			SetControlSpecifiedEulerAngle(InControlElement, EulerAngle, bIsInitial);
+			FEulerTransform EulerTransform(InTransform);
+			switch(InControlElement->Settings.ControlType)
+			{
+				case ERigControlType::Transform:
+				case ERigControlType::Rotator:
+					{
+						FVector EulerAngle(EulerTransform.Rotation.Roll, EulerTransform.Rotation.Pitch, EulerTransform.Rotation.Yaw);
+						SetControlSpecifiedEulerAngle(InControlElement, EulerAngle, bIsInitial);
+						break;
+					}
+				case ERigControlType::TransformNoScale:
+					{
+						const FTransformNoScale NoScale = EulerTransform.ToFTransform();
+						FVector EulerAngle(EulerTransform.Rotation.Roll, EulerTransform.Rotation.Pitch, EulerTransform.Rotation.Yaw);
+						SetControlSpecifiedEulerAngle(InControlElement, EulerAngle, bIsInitial);
+						break;
+					}
+				case ERigControlType::EulerTransform:
+					{
+						FVector EulerAngle(EulerTransform.Rotation.Roll, EulerTransform.Rotation.Pitch, EulerTransform.Rotation.Yaw);
+						FQuat Quat = GetControlQuaternion(InControlElement, EulerAngle);
+						SetControlSpecifiedEulerAngle(InControlElement, EulerAngle, bIsInitial);
+						break;
+					}
+			}
 		}
 	}
 
