Control Rig: Fix test data not comparing the correspondent input/output frames
#rnx

diff --git a/Engine/Plugins/Animation/ControlRig/Source/ControlRig/Private/ControlRigTestData.cpp b/Engine/Plugins/Animation/ControlRig/Source/ControlRig/Private/ControlRigTestData.cpp
--- a/Engine/Plugins/Animation/ControlRig/Source/ControlRig/Private/ControlRigTestData.cpp
+++ b/Engine/Plugins/Animation/ControlRig/Source/ControlRig/Private/ControlRigTestData.cpp
@@ -488,7 +488,7 @@ bool UControlRigTestData::SetupReplay(UControlRig* InControlRig, bool bGroundTru
 			
 			const FRigPose CurrentPose = ControlRig->GetHierarchy()->GetPose();
 
-			const int32 FrameIndex = GetFrameIndexForTime(ControlRig->GetAbsoluteTime(), false);
+			const int32 FrameIndex = LastFrameIndex; 
 			const FControlRigTestDataFrame& Frame = OutputFrames[FrameIndex];
 
 			if (Frame.bTestMetadata)
diff --git a/Engine/Plugins/Animation/ControlRig/Source/ControlRigEditor/Private/Editor/ControlRigEditor.cpp b/Engine/Plugins/Animation/ControlRig/Source/ControlRigEditor/Private/Editor/ControlRigEditor.cpp
--- a/Engine/Plugins/Animation/ControlRig/Source/ControlRigEditor/Private/Editor/ControlRigEditor.cpp
+++ b/Engine/Plugins/Animation/ControlRig/Source/ControlRigEditor/Private/Editor/ControlRigEditor.cpp
@@ -624,16 +624,14 @@ bool FControlRigEditor::RecordTestData(double InRecordingDuration)
 		return false;
 	}
 	
-	if(!TestDataStrongPtr.IsValid())
-	{
-		// create a new test asset
-		static const FString Folder = TEXT("/Game/Animation/ControlRig/NoCook/");
-		const FString DesiredPackagePath =  FString::Printf(TEXT("%s/%s_TestData"), *Folder, *GetControlRigBlueprint()->GetName());
+	
+	// create a new test asset
+	static const FString Folder = TEXT("/Game/Animation/ControlRig/NoCook/");
+	const FString DesiredPackagePath =  FString::Printf(TEXT("%s/%s_TestData"), *Folder, *GetControlRigBlueprint()->GetName());
 
-		if(UControlRigTestData* TestData = UControlRigTestData::CreateNewAsset(DesiredPackagePath, GetControlRigBlueprint()->GetPathName()))
-		{
-			SetTestAssetPath(TestData->GetPathName());
-		}
+	if(UControlRigTestData* TestData = UControlRigTestData::CreateNewAsset(DesiredPackagePath, GetControlRigBlueprint()->GetPathName()))
+	{
+		SetTestAssetPath(TestData->GetPathName());
 	}
 	
 	if(UControlRigTestData* TestData = TestDataStrongPtr.Get())
