Control Rig: Add a cvar to allow disabling tick dependency checking when setting up dependencies for the CR component.

#rnx

diff --git a/Engine/Plugins/Animation/ControlRig/Source/ControlRig/Private/ControlRigComponent.cpp b/Engine/Plugins/Animation/ControlRig/Source/ControlRig/Private/ControlRigComponent.cpp
--- a/Engine/Plugins/Animation/ControlRig/Source/ControlRig/Private/ControlRigComponent.cpp
+++ b/Engine/Plugins/Animation/ControlRig/Source/ControlRig/Private/ControlRigComponent.cpp
@@ -35,6 +35,14 @@ FControlRigAnimInstanceProxy* FControlRigComponentMappedElement::GetAnimProxyOnG
 	return nullptr;
 }
 
+
+TAutoConsoleVariable<bool> CVarControlRigDisableTickDependencyCheck(
+	TEXT("ControlRig.DisableTickDependencyCheck"),
+	false,
+	TEXT("Disables tick dependency checking, which prevents cycles. Use carefully.\n")
+	);
+
+
 ////////////////////////////////////////////////////////////////////////////////////////
 
 #if WITH_EDITOR
@@ -2179,6 +2187,11 @@ bool UControlRigComponent::HasTickDependency_Impl(const UObject* InDependent, co
 
 bool UControlRigComponent::AllowsTickDependency(const UObject* InDependent, const UObject* InDependency, FString* OutErrorMessage)
 {
+	if (CVarControlRigDisableTickDependencyCheck.GetValueOnAnyThread())
+	{
+		return true;
+	}
+	
 	if(InDependent == InDependency)
 	{
 #if WITH_EDITOR
