Constraints: ECS version of 36626478

+ removed some useless includes

#rnx

diff --git a/Engine/Plugins/Animation/ControlRig/Source/ControlRig/Private/Sequencer/MovieSceneControlRigSystem.cpp b/Engine/Plugins/Animation/ControlRig/Source/ControlRig/Private/Sequencer/MovieSceneControlRigSystem.cpp
--- a/Engine/Plugins/Animation/ControlRig/Source/ControlRig/Private/Sequencer/MovieSceneControlRigSystem.cpp
+++ b/Engine/Plugins/Animation/ControlRig/Source/ControlRig/Private/Sequencer/MovieSceneControlRigSystem.cpp
@@ -10,27 +10,21 @@
 #include "MovieSceneTracksComponentTypes.h"
 #include "EntitySystem/MovieSceneInterrogation.h"
 #include "EntitySystem/MovieSceneEntitySystemRunner.h"
-#include "EntitySystem/MovieSceneBoundObjectInstantiator.h"
-#include "EntitySystem/MovieSceneBoundSceneComponentInstantiator.h"
 #include "EntitySystem/MovieSceneEntitySystemTask.h"
 #include "EntitySystem/MovieSceneComponentRegistry.h"
-#include "EntitySystem/MovieSceneRootInstantiatorSystem.h"
 #include "EntitySystem/MovieSceneEntityMutations.h"
 #include "Evaluation/PreAnimatedState/MovieScenePreAnimatedObjectStorage.h"
 #include "Evaluation/PreAnimatedState/MovieScenePreAnimatedStorageID.inl"
 #include "EntitySystem/MovieScenePreAnimatedStateSystem.h"
 #include "EntitySystem/BuiltInComponentTypes.h"
-#include "MovieSceneTracksComponentTypes.h"
 #include "Systems/FloatChannelEvaluatorSystem.h"
 #include "Systems/DoubleChannelEvaluatorSystem.h"
 #include "Systems/MovieScenePiecewiseDoubleBlenderSystem.h"
-#include "Systems/MovieSceneHierarchicalBiasSystem.h"
 #include "Systems/MovieSceneInitialValueSystem.h"
 #include "Systems/WeightAndEasingEvaluatorSystem.h"
 #include "EntitySystem/MovieSceneEntityGroupingSystem.h"
 #include "ControlRig.h"
 #include "Components/SkeletalMeshComponent.h"
-#include "AnimCustomInstanceHelper.h"
 #include "Sequencer/ControlRigLayerInstance.h"
 #include "IControlRigObjectBinding.h"
 #include "ControlRigObjectBinding.h"
@@ -42,7 +36,7 @@
 #include "Rigs/RigHierarchy.h"
 #include "TransformConstraint.h"
 #include "TransformableHandle.h"
-#include "AnimationCoreLibrary.h"
+#include "Constraints/ControlRigTransformableHandle.h"
 
 #include UE_INLINE_GENERATED_CPP_BY_NAME(MovieSceneControlRigSystem)
 
@@ -151,7 +145,7 @@ UTickableConstraint* CreateConstraintIfNeeded(UWorld* InWorld, UTickableConstrai
 		if (Controller.GetConstraint(Constraint->ConstraintID) == nullptr)
 		{
 			Controller.AddConstraint(Constraint);
-			//need to reconstuct channels here.. note this is now lazy and so will recreate it next time view requests it
+			//need to reconstruct channels here.. note this is now lazy and so will recreate it next time view requests it
 			//but only do it if the control rig has a valid world it may not for example in PIE
 			if (InSection->GetControlRig() && InSection->GetControlRig()->GetWorld())
 			{
@@ -186,11 +180,11 @@ struct FControlRigInitialValueProcessor : IInitialValueProcessor
 		: BuiltInComponents(FBuiltInComponentTypes::Get())
 		, ControlRigComponents(FControlRigComponentTypes::Get())
 		, TracksComponents(FMovieSceneTracksComponentTypes::Get())
+		, Interrogation(nullptr)
+		, InitialValueCache(nullptr)
 		, WriteContext(FEntityAllocationWriteContext::NewAllocation())
-	{
-		Interrogation     = nullptr;
-		InitialValueCache = nullptr;
-	}
+		, ParameterSystem(nullptr)
+	{}
 
 	virtual void Initialize(UMovieSceneEntitySystemLinker* Linker, FInitialValueCache* InInitialValueCache) override
 	{
@@ -1162,18 +1156,82 @@ struct FEvaluateBaseControlRigs
 			}
 		}
 
-		//unfortunately for Constraints with ControlRig we need to resolve all Parents also. Don't need to do children since they wil be handled by
-		//the channel resolve above
-		const FConstraintsManagerController& Controller = FConstraintsManagerController::Get(BoundObject->GetWorld());
-		const TArray< TWeakObjectPtr<UTickableConstraint>> AllConstraints = Controller.GetAllConstraints();
-		for (const TWeakObjectPtr<UTickableConstraint>& TickConstraint : AllConstraints)
+		// for Constraints with ControlRig we need to resolve all Parents also
+		// Don't need to do children since they wil be handled by the channel resolve above
+		ResolveParentHandles(BoundObject, ControlRig, InOutBaseEvalData, SequenceID, SharedPlaybackState);
+	}
+
+	static void ResolveParentHandles(
+		const UObject* InBoundObject, UControlRig* InControlRigInstance,
+		const FBaseControlRigEvalData& InBaseEvalData,
+		const FMovieSceneSequenceID& InSequenceID,
+		const TSharedRef<FSharedPlaybackState>& InSharedPlaybackState)
+	{
+		if (!InBoundObject)
+		{
+			return;
+		}
+
+		UWorld* BoundObjectWorld = InBoundObject->GetWorld();
+		const bool bIsGameWorld = InBoundObject->GetWorld() ? InBoundObject->GetWorld()->IsGameWorld() : false;
+		
+		UMovieSceneControlRigParameterTrack* ControlRigTrack = InBaseEvalData.Section->GetTypedOuter<UMovieSceneControlRigParameterTrack>();
+
+		// is this control rig a game world instance of this section's rig?
+		auto WasAGameInstance = [ControlRigTrack](const UControlRig* InRigToTest)
+		{
+			return ControlRigTrack ? ControlRigTrack->IsAGameInstance(InRigToTest) : false;
+		};
+
+		// is the parent handle of this constraint related to this section?
+		// this return true if the handle's control rig has been spawned by the ControlRigTrack (whether in Editor or Game)
+		// if false, it means that the handle represents another control on another control rig so we don't need to resolve it here
+		// note that it returns true if ControlRigTrack is null (is this possible?!) or if the ControlRig is null (we can't infer anything from this)
+		auto ShouldResolveParent = [ControlRigTrack](const UTransformableControlHandle* ParentControlHandle)
 		{
-			if (UTickableTransformConstraint* TransformConstraint = Cast< UTickableTransformConstraint>(TickConstraint.Get()))
+			if (!ParentControlHandle)
 			{
-				if (TransformConstraint->ParentTRSHandle)
+				return false;
+			}
+
+			if (!ControlRigTrack)
+			{
+				// cf. UObjectBaseUtility::IsInOuter
+				return true;
+			}
+			
+			return ParentControlHandle->ControlRig ? ParentControlHandle->ControlRig->IsInOuter(ControlRigTrack) : true;
+		};
+
+		// this is the default's section rig. when bIsGameWorld is false, InControlRigInstance should be equal to SectionRig
+		const UControlRig* SectionRig = InBaseEvalData.Section->GetControlRig();
+
+		const FConstraintsManagerController& Controller = FConstraintsManagerController::Get(BoundObjectWorld);
+		const TArray< TWeakObjectPtr<UTickableConstraint>> Constraints = Controller.GetAllConstraints();
+
+		for (const TWeakObjectPtr<UTickableConstraint>& TickConstraint : Constraints)
+		{
+			UTickableTransformConstraint* TransformConstraint = Cast< UTickableTransformConstraint>(TickConstraint.Get());
+			UTransformableControlHandle* ParentControlHandle = TransformConstraint ? Cast<UTransformableControlHandle>(TransformConstraint->ParentTRSHandle) : nullptr;
+			if (ParentControlHandle && ShouldResolveParent(ParentControlHandle))
+			{
+				if (bIsGameWorld)
+				{
+					// switch from section's rig to the game instance
+					if (ParentControlHandle->ControlRig == SectionRig)
+					{
+						ParentControlHandle->ResolveBoundObjects(InSequenceID, InSharedPlaybackState, InControlRigInstance);
+						TransformConstraint->EnsurePrimaryDependency(BoundObjectWorld);
+					}
+				}
+				else
 				{
-					TransformConstraint->ParentTRSHandle->ResolveBoundObjects(SequenceID, SharedPlaybackState, ControlRig);
-					TransformConstraint->EnsurePrimaryDependency(BoundObject->GetWorld());
+					// switch from the game instance to the section's rig
+					if (WasAGameInstance(ParentControlHandle->ControlRig.Get()))
+					{
+						ParentControlHandle->ResolveBoundObjects(InSequenceID, InSharedPlaybackState, InControlRigInstance);
+						TransformConstraint->EnsurePrimaryDependency(BoundObjectWorld);
+					}
 				}
 			}
 		}
diff --git a/Engine/Plugins/Animation/ControlRig/Source/ControlRig/Public/Sequencer/MovieSceneControlRigSystem.h b/Engine/Plugins/Animation/ControlRig/Source/ControlRig/Public/Sequencer/MovieSceneControlRigSystem.h
--- a/Engine/Plugins/Animation/ControlRig/Source/ControlRig/Public/Sequencer/MovieSceneControlRigSystem.h
+++ b/Engine/Plugins/Animation/ControlRig/Source/ControlRig/Public/Sequencer/MovieSceneControlRigSystem.h
@@ -3,7 +3,6 @@
 #pragma once
 
 #include "EntitySystem/MovieSceneEntityIDs.h"
-#include "EntitySystem/MovieSceneEntityInstantiatorSystem.h"
 #include "EntitySystem/MovieSceneEntitySystem.h"
 #include "EntitySystem/MovieSceneOverlappingEntityTracker.h"
 #include "Evaluation/PreAnimatedState/MovieScenePreAnimatedObjectStorage.h"
