Cotnrol Rig: Array Connectors further UI work

#rnx

diff --git a/Engine/Plugins/Animation/ControlRig/Source/ControlRig/Private/ModularRigController.cpp b/Engine/Plugins/Animation/ControlRig/Source/ControlRig/Private/ModularRigController.cpp
--- a/Engine/Plugins/Animation/ControlRig/Source/ControlRig/Private/ModularRigController.cpp
+++ b/Engine/Plugins/Animation/ControlRig/Source/ControlRig/Private/ModularRigController.cpp
@@ -433,6 +433,46 @@ bool UModularRigController::ConnectConnectorToElements(const FRigElementKey& InC
 	return true;
 }
 
+bool UModularRigController::AddTargetToArrayConnector(const FRigElementKey& InConnectorKey, const FRigElementKey& InTargetKey, bool bSetupUndo,
+	bool bAutoResolveOtherConnectors, bool bCheckValidConnection)
+{
+	TArray<FRigElementKey> ExistingTargets = Model->Connections.FindTargetsFromConnector(InConnectorKey);
+	ExistingTargets.Add(InTargetKey);
+	
+	if(bCheckValidConnection)
+	{
+		return ConnectConnectorToElements(InConnectorKey, ExistingTargets, bSetupUndo, bAutoResolveOtherConnectors, bCheckValidConnection);
+	}
+
+	FString ConnectorParentPath, ConnectorName;
+	(void)URigHierarchy::SplitNameSpace(InConnectorKey.Name.ToString(), &ConnectorParentPath, &ConnectorName);
+	const FRigModuleReference* Module = FindModule(ConnectorParentPath);
+	if(Module == nullptr)
+	{
+		return false;
+	}
+
+	FRigElementKeyRedirector::FKeyArray TargetKeyArray;
+	TargetKeyArray.Append(ExistingTargets);
+
+	UBlueprint* Blueprint = Cast<UBlueprint>(GetOuter());
+#if WITH_EDITOR
+	TSharedPtr<FScopedTransaction> TransactionPtr;
+	if (bSetupUndo)
+	{
+		TransactionPtr = MakeShared<FScopedTransaction>(NSLOCTEXT("ModularRigController", "AddTargetToArrayConnector", "Add Target to Array Connector"), !GIsTransacting);
+		Blueprint->Modify();
+	}
+#endif 
+
+	Model->Connections.RemoveConnection(InConnectorKey);
+	Model->Connections.AddConnection(InConnectorKey, TargetKeyArray);
+
+	Notify(EModularRigNotification::ConnectionChanged, Module);
+
+	return true;
+}
+
 bool UModularRigController::DisconnectConnector(const FRigElementKey& InConnectorKey, bool bDisconnectSubModules, bool bSetupUndo)
 {
 	return DisconnectConnector_Internal(InConnectorKey, bDisconnectSubModules, nullptr, bSetupUndo);
diff --git a/Engine/Plugins/Animation/ControlRig/Source/ControlRig/Private/Rigs/RigConnectionRules.cpp b/Engine/Plugins/Animation/ControlRig/Source/ControlRig/Private/Rigs/RigConnectionRules.cpp
--- a/Engine/Plugins/Animation/ControlRig/Source/ControlRig/Private/Rigs/RigConnectionRules.cpp
+++ b/Engine/Plugins/Animation/ControlRig/Source/ControlRig/Private/Rigs/RigConnectionRules.cpp
@@ -172,7 +172,11 @@ const FRigTransformElement* FRigConnectionRuleInput::ResolveConnector(const FRig
 	{
 		if(!Cache->IsEmpty())
 		{
-			return Cast<FRigTransformElement>(Cache->operator[](0).GetElement());
+			const FCachedRigElement& CachedElement = Cache->operator[](0);
+			if(const_cast<FCachedRigElement&>(CachedElement).UpdateCache(GetHierarchy()))
+			{
+				return Cast<FRigTransformElement>(CachedElement.GetElement());
+			}
 		}
 	}
 
diff --git a/Engine/Plugins/Animation/ControlRig/Source/ControlRig/Private/Rigs/RigHierarchyCache.cpp b/Engine/Plugins/Animation/ControlRig/Source/ControlRig/Private/Rigs/RigHierarchyCache.cpp
--- a/Engine/Plugins/Animation/ControlRig/Source/ControlRig/Private/Rigs/RigHierarchyCache.cpp
+++ b/Engine/Plugins/Animation/ControlRig/Source/ControlRig/Private/Rigs/RigHierarchyCache.cpp
@@ -228,9 +228,13 @@ void FRigElementKeyRedirector::Add(const FRigElementKey& InSource, const FKeyArr
 		return;
 	}
 
-	if (IsValid(InTargets))
+	if(!InTargets.IsEmpty())
 	{
-		InternalKeyToExternalKey.Add(InSource, Convert(InTargets, InHierarchy, true));
+		const FCachedKeyArray Cache = Convert(InTargets, InHierarchy, true, true);
+		if(!Cache.IsEmpty())
+		{
+			InternalKeyToExternalKey.Add(InSource, Cache);
+		}
 	}
 	
 	ExternalKeys.Add(InSource, InTargets);
@@ -256,13 +260,16 @@ FRigElementKeyRedirector::FKeyArray FRigElementKeyRedirector::Convert(const TArr
 }
 
 FRigElementKeyRedirector::FCachedKeyArray FRigElementKeyRedirector::Convert(const FKeyArray& InKeys, const URigHierarchy* InHierarchy,
-	bool bForceStoreKey)
+	bool bForceStoreKey, bool bOnlyValidItems)
 {
 	FCachedKeyArray Result;
 	Result.Reserve(InKeys.Num());
 	for(const FRigElementKey& Key : InKeys)
 	{
-		Result.Emplace(Key, InHierarchy, bForceStoreKey);
+		if(!bOnlyValidItems || Key.IsValid())
+		{
+			Result.Emplace(Key, InHierarchy, bForceStoreKey);
+		}
 	}
 	return Result;
 }
\ No newline at end of file
diff --git a/Engine/Plugins/Animation/ControlRig/Source/ControlRig/Private/Units/Execution/RigUnit_RigModules.cpp b/Engine/Plugins/Animation/ControlRig/Source/ControlRig/Private/Units/Execution/RigUnit_RigModules.cpp
--- a/Engine/Plugins/Animation/ControlRig/Source/ControlRig/Private/Units/Execution/RigUnit_RigModules.cpp
+++ b/Engine/Plugins/Animation/ControlRig/Source/ControlRig/Private/Units/Execution/RigUnit_RigModules.cpp
@@ -8,25 +8,50 @@
 
 FRigUnit_ResolveConnector_Execute()
 {
-	Result = Connector;
+	TArray<FRigElementKey> ResultArray;
+	FRigUnit_ResolveArrayConnector::StaticExecute(ExecuteContext, Connector, SkipSocket, ResultArray, bIsConnected);
+	if(ResultArray.IsEmpty())
+	{
+		Result = Connector;
+	}
+	else
+	{
+		Result = ResultArray[0];
+	}
+}
+
+FRigUnit_ResolveArrayConnector_Execute()
+{
+	Result = {Connector};
 
 	if (const URigHierarchy* Hierarchy = ExecuteContext.Hierarchy)
 	{
-		if(const FRigBaseElement* ResolvedElement = Hierarchy->Find(Connector))
+		Result = Hierarchy->GetResolvedTargets(Connector);
+		if(SkipSocket && !Result.IsEmpty())
 		{
-			Result = ResolvedElement->GetKey();
-			if(SkipSocket && Result.IsValid() && Result.Type == ERigElementType::Socket)
+			for(FRigElementKey& ResultKey : Result)
 			{
-				const FRigElementKey ParentOfSocket = Hierarchy->GetFirstParent(Result);
-				if(ParentOfSocket.IsValid())
+				if(ResultKey.Type == ERigElementType::Socket)
 				{
-					Result = ParentOfSocket;
+					const FRigElementKey ParentOfSocket = Hierarchy->GetFirstParent(ResultKey);
+					if(ParentOfSocket.IsValid())
+					{
+						ResultKey = ParentOfSocket;
+					}
 				}
 			}
 		}
 	}
 
-	bIsConnected = Result != Connector;
+	bIsConnected = false;
+	if(Result.Num() == 1)
+	{
+		bIsConnected = Result[0] != Connector;
+	}
+	else if(Result.Num() > 1)
+	{
+		bIsConnected = true;
+	}
 }
 
 FRigUnit_GetCurrentNameSpace_Execute()
diff --git a/Engine/Plugins/Animation/ControlRig/Source/ControlRig/Public/ModularRigController.h b/Engine/Plugins/Animation/ControlRig/Source/ControlRig/Public/ModularRigController.h
--- a/Engine/Plugins/Animation/ControlRig/Source/ControlRig/Public/ModularRigController.h
+++ b/Engine/Plugins/Animation/ControlRig/Source/ControlRig/Public/ModularRigController.h
@@ -40,6 +40,9 @@ class CONTROLRIG_API UModularRigController : public UObject
 	UFUNCTION(BlueprintCallable, Category = "Control Rig | Modules")
 	bool ConnectConnectorToElements(const FRigElementKey& InConnectorKey, const TArray<FRigElementKey>& InTargetKeys, bool bSetupUndo = true, bool bAutoResolveOtherConnectors = true, bool bCheckValidConnection = true);
 
+	UFUNCTION(BlueprintCallable, Category = "Control Rig | Modules")
+	bool AddTargetToArrayConnector(const FRigElementKey& InConnectorKey, const FRigElementKey& InTargetKey, bool bSetupUndo = true, bool bAutoResolveOtherConnectors = true, bool bCheckValidConnection = true);
+
 	UFUNCTION(BlueprintCallable, Category = "Control Rig | Modules")
 	bool DisconnectConnector(const FRigElementKey& InConnectorKey, bool bDisconnectSubModules = false, bool bSetupUndo = true);
 	bool DisconnectConnector_Internal(const FRigElementKey& InConnectorKey, bool bDisconnectSubModules = false, TMap<FRigElementKey, TArray<FRigElementKey>>*OutRemovedConnections = nullptr, bool bSetupUndo = true);
diff --git a/Engine/Plugins/Animation/ControlRig/Source/ControlRig/Public/ModularRigModel.h b/Engine/Plugins/Animation/ControlRig/Source/ControlRig/Public/ModularRigModel.h
--- a/Engine/Plugins/Animation/ControlRig/Source/ControlRig/Public/ModularRigModel.h
+++ b/Engine/Plugins/Animation/ControlRig/Source/ControlRig/Public/ModularRigModel.h
@@ -237,7 +237,7 @@ public:
 		int32 Index = FindConnectionIndex(InConnectorKey);
 		if (ConnectionList.IsValidIndex(Index))
 		{
-			check(ConnectionList[Index].Targets.Num() == 1);
+			check(ConnectionList[Index].Targets.Num() >= 1);
 			return ConnectionList[Index].Targets[0];
 		}
 		static const FRigElementKey EmptyKey;
diff --git a/Engine/Plugins/Animation/ControlRig/Source/ControlRig/Public/Rigs/RigHierarchyCache.h b/Engine/Plugins/Animation/ControlRig/Source/ControlRig/Public/Rigs/RigHierarchyCache.h
--- a/Engine/Plugins/Animation/ControlRig/Source/ControlRig/Public/Rigs/RigHierarchyCache.h
+++ b/Engine/Plugins/Animation/ControlRig/Source/ControlRig/Public/Rigs/RigHierarchyCache.h
@@ -394,14 +394,14 @@ public:
 	FCachedKeyArray* Find(const FRigElementKey& InKey) { return InternalKeyToExternalKey.Find(InKey); }
 	const FRigElementKey* FindReverse(const FRigElementKey& InKey) const;
 	uint32 GetHash() const { return Hash; }
-	
-private:
-
-	void Add(const FRigElementKey& InSource, const FKeyArray& InTargets, const URigHierarchy* InHierarchy);
 
 	static FKeyArray Convert(const FCachedKeyArray& InCachedKeys);
 	static FKeyArray Convert(const TArray<FRigElementKey>& InKeys);
-	static FCachedKeyArray Convert(const FKeyArray& InKeys, const URigHierarchy* InHierarchy, bool bForceStoreKey = false);
+	static FCachedKeyArray Convert(const FKeyArray& InKeys, const URigHierarchy* InHierarchy, bool bForceStoreKey = false, bool bOnlyValidItems = false);
+
+private:
+
+	void Add(const FRigElementKey& InSource, const FKeyArray& InTargets, const URigHierarchy* InHierarchy);
 
 	template<typename T = FKeyArray>
 	static bool IsValid(const T& InKeys)
diff --git a/Engine/Plugins/Animation/ControlRig/Source/ControlRig/Public/Units/Execution/RigUnit_RigModules.h b/Engine/Plugins/Animation/ControlRig/Source/ControlRig/Public/Units/Execution/RigUnit_RigModules.h
--- a/Engine/Plugins/Animation/ControlRig/Source/ControlRig/Public/Units/Execution/RigUnit_RigModules.h
+++ b/Engine/Plugins/Animation/ControlRig/Source/ControlRig/Public/Units/Execution/RigUnit_RigModules.h
@@ -62,6 +62,50 @@ struct CONTROLRIG_API FRigUnit_ResolveConnector : public FRigUnit_RigModulesBase
 	bool bIsConnected;
 };
 
+/**
+* Returns the resolved item of the connector.
+*/
+USTRUCT(meta=(DisplayName="Get Array Connection", Keywords="Connector,GetResolved,Target,Resolve", Varying))
+struct CONTROLRIG_API FRigUnit_ResolveArrayConnector : public FRigUnit_RigModulesBase
+{
+	GENERATED_BODY()
+
+	FRigUnit_ResolveArrayConnector()
+	{
+		Connector = FRigElementKey(NAME_None, ERigElementType::Connector);
+		SkipSocket = false;
+		bIsConnected = false;
+	}
+
+	RIGVM_METHOD()
+	virtual void Execute() override;
+
+	/*
+	 * The key of the connector to resolve
+	 */
+	UPROPERTY(meta = (Input, ExpandByDefault))
+	FRigElementKey Connector;
+
+	/*
+	 * If the connector is resolved to a socket the node
+	 * will return the socket's direct parent (skipping it).
+	 */
+	UPROPERTY(meta = (Input, ExpandByDefault))
+	bool SkipSocket;
+
+	/*
+	 * The resulting item the connector is resolved to
+	 */
+	UPROPERTY(meta = (Output))
+	TArray<FRigElementKey> Result;
+
+	/*
+	 * Returns true if the connector is resolved to a target.
+	 */
+	UPROPERTY(meta = (Output))
+	bool bIsConnected;
+};
+
 /**
  * Returns the currently used namespace of the module.
  */
diff --git a/Engine/Plugins/Animation/ControlRig/Source/ControlRigDeveloper/Private/ControlRigBlueprint.cpp b/Engine/Plugins/Animation/ControlRig/Source/ControlRigDeveloper/Private/ControlRigBlueprint.cpp
--- a/Engine/Plugins/Animation/ControlRig/Source/ControlRigDeveloper/Private/ControlRigBlueprint.cpp
+++ b/Engine/Plugins/Animation/ControlRig/Source/ControlRigDeveloper/Private/ControlRigBlueprint.cpp
@@ -1144,6 +1144,11 @@ void UControlRigBlueprint::UpdateExposedModuleConnectors() const
 }
 
 bool UControlRigBlueprint::ResolveConnector(const FRigElementKey& DraggedKey, const FRigElementKey& TargetKey, bool bSetupUndoRedo)
+{
+	return ResolveConnectorToArray(DraggedKey, {TargetKey}, bSetupUndoRedo);
+}
+
+bool UControlRigBlueprint::ResolveConnectorToArray(const FRigElementKey& DraggedKey, const TArray<FRigElementKey>& TargetKeys, bool bSetupUndoRedo)
 {
 	FScopedTransaction Transaction(LOCTEXT("ResolveConnector", "Resolve connector"));
 
@@ -1151,21 +1156,41 @@ bool UControlRigBlueprint::ResolveConnector(const FRigElementKey& DraggedKey, co
 	{
 		Modify();
 	}
-	if(TargetKey.IsValid())
+
+	TArray<FRigElementKey> FilteredKeys = TargetKeys;
+	FilteredKeys.RemoveAll([](const FRigElementKey& Key)
+	{
+		return !Key.IsValid();
+	});
+	
+	if(!FilteredKeys.IsEmpty())
 	{
 		FRigElementKeyCollection& ExistingTargetKeys = ArrayConnectionMap.FindOrAdd(DraggedKey);
-		if(ExistingTargetKeys.Keys.Num() == 1 && ExistingTargetKeys.Keys.Contains(TargetKey))
+
+		if(ExistingTargetKeys.Num() == FilteredKeys.Num())
 		{
-			return false;
+			bool bCompleteMatch = true;
+			for(int32 Index = 0; Index < ExistingTargetKeys.Num(); Index++)
+			{
+				if(ExistingTargetKeys[Index] != FilteredKeys[Index])
+				{
+					bCompleteMatch = false;
+					break;
+				}
+			}
+			if(bCompleteMatch)
+			{
+				return false;
+			}
 		}
-		ExistingTargetKeys.Keys = {TargetKey};
+		ExistingTargetKeys.Keys = FilteredKeys;
 
 		if (IsModularRig())
 		{
 			// Add connection to the model
 			if (UModularRigController* Controller = GetModularRigController())
 			{
-				Controller->ConnectConnectorToElement(DraggedKey, TargetKey, bSetupUndoRedo, ModularRigSettings.bAutoResolve);
+				Controller->ConnectConnectorToElements(DraggedKey, FilteredKeys, bSetupUndoRedo, ModularRigSettings.bAutoResolve);
 			}
 		}
 		else
diff --git a/Engine/Plugins/Animation/ControlRig/Source/ControlRigDeveloper/Public/ControlRigBlueprint.h b/Engine/Plugins/Animation/ControlRig/Source/ControlRigDeveloper/Public/ControlRigBlueprint.h
--- a/Engine/Plugins/Animation/ControlRig/Source/ControlRigDeveloper/Public/ControlRigBlueprint.h
+++ b/Engine/Plugins/Animation/ControlRig/Source/ControlRigDeveloper/Public/ControlRigBlueprint.h
@@ -224,6 +224,7 @@ protected:
 	FOnRigTypeChanged OnRigTypeChangedDelegate;
 	
 	bool ResolveConnector(const FRigElementKey& DraggedKey, const FRigElementKey& TargetKey, bool bSetupUndoRedo = true);
+	bool ResolveConnectorToArray(const FRigElementKey& DraggedKey, const TArray<FRigElementKey>& TargetKeys, bool bSetupUndoRedo = true);
 
 	void UpdateConnectionMapFromModel();
 
diff --git a/Engine/Plugins/Animation/ControlRig/Source/ControlRigEditor/Private/ControlRigElementDetails.cpp b/Engine/Plugins/Animation/ControlRig/Source/ControlRigEditor/Private/ControlRigElementDetails.cpp
--- a/Engine/Plugins/Animation/ControlRig/Source/ControlRigEditor/Private/ControlRigElementDetails.cpp
+++ b/Engine/Plugins/Animation/ControlRig/Source/ControlRigEditor/Private/ControlRigElementDetails.cpp
@@ -6176,10 +6176,14 @@ void FRigConnectorElementDetails::CustomizeSettings(IDetailLayoutBuilder& Detail
 	}
 
 	IsArrayHandle = SettingsHandle->GetChildHandle(GET_MEMBER_NAME_CHECKED(FRigConnectorSettings, bIsArray));
-	SettingsCategory
+	IDetailPropertyRow& IsArrayPropertyRow = SettingsCategory
 		.AddProperty(IsArrayHandle)
-		.IsEnabled(!IsAnyConnectorImported())
-		.Visibility(TAttribute<EVisibility>::CreateSP(this, &FRigConnectorElementDetails::GetArrayVisibility));
+		.IsEnabled(IsArrayEnabled());
+
+	if(!IsArrayEnabled())
+	{
+		IsArrayPropertyRow.ToolTip(LOCTEXT("PrimaryConnectorsDontAllowArrayToolTip", "Primary Connectors don't support arrays. Add a secondary connector for that."));
+	}
 
 	if(!bHideRules)
 	{
@@ -6229,10 +6233,14 @@ TOptional<bool> FRigConnectorElementDetails::GetIsConnectorArray() const
 	return Result;
 }
 
-EVisibility FRigConnectorElementDetails::GetArrayVisibility() const
+bool FRigConnectorElementDetails::IsArrayEnabled() const
 {
 	const TOptional<EConnectorType> ConnectorType = GetConnectorType();
-	return ConnectorType.Get(EConnectorType::Primary) == EConnectorType::Primary ? EVisibility::Collapsed : EVisibility::Visible;
+	if(ConnectorType.Get(EConnectorType::Primary) != EConnectorType::Primary)
+	{
+		return !IsAnyConnectorImported();
+	}
+	return false;
 }
 
 void FRigSocketElementDetails::CustomizeDetails(IDetailLayoutBuilder& DetailBuilder)
diff --git a/Engine/Plugins/Animation/ControlRig/Source/ControlRigEditor/Private/ControlRigElementDetails.h b/Engine/Plugins/Animation/ControlRig/Source/ControlRigEditor/Private/ControlRigElementDetails.h
--- a/Engine/Plugins/Animation/ControlRig/Source/ControlRigEditor/Private/ControlRigElementDetails.h
+++ b/Engine/Plugins/Animation/ControlRig/Source/ControlRigEditor/Private/ControlRigElementDetails.h
@@ -872,7 +872,7 @@ private:
 
 	TOptional<EConnectorType> GetConnectorType() const;
 	TOptional<bool> GetIsConnectorArray() const;
-	EVisibility GetArrayVisibility() const;
+	bool IsArrayEnabled() const;
 
 	TSharedPtr<IPropertyHandle> ConnectorTypeHandle;
 	TSharedPtr<IPropertyHandle> IsArrayHandle;
diff --git a/Engine/Plugins/Animation/ControlRig/Source/ControlRigEditor/Private/ControlRigModuleDetails.cpp b/Engine/Plugins/Animation/ControlRig/Source/ControlRigEditor/Private/ControlRigModuleDetails.cpp
--- a/Engine/Plugins/Animation/ControlRig/Source/ControlRigEditor/Private/ControlRigModuleDetails.cpp
+++ b/Engine/Plugins/Animation/ControlRig/Source/ControlRigEditor/Private/ControlRigModuleDetails.cpp
@@ -312,20 +312,65 @@ void FRigModuleInstanceDetails::CustomizeDetails(IDetailLayoutBuilder& DetailBui
 				}
 			}
 		}
+
+		TArray<FRigModuleConnector> Connectors;
+		TArray<TOptional<bool>> IsArrayConnector;
 		if (bDisplayConnectors)
 		{
-			TArray<FRigModuleConnector> Connectors = GetConnectors();
-
+			Connectors = GetConnectors();
+			
 			// sort connectors primary first, then secondary, then optional
 			Algo::SortBy(Connectors, [](const FRigModuleConnector& Connector) -> int32
 			{
 				return Connector.IsPrimary() ? 0 : (Connector.IsOptional() ? 2 : 1);
 			});
+
+			IsArrayConnector.Reserve(Connectors.Num());
 			
 			for(const FRigModuleConnector& Connector : Connectors)
 			{
 				const FText Label = FText::FromString(Connector.Name);
-				TSharedPtr<SVerticalBox> ButtonBox;
+				IsArrayConnector.Emplace();
+				TOptional<bool>& IsArray = IsArrayConnector.Last();
+				
+				TArray<FRigElementResolveResult> Matches;
+				for (int32 ModuleIndex=0; ModuleIndex<PerModuleInfos.Num(); ++ModuleIndex)
+				{
+					const FPerModuleInfo& Info = PerModuleInfos[ModuleIndex];
+					if (const FRigModuleInstance* Module = Info.GetModule())
+					{
+						if (UModularRig* ModularRig = Info.GetModularRig())
+						{
+							if (URigHierarchy* Hierarchy = ModularRig->GetHierarchy())
+							{
+								FString ConnectorPath = FString::Printf(TEXT("%s:%s"), *Info.Path, *Connector.Name);
+								FRigElementKey ConnectorKey(*ConnectorPath, ERigElementType::Connector);
+								if (FRigConnectorElement* ConnectorElement = Cast<FRigConnectorElement>(Hierarchy->Find(ConnectorKey)))
+								{
+									if(ModuleIndex == 0)
+									{
+										IsArray = ConnectorElement->IsArrayConnector();
+									}
+									else if(IsArray.Get(false) != ConnectorElement->IsArrayConnector())
+									{
+										IsArray.Reset();
+									}
+								}
+							}
+						}
+					}
+				}
+			}
+		}
+
+		if(bDisplayConnectors)
+		{
+			for(int32 ConnectorIndex = 0; ConnectorIndex < Connectors.Num(); ConnectorIndex++)
+			{
+				const FRigModuleConnector& Connector = Connectors[ConnectorIndex];
+				const FString ConnectorName = Connector.Name;
+				const FText Label = FText::FromString(Connector.Name);
+				bool bIsArrayConnector = IsArrayConnector[ConnectorIndex].Get(false);
 
 				TArray<FRigElementResolveResult> Matches;
 				for (int32 ModuleIndex=0; ModuleIndex<PerModuleInfos.Num(); ++ModuleIndex)
@@ -367,6 +412,7 @@ void FRigModuleInstanceDetails::CustomizeDetails(IDetailLayoutBuilder& DetailBui
 				});
 
 				TArray<FRigElementKey> MatchingKeys;
+				MatchingKeys.Reserve(Matches.Num());
 				for(const FRigElementResolveResult& SingleMatch : Matches)
 				{
 					MatchingKeys.Add(SingleMatch.GetKey());
@@ -375,287 +421,404 @@ void FRigModuleInstanceDetails::CustomizeDetails(IDetailLayoutBuilder& DetailBui
 				{
 					return MatchingKeys.Contains(InTarget);
 				});
-				TreeDelegates.OnGetSelection.BindLambda([this, Connector]() -> TArray<FRigElementKey>
+
+				static const FSlateBrush* PrimaryBrush = FControlRigEditorStyle::Get().GetBrush("ControlRig.ConnectorPrimary");
+				static const FSlateBrush* SecondaryBrush = FControlRigEditorStyle::Get().GetBrush("ControlRig.ConnectorSecondary");
+				static const FSlateBrush* OptionalBrush = FControlRigEditorStyle::Get().GetBrush("ControlRig.ConnectorOptional");
+
+				const FSlateBrush* IconBrush = Connector.IsPrimary() ? PrimaryBrush : (Connector.IsOptional() ? OptionalBrush : SecondaryBrush);
+
+				TSharedRef<SVerticalBox> ValueContentVerticalBox = SNew(SVerticalBox);
+
+				TArray<FRigElementKey> CurrentTargetKeys;
+				if (PerModuleInfos.Num() >= 1)
 				{
-					TOptional<FRigElementKey> Target;
-					FRigElementKeyRedirector Redirector = PerModuleInfos[0].GetModularRig()->GetElementKeyRedirector();
+					TOptional<TArray<FRigElementKey>> TargetKeys;
 					for (int32 ModuleIndex=0; ModuleIndex<PerModuleInfos.Num(); ++ModuleIndex)
 					{
 						FString ConnectorPath = FString::Printf(TEXT("%s:%s"), *PerModuleInfos[ModuleIndex].Path, *Connector.Name);
 						FRigElementKey ConnectorKey(*ConnectorPath, ERigElementType::Connector);
-						if (const FRigElementKeyRedirector::FKeyArray* Keys = Redirector.FindExternalKey(ConnectorKey))
+						if(const UControlRigBlueprint* Blueprint = PerModuleInfos[ModuleIndex].GetBlueprint())
 						{
-							for(const FRigElementKey& Key : *Keys)
+							const TArray<FRigElementKey> CurrentTargets = Blueprint->ModularRigModel.Connections.FindTargetsFromConnector(ConnectorKey);
+							if (!TargetKeys.IsSet())
+							{
+								TargetKeys = CurrentTargets;
+							}
+							else
 							{
-								if (!Target.IsSet())
+								const TArray<FRigElementKey>& FlatKeys = TargetKeys.GetValue();
+								if(FlatKeys.Num() != CurrentTargets.Num())
 								{
-									Target = Key;
+									TargetKeys.Reset();
+									break;
 								}
-								else if (Target.GetValue() != Key)
+
+								for(int32 TargetKeyIndex = 0; TargetKeyIndex < CurrentTargets.Num(); TargetKeyIndex++)
 								{
-									Target.GetValue().Name = *ControlRigModuleDetailsMultipleValues.ToString();
-									return {Target.GetValue()};
+									if(FlatKeys[TargetKeyIndex] != CurrentTargets[TargetKeyIndex])
+									{
+										TargetKeys.Reset();
+										break;
+									}
+								}
+
+								if(!TargetKeys.IsSet())
+								{
+									break;
 								}
 							}
-						}
-						else
-						{
-							return {FRigElementKey(*ControlRigModuleDetailsMultipleValues.ToString(), ERigElementType::None)};
+							break;
 						}
 					}
-					return {};
-				});
-				TreeDelegates.OnSelectionChanged.BindSP(this, &FRigModuleInstanceDetails::OnConnectorTargetChanged, Connector);
-			
-				static const FSlateBrush* PrimaryBrush = FControlRigEditorStyle::Get().GetBrush("ControlRig.ConnectorPrimary");
-				static const FSlateBrush* SecondaryBrush = FControlRigEditorStyle::Get().GetBrush("ControlRig.ConnectorSecondary");
-				static const FSlateBrush* OptionalBrush = FControlRigEditorStyle::Get().GetBrush("ControlRig.ConnectorOptional");
 
-				const FSlateBrush* IconBrush = Connector.IsPrimary() ? PrimaryBrush : (Connector.IsOptional() ? OptionalBrush : SecondaryBrush);
-				TSharedPtr<SSearchableRigHierarchyTreeView>& SearchableTreeView = ConnectionListBox.FindOrAdd(Connector.Name);
+					if(TargetKeys.IsSet())
+					{
+						CurrentTargetKeys = TargetKeys.GetValue();
+					}
+				}
 
-				ConnectionsCategory.AddCustomRow(Label)
-					.NameContent()
+				if(bIsArrayConnector)
+				{
+					ValueContentVerticalBox->AddSlot()
+					.AutoHeight()
+					.Padding(0)
 					[
 						SNew(SHorizontalBox)
 
 						+SHorizontalBox::Slot()
+						.Padding(2, 4, 0, 4)
 						.AutoWidth()
-						.Padding(0.f, 0.f, 4.f, 0.f)
 						.HAlign(HAlign_Left)
 						.VAlign(VAlign_Center)
 						[
-							SNew(SImage)
-							.Image(IconBrush)
-							.ColorAndOpacity(FSlateColor::UseForeground())
-							.DesiredSizeOverride(FVector2D(16, 16))
+							SNew(STextBlock)
+							.Font(DetailBuilder.GetDetailFont())
+							.Text(FText::FromString(FString::Printf(TEXT("%d Connected Targets"), CurrentTargetKeys.Num())))
 						]
-						
+
 						+SHorizontalBox::Slot()
+						.Padding(4, 0, 0, 0)
 						.AutoWidth()
-						.Padding(0.f, 0.f, 0.f, 0.f)
 						.HAlign(HAlign_Left)
 						.VAlign(VAlign_Center)
 						[
-							SNew(STextBlock)
-							.Text(Label)
-							.Font(IDetailLayoutBuilder::GetDetailFont())
+							SNew(SButton)
 							.IsEnabled(true)
+							.ButtonStyle(FAppStyle::Get(), "SimpleButton")
+							.ContentPadding(FMargin(1, 0))
+							.OnClicked(this, &FRigModuleInstanceDetails::OnAddTargetToArrayConnector, ConnectorName, DetailBuilder.GetPropertyUtilities())
+							.HAlign(HAlign_Right)
+							.ToolTipText(LOCTEXT("AddTargetToArrayConnectors", "Adds a new target for this array connector"))
+							.VAlign(VAlign_Center)
+							[
+								SNew(SImage)
+								.Image(FAppStyle::Get().GetBrush("Icons.PlusCircle"))
+								.ColorAndOpacity(FSlateColor::UseForeground())
+							]
 						]
-					]
-					.ValueContent()
-					[
-						SNew(SHorizontalBox)
 
-						// Combo button
 						+SHorizontalBox::Slot()
+						.Padding(4, 0, 0, 0)
 						.AutoWidth()
-						.Padding(0.f, 0.f, 0.f, 0.f)
+						.HAlign(HAlign_Left)
+						.VAlign(VAlign_Center)
 						[
-							SNew( SComboButton )
-							.ContentPadding(3)
-							.MenuPlacement(MenuPlacement_BelowAnchor)
-							.OnComboBoxOpened(this, &FRigModuleInstanceDetails::PopulateConnectorTargetList, Connector.Name)
-							.ButtonContent()
+							SNew(SButton)
+							.IsEnabled(true)
+							.ButtonStyle(FAppStyle::Get(), "SimpleButton")
+							.ContentPadding(FMargin(1, 0))
+							.OnClicked(this, &FRigModuleInstanceDetails::OnClearTargetsForArrayConnector, ConnectorName, DetailBuilder.GetPropertyUtilities())
+							.HAlign(HAlign_Right)
+							.ToolTipText(LOCTEXT("ClearTargetsForArrayConnector", "Removes all targets for this array connector"))
+							.VAlign(VAlign_Center)
 							[
-								// Wrap in configurable box to restrain height/width of menu
-								SNew(SBox)
-								.MinDesiredWidth(150.0f)
-								[
-									SAssignNew(ButtonBox, SVerticalBox)
-								]
+								SNew(SImage)
+								.Image(FAppStyle::Get().GetBrush("Icons.Delete"))
+								.ColorAndOpacity(FSlateColor::UseForeground())
+							]
+						]
+					];
+				}
+
+				for(int32 TargetIndex = 0; TargetIndex < CurrentTargetKeys.Num(); TargetIndex++)
+				{
+					const FString SearchableTreeViewKey = FString::Printf(TEXT("%s_%d"), *Connector.Name, TargetIndex);
+					
+					const FRigElementKey& TargetKey = CurrentTargetKeys[TargetIndex];
+
+					TreeDelegates.OnGetSelection.BindLambda([this, ConnectorName, TargetIndex]() -> TArray<FRigElementKey>
+					{
+						TOptional<FRigElementKey> Target;
+						FRigElementKeyRedirector Redirector = PerModuleInfos[0].GetModularRig()->GetElementKeyRedirector();
+						for (int32 ModuleIndex=0; ModuleIndex<PerModuleInfos.Num(); ++ModuleIndex)
+						{
+							FString ConnectorPath = FString::Printf(TEXT("%s:%s"), *PerModuleInfos[ModuleIndex].Path, *ConnectorName);
+							FRigElementKey ConnectorKey(*ConnectorPath, ERigElementType::Connector);
+							if(const UControlRigBlueprint* Blueprint = PerModuleInfos[ModuleIndex].GetBlueprint())
+							{
+								const TArray<FRigElementKey> CurrentTargets = Blueprint->ModularRigModel.Connections.FindTargetsFromConnector(ConnectorKey);
+								if(CurrentTargets.IsValidIndex(TargetIndex))
+								{
+									if (!Target.IsSet())
+									{
+										Target = CurrentTargets[TargetIndex];
+									}
+									else if (Target.GetValue() != CurrentTargets[TargetIndex])
+									{
+										Target.GetValue().Name = *ControlRigModuleDetailsMultipleValues.ToString();
+										return {Target.GetValue()};
+									}
+								}
+								else if(Target.IsSet())
+								{
+									Target.GetValue().Name = *ControlRigModuleDetailsMultipleValues.ToString();
+									return {Target.GetValue()};
+								}
+							}
+							else
+							{
+								return {FRigElementKey(*ControlRigModuleDetailsMultipleValues.ToString(), ERigElementType::None)};
+							}
+						}
+						return {};
+					});
+								
+					TreeDelegates.OnSelectionChanged.BindSP(this, &FRigModuleInstanceDetails::OnConnectorTargetChanged, Connector, TargetIndex);
+
+					TSharedPtr<SSearchableRigHierarchyTreeView>& SearchableTreeView = ConnectionListBox.FindOrAdd(SearchableTreeViewKey);
+					SearchableTreeView = SNew(SSearchableRigHierarchyTreeView)
+						.RigTreeDelegates(TreeDelegates);
+					
+					TSharedRef<SHorizontalBox> ValueContentHorizontalBox = SNew(SHorizontalBox);
+					ValueContentVerticalBox->AddSlot()
+					.AutoHeight()
+					.Padding(0)
+					[
+						ValueContentHorizontalBox
+					];
+
+					TSharedPtr<SVerticalBox> ButtonBox;
+
+					// Combo button
+					ValueContentHorizontalBox->AddSlot()
+					.AutoWidth()
+					.Padding(0.f, 0.f, 0.f, 0.f)
+					[
+						SNew( SComboButton )
+						.ContentPadding(3)
+						.MenuPlacement(MenuPlacement_BelowAnchor)
+						.OnComboBoxOpened(this, &FRigModuleInstanceDetails::PopulateConnectorTargetList, SearchableTreeViewKey)
+						.ButtonContent()
+						[
+							// Wrap in configurable box to restrain height/width of menu
+							SNew(SBox)
+							.MinDesiredWidth(150.0f)
+							[
+								SAssignNew(ButtonBox, SVerticalBox)
 							]
-							.MenuContent()
+						]
+						.MenuContent()
+						[
+							SNew(SBorder)
+							.Visibility(EVisibility::Visible)
+							.BorderImage(FAppStyle::GetBrush("Menu.Background"))
 							[
-								SNew(SBorder)
-								.Visibility(EVisibility::Visible)
-								.BorderImage(FAppStyle::GetBrush("Menu.Background"))
-								[
-									SAssignNew(SearchableTreeView, SSearchableRigHierarchyTreeView)
-										.RigTreeDelegates(TreeDelegates)
-								]
+								SearchableTreeView.ToSharedRef()
 							]
 						]
+					];
 
-						+SHorizontalBox::Slot()
-						.AutoWidth()
+					ValueContentHorizontalBox->AddSlot()
+					.AutoWidth()
+					.VAlign(VAlign_Center)
+					.Padding(0.f, 0.f, 0.f, 0.f)
+					[
+						SNew(SVerticalBox)
+						+SVerticalBox::Slot()
 						.VAlign(VAlign_Center)
-						.Padding(0.f, 0.f, 0.f, 0.f)
+						.AutoHeight()
 						[
-							SNew(SVerticalBox)
-							+SVerticalBox::Slot()
-							.VAlign(VAlign_Center)
-							.AutoHeight()
+							// Reset button
+							SNew(SHorizontalBox)
+							+SHorizontalBox::Slot()
+							.AutoWidth()
+							.Padding(0.f, 0.f, 0.f, 0.f)
 							[
-								// Reset button
-								SNew(SHorizontalBox)
-								+SHorizontalBox::Slot()
-								.AutoWidth()
-								.Padding(0.f, 0.f, 0.f, 0.f)
+								SAssignNew(ResetConnectorButton.FindOrAdd(Connector.Name), SButton)
+								.ButtonStyle( FAppStyle::Get(), "NoBorder" )
+								.ButtonColorAndOpacity_Lambda([this, Connector]()
+								{
+									const TSharedPtr<SButton>& Button = ResetConnectorButton.FindRef(Connector.Name);
+									return Button.IsValid() && Button->IsHovered()
+										? FSlateColor(FLinearColor(1.f, 1.f, 1.f, 0.8))
+										: FSlateColor(FLinearColor(1.f, 1.f, 1.f, 0.4));
+								})
+								.OnClicked_Lambda([this, Connector]()
+								{
+									for (FPerModuleInfo& Info : PerModuleInfos)
+									{
+										FString ConnectorPath = FString::Printf(TEXT("%s:%s"), *Info.Path, *Connector.Name);
+										FRigElementKey ConnectorKey(*ConnectorPath, ERigElementType::Connector);
+										Info.GetBlueprint()->GetModularRigController()->DisconnectConnector(ConnectorKey);
+									}
+									return FReply::Handled();
+								})
+								.ContentPadding(1.f)
+								.ToolTipText(NSLOCTEXT("ControlRigModuleDetails", "Reset_Connector", "Reset Connector"))
 								[
-									SAssignNew(ResetConnectorButton.FindOrAdd(Connector.Name), SButton)
-									.ButtonStyle( FAppStyle::Get(), "NoBorder" )
-									.ButtonColorAndOpacity_Lambda([this, Connector]()
+									SNew(SImage)
+									.ColorAndOpacity_Lambda( [this, Connector]()
 									{
 										const TSharedPtr<SButton>& Button = ResetConnectorButton.FindRef(Connector.Name);
 										return Button.IsValid() && Button->IsHovered()
-											? FSlateColor(FLinearColor(1.f, 1.f, 1.f, 0.8))
-											: FSlateColor(FLinearColor(1.f, 1.f, 1.f, 0.4));
+										? FSlateColor(FLinearColor(1.f, 1.f, 1.f, 0.8))
+										: FSlateColor(FLinearColor(1.f, 1.f, 1.f, 0.4));
 									})
-									.OnClicked_Lambda([this, Connector]()
-									{
-										for (FPerModuleInfo& Info : PerModuleInfos)
-										{
-											FString ConnectorPath = FString::Printf(TEXT("%s:%s"), *Info.Path, *Connector.Name);
-											FRigElementKey ConnectorKey(*ConnectorPath, ERigElementType::Connector);
-											Info.GetBlueprint()->GetModularRigController()->DisconnectConnector(ConnectorKey);
-										}
-										return FReply::Handled();
-									})
-									.ContentPadding(1.f)
-									.ToolTipText(NSLOCTEXT("ControlRigModuleDetails", "Reset_Connector", "Reset Connector"))
-									[
-										SNew(SImage)
-										.ColorAndOpacity_Lambda( [this, Connector]()
-										{
-											const TSharedPtr<SButton>& Button = ResetConnectorButton.FindRef(Connector.Name);
-											return Button.IsValid() && Button->IsHovered()
-											? FSlateColor(FLinearColor(1.f, 1.f, 1.f, 0.8))
-											: FSlateColor(FLinearColor(1.f, 1.f, 1.f, 0.4));
-										})
-										.Image(FSlateIcon(FAppStyle::Get().GetStyleSetName(), "PropertyWindow.DiffersFromDefault").GetIcon())
-									]
+									.Image(FSlateIcon(FAppStyle::Get().GetStyleSetName(), "PropertyWindow.DiffersFromDefault").GetIcon())
 								]
+							]
 
-								// Use button
-								+SHorizontalBox::Slot()
-								.AutoWidth()
-								.Padding(0.f, 0.f, 0.f, 0.f)
-								[
-									SAssignNew(UseSelectedButton.FindOrAdd(Connector.Name), SButton)
-									.ButtonStyle( FAppStyle::Get(), "NoBorder" )
-									.ButtonColorAndOpacity_Lambda([this, Connector]()
-									{
-										const TSharedPtr<SButton>& Button = UseSelectedButton.FindRef(Connector.Name);
-										return Button.IsValid() && Button->IsHovered()
-											? FSlateColor(FLinearColor(1.f, 1.f, 1.f, 0.8))
-											: FSlateColor(FLinearColor(1.f, 1.f, 1.f, 0.4));
-									})
-									.OnClicked_Lambda([this, Connector]()
+							// Use button
+							+SHorizontalBox::Slot()
+							.AutoWidth()
+							.Padding(0.f, 0.f, 0.f, 0.f)
+							[
+								SAssignNew(UseSelectedButton.FindOrAdd(Connector.Name), SButton)
+								.ButtonStyle( FAppStyle::Get(), "NoBorder" )
+								.ButtonColorAndOpacity_Lambda([this, Connector]()
+								{
+									const TSharedPtr<SButton>& Button = UseSelectedButton.FindRef(Connector.Name);
+									return Button.IsValid() && Button->IsHovered()
+										? FSlateColor(FLinearColor(1.f, 1.f, 1.f, 0.8))
+										: FSlateColor(FLinearColor(1.f, 1.f, 1.f, 0.4));
+								})
+								.OnClicked_Lambda([this, Connector, TargetIndex]()
+								{
+									if (UModularRig* ModularRig = PerModuleInfos[0].GetModularRig())
 									{
-										if (UModularRig* ModularRig = PerModuleInfos[0].GetModularRig())
+										const TArray<FRigElementKey>& Selected = ModularRig->GetHierarchy()->GetSelectedKeys();
+										if (Selected.Num() > 0)
 										{
-											const TArray<FRigElementKey>& Selected = ModularRig->GetHierarchy()->GetSelectedKeys();
-											if (Selected.Num() > 0)
+											for (FPerModuleInfo& Info : PerModuleInfos)
 											{
-												for (FPerModuleInfo& Info : PerModuleInfos)
+												FString ConnectorPath = FString::Printf(TEXT("%s:%s"), *Info.Path, *Connector.Name);
+												FRigElementKey ConnectorKey(*ConnectorPath, ERigElementType::Connector);
+
+												TArray<FRigElementKey> Targets = Info.GetBlueprint()->ModularRigModel.Connections.FindTargetsFromConnector(ConnectorKey);
+												if(Targets.IsValidIndex(TargetIndex))
+												{
+													Targets[TargetIndex] = Selected[0];
+												}
+												else
 												{
-													FString ConnectorPath = FString::Printf(TEXT("%s:%s"), *Info.Path, *Connector.Name);
-													FRigElementKey ConnectorKey(*ConnectorPath, ERigElementType::Connector);
-													Info.GetBlueprint()->GetModularRigController()->ConnectConnectorToElement(ConnectorKey, Selected[0], true, ModularRig->GetModularRigSettings().bAutoResolve);
+													Targets.Add(Selected[0]);
 												}
+												
+												Info.GetBlueprint()->GetModularRigController()->ConnectConnectorToElements(ConnectorKey, Targets, true, ModularRig->GetModularRigSettings().bAutoResolve);
 											}
 										}
-										return FReply::Handled();
-									})
-									.ContentPadding(1.f)
-									.ToolTipText(NSLOCTEXT("ControlRigModuleDetails", "Use_Selected", "Use Selected"))
-									[
-										SNew(SImage)
-										.ColorAndOpacity_Lambda( [this, Connector]()
-										{
-											const TSharedPtr<SButton>& Button = UseSelectedButton.FindRef(Connector.Name);
-											return Button.IsValid() && Button->IsHovered()
-											? FSlateColor(FLinearColor(1.f, 1.f, 1.f, 0.8))
-											: FSlateColor(FLinearColor(1.f, 1.f, 1.f, 0.4));
-										})
-										.Image(FAppStyle::GetBrush("Icons.CircleArrowLeft"))
-									]
-								]
-
-								// Select in hierarchy button
-								+SHorizontalBox::Slot()
-								.AutoWidth()
-								.Padding(0.f, 0.f, 0.f, 0.f)
+									}
+									return FReply::Handled();
+								})
+								.ContentPadding(1.f)
+								.ToolTipText(NSLOCTEXT("ControlRigModuleDetails", "Use_Selected", "Use Selected"))
 								[
-									SAssignNew(SelectElementButton.FindOrAdd(Connector.Name), SButton)
-									.ButtonStyle( FAppStyle::Get(), "NoBorder" )
-									.ButtonColorAndOpacity_Lambda([this, Connector]()
+									SNew(SImage)
+									.ColorAndOpacity_Lambda( [this, Connector]()
 									{
-										const TSharedPtr<SButton>& Button = SelectElementButton.FindRef(Connector.Name);
+										const TSharedPtr<SButton>& Button = UseSelectedButton.FindRef(Connector.Name);
 										return Button.IsValid() && Button->IsHovered()
-											? FSlateColor(FLinearColor(1.f, 1.f, 1.f, 0.8))
-											: FSlateColor(FLinearColor(1.f, 1.f, 1.f, 0.4));
+										? FSlateColor(FLinearColor(1.f, 1.f, 1.f, 0.8))
+										: FSlateColor(FLinearColor(1.f, 1.f, 1.f, 0.4));
 									})
-									.OnClicked_Lambda([this, Connector]()
+									.Image(FAppStyle::GetBrush("Icons.CircleArrowLeft"))
+								]
+							]
+
+							// Select in hierarchy button
+							+SHorizontalBox::Slot()
+							.AutoWidth()
+							.Padding(0.f, 0.f, 0.f, 0.f)
+							[
+								SAssignNew(SelectElementButton.FindOrAdd(Connector.Name), SButton)
+								.ButtonStyle( FAppStyle::Get(), "NoBorder" )
+								.ButtonColorAndOpacity_Lambda([this, Connector]()
+								{
+									const TSharedPtr<SButton>& Button = SelectElementButton.FindRef(Connector.Name);
+									return Button.IsValid() && Button->IsHovered()
+										? FSlateColor(FLinearColor(1.f, 1.f, 1.f, 0.8))
+										: FSlateColor(FLinearColor(1.f, 1.f, 1.f, 0.4));
+								})
+								.OnClicked_Lambda([this, Connector]()
+								{
+									if (UModularRig* ModularRig = PerModuleInfos[0].GetModularRig())
 									{
-										if (UModularRig* ModularRig = PerModuleInfos[0].GetModularRig())
+										FString ConnectorPath = FString::Printf(TEXT("%s:%s"), *PerModuleInfos[0].Path, *Connector.Name);
+										FRigElementKey ConnectorKey(*ConnectorPath, ERigElementType::Connector);
+										if (const FRigElementKeyRedirector::FKeyArray* TargetKeys = ModularRig->GetElementKeyRedirector().FindExternalKey(ConnectorKey))
 										{
-											FString ConnectorPath = FString::Printf(TEXT("%s:%s"), *PerModuleInfos[0].Path, *Connector.Name);
-											FRigElementKey ConnectorKey(*ConnectorPath, ERigElementType::Connector);
-											if (const FRigElementKeyRedirector::FKeyArray* TargetKeys = ModularRig->GetElementKeyRedirector().FindExternalKey(ConnectorKey))
-											{
-												TArray<FRigElementKey> KeysToSelect;
-												KeysToSelect.Append(*TargetKeys);
-												ModularRig->GetHierarchy()->GetController()->SetSelection(KeysToSelect, true);
-											}
+											TArray<FRigElementKey> KeysToSelect;
+											KeysToSelect.Append(*TargetKeys);
+											ModularRig->GetHierarchy()->GetController()->SetSelection(KeysToSelect, true);
 										}
-										return FReply::Handled();
+									}
+									return FReply::Handled();
+								})
+								.ContentPadding(1.f)
+								.ToolTipText(NSLOCTEXT("ControlRigModuleDetails", "Select_Element", "Select Element"))
+								[
+									SNew(SImage)
+									.ColorAndOpacity_Lambda( [this, Connector]()
+									{
+										const TSharedPtr<SButton>& Button = SelectElementButton.FindRef(Connector.Name);
+										return Button.IsValid() && Button->IsHovered()
+										? FSlateColor(FLinearColor(1.f, 1.f, 1.f, 0.8))
+										: FSlateColor(FLinearColor(1.f, 1.f, 1.f, 0.4));
 									})
-									.ContentPadding(1.f)
-									.ToolTipText(NSLOCTEXT("ControlRigModuleDetails", "Select_Element", "Select Element"))
-									[
-										SNew(SImage)
-										.ColorAndOpacity_Lambda( [this, Connector]()
-										{
-											const TSharedPtr<SButton>& Button = SelectElementButton.FindRef(Connector.Name);
-											return Button.IsValid() && Button->IsHovered()
-											? FSlateColor(FLinearColor(1.f, 1.f, 1.f, 0.8))
-											: FSlateColor(FLinearColor(1.f, 1.f, 1.f, 0.4));
-										})
-										.Image(FAppStyle::GetBrush("Icons.Search"))
-									]
+									.Image(FAppStyle::GetBrush("Icons.Search"))
 								]
 							]
 						]
 					];
 
-				FRigElementKey CurrentTargetKey;
-				if (PerModuleInfos.Num() >= 1)
-				{
-					TOptional<FRigElementKey> TargetKey;
-					for (int32 ModuleIndex=0; ModuleIndex<PerModuleInfos.Num(); ++ModuleIndex)
-					{
-						FString ConnectorPath = FString::Printf(TEXT("%s:%s"), *PerModuleInfos[ModuleIndex].Path, *Connector.Name);
-						FRigElementKey ConnectorKey(*ConnectorPath, ERigElementType::Connector);
-						if (const FRigElementKeyRedirector::FKeyArray* Keys = PerModuleInfos[ModuleIndex].GetModularRig()->GetElementKeyRedirector().FindExternalKey(ConnectorKey))
-						{
-							for(const FRigElementKey& Key : *Keys)
-							{
-								if (!TargetKey.IsSet())
-								{
-									TargetKey = Key;
-								}
-								else
-								{
-									if (TargetKey.GetValue() != Key)
-									{
-										CurrentTargetKey.Name = *ControlRigModuleDetailsMultipleValues.ToString();
-										break;
-									}
-								}
-							}
-							break;
-						}
-					}
-
-					if(TargetKey.IsSet())
-					{
-						CurrentTargetKey = TargetKey.GetValue();
-					}
+					TPair<const FSlateBrush*, FSlateColor> IconAndColor = SRigHierarchyItem::GetBrushForElementType(PerModuleInfos[0].GetModularRig()->GetHierarchy(), TargetKey);
+					PopulateConnectorCurrentTarget(ButtonBox, IconAndColor.Key, IconAndColor.Value, FText::FromName(TargetKey.Name));
 				}
-				TPair<const FSlateBrush*, FSlateColor> IconAndColor = SRigHierarchyItem::GetBrushForElementType(PerModuleInfos[0].GetModularRig()->GetHierarchy(), CurrentTargetKey);
-				PopulateConnectorCurrentTarget(ButtonBox, IconAndColor.Key, IconAndColor.Value, FText::FromName(CurrentTargetKey.Name));
+
+				ConnectionsCategory.AddCustomRow(Label)
+				.NameContent()
+				[
+					SNew(SHorizontalBox)
+
+					+SHorizontalBox::Slot()
+					.AutoWidth()
+					.Padding(0.f, 0.f, 4.f, 0.f)
+					.HAlign(HAlign_Left)
+					.VAlign(VAlign_Center)
+					[
+						SNew(SImage)
+						.Image(IconBrush)
+						.ColorAndOpacity(FSlateColor::UseForeground())
+						.DesiredSizeOverride(FVector2D(16, 16))
+					]
+									
+					+SHorizontalBox::Slot()
+					.AutoWidth()
+					.Padding(0.f, 0.f, 0.f, 0.f)
+					.HAlign(HAlign_Left)
+					.VAlign(VAlign_Center)
+					[
+						SNew(STextBlock)
+						.Text(Label)
+						.Font(IDetailLayoutBuilder::GetDetailFont())
+						.IsEnabled(true)
+					]
+				]
+				.ValueContent()
+				[
+					ValueContentVerticalBox
+				];
 			}
 		}
 	}
@@ -1128,7 +1291,7 @@ void FRigModuleInstanceDetails::OnConfigValueChanged(const FName InVariableName)
 	}
 }
 
-void FRigModuleInstanceDetails::OnConnectorTargetChanged(TSharedPtr<FRigTreeElement> Selection, ESelectInfo::Type SelectInfo, const FRigModuleConnector InConnector)
+void FRigModuleInstanceDetails::OnConnectorTargetChanged(TSharedPtr<FRigTreeElement> Selection, ESelectInfo::Type SelectInfo, const FRigModuleConnector InConnector, int32 InTargetIndex)
 {
 	if (SelectInfo == ESelectInfo::OnNavigation)
 	{
@@ -1142,10 +1305,32 @@ void FRigModuleInstanceDetails::OnConnectorTargetChanged(TSharedPtr<FRigTreeElem
 		{
 			FString ConnectorPath = FString::Printf(TEXT("%s:%s"), *Info.Path, *InConnector.Name);
 			FRigElementKey ConnectorKey(*ConnectorPath, ERigElementType::Connector);
-			if (Selection.IsValid())
+
+			FRigElementKey ChosenTargetKey;
+			if(Selection.IsValid())
+			{
+				ChosenTargetKey = Selection->Key;
+			}
+			
+			TArray<FRigElementKey> Targets = Info.GetBlueprint()->ModularRigModel.Connections.FindTargetsFromConnector(ConnectorKey);
+			if(Targets.IsValidIndex(InTargetIndex))
+			{
+				Targets[InTargetIndex] = ChosenTargetKey;
+			}
+			else
+			{
+				Targets.Add(ChosenTargetKey);
+			}
+
+			Targets.RemoveAll([](const FRigElementKey& Key) -> bool
+			{
+				return !Key.IsValid();
+			});
+			
+			if (!Targets.IsEmpty())
 			{
 				const FModularRigSettings& Settings = Info.GetModularRig()->GetModularRigSettings();
-				Controller->ConnectConnectorToElement(ConnectorKey, Selection->Key, true, Settings.bAutoResolve);
+				Controller->ConnectConnectorToElements(ConnectorKey, Targets, true, Settings.bAutoResolve);
 			}
 			else
 			{
@@ -1460,5 +1645,45 @@ void FRigModuleInstanceDetails::HandleChangeBinding(const FProperty* InProperty,
 	}
 }
 
+FReply FRigModuleInstanceDetails::OnAddTargetToArrayConnector(const FString InConnectorName, const TSharedRef<IPropertyUtilities> PropertyUtilities)
+{
+	FScopedTransaction Transaction(LOCTEXT("AddTargetToArrayConnector", "Add Target To Array Connector"));
+	for(const FPerModuleInfo& Info : PerModuleInfos)
+	{
+		if (UControlRigBlueprint* Blueprint = Info.GetBlueprint())
+		{
+			if (const FRigModuleInstance* ModuleInstance = Info.GetModule())
+			{
+				UModularRigController* Controller = Blueprint->GetModularRigController();
+				FString ConnectorPath = FString::Printf(TEXT("%s:%s"), *Info.Path, *InConnectorName);
+				FRigElementKey ConnectorKey(*ConnectorPath, ERigElementType::Connector);
+				Controller->AddTargetToArrayConnector(ConnectorKey, FRigElementKey(), true, false, false);
+				Blueprint->RecompileModularRig();
+			}
+		}
+	}
+	PropertyUtilities->RequestForceRefresh();
+	return FReply::Handled();
+}
+
+FReply FRigModuleInstanceDetails::OnClearTargetsForArrayConnector(const FString InConnectorName, const TSharedRef<IPropertyUtilities> PropertyUtilities)
+{
+	FScopedTransaction Transaction(LOCTEXT("ClearTargetsForArrayConnector", "Clear Targets For Array Connector"));
+	for(const FPerModuleInfo& Info : PerModuleInfos)
+	{
+		if (UControlRigBlueprint* Blueprint = Info.GetBlueprint())
+		{
+			if (const FRigModuleInstance* ModuleInstance = Info.GetModule())
+			{
+				UModularRigController* Controller = Blueprint->GetModularRigController();
+				FString ConnectorPath = FString::Printf(TEXT("%s:%s"), *Info.Path, *InConnectorName);
+				FRigElementKey ConnectorKey(*ConnectorPath, ERigElementType::Connector);
+				Controller->DisconnectConnector(ConnectorKey);
+			}
+		}
+	}
+	PropertyUtilities->RequestForceRefresh();
+	return FReply::Handled();
+}
 
 #undef LOCTEXT_NAMESPACE
diff --git a/Engine/Plugins/Animation/ControlRig/Source/ControlRigEditor/Private/ControlRigModuleDetails.h b/Engine/Plugins/Animation/ControlRig/Source/ControlRigEditor/Private/ControlRigModuleDetails.h
--- a/Engine/Plugins/Animation/ControlRig/Source/ControlRigEditor/Private/ControlRigModuleDetails.h
+++ b/Engine/Plugins/Animation/ControlRig/Source/ControlRigEditor/Private/ControlRigModuleDetails.h
@@ -58,7 +58,7 @@ public:
 		const FText& InTitle);
 
 	void OnConfigValueChanged(const FName InVariableName);
-	void OnConnectorTargetChanged(TSharedPtr<FRigTreeElement> Selection, ESelectInfo::Type SelectInfo, const FRigModuleConnector InConnector);
+	void OnConnectorTargetChanged(TSharedPtr<FRigTreeElement> Selection, ESelectInfo::Type SelectInfo, const FRigModuleConnector InConnector, int32 InTargetIndex);
 	
 	struct FPerModuleInfo
 	{
@@ -142,6 +142,8 @@ protected:
 	bool CanRemoveBinding(FName InPropertyName) const;
 	void HandleRemoveBinding(FName InPropertyName) const;
 	void HandleChangeBinding(const FProperty* InProperty, const FString& InNewVariablePath) const;
+	FReply OnAddTargetToArrayConnector(const FString InConnectorName, const TSharedRef<IPropertyUtilities> PropertyUtilities);
+	FReply OnClearTargetsForArrayConnector(const FString InConnectorName, const TSharedRef<IPropertyUtilities> PropertyUtilities);
 
 	TArray<FPerModuleInfo> PerModuleInfos;
 	TMap<FString, TSharedPtr<SSearchableRigHierarchyTreeView>> ConnectionListBox;
@@ -151,3 +153,4 @@ protected:
 	TMap<FString, TSharedPtr<SButton>> SelectElementButton;
 	TMap<FString, TSharedPtr<SButton>> ResetConnectorButton;
 };
+
diff --git a/Engine/Plugins/Animation/ControlRig/Source/ControlRigEditor/Private/Editor/SModularRigTreeView.cpp b/Engine/Plugins/Animation/ControlRig/Source/ControlRigEditor/Private/Editor/SModularRigTreeView.cpp
--- a/Engine/Plugins/Animation/ControlRig/Source/ControlRigEditor/Private/Editor/SModularRigTreeView.cpp
+++ b/Engine/Plugins/Animation/ControlRig/Source/ControlRigEditor/Private/Editor/SModularRigTreeView.cpp
@@ -159,34 +159,40 @@ void SModularRigModelItem::PopulateConnectorTargetList(const FRigElementKey InCo
 	}
 }
 
-void SModularRigModelItem::PopulateConnectorCurrentTarget(TSharedPtr<SVerticalBox> InListBox, const FRigElementKey& InConnectorKey, const FRigElementKey& InTargetKey, const FSlateBrush* InBrush, const FSlateColor& InColor, const FText& InTitle)
+void SModularRigModelItem::PopulateConnectorCurrentTarget(TSharedPtr<SVerticalBox> InListBox, const FRigElementKey& InConnectorKey, const TArray<FRigElementKey>& InTargetKeys, const FSlateBrush* InBrush, const FSlateColor& InColor, const FText& InTitle)
 {
 	static const FSlateBrush* RoundedBoxBrush = FControlRigEditorStyle::Get().GetBrush(TEXT("ControlRig.SpacePicker.RoundedRect"));
 
-	TAttribute<FSlateColor> TextColor = TAttribute<FSlateColor>::CreateLambda([this, InConnectorKey, InTargetKey]() -> FSlateColor
+	TAttribute<FSlateColor> TextColor;
+	const FRigElementKey FirstTargetKey = InTargetKeys.Num() == 1 ? InTargetKeys[0] : FRigElementKey();
+
+	if(InTargetKeys.Num() == 1)
 	{
-		if(const UModularRig* ModularRig = Delegates.GetModularRig())
+		TextColor = TAttribute<FSlateColor>::CreateLambda([this, InConnectorKey, FirstTargetKey]() -> FSlateColor
 		{
-			if(const URigHierarchy* Hierarchy = ModularRig->GetHierarchy())
+			if(const UModularRig* ModularRig = Delegates.GetModularRig())
 			{
-				if(!Hierarchy->Contains(InTargetKey))
+				if(const URigHierarchy* Hierarchy = ModularRig->GetHierarchy())
 				{
-					if(!InTargetKey.IsValid())
+					if(!Hierarchy->Contains(FirstTargetKey))
 					{
-						if(const FRigConnectorElement* Connector = Hierarchy->Find<FRigConnectorElement>(InConnectorKey))
+						if(!FirstTargetKey.IsValid())
 						{
-							if(Connector->IsOptional())
+							if(const FRigConnectorElement* Connector = Hierarchy->Find<FRigConnectorElement>(InConnectorKey))
 							{
-								return FSlateColor::UseForeground();
+								if(Connector->IsOptional())
+								{
+									return FSlateColor::UseForeground();
+								}
 							}
 						}
+						return FLinearColor::Red;
 					}
-					return FLinearColor::Red;
 				}
 			}
-		}
-		return FSlateColor::UseForeground();
-	});
+			return FSlateColor::UseForeground();
+		});
+	}
 	
 	TSharedPtr<SHorizontalBox> RowBox, ButtonBox;
 	InListBox->AddSlot()
@@ -204,12 +210,12 @@ void SModularRigModelItem::PopulateConnectorCurrentTarget(TSharedPtr<SVerticalBo
 			SNew( SButton )
 			.ButtonStyle(FAppStyle::Get(), "SimpleButton")
 			.ContentPadding(FMargin(0.0))
-			.OnClicked_Lambda([this, InConnectorKey, InTargetKey]()
+			.OnClicked_Lambda([this, InConnectorKey]()
 			{
-				//Controller->ConnectConnectorToElement(InConnectorKey, InTargetKey, true, Info.GetModularRig()->GetModularRigSettings().bAutoResolve);
 				PopulateConnectorTargetList(InConnectorKey);
 				return FReply::Handled();
 			})
+			.IsEnabled(InTargetKeys.Num() == 1)
 			[
 				SAssignNew(RowBox, SHorizontalBox)
 				+ SHorizontalBox::Slot()
@@ -432,6 +438,25 @@ TSharedRef<SWidget> SModularRigModelItem::GenerateWidgetForColumn(const FName& C
 	{
 		TSharedPtr<SVerticalBox> ComboButtonBox;
 
+		if (const UModularRig* ModularRig = Delegates.GetModularRig())
+		{
+			if (const FRigModuleInstance* Module = ModularRig->FindModule(WeakRigTreeElement.Pin()->ModulePath))
+			{
+				if (URigHierarchy* Hierarchy = ModularRig->GetHierarchy())
+				{
+					if (const FRigConnectorElement* ConnectorElement = Cast<FRigConnectorElement>(Hierarchy->Find(ConnectorKey)))
+					{
+						if(ConnectorElement->IsArrayConnector())
+						{
+							return SNew( STextBlock )
+							.Text( LOCTEXT("ArrayConnectorPleaseRelyOnDetailsPanel", "WIP. Array Connector. Use Details Panel.") )
+							.Font( IDetailLayoutBuilder::GetDetailFont() );
+						}
+					}
+				}
+			}
+		}
+		
 		FRigTreeDelegates TreeDelegates;
 		TreeDelegates.OnGetHierarchy = FOnGetRigTreeHierarchy::CreateLambda([this]()
 		{
@@ -512,22 +537,36 @@ TSharedRef<SWidget> SModularRigModelItem::GenerateWidgetForColumn(const FName& C
 
 		const UModularRig* ModularRig = Delegates.GetModularRig();
 		const FRigElementKeyRedirector& Redirector = ModularRig->GetElementKeyRedirector();
-		FRigElementKey CurrentTargetKey;
+		TArray<FRigElementKey> CurrentTargetKeys;
 		if (const FRigElementKeyRedirector::FKeyArray* Keys = Redirector.FindExternalKey(ConnectorKey))
 		{
-			if(!Keys->IsEmpty())
+			CurrentTargetKeys.Append(*Keys);
+		}
+		if(CurrentTargetKeys.IsEmpty())
+		{
+			CurrentTargetKeys.Add(FRigElementKey());
+		}
+
+		TPair<const FSlateBrush*, FSlateColor> IconAndColor;
+		FText CurrentTargetShortName;
+		
+		if(CurrentTargetKeys.Num() == 1)
+		{
+			const FRigElementKey& FirstTargetKey = CurrentTargetKeys[0];
+			IconAndColor = SRigHierarchyItem::GetBrushForElementType(ModularRig->GetHierarchy(), FirstTargetKey);
+
+			CurrentTargetShortName = FText::FromString(ModularRig->GetShortestDisplayPathForElement(FirstTargetKey, false));
+			if(CurrentTargetShortName.IsEmpty())
 			{
-				CurrentTargetKey = (*Keys)[0];
+				CurrentTargetShortName = FText::FromName(FirstTargetKey.Name);
 			}
 		}
-		TPair<const FSlateBrush*, FSlateColor> IconAndColor = SRigHierarchyItem::GetBrushForElementType(ModularRig->GetHierarchy(), CurrentTargetKey);
-
-		FText CurrentTargetShortName = FText::FromString(ModularRig->GetShortestDisplayPathForElement(CurrentTargetKey, false));
-		if(CurrentTargetShortName.IsEmpty())
+		else
 		{
-			CurrentTargetShortName = FText::FromName(CurrentTargetKey.Name);
+			CurrentTargetShortName = LOCTEXT("MultipleValues", "Multiple Values");
+			ConnectorComboBox->SetEnabled(false);
 		}
-		PopulateConnectorCurrentTarget(ComboButtonBox, ConnectorKey, CurrentTargetKey, IconAndColor.Key, IconAndColor.Value, CurrentTargetShortName);
+		PopulateConnectorCurrentTarget(ComboButtonBox, ConnectorKey, CurrentTargetKeys, IconAndColor.Key, IconAndColor.Value, CurrentTargetShortName);
 
 		return Widget;
 	}
diff --git a/Engine/Plugins/Animation/ControlRig/Source/ControlRigEditor/Private/Editor/SModularRigTreeView.h b/Engine/Plugins/Animation/ControlRig/Source/ControlRigEditor/Private/Editor/SModularRigTreeView.h
--- a/Engine/Plugins/Animation/ControlRig/Source/ControlRigEditor/Private/Editor/SModularRigTreeView.h
+++ b/Engine/Plugins/Animation/ControlRig/Source/ControlRigEditor/Private/Editor/SModularRigTreeView.h
@@ -161,7 +161,7 @@ public:
 	void PopulateConnectorCurrentTarget(
 		TSharedPtr<SVerticalBox> InListBox, 
 		const FRigElementKey& InConnectorKey,
-		const FRigElementKey& InTargetKey,
+		const TArray<FRigElementKey>& InTargetKeys,
 		const FSlateBrush* InBrush,
 		const FSlateColor& InColor,
 		const FText& InTitle);
diff --git a/Engine/Plugins/Animation/ControlRig/Source/ControlRigEditor/Private/Editor/SRigHierarchy.cpp b/Engine/Plugins/Animation/ControlRig/Source/ControlRigEditor/Private/Editor/SRigHierarchy.cpp
--- a/Engine/Plugins/Animation/ControlRig/Source/ControlRigEditor/Private/Editor/SRigHierarchy.cpp
+++ b/Engine/Plugins/Animation/ControlRig/Source/ControlRigEditor/Private/Editor/SRigHierarchy.cpp
@@ -3014,7 +3014,31 @@ FReply SRigHierarchy::OnAcceptDrop(const FDragDropEvent& DragDropEvent, EItemDro
 		FRigElementKey::StaticStruct()->ImportText(*TagDragDropOp->GetIdentifier(), &DraggedKey, nullptr, EPropertyPortFlags::PPF_None, nullptr, FRigElementKey::StaticStruct()->GetName(), true);
 		if(TargetItem.IsValid())
 		{
-			return ResolveConnector(DraggedKey, TargetItem->Key);
+			TArray<FRigElementKey> TargetsForConnector = {TargetItem->Key};
+			
+			// do we want to add this target to an array connector?
+			if(DragDropEvent.GetModifierKeys().IsShiftDown())
+			{
+				if(const UControlRig* ControlRig = Cast<UControlRig>(ControlRigBeingDebuggedPtr.Get()))
+				{
+					if(URigHierarchy* Hierarchy = ControlRig->GetHierarchy())
+					{
+						if(const FRigConnectorElement* Connector = Hierarchy->Find<FRigConnectorElement>(DraggedKey))
+						{
+							if(Connector->IsArrayConnector())
+							{
+								if(const FRigElementKeyRedirector::FCachedKeyArray* Cache = ControlRigBeingDebuggedPtr->GetElementKeyRedirector().Find(Connector->GetKey()))
+								{
+									TargetsForConnector.Reset();
+									TargetsForConnector.Append(FRigElementKeyRedirector::Convert(*Cache));
+									TargetsForConnector.AddUnique(TargetItem->Key);
+								}
+							}
+						}
+					}
+				}
+			}
+			return ResolveConnectorToArray(DraggedKey, TargetsForConnector);
 		}
 		return ResolveConnector(DraggedKey, FRigElementKey());
 	}
@@ -3758,6 +3782,11 @@ FReply SRigHierarchy::ReparentOrMatchTransform(const TArray<FRigElementKey>& Dra
 }
 
 FReply SRigHierarchy::ResolveConnector(const FRigElementKey& DraggedKey, const FRigElementKey& TargetKey)
+{
+	return ResolveConnectorToArray(DraggedKey, {TargetKey});
+}
+
+FReply SRigHierarchy::ResolveConnectorToArray(const FRigElementKey& DraggedKey, const TArray<FRigElementKey>& TargetKeys)
 {
 	if (UControlRigBlueprint* Blueprint = ControlRigEditor.Pin()->GetControlRigBlueprint())
 	{
@@ -3765,7 +3794,7 @@ FReply SRigHierarchy::ResolveConnector(const FRigElementKey& DraggedKey, const F
 		{
 			if(DebuggedHierarchy->Contains(DraggedKey))
 			{
-				if(Blueprint->ResolveConnector(DraggedKey, TargetKey))
+				if(Blueprint->ResolveConnectorToArray(DraggedKey, TargetKeys))
 				{
 					RefreshTreeView();
 					return FReply::Handled();
diff --git a/Engine/Plugins/Animation/ControlRig/Source/ControlRigEditor/Private/Editor/SRigHierarchy.h b/Engine/Plugins/Animation/ControlRig/Source/ControlRigEditor/Private/Editor/SRigHierarchy.h
--- a/Engine/Plugins/Animation/ControlRig/Source/ControlRigEditor/Private/Editor/SRigHierarchy.h
+++ b/Engine/Plugins/Animation/ControlRig/Source/ControlRigEditor/Private/Editor/SRigHierarchy.h
@@ -210,7 +210,8 @@ private:
 	void HandleAlign(const FToolMenuContext& Context);
 	FReply ReparentOrMatchTransform(const TArray<FRigElementKey>& DraggedKeys, FRigElementKey TargetKey, bool bReparentItems, int32 LocalIndex = INDEX_NONE);
 	FReply ResolveConnector(const FRigElementKey& DraggedKey, const FRigElementKey& TargetKey);
-	
+	FReply ResolveConnectorToArray(const FRigElementKey& DraggedKey, const TArray<FRigElementKey>& TargetKeys);
+
 	FName CreateUniqueName(const FName& InBaseName, ERigElementType InElementType) const;
 
 	void ClearDetailPanel() const;
