Control Rig: Introduce steped sample track for storing test data

#rnx

diff --git a/Engine/Plugins/Animation/ControlRig/Source/ControlRig/Private/ControlRigStepSampleTrack.cpp b/Engine/Plugins/Animation/ControlRig/Source/ControlRig/Private/ControlRigStepSampleTrack.cpp
new file mode 100644
--- /dev/null
+++ b/Engine/Plugins/Animation/ControlRig/Source/ControlRig/Private/ControlRigStepSampleTrack.cpp
@@ -0,0 +1,168 @@
+// Copyright Epic Games, Inc. All Rights Reserved.
+
+#include "ControlRigStepSampleTrack.h"
+
+FStepSampleTrackBase::FStepSampleTrackBase()
+: TrackIndex(INDEX_NONE)
+, Mode(Invalid)
+, ValuesUpperBound(0)
+{
+}
+
+FStepSampleTrackBase::FStepSampleTrackBase(int32 InTrackIndex)
+: TrackIndex(InTrackIndex)
+, Mode(Invalid)
+, ValuesUpperBound(0)
+{
+}
+
+int32 FStepSampleTrackBase::Num() const
+{
+	return ValuesUpperBound;
+}
+
+int32 FStepSampleTrackBase::NumSamples() const
+{
+	switch(Mode)
+	{
+		case Singleton:
+		{
+			return 1;
+		}
+		case Sampled:
+		{
+			return Indices.Num();
+		}
+		case Complete:
+		{
+			return Num();
+		}
+		case Invalid:
+		default:
+		{
+			break;
+		}
+	}
+	return 0;
+}
+
+bool FStepSampleTrackBase::GetSampleIndexAt(int32 InIndex, int32& InOutSampleIndex) const
+{
+	switch(Mode)
+	{
+		case Singleton:
+		{
+			InOutSampleIndex = 0;
+			return true;
+		}
+		case Sampled:
+		{
+			if(Indices.IsEmpty())
+			{
+				break;
+			}
+			if(!Indices.IsValidIndex(InOutSampleIndex))
+			{
+				InOutSampleIndex = 0;
+			}
+			while(Indices.IsValidIndex(InOutSampleIndex - 1) && Indices[InOutSampleIndex] > InIndex)
+			{
+				InOutSampleIndex--;
+			}
+			while(Indices.IsValidIndex(InOutSampleIndex + 1) && Indices[InOutSampleIndex + 1] <= InIndex)
+			{
+				InOutSampleIndex++;
+			}
+			return true;
+		}
+		case Complete:
+		{
+			InOutSampleIndex = InIndex;
+			return true;
+		}
+		case Invalid:
+		default:
+		{
+			break;
+		}
+	}
+	return false;
+}
+
+void FStepSampleTrackBase::Reset()
+{
+	TrackIndex = INDEX_NONE;
+	Mode = Invalid;
+	Indices.Reset();
+	ValuesUpperBound = 0;
+}
+
+void FStepSampleTrackBase::Serialize(FArchive& InArchive)
+{
+	if(InArchive.IsLoading())
+	{
+		Reset();
+	}
+		
+	InArchive << TrackIndex;
+	InArchive << (uint8&)Mode;
+	InArchive << ValuesUpperBound;
+
+	switch(Mode)
+	{
+		// only sampled tracks require the indices array
+		case Sampled:
+		{
+			InArchive << Indices;
+		}
+		default:
+		{
+			break;
+		}
+	}
+
+	if(InArchive.IsLoading())
+	{
+		Shrink();
+	}
+}
+
+void FStepSampleTrackBase::AddIndex(bool bOnlyIncreaseUpperBound)
+{
+	if(!bOnlyIncreaseUpperBound)
+	{
+		Indices.Add(ValuesUpperBound);
+	}
+	ValuesUpperBound++;
+
+	switch(Mode)
+	{
+		case Singleton:
+		case Sampled:
+		case Complete:
+		{
+			if(Indices.Num() == ValuesUpperBound)
+			{
+				Mode = Complete;
+			}
+			else
+			{
+				if(Indices.Num() == 1)
+				{
+					Mode = Singleton;
+				}
+				else
+				{
+					Mode = Sampled;
+				}
+			}
+			break;
+		}
+		case Invalid:
+		default:
+		{
+			Mode = Singleton;
+			break;
+		}
+	}
+}
diff --git a/Engine/Plugins/Animation/ControlRig/Source/ControlRig/Public/ControlRigStepSampleTrack.h b/Engine/Plugins/Animation/ControlRig/Source/ControlRig/Public/ControlRigStepSampleTrack.h
new file mode 100644
--- /dev/null
+++ b/Engine/Plugins/Animation/ControlRig/Source/ControlRig/Public/ControlRigStepSampleTrack.h
@@ -0,0 +1,157 @@
+// Copyright Epic Games, Inc. All Rights Reserved.
+
+#pragma once
+
+#include "CoreMinimal.h"
+#include "Templates/SharedPointer.h"
+#include "Math/UnrealMath.h"
+#include "Math/UnrealMathUtility.h"
+
+class CONTROLRIG_API FStepSampleTrackBase : public TSharedFromThis<FStepSampleTrackBase>
+{
+public:
+	
+	enum EMode : uint8
+	{
+		Invalid,
+		Singleton,
+		Sampled,
+		Complete
+	};
+	
+	FStepSampleTrackBase();
+	FStepSampleTrackBase(int32 InTrackIndex);
+	virtual ~FStepSampleTrackBase() {}
+
+	bool IsValid() const { return Mode != Invalid; }
+	int32 GetTrackIndex() const { return TrackIndex; }
+	EMode GetMode() const { return Mode; }
+	int32 Num() const;
+	int32 NumSamples() const;
+	bool GetSampleIndexAt(int32 InIndex, int32& InOutSampleIndex) const;
+
+	virtual void Reserve(int32 InCount)
+	{
+		Indices.Reserve(InCount);
+	}
+
+	virtual void Reset();
+
+	virtual void Shrink()
+	{
+		Indices.Shrink();
+	}
+
+	virtual void Empty()
+	{
+		Indices.Empty();
+	}
+
+	virtual void Serialize(FArchive& InArchive);
+
+protected:
+
+	void AddIndex(bool bOnlyIncreaseUpperBound);
+
+	int32 TrackIndex;
+	EMode Mode;
+	TArray<int32> Indices;
+	int32 ValuesUpperBound;
+};
+
+template<typename T>
+class FStepSampleTrack : public FStepSampleTrackBase
+{
+public:
+
+	FStepSampleTrack()
+	{
+	}
+	
+	FStepSampleTrack(int32 InTrackIndex)
+	: FStepSampleTrackBase(InTrackIndex)
+	{
+	}
+
+	virtual ~FStepSampleTrack() override {}
+
+	const T& GetValueAt(int32 InIndex, int32& InOutSampleIndex) const
+	{
+		verify(GetSampleIndexAt(InIndex, InOutSampleIndex));
+		return GetValueForSample(InOutSampleIndex);
+	}
+
+	const T& GetValueForSample(int32 InSampleIndex) const
+	{
+		return Values[InSampleIndex];
+	}
+
+	void AddSample(const T& InSample, float InTolerance = KINDA_SMALL_NUMBER)
+	{
+		if(Values.IsEmpty())
+		{
+			AddIndex(false);
+			Values.Add(InSample);
+		}
+		else if(Equals(Values.Last(), InSample, InTolerance))
+		{
+			AddIndex(true);
+		}
+		else
+		{
+			AddIndex(false);
+			Values.Add(InSample);
+		}
+	}
+
+	static bool Equals(const T& InA, const T& InB, float InTolerance)
+	{
+		return InA == InB;
+	}
+
+	virtual void Reset() override
+	{
+		FStepSampleTrackBase::Reset();
+		Values.Reset();
+	}
+
+	virtual void Reserve(int32 InCount) override
+	{
+		FStepSampleTrackBase::Reserve(InCount);
+		Values.Reserve(InCount);
+	}
+
+	virtual void Shrink() override
+	{
+		FStepSampleTrackBase::Shrink();
+		Values.Shrink();
+	}
+
+	virtual void Empty() override
+	{
+		FStepSampleTrackBase::Empty();
+		Values.Empty();
+	}
+
+	virtual void Serialize(FArchive& InArchive) override
+	{
+		FStepSampleTrackBase::Serialize(InArchive);
+		InArchive << Values;
+	}
+
+protected:
+
+	TArray<T> Values;
+};
+
+template<>
+inline bool FStepSampleTrack<float>::Equals(const float& InA, const float& InB, float InTolerance)
+{
+	return FMath::IsNearlyEqual(InA, InB, InTolerance);
+}
+
+template<>
+inline bool FStepSampleTrack<double>::Equals(const double& InA, const double& InB, float InTolerance)
+{
+	return FMath::IsNearlyEqual<double>(InA, InB, InTolerance);
+}
\ No newline at end of file
