Control Rig: Move ownership of performing tests to test data class, add flag to skip tests

#rnx

diff --git a/Engine/Plugins/Animation/ControlRig/Source/ControlRig/Private/ControlRigTestData.cpp b/Engine/Plugins/Animation/ControlRig/Source/ControlRig/Private/ControlRigTestData.cpp
--- a/Engine/Plugins/Animation/ControlRig/Source/ControlRig/Private/ControlRigTestData.cpp
+++ b/Engine/Plugins/Animation/ControlRig/Source/ControlRig/Private/ControlRigTestData.cpp
@@ -608,6 +608,162 @@ bool UControlRigTestData::IsReplaying() const
 	return ReplayControlRig.IsValid();
 }
 
+bool UControlRigTestData::PerformTest(UControlRig* InSubject, TFunction<void(EMessageSeverity::Type, const FString&)> InLogFunction) const
+{
+	bool bSuccess = true;
+
+	// initialize
+	InSubject->RequestInit();
+	Initial.Restore(InSubject, true);
+
+	// run the construction event
+	InSubject->SetEventQueue({FRigUnit_PrepareForExecution::EventName});
+	InSubject->Evaluate_AnyThread();
+
+	// now run all of the frames
+	InSubject->SetEventQueue({FRigUnit_BeginExecution::EventName});
+	for(int32 FrameIndex = 0; FrameIndex < FMath::Min<int32>(InputFrames.Num(), OutputFrames.Num()); FrameIndex++)
+	{
+		const FControlRigTestDataFrame& InputFrame = InputFrames[FrameIndex];
+		const FControlRigTestDataFrame& OutputFrame = OutputFrames[FrameIndex];
+		
+		InSubject->SetAbsoluteAndDeltaTime(InputFrame.AbsoluteTime, InputFrame.DeltaTime);
+		if(!InputFrame.Restore(InSubject))
+		{
+			bSuccess = false;
+		}
+		
+		InSubject->Evaluate_AnyThread();
+
+		// skip the frame's test results
+		if(FramesToSkip.Contains(FrameIndex))
+		{
+			continue;
+		}
+
+		const FRigPose CurrentPose = InSubject->GetHierarchy()->GetPose();
+		const FRigPose& ExpectedPose = OutputFrame.Pose;
+
+		for(const FRigPoseElement& ExpectedPoseElement : ExpectedPose)
+		{
+			const int32 CurrentPoseIndex = CurrentPose.GetIndex(ExpectedPoseElement.Index.GetKey());
+			if(CurrentPoseIndex != INDEX_NONE)
+			{
+				const FRigPoseElement& CurrentPoseElement = CurrentPose[CurrentPoseIndex];
+
+				if(!CurrentPoseElement.LocalTransform.Equals(ExpectedPoseElement.LocalTransform, Tolerance))
+				{
+					InLogFunction(EMessageSeverity::Error, FString::Printf(TEXT("Frame[%03d]: Expected Rig Element '%s' doesn't match. (%s vs expected %s)."),
+						FrameIndex,
+						*ExpectedPoseElement.Index.GetKey().ToString(),
+						*CurrentPoseElement.LocalTransform.ToString(),
+						*ExpectedPoseElement.LocalTransform.ToString()
+					));
+					bSuccess = false;
+				}
+			}
+			else
+			{
+				InLogFunction(EMessageSeverity::Error, FString::Printf(TEXT("Frame[%03d]: Expected Rig Element '%s' is missing."),
+					FrameIndex,
+					*ExpectedPoseElement.Index.GetKey().ToString()
+				));
+				bSuccess = false;
+			}
+		}
+
+		if (OutputFrame.bTestMetadata)
+		{
+			const TMap<FRigElementKey, URigHierarchy::FMetadataStorage>& ExpectedMetadata = OutputFrame.MetadataMap;
+			URigHierarchy* Hierarchy = InSubject->GetHierarchy();
+#if WITH_EDITOR
+			Hierarchy->ForEach([ExpectedMetadata, Hierarchy, InLogFunction, FrameIndex](const FRigBaseElement* Element)
+			{
+				const TArray<FName> MetadataNames = Hierarchy->GetMetadataNames(Element->GetKey());
+				const URigHierarchy::FMetadataStorage* ExpectedElementMetadata = ExpectedMetadata.Find(Element->GetKey());
+				if (ExpectedElementMetadata)
+				{
+					if (ExpectedElementMetadata->MetadataMap.Num() != MetadataNames.Num())
+					{
+						const FString Message = FString::Printf(TEXT("Metadata mismatch in element %s: number of metadata elements does not match"), *Element->GetKey().ToString());
+						InLogFunction(EMessageSeverity::Error, Message);
+					}
+					for(TTuple<FName, FRigBaseMetadata*> Pair : ExpectedElementMetadata->MetadataMap)
+					{
+						if(FRigBaseMetadata* Value = Hierarchy->FindMetadataForElement(Element, Pair.Key, Pair.Value->GetType()))
+						{
+							const FProperty* Property = Value->GetValueProperty();
+							check(Value->IsValid());
+							check(Pair.Value->IsValid());
+							if(!Property->Identical(Value->GetValueData(), Pair.Value->GetValueData()))
+							{
+								FString Expected, Received;
+								Property->ExportText_Direct(Expected, Pair.Value->GetValueData(), nullptr, nullptr, PPF_None, nullptr);
+								Property->ExportText_Direct(Received, Value->GetValueData(), nullptr, nullptr, PPF_None, nullptr);
+								const FString Message = FString::Printf(TEXT("Metadata mismatch in element %s: property %s expected %s, but found %s"), *Element->GetKey().ToString(), *Pair.Key.ToString(), *Expected, *Received);
+								InLogFunction(EMessageSeverity::Error, Message);
+							}
+						}
+						else
+						{
+							const FString Message = FString::Printf(TEXT("Metadata mismatch in element %s: property %s expected, but was not found"), *Element->GetKey().ToString(), *Pair.Key.ToString());
+							InLogFunction(EMessageSeverity::Error, Message);
+						}
+					}
+				}
+				else if (Hierarchy->HasMetadata(Element))
+				{
+					const FString Message = FString::Printf(TEXT("Metadata mismatch in element %s: number of metadata elements does not match"), *Element->GetKey().ToString());
+					InLogFunction(EMessageSeverity::Error, Message);
+				}
+				return true;
+			});
+#endif
+		}
+
+		const TArray<FRigVMExternalVariable> ExternalVariables = InSubject->GetExternalVariables(); 
+		for(const FRigVMExternalVariable& ExternalVariable : ExternalVariables)
+		{
+			const FControlRigTestDataVariable* VariableData = OutputFrame.Variables.FindByPredicate(
+				[ExternalVariable](const FControlRigTestDataVariable& InVariable) -> bool
+				{
+					return InVariable.Name == ExternalVariable.Name &&
+						InVariable.CPPType == ExternalVariable.TypeName;
+				}
+			);
+
+			// cases of missing variables etc are already caught in the test by the
+			// FControlRigTestDataFrame::Restore earlier.
+			
+			if(VariableData)
+			{
+				FString CurrentValue;
+				ExternalVariable.Property->ExportText_Direct(
+					CurrentValue,
+					ExternalVariable.Memory,
+					ExternalVariable.Memory,
+					nullptr,
+					PPF_None,
+					nullptr
+				);
+
+				if(CurrentValue != VariableData->Value)
+				{
+					UE_LOG(LogControlRig, Error, TEXT("Variable '%s' doesn't match: %s vs expected '%s'"),
+						*ExternalVariable.Name.ToString(),
+						*CurrentValue,
+						*VariableData->Value
+					);
+					bSuccess = false;
+				}
+			}
+
+		}
+	}
+
+	return bSuccess;
+}
+
 void UControlRigTestData::ClearDelegates(UControlRig* InControlRig)
 {
 	if(InControlRig)
diff --git a/Engine/Plugins/Animation/ControlRig/Source/ControlRig/Public/ControlRigTestData.h b/Engine/Plugins/Animation/ControlRig/Source/ControlRig/Public/ControlRigTestData.h
--- a/Engine/Plugins/Animation/ControlRig/Source/ControlRig/Public/ControlRigTestData.h
+++ b/Engine/Plugins/Animation/ControlRig/Source/ControlRig/Public/ControlRigTestData.h
@@ -5,6 +5,7 @@
 #include "CoreMinimal.h"
 #include "Rigs/RigHierarchy.h"
 #include "Rigs/RigHierarchyPose.h"
+#include "Logging/TokenizedMessage.h"
 
 #include "ControlRigTestData.generated.h"
 
@@ -90,6 +91,7 @@ public:
 
 	UControlRigTestData()
 		: Tolerance(0.001)
+		, EnableTest(true)
 		, LastFrameIndex(INDEX_NONE)
 		, DesiredRecordingDuration(0.0)
 		, TimeAtStartOfRecording(0.0)
@@ -148,6 +150,16 @@ public:
 	UPROPERTY(EditAnywhere, BlueprintReadOnly, Category = "ControlRigTestData")
 	TArray<FName> EventQueue;
 
+	UPROPERTY(EditAnywhere, BlueprintReadOnly, Category = "ControlRigTestData", AssetRegistrySearchable)
+	bool EnableTest;
+
+	bool IsTestEnabled() const
+	{
+		return EnableTest;
+	}
+
+	bool PerformTest(UControlRig* InSubject, TFunction<void(EMessageSeverity::Type, const FString&)> InLogFunction) const;
+
 private:
 
 	void ClearDelegates(UControlRig* InControlRig);
