Control Rig: Fix LOCTEXT conflict

#rnx

diff --git a/Engine/Plugins/Animation/ControlRig/Source/ControlRigEditor/Private/ControlRigModuleDetails.cpp b/Engine/Plugins/Animation/ControlRig/Source/ControlRigEditor/Private/ControlRigModuleDetails.cpp
--- a/Engine/Plugins/Animation/ControlRig/Source/ControlRigEditor/Private/ControlRigModuleDetails.cpp
+++ b/Engine/Plugins/Animation/ControlRig/Source/ControlRigEditor/Private/ControlRigModuleDetails.cpp
@@ -1668,7 +1668,7 @@ FReply FRigModuleInstanceDetails::OnAddTargetToArrayConnector(const FString InCo
 
 FReply FRigModuleInstanceDetails::OnClearTargetsForArrayConnector(const FString InConnectorName, const TSharedRef<IPropertyUtilities> PropertyUtilities)
 {
-	FScopedTransaction Transaction(LOCTEXT("ClearTargetsForArrayConnector", "Clear Targets For Array Connector"));
+	FScopedTransaction Transaction(LOCTEXT("TransactionClearTargetsForArrayConnector", "Clear Targets For Array Connector"));
 	for(const FPerModuleInfo& Info : PerModuleInfos)
 	{
 		if (UControlRigBlueprint* Blueprint = Info.GetBlueprint())
