Control Rig: Step sample track container + archive reader writer WIP

#rnx

diff --git a/Engine/Plugins/Animation/ControlRig/Source/ControlRig/Private/ControlRigStepSampleTrack.cpp b/Engine/Plugins/Animation/ControlRig/Source/ControlRig/Private/ControlRigStepSampleTrack.cpp
deleted file mode 100644
diff --git a/Engine/Plugins/Animation/ControlRig/Source/ControlRig/Private/Tracks/SampleTrack.cpp b/Engine/Plugins/Animation/ControlRig/Source/ControlRig/Private/Tracks/SampleTrack.cpp
new file mode 100644
--- /dev/null
+++ b/Engine/Plugins/Animation/ControlRig/Source/ControlRig/Private/Tracks/SampleTrack.cpp
@@ -0,0 +1,79 @@
+// Copyright Epic Games, Inc. All Rights Reserved.
+
+#include "Tracks/SampleTrack.h"
+
+//////////////////////////////////////////////////////////////////////////////////////////////////////////////////////////////
+/////// FSampleTrack
+//////////////////////////////////////////////////////////////////////////////////////////////////////////////////////////////
+
+template<>
+void FSampleTrack<FInstancedStruct>::Serialize(FArchive& InArchive)
+{
+	FSampleTrackBase::Serialize(InArchive);
+	
+	int32 ArraySize = Values.Num();
+	InArchive << ArraySize;
+	
+	if(InArchive.IsLoading())
+	{
+		ValuesStorage.Reset();
+		ValuesStorage.SetNum(ArraySize);
+	}
+
+	for(int32 Index = 0; Index < ArraySize; Index++)
+	{
+		ValuesStorage[Index].Serialize(InArchive);
+	}
+
+	if(InArchive.IsLoading())
+	{
+		UpdateValueArrayView();
+	}
+}
+
+template<>
+void FSampleTrack<TArray<FInstancedStruct>>::Serialize(FArchive& InArchive)
+{
+	FSampleTrackBase::Serialize(InArchive);
+	
+	int32 OuterArraySize = ValuesStorage.Num();
+	InArchive << OuterArraySize;
+	
+	if(InArchive.IsLoading())
+	{
+		ValuesStorage.Reset();
+		ValuesStorage.SetNum(OuterArraySize);
+	}
+
+	for(int32 OuterIndex = 0; OuterIndex < OuterArraySize; OuterIndex++)
+	{
+		int32 InnerArraySize = ValuesStorage[OuterIndex].Num();
+		InArchive << InnerArraySize;
+	
+		if(InArchive.IsLoading())
+		{
+			ValuesStorage[OuterIndex].Reset();
+			ValuesStorage[OuterIndex].SetNum(InnerArraySize);
+		}
+
+		for(int32 InnerIndex = 0; InnerIndex < InnerArraySize; InnerIndex++)
+		{
+			ValuesStorage[OuterIndex][InnerIndex].Serialize(InArchive);
+		}
+	}
+
+	if(InArchive.IsLoading())
+	{
+		UpdateValueArrayView();
+	}
+}
+
+template<>
+int32 FSampleTrack<FInstancedStruct>::GetSizePerValue() const
+{
+	if(ScriptStruct)
+	{
+		return sizeof(FInstancedStruct) + ScriptStruct->GetStructureSize();
+	}
+	return sizeof(FInstancedStruct);
+}
diff --git a/Engine/Plugins/Animation/ControlRig/Source/ControlRig/Private/Tracks/SampleTrackArchive.cpp b/Engine/Plugins/Animation/ControlRig/Source/ControlRig/Private/Tracks/SampleTrackArchive.cpp
new file mode 100644
--- /dev/null
+++ b/Engine/Plugins/Animation/ControlRig/Source/ControlRig/Private/Tracks/SampleTrackArchive.cpp
@@ -0,0 +1,163 @@
+// Copyright Epic Games, Inc. All Rights Reserved.
+
+#include "Tracks/SampleTrackArchive.h"
+#include "Misc/Compression.h"
+#include "RigVMTypeUtils.h"
+#include "ControlRig.h"
+
+#include UE_INLINE_GENERATED_CPP_BY_NAME(SampleTrackArchive)
+
+//////////////////////////////////////////////////////////////////////////////////////////////////////////////////////////////
+/////// FSampleTrackMemoryData
+//////////////////////////////////////////////////////////////////////////////////////////////////////////////////////////////
+
+void FSampleTrackMemoryData::Serialize(FArchive& Ar)
+{
+	Ar << Buffer;
+	Ar << Names;
+	Ar << ObjectPaths;
+}
+
+//////////////////////////////////////////////////////////////////////////////////////////////////////////////////////////////
+/////// FSampleTrackMemoryWriter
+//////////////////////////////////////////////////////////////////////////////////////////////////////////////////////////////
+
+FSampleTrackMemoryWriter::FSampleTrackMemoryWriter(FSampleTrackMemoryData& InOutData, bool bIsPersistent)
+: FMemoryWriter(InOutData.Buffer, bIsPersistent)
+, Data(InOutData)
+{
+}
+	
+FArchive& FSampleTrackMemoryWriter::operator<<(FName& Value)
+{
+	int32 Index = INDEX_NONE;
+	if(const int32* ExistingIndexPtr = NameToIndex.Find(Value))
+	{
+		Index = *ExistingIndexPtr;
+	}
+	else
+	{
+		Index = Data.Names.Find(Value);
+		if(Index == INDEX_NONE)
+		{
+			Index = Data.Names.Add(Value);
+		}
+		NameToIndex.Add(Value, Index);
+	}
+	*this << Index;
+	return *this;
+}
+
+FArchive& FSampleTrackMemoryWriter::operator<<(FText& Value)
+{
+	FString ValueString = Value.ToString();
+	*this << ValueString;
+	return *this;
+}
+
+FArchive& FSampleTrackMemoryWriter::operator<<(class UObject*& Value)
+{
+	int32 Index = INDEX_NONE;
+	if(Value)
+	{
+		if(const int32* ExistingIndex = ObjectToIndex.Find(Value))
+		{
+			Index = *ExistingIndex;
+		}
+		else
+		{
+			const FString ObjectPath = Value->GetPathName();
+			Index = Data.ObjectPaths.Find(ObjectPath);
+			if(Index == INDEX_NONE)
+			{
+				Index = Data.ObjectPaths.Add(ObjectPath);
+			}
+			ObjectToIndex.Add(Value, Index);
+		}
+	}
+	*this << Index;
+	return *this;
+}
+
+FArchive& FSampleTrackMemoryWriter::operator<<(FObjectPtr& Value)
+{
+	UObject* Object = Value.Get();
+	*this << Object;
+	return *this;
+}
+
+//////////////////////////////////////////////////////////////////////////////////////////////////////////////////////////////
+/////// FSampleTrackMemoryReader
+//////////////////////////////////////////////////////////////////////////////////////////////////////////////////////////////
+
+FSampleTrackMemoryReader::FSampleTrackMemoryReader(FSampleTrackMemoryData& InOutData, bool bIsPersistent)
+: FMemoryReader(InOutData.Buffer, bIsPersistent)
+, Data(InOutData)
+{
+	static TMap<FString, TWeakObjectPtr<UObject>> PathToObject;
+	static FCriticalSection PathToObjectMutex;
+	FScopeLock _(&PathToObjectMutex);
+
+	for(const FString& ObjectPath : Data.ObjectPaths)
+	{
+		if(TWeakObjectPtr<UObject>* ExistingObject = PathToObject.Find(ObjectPath))
+		{
+			Objects.Add(ExistingObject->Get());
+		}
+		else
+		{
+			UObject* Object = RigVMTypeUtils::FindObjectFromCPPTypeObjectPath(ObjectPath);
+			if(Object == nullptr)
+			{
+				UE_LOG(LogControlRig, Error, TEXT("FSampleTrackMemoryReader: The object '%s' could not be resolved."), *ObjectPath);
+			}
+			else
+			{
+				PathToObject.Add(ObjectPath, Object);
+			}
+			Objects.Add(Object);
+		}
+	}
+}
+
+FArchive& FSampleTrackMemoryReader::operator<<(FName& Value)
+{
+	int32 NameIndex = INDEX_NONE;
+	*this << NameIndex;
+	check(Data.Names.IsValidIndex(NameIndex));
+	Value = Data.Names[NameIndex];
+	return *this;
+}
+	
+FArchive& FSampleTrackMemoryReader::operator<<(FText& Value)
+{
+	FString ValueString;
+	*this << ValueString;
+	Value = FText::FromString(ValueString);
+	return *this;
+}
+
+FArchive& FSampleTrackMemoryReader::operator<<(class UObject*& Value)
+{
+	int32 Index = INDEX_NONE;
+	*this << Index;
+
+	if(Index == INDEX_NONE)
+	{
+		Value = nullptr;
+	}
+	else
+	{
+		check(Objects.IsValidIndex(Index));
+		Value = Objects[Index];
+	}
+	return *this;
+}
+
+FArchive& FSampleTrackMemoryReader::operator<<(FObjectPtr& Value)
+{
+	UObject* Object = nullptr;
+	*this << Object;
+	Value = Object;
+	return *this;
+}
diff --git a/Engine/Plugins/Animation/ControlRig/Source/ControlRig/Private/Tracks/SampleTrackBase.cpp b/Engine/Plugins/Animation/ControlRig/Source/ControlRig/Private/Tracks/SampleTrackBase.cpp
new file mode 100644
--- /dev/null
+++ b/Engine/Plugins/Animation/ControlRig/Source/ControlRig/Private/Tracks/SampleTrackBase.cpp
@@ -0,0 +1,501 @@
+// Copyright Epic Games, Inc. All Rights Reserved.
+
+#include "Tracks/SampleTrackBase.h"
+#include "Tracks/SampleTrackContainer.h"
+
+//////////////////////////////////////////////////////////////////////////////////////////////////////////////////////////////
+/////// FSampleTrackBase
+//////////////////////////////////////////////////////////////////////////////////////////////////////////////////////////////
+
+FSampleTrackBase::FSampleTrackBase()
+: Names({NAME_None})
+, TrackIndex(INDEX_NONE)
+, ReferencedTrackIndex(INDEX_NONE)
+, ReferencedTimeIndicesRange({INDEX_NONE,INDEX_NONE})
+, ReferencedAtlasRange({INDEX_NONE,INDEX_NONE})
+, ReferencedValuesRange({INDEX_NONE,INDEX_NONE})
+, ScriptStruct(nullptr)
+, Mode(EMode_Invalid)
+, NumTimesInContainer(0)
+, Container(nullptr)
+{
+}
+
+FSampleTrackBase::FSampleTrackBase(const FName& InName)
+: Names({InName})
+, TrackIndex(INDEX_NONE)
+, ReferencedTrackIndex(INDEX_NONE)
+, ReferencedTimeIndicesRange({INDEX_NONE,INDEX_NONE})
+, ReferencedAtlasRange({INDEX_NONE,INDEX_NONE})
+, ReferencedValuesRange({INDEX_NONE,INDEX_NONE})
+, ScriptStruct(nullptr)
+, Mode(EMode_Invalid)
+, NumTimesInContainer(0)
+, Container(nullptr)
+{
+}
+
+int32 FSampleTrackBase::NumTimes() const
+{
+	return NumTimesInContainer;
+}
+
+int32 FSampleTrackBase::NumSamples() const
+{
+	switch(Mode)
+	{
+		case EMode_Singleton:
+		{
+			return 1;
+		}
+		case EMode_Sampled:
+		{
+			return TimeIndices.Num();
+		}
+		case EMode_Complete:
+		{
+			return NumTimes();
+		}
+		case EMode_Raw:
+		case EMode_Invalid:
+		default:
+		{
+			break;
+		}
+	}
+	return 0;
+}
+
+bool FSampleTrackBase::GetSampleIndexForTimeIndex(int32 InTimeIndex, int32& InOutSampleIndex) const
+{
+	switch(Mode)
+	{
+		case EMode_Singleton:
+		{
+			InOutSampleIndex = 0;
+			return true;
+		}
+		case EMode_Sampled:
+		{
+			InOutSampleIndex = FMath::Clamp(InOutSampleIndex, 0, TimeIndices.Num()-1);
+
+			if(TimeIndices.IsEmpty())
+			{
+				break;
+			}
+			if(!TimeIndices.IsValidIndex(InOutSampleIndex))
+			{
+				InOutSampleIndex = 0;
+			}
+			while(TimeIndices.IsValidIndex(InOutSampleIndex - 1) && TimeIndices[InOutSampleIndex] > InTimeIndex)
+			{
+				InOutSampleIndex--;
+			}
+			while(TimeIndices.IsValidIndex(InOutSampleIndex + 1) && TimeIndices[InOutSampleIndex + 1] <= InTimeIndex)
+			{
+				InOutSampleIndex++;
+			}
+			return true;
+		}
+		case EMode_Complete:
+		{
+			InOutSampleIndex = InTimeIndex;
+			return true;
+		}
+		case EMode_Invalid:
+		case EMode_Raw:
+		default:
+		{
+			break;
+		}
+	}
+	return false;
+}
+
+bool FSampleTrackBase::StoresValueForTimeIndex(int32 InTimeIndex) const
+{
+	switch(Mode)
+	{
+		case EMode_Singleton:
+		{
+			return InTimeIndex == 0;
+		}
+		case EMode_Sampled:
+		{
+			return TimeIndices.Contains(InTimeIndex);
+		}
+		case EMode_Complete:
+		{
+			return true;
+		}
+		case EMode_Invalid:
+		default:
+		{
+			break;
+		}
+	}
+	return false;
+}
+
+bool FSampleTrackBase::IsIdentical(const FSampleTrackBase* InOther, float InTolerance)
+{
+	if(IsReferenced() || InOther->IsReferenced())
+	{
+		return false;
+	}
+	if(GetTrackType() != InOther->GetTrackType())
+	{
+		return false;
+	}
+	if(GetMode() != InOther->GetMode())
+	{
+		return false;
+	}
+	if(NumTimes() != InOther->NumTimes())
+	{
+		return false;
+	}
+	if(NumSamples() != InOther->NumSamples())
+	{
+		return false;
+	}
+	if(GetMode() == EMode_Sampled)
+	{
+		if(TimeIndices.Num() != InOther->TimeIndices.Num())
+		{
+			return false;
+		}
+		for(int32 Index = 0; Index < TimeIndices.Num(); Index++)
+		{
+			if(TimeIndices[Index] != InOther->TimeIndices[Index])
+			{
+				return false;
+			}
+		}
+	}
+	if(Atlas.Num() != InOther->Atlas.Num())
+	{
+		return false;
+	}
+	for(int32 Index = 0; Index < Atlas.Num(); Index++)
+	{
+		if(Atlas[Index] != InOther->Atlas[Index])
+		{
+			return false;
+		}
+	}
+	return true;
+}
+
+void FSampleTrackBase::Reserve(int32 InSampleCount, int32 InValueCount)
+{
+	check(!IsReferenced());
+	TimeIndicesStorage.Reserve(InSampleCount);
+}
+
+void FSampleTrackBase::Reset()
+{
+	Mode = EMode_Invalid;
+	TimeIndicesStorage.Reset();
+	AtlasStorage.Reset(0);
+	NumTimesInContainer = 0;
+	ReferencedTrackIndex = INDEX_NONE;
+	ReferencedTimeIndicesRange = {INDEX_NONE, INDEX_NONE};
+	ReferencedAtlasRange = {INDEX_NONE, INDEX_NONE};
+	ReferencedValuesRange = {INDEX_NONE, INDEX_NONE};
+	UpdateTimeAndAtlasArrayViews();
+}
+
+void FSampleTrackBase::Shrink()
+{
+	if(Mode == EMode_Sampled || Mode == EMode_Raw)
+	{
+		TimeIndicesStorage.Shrink();
+	}
+	else
+	{
+		TimeIndicesStorage.Empty();
+	}
+	AtlasStorage.Shrink();
+	UpdateTimeAndAtlasArrayViews();
+}
+
+void FSampleTrackBase::Empty()
+{
+	TimeIndicesStorage.Empty();
+	AtlasStorage.Empty();
+	UpdateTimeAndAtlasArrayViews();
+}
+
+void FSampleTrackBase::Serialize(FArchive& InArchive)
+{
+	if(InArchive.IsLoading())
+	{
+		Reset();
+	}
+		
+	InArchive << Names;
+	InArchive << (uint8&)Mode;
+	InArchive << NumTimesInContainer;
+	InArchive << TrackIndex;
+	InArchive << ReferencedTrackIndex;
+	InArchive << ReferencedTimeIndicesRange;
+	InArchive << ReferencedAtlasRange;
+	InArchive << ReferencedValuesRange;
+
+	bool bHasScriptStruct = ScriptStruct != nullptr;
+	InArchive << bHasScriptStruct;
+	
+	if(InArchive.IsLoading())
+	{
+		if(bHasScriptStruct)
+		{
+			UScriptStruct* SerializedScriptStruct = nullptr;
+			InArchive << SerializedScriptStruct;
+			ScriptStruct = SerializedScriptStruct;
+		}
+		else
+		{
+			ScriptStruct = nullptr;
+		}
+	}
+	else
+	{
+		if(bHasScriptStruct)
+		{
+			UScriptStruct* ScriptStructToSerialize = const_cast<UScriptStruct*>(ScriptStruct);
+			InArchive << ScriptStructToSerialize;
+		}
+	}
+
+	switch(Mode)
+	{
+		// only sampled tracks require the indices array
+		case EMode_Sampled:
+		case EMode_Raw:
+		{
+			InArchive << TimeIndicesStorage;
+		}
+		default:
+		{
+			break;
+		}
+	}
+
+	InArchive << AtlasStorage;
+
+	if(InArchive.IsLoading())
+	{
+		UpdateTimeAndAtlasArrayViews();
+		Shrink();
+	}
+}
+
+const FSampleTrackBase* FSampleTrackBase::GetReferencedTrack() const
+{
+	check(IsReferenced());
+	check(Container);
+	check(ReferencedTrackIndex != INDEX_NONE);
+	if(Container->NumTracks() <= ReferencedTrackIndex)
+	{
+		return nullptr;
+	}
+	return Container->GetTrack(ReferencedTrackIndex).Get();
+}
+
+FSampleTrackBase* FSampleTrackBase::GetReferencedTrack()
+{
+	check(IsReferenced());
+	check(Container);
+	check(ReferencedTrackIndex != INDEX_NONE);
+	if(Container->NumTracks() <= ReferencedTrackIndex)
+	{
+		return nullptr;
+	}
+	return Container->GetTrack(ReferencedTrackIndex).Get();
+}
+
+void FSampleTrackBase::UpdateArrayViews()
+{
+	UpdateTimeAndAtlasArrayViews();
+}
+
+void FSampleTrackBase::UpdateTimeAndAtlasArrayViews()
+{
+	if(IsReferenced())
+	{
+		FSampleTrackBase* ReferencedTrack = GetReferencedTrack();
+		if(ReferencedTrack == nullptr)
+		{
+			TimeIndices = TArrayView<int32>();
+			Atlas = TArrayView<int32>();
+			return;
+		}
+
+		if(ReferencedTimeIndicesRange.Get<1>() <= 0)
+		{
+			TimeIndices = TArrayView<int32>();
+		}
+		else
+		{
+			check(ReferencedTrack->TimeIndicesStorage.IsValidIndex(ReferencedTimeIndicesRange.Get<0>()));
+			check(ReferencedTrack->TimeIndicesStorage.IsValidIndex(ReferencedTimeIndicesRange.Get<0>() + ReferencedTimeIndicesRange.Get<1>() - 1));
+		
+			TimeIndices = TArrayView<int32>(
+				ReferencedTrack->TimeIndicesStorage.GetData() + ReferencedTimeIndicesRange.Get<0>(),
+				ReferencedTimeIndicesRange.Get<1>());
+		}
+
+		if(ReferencedAtlasRange.Get<1>() <= 0)
+		{
+			Atlas = TArrayView<int32>();
+		}
+		else
+		{
+			check(ReferencedTrack->AtlasStorage.IsValidIndex(ReferencedAtlasRange.Get<0>()));
+			check(ReferencedTrack->AtlasStorage.IsValidIndex(ReferencedAtlasRange.Get<0>() + ReferencedAtlasRange.Get<1>() - 1));
+			
+			Atlas = TArrayView<int32>(
+				ReferencedTrack->AtlasStorage.GetData() + ReferencedAtlasRange.Get<0>(),
+				ReferencedAtlasRange.Get<1>());
+		}
+	}
+	else
+	{
+		TimeIndices = TArrayView<int32>(TimeIndicesStorage.GetData(), TimeIndicesStorage.Num());
+		Atlas = TArrayView<int32>(AtlasStorage.GetData(), AtlasStorage.Num());
+	}
+}
+
+bool FSampleTrackBase::UsesAtlas(int32 InAtlasIndex) const
+{
+	if(Atlas.IsEmpty())
+	{
+		return false;
+	}
+#if WITH_EDITOR	
+	if(InAtlasIndex != INDEX_NONE)
+	{
+		check(Atlas.IsValidIndex(InAtlasIndex));
+	}
+#endif
+	return true;
+}
+
+void FSampleTrackBase::AddTimeIndex(bool bOnlyIncreaseUpperBound)
+{
+	check(!IsReferenced());
+
+	if(!bOnlyIncreaseUpperBound)
+	{
+		TimeIndicesStorage.Add(NumTimesInContainer);
+		UpdateTimeAndAtlasArrayViews();
+	}
+	NumTimesInContainer++;
+
+	switch(Mode)
+	{
+		case EMode_Singleton:
+		case EMode_Sampled:
+		case EMode_Complete:
+		{
+			if(TimeIndices.Num() == NumTimesInContainer)
+			{
+				Mode = EMode_Complete;
+			}
+			else if(TimeIndices.Num() == 1)
+			{
+				Mode = EMode_Singleton;
+			}
+			else
+			{
+				Mode = EMode_Sampled;
+			}
+			break;
+		}
+		case EMode_Invalid:
+		default:
+		{
+			Mode = EMode_Singleton;
+			break;
+		}
+	}
+}
+
+int32 FSampleTrackBase::RemoveObsoleteTimes(int32 InNumSamplesToRemove, const TArray<bool>& InRemoveIndexMap, const TArray<int32>& InNewIndexMap)
+{
+	if(!IsValid() || IsReferenced())
+	{
+		return 0;
+	}
+
+	check(NumTimesInContainer == InRemoveIndexMap.Num());
+	check(NumTimesInContainer == InNewIndexMap.Num());
+	
+	if(InNumSamplesToRemove == 0)
+	{
+		return 0;
+	}
+
+	switch(Mode)
+	{
+		case EMode_Singleton:
+		{
+			// tracks can't be empty and still be valid.
+			// so we'll keep one value in here at all times (no pun intended).
+			break;
+		}
+		case EMode_Sampled:
+		{
+			int32 NumRemoved = TimeIndicesStorage.RemoveAll([InRemoveIndexMap](int32 TimeIndex)
+			{
+				return InRemoveIndexMap[TimeIndex];
+			});
+
+			for(int32& TimeIndex : TimeIndicesStorage)
+			{
+				if(InNewIndexMap[TimeIndex] != TimeIndex)
+				{
+					TimeIndex = InNewIndexMap[TimeIndex];
+				}
+			}
+
+			NumTimesInContainer = FMath::Max(1, NumTimesInContainer - InNumSamplesToRemove);
+			UpdateTimeAndAtlasArrayViews();
+			return NumRemoved;
+		}
+		case EMode_Complete:
+		{
+			// the indices array is empty - there's no need to shrink them,
+			// only the values have to be adapted.
+			return InNumSamplesToRemove;
+		}
+		case EMode_Invalid:
+		default:
+		{
+			checkNoEntry();
+			break;
+		}
+	}
+
+	return 0;
+}
+
+bool FSampleTrackBase::LocalizeValues()
+{
+	if(IsReferenced())
+	{
+		// first update all array views so that we can copy the data off of the raw track
+		UpdateArrayViews();
+
+		ReferencedTrackIndex = INDEX_NONE;
+		ReferencedTimeIndicesRange = {INDEX_NONE, INDEX_NONE};
+		ReferencedAtlasRange = {INDEX_NONE, INDEX_NONE};
+		ReferencedValuesRange = {INDEX_NONE, INDEX_NONE};
+		TimeIndicesStorage = TimeIndices;
+		AtlasStorage = Atlas;
+		UpdateTimeAndAtlasArrayViews();
+		return true;
+	}
+	return false;
+}
diff --git a/Engine/Plugins/Animation/ControlRig/Source/ControlRig/Private/Tracks/SampleTrackContainer.cpp b/Engine/Plugins/Animation/ControlRig/Source/ControlRig/Private/Tracks/SampleTrackContainer.cpp
new file mode 100644
--- /dev/null
+++ b/Engine/Plugins/Animation/ControlRig/Source/ControlRig/Private/Tracks/SampleTrackContainer.cpp
@@ -0,0 +1,1083 @@
+// Copyright Epic Games, Inc. All Rights Reserved.
+
+#include "Tracks/SampleTrackContainer.h"
+#include "Tracks/SampleTrackArchive.h"
+#include "Misc/Compression.h"
+
+#include UE_INLINE_GENERATED_CPP_BY_NAME(SampleTrackContainer)
+
+//////////////////////////////////////////////////////////////////////////////////////////////////////////////////////////////
+/////// FSampleTrackContainer
+//////////////////////////////////////////////////////////////////////////////////////////////////////////////////////////////
+
+FSampleTrackHost::FSampleTrackHost()
+: Container(new FSampleTrackContainer)
+{
+}
+
+FSampleTrackHost::FSampleTrackHost(const FSampleTrackHost& InOther)
+: Container(new FSampleTrackContainer)
+{
+	*this = InOther;
+}
+
+FSampleTrackHost& FSampleTrackHost::operator=(const FSampleTrackHost& InOther)
+{
+	if(Container == nullptr)
+	{
+		Container = MakeUnique<FSampleTrackContainer>();
+	}
+	
+	if(InOther.GetContainer())
+	{
+		FSampleTrackMemoryData ArchiveData;
+		FSampleTrackMemoryWriter ArchiveWriter(ArchiveData);
+		const_cast<FSampleTrackHost&>(InOther).GetContainer()->Serialize(ArchiveWriter);
+
+		FSampleTrackMemoryReader ArchiveReader(ArchiveData);
+		Container->Serialize(ArchiveReader);
+	}
+	else
+	{
+		Container->Reset();
+	}
+	
+	return *this;
+}
+
+void FSampleTrackHost::Serialize(FArchive& InArchive)
+{
+	if(Container == nullptr)
+	{
+		Container = MakeUnique<FSampleTrackContainer>();
+	}
+	Container->Serialize(InArchive);
+}
+
+const FSampleTrackContainer* FSampleTrackHost::GetContainer() const
+{
+	return Container.Get();
+}
+
+FSampleTrackContainer* FSampleTrackHost::GetContainer()
+{
+	return Container.Get();
+}
+
+//////////////////////////////////////////////////////////////////////////////////////////////////////////////////////////////
+/////// FSampleTrackContainer
+//////////////////////////////////////////////////////////////////////////////////////////////////////////////////////////////
+
+FSampleTrackContainer::FSampleTrackContainer()
+: bForceToUseCompression(false)
+{
+}
+
+FSampleTrackContainer::~FSampleTrackContainer()
+{
+}
+
+void FSampleTrackContainer::Reset()
+{
+	TimeSamples.Reset();
+	NameToIndex.Reset();
+	Tracks.Reset();
+}
+
+void FSampleTrackContainer::Shrink()
+{
+	TimeSamples.Shrink();
+	NameToIndex.Shrink();
+
+	for(const TSharedPtr<FSampleTrackBase>& Track : Tracks)
+	{
+		Track->Shrink();
+	}
+	for(const TSharedPtr<FSampleTrackBase>& Track : Tracks)
+	{
+		Track->UpdateArrayViews();
+	}
+}
+
+void FSampleTrackContainer::Compact(float InTolerance)
+{
+	RemoveInvalidTracks(false);
+	RemoveObsoleteTimes();
+	MergeDuplicateTracks(false, InTolerance);
+	MergeTypedTracks(false, InTolerance);
+	ConvertTracksToSampled();
+	ConvertTracksToComplete();
+	EnableTrackAtlas(InTolerance);
+	UpdateNameToIndexMap();
+	Shrink();
+}
+
+void FSampleTrackContainer::Reserve(int32 InNum)
+{
+	TimeSamples.Reserve(InNum);
+
+	for(const TSharedPtr<FSampleTrackBase>& Track : Tracks)
+	{
+		Track->Reserve(InNum, InNum);
+	}
+	for(const TSharedPtr<FSampleTrackBase>& Track : Tracks)
+	{
+		Track->UpdateArrayViews();
+	}
+}
+
+void FSampleTrackContainer::Serialize(FArchive& InArchive)
+{
+	InArchive.UsingCustomVersion(FControlRigObjectVersion::GUID);
+	
+	bool bUseCompression = bForceToUseCompression || (
+		!((InArchive.GetPortFlags() & PPF_Duplicate) != 0) &&
+		InArchive.IsPersistent() &&
+		!InArchive.IsObjectReferenceCollector() &&
+		!InArchive.IsCountingMemory() &&
+		!InArchive.ShouldSkipBulkData() &&
+		!InArchive.IsTransacting()
+	);
+
+	// compact ourselves to as small as possible
+	if(InArchive.IsSaving() && bUseCompression && InArchive.IsPersistent())
+	{
+		Compact();
+	}
+
+	InArchive << bUseCompression;
+
+	FArchive* InnerArchive = &InArchive;
+	TUniquePtr<FMemoryArchive> MemoryArchive;
+
+	FSampleTrackMemoryData ArchiveData;
+	TArray<uint8> CompressedBytes;
+
+	if(InArchive.IsLoading())
+	{
+		Reset();
+
+		if(bUseCompression)
+		{
+			InArchive << ArchiveData;
+			
+			int32 UncompressedSize = 0;
+			InArchive << UncompressedSize;
+
+			bool bStoreCompressedBytes = false;
+			InArchive << bStoreCompressedBytes;
+
+			if(bStoreCompressedBytes)
+			{
+				Swap(CompressedBytes, ArchiveData.Buffer);
+				ArchiveData.Buffer.SetNumUninitialized(UncompressedSize);
+
+				const bool bSuccess = FCompression::UncompressMemory(NAME_Oodle, ArchiveData.Buffer.GetData(), UncompressedSize, CompressedBytes.GetData(), CompressedBytes.Num());
+				verify(bSuccess);
+			}
+
+			MemoryArchive = MakeUnique<FSampleTrackMemoryReader>(ArchiveData, InArchive.IsPersistent());
+			InnerArchive = MemoryArchive.Get();
+			InnerArchive->SetCustomVersions(InArchive.GetCustomVersions());
+		}
+	}
+	else
+	{
+		if(bUseCompression)
+		{
+			MemoryArchive = MakeUnique<FSampleTrackMemoryWriter>(ArchiveData, InArchive.IsPersistent());
+			InnerArchive = MemoryArchive.Get();
+			InnerArchive->SetCustomVersions(InArchive.GetCustomVersions());
+		}
+	}
+	
+	(*InnerArchive) << TimeSamples;
+
+	int32 NumTracks = Tracks.Num();
+	(*InnerArchive) << NumTracks;
+	
+	if(InArchive.IsLoading())
+	{
+		// first create all of the tracks
+		for(int32 TrackIndex = 0; TrackIndex < NumTracks; TrackIndex++)
+		{
+			uint8 TrackType = 0;
+			(*InnerArchive) << TrackType;
+			AddTrack(MakeTrack((FSampleTrackBase::ETrackType)TrackType));
+		}
+
+		// now load the track in the order they were saved (raw first, everything else after)
+		for(int32 TrackIndex = 0; TrackIndex < NumTracks; TrackIndex++)
+		{
+			int32 TrackIndexToLoad = INDEX_NONE;
+			(*InnerArchive) << TrackIndexToLoad;
+			Tracks[TrackIndexToLoad]->Serialize((*InnerArchive));
+		}
+
+		// update the array views on the tracks
+		for(int32 TrackIndex = 0; TrackIndex < NumTracks; TrackIndex++)
+		{
+			GetTrack(TrackIndex)->UpdateArrayViews();
+		}
+
+		// the tracks were added with no names present,
+		// so let's take care of updating the map now
+		UpdateNameToIndexMap();
+	}
+	else
+	{
+		// first save all of the types of tracks
+		for(const TSharedPtr<FSampleTrackBase>& Track : Tracks)
+		{
+			uint8 TrackType = (uint8)Track->GetTrackType();
+			(*InnerArchive) << TrackType;
+		}
+
+		// then save all of the raw tracks (tracks referenced by other tracks)
+		for(const TSharedPtr<FSampleTrackBase>& Track : Tracks)
+		{
+			if(Track->GetMode() != FSampleTrackBase::EMode_Raw)
+			{
+				continue;
+			}
+			int32 TrackIndex = Track->GetTrackIndex();
+			(*InnerArchive) << TrackIndex;
+			Track->Serialize((*InnerArchive));
+		}
+
+		// finally save all of the other tracks
+		for(const TSharedPtr<FSampleTrackBase>& Track : Tracks)
+		{
+			if(Track->GetMode() == FSampleTrackBase::EMode_Raw)
+			{
+				continue;
+			}
+			int32 TrackIndex = Track->GetTrackIndex();
+			(*InnerArchive) << TrackIndex;
+			Track->Serialize((*InnerArchive));
+		}
+		
+		if(bUseCompression)
+		{
+			// It is possible for compression to actually increase the size of the data, so we over allocate here to handle that.
+			int32 UncompressedSize = ArchiveData.Buffer.Num();
+			int32 CompressedSize = FCompression::CompressMemoryBound(NAME_Oodle, UncompressedSize);
+
+			CompressedBytes.SetNumUninitialized(CompressedSize);
+
+			bool bStoreCompressedBytes = FCompression::CompressMemory(NAME_Oodle, CompressedBytes.GetData(), CompressedSize, ArchiveData.Buffer.GetData(), ArchiveData.Buffer.Num(), COMPRESS_BiasMemory);
+			if(bStoreCompressedBytes)
+			{
+				if(CompressedSize < UncompressedSize)
+				{
+					CompressedBytes.SetNum(CompressedSize);
+					ArchiveData.Buffer = MoveTemp(CompressedBytes);
+				}
+				else
+				{
+					bStoreCompressedBytes = false;
+				}
+			}
+			
+			InArchive << ArchiveData;
+			InArchive << UncompressedSize;
+			InArchive << bStoreCompressedBytes;
+		}
+	}
+}
+
+void FSampleTrackContainer::SetForceToUseCompression(bool InForce)
+{
+	bForceToUseCompression = InForce;
+}
+
+TSharedPtr<FSampleTrack<bool>> FSampleTrackContainer::AddBoolTrack(const FName& InName)
+{
+	TSharedPtr<FSampleTrack<bool>> Track = MakeShareable(new FSampleTrack<bool>(InName));
+	AddTrack(Track);
+	return Track;
+}
+
+TSharedPtr<FSampleTrack<int32>> FSampleTrackContainer::AddInt32Track(const FName& InName)
+{
+	TSharedPtr<FSampleTrack<int32>> Track = MakeShareable(new FSampleTrack<int32>(InName));
+	AddTrack(Track);
+	return Track;
+}
+
+TSharedPtr<FSampleTrack<uint32>> FSampleTrackContainer::AddUint32Track(const FName& InName)
+{
+	TSharedPtr<FSampleTrack<uint32>> Track = MakeShareable(new FSampleTrack<uint32>(InName));
+	AddTrack(Track);
+	return Track;
+}
+
+TSharedPtr<FSampleTrack<float>> FSampleTrackContainer::AddFloatTrack(const FName& InName)
+{
+	TSharedPtr<FSampleTrack<float>> Track = MakeShareable(new FSampleTrack<float>(InName));
+	AddTrack(Track);
+	return Track;
+}
+
+TSharedPtr<FSampleTrack<FName>> FSampleTrackContainer::AddNameTrack(const FName& InName)
+{
+	TSharedPtr<FSampleTrack<FName>> Track = MakeShareable(new FSampleTrack<FName>(InName));
+	AddTrack(Track);
+	return Track;
+}
+
+TSharedPtr<FSampleTrack<FString>> FSampleTrackContainer::AddStringTrack(const FName& InName)
+{
+	TSharedPtr<FSampleTrack<FString>> Track = MakeShareable(new FSampleTrack<FString>(InName));
+	AddTrack(Track);
+	return Track;
+}
+
+TSharedPtr<FSampleTrack<FVector3f>> FSampleTrackContainer::AddVectorTrack(const FName& InName)
+{
+	TSharedPtr<FSampleTrack<FVector3f>> Track = MakeShareable(new FSampleTrack<FVector3f>(InName));
+	AddTrack(Track);
+	return Track;
+}
+
+TSharedPtr<FSampleTrack<FQuat4f>> FSampleTrackContainer::AddQuatTrack(const FName& InName)
+{
+	TSharedPtr<FSampleTrack<FQuat4f>> Track = MakeShareable(new FSampleTrack<FQuat4f>(InName));
+	AddTrack(Track);
+	return Track;
+}
+
+TSharedPtr<FSampleTrack<FTransform3f>> FSampleTrackContainer::AddTransformTrack(const FName& InName)
+{
+	TSharedPtr<FSampleTrack<FTransform3f>> Track = MakeShareable(new FSampleTrack<FTransform3f>(InName));
+	AddTrack(Track);
+	return Track;
+}
+
+TSharedPtr<FSampleTrack<FLinearColor>> FSampleTrackContainer::AddLinearColorTrack(const FName& InName)
+{
+	TSharedPtr<FSampleTrack<FLinearColor>> Track = MakeShareable(new FSampleTrack<FLinearColor>(InName));
+	AddTrack(Track);
+	return Track;
+}
+
+TSharedPtr<FSampleTrack<FRigElementKey>> FSampleTrackContainer::AddRigElementKeyTrack(const FName& InName)
+{
+	TSharedPtr<FSampleTrack<FRigElementKey>> Track = MakeShareable(new FSampleTrack<FRigElementKey>(InName));
+	AddTrack(Track);
+	return Track;
+}
+
+TSharedPtr<FSampleTrack<FRigComponentKey>> FSampleTrackContainer::AddRigComponentKeyTrack(const FName& InName)
+{
+	TSharedPtr<FSampleTrack<FRigComponentKey>> Track = MakeShareable(new FSampleTrack<FRigComponentKey>(InName));
+	AddTrack(Track);
+	return Track;
+}
+
+TSharedPtr<FSampleTrack<FInstancedStruct>> FSampleTrackContainer::AddStructTrack(const FName& InName, const UScriptStruct* InScriptStruct)
+{
+	check(InScriptStruct);
+	TSharedPtr<FSampleTrack<FInstancedStruct>> Track = MakeShareable(new FSampleTrack<FInstancedStruct>(InName));
+	Track->ScriptStruct = InScriptStruct;
+	AddTrack(Track);
+	return Track;
+}
+
+TSharedPtr<FSampleTrack<TArray<bool>>> FSampleTrackContainer::AddBoolArrayTrack(const FName& InName)
+{
+	TSharedPtr<FSampleTrack<TArray<bool>>> Track = MakeShareable(new FSampleTrack<TArray<bool>>(InName));
+	AddTrack(Track);
+	return Track;
+}
+
+TSharedPtr<FSampleTrack<TArray<int32>>> FSampleTrackContainer::AddInt32ArrayTrack(const FName& InName)
+{
+	TSharedPtr<FSampleTrack<TArray<int32>>> Track = MakeShareable(new FSampleTrack<TArray<int32>>(InName));
+	AddTrack(Track);
+	return Track;
+}
+
+TSharedPtr<FSampleTrack<TArray<uint32>>> FSampleTrackContainer::AddUint32ArrayTrack(const FName& InName)
+{
+	TSharedPtr<FSampleTrack<TArray<uint32>>> Track = MakeShareable(new FSampleTrack<TArray<uint32>>(InName));
+	AddTrack(Track);
+	return Track;
+}
+
+TSharedPtr<FSampleTrack<TArray<float>>> FSampleTrackContainer::AddFloatArrayTrack(const FName& InName)
+{
+	TSharedPtr<FSampleTrack<TArray<float>>> Track = MakeShareable(new FSampleTrack<TArray<float>>(InName));
+	AddTrack(Track);
+	return Track;
+}
+
+TSharedPtr<FSampleTrack<TArray<FName>>> FSampleTrackContainer::AddNameArrayTrack(const FName& InName)
+{
+	TSharedPtr<FSampleTrack<TArray<FName>>> Track = MakeShareable(new FSampleTrack<TArray<FName>>(InName));
+	AddTrack(Track);
+	return Track;
+}
+
+TSharedPtr<FSampleTrack<TArray<FString>>> FSampleTrackContainer::AddStringArrayTrack(const FName& InName)
+{
+	TSharedPtr<FSampleTrack<TArray<FString>>> Track = MakeShareable(new FSampleTrack<TArray<FString>>(InName));
+	AddTrack(Track);
+	return Track;
+}
+
+TSharedPtr<FSampleTrack<TArray<FVector3f>>> FSampleTrackContainer::AddVectorArrayTrack(const FName& InName)
+{
+	TSharedPtr<FSampleTrack<TArray<FVector3f>>> Track = MakeShareable(new FSampleTrack<TArray<FVector3f>>(InName));
+	AddTrack(Track);
+	return Track;
+}
+
+TSharedPtr<FSampleTrack<TArray<FQuat4f>>> FSampleTrackContainer::AddQuatArrayTrack(const FName& InName)
+{
+	TSharedPtr<FSampleTrack<TArray<FQuat4f>>> Track = MakeShareable(new FSampleTrack<TArray<FQuat4f>>(InName));
+	AddTrack(Track);
+	return Track;
+}
+
+TSharedPtr<FSampleTrack<TArray<FTransform3f>>> FSampleTrackContainer::AddTransformArrayTrack(const FName& InName)
+{
+	TSharedPtr<FSampleTrack<TArray<FTransform3f>>> Track = MakeShareable(new FSampleTrack<TArray<FTransform3f>>(InName));
+	AddTrack(Track);
+	return Track;
+}
+
+TSharedPtr<FSampleTrack<TArray<FLinearColor>>> FSampleTrackContainer::AddLinearColorArrayTrack(const FName& InName)
+{
+	TSharedPtr<FSampleTrack<TArray<FLinearColor>>> Track = MakeShareable(new FSampleTrack<TArray<FLinearColor>>(InName));
+	AddTrack(Track);
+	return Track;
+}
+
+TSharedPtr<FSampleTrack<TArray<FRigElementKey>>> FSampleTrackContainer::AddRigElementKeyArrayTrack(const FName& InName)
+{
+	TSharedPtr<FSampleTrack<TArray<FRigElementKey>>> Track = MakeShareable(new FSampleTrack<TArray<FRigElementKey>>(InName));
+	AddTrack(Track);
+	return Track;
+}
+
+TSharedPtr<FSampleTrack<TArray<FRigComponentKey>>> FSampleTrackContainer::AddRigComponentKeyArrayTrack(const FName& InName)
+{
+	TSharedPtr<FSampleTrack<TArray<FRigComponentKey>>> Track = MakeShareable(new FSampleTrack<TArray<FRigComponentKey>>(InName));
+	AddTrack(Track);
+	return Track;
+}
+
+TSharedPtr<FSampleTrack<TArray<FInstancedStruct>>> FSampleTrackContainer::AddStructArrayTrack(const FName& InName, const UScriptStruct* InScriptStruct)
+{
+	check(InScriptStruct);
+	TSharedPtr<FSampleTrack<TArray<FInstancedStruct>>> Track = MakeShareable(new FSampleTrack<TArray<FInstancedStruct>>(InName));
+	Track->ScriptStruct = InScriptStruct;
+	AddTrack(Track);
+	return Track;
+}
+
+int32 FSampleTrackContainer::GetTrackIndex(const FName& InName) const
+{
+	if(const int32* Index = NameToIndex.Find(InName))
+	{
+		return *Index;
+	}
+	return INDEX_NONE;
+}
+
+TSharedPtr<const FSampleTrackBase> FSampleTrackContainer::GetTrack(int32 InIndex) const
+{
+	return ConstCastSharedPtr<FSampleTrackBase>(Tracks[InIndex]);
+}
+
+TSharedPtr<FSampleTrackBase> FSampleTrackContainer::GetTrack(int32 InIndex)
+{
+	return Tracks[InIndex];
+}
+
+TSharedPtr<const FSampleTrackBase> FSampleTrackContainer::FindTrack(const FName& InName) const
+{
+	const int32 TrackIndex = GetTrackIndex(InName);
+	if(Tracks.IsValidIndex(TrackIndex))
+	{
+		return GetTrack(TrackIndex);
+	}
+	return nullptr;
+}
+
+TSharedPtr<FSampleTrackBase> FSampleTrackContainer::FindTrack(const FName& InName)
+{
+	const int32 TrackIndex = GetTrackIndex(InName);
+	if(Tracks.IsValidIndex(TrackIndex))
+	{
+		return GetTrack(TrackIndex);
+	}
+	return nullptr;
+}
+
+void FSampleTrackContainer::AddTrack(TSharedPtr<FSampleTrackBase> InTrack)
+{
+	InTrack->TrackIndex = Tracks.Num();
+	InTrack->Container = this;
+
+	// a track can have many names - so we need to make sure each one is unique
+	for(int32 Index = 0; Index < InTrack->Names.Num(); Index++)
+	{
+		FName& Name = InTrack->Names[Index];
+		int32 NameSuffix = 3; // causes the first duplicate name to be called "Foo_2"
+		while (NameToIndex.Contains(Name))
+		{
+			if(NameSuffix == 3)
+			{
+				Name = FName(Name, NameSuffix);
+			}
+			else
+			{
+				Name.SetNumber(NameSuffix);
+			}
+			NameSuffix++;
+		}
+		NameToIndex.Add(Name, InTrack->TrackIndex);
+	}
+	
+	Tracks.Add(InTrack);
+}
+
+TSharedPtr<FSampleTrackBase> FSampleTrackContainer::MakeTrack(FSampleTrackBase::ETrackType InTrackType)
+{
+	switch(InTrackType)
+	{
+		case FSampleTrackBase::ETrackType_Bool:
+		{
+			return MakeShareable(new FSampleTrack<bool>());
+		}
+		case FSampleTrackBase::ETrackType_Int32:
+		{
+			return MakeShareable(new FSampleTrack<int32>());
+		}
+		case FSampleTrackBase::ETrackType_Uint32:
+		{
+			return MakeShareable(new FSampleTrack<uint32>());
+		}
+		case FSampleTrackBase::ETrackType_Float:
+		{
+			return MakeShareable(new FSampleTrack<float>());
+		}
+		case FSampleTrackBase::ETrackType_Name:
+		{
+			return MakeShareable(new FSampleTrack<FName>());
+		}
+		case FSampleTrackBase::ETrackType_String:
+		{
+			return MakeShareable(new FSampleTrack<FString>());
+		}
+		case FSampleTrackBase::ETrackType_Vector3f:
+		{
+			return MakeShareable(new FSampleTrack<FVector3f>());
+		}
+		case FSampleTrackBase::ETrackType_Quatf:
+		{
+			return MakeShareable(new FSampleTrack<FQuat4f>());
+		}
+		case FSampleTrackBase::ETrackType_Transformf:
+		{
+			return MakeShareable(new FSampleTrack<FTransform3f>());
+		}
+		case FSampleTrackBase::ETrackType_LinearColor:
+		{
+			return MakeShareable(new FSampleTrack<FLinearColor>());
+		}
+		case FSampleTrackBase::ETrackType_ElementKey:
+		{
+			return MakeShareable(new FSampleTrack<FRigElementKey>());
+		}
+		case FSampleTrackBase::ETrackType_ComponentKey:
+		{
+			return MakeShareable(new FSampleTrack<FRigComponentKey>());
+		}
+		case FSampleTrackBase::ETrackType_Struct:
+		{
+			return MakeShareable(new FSampleTrack<FInstancedStruct>());
+		}
+		case FSampleTrackBase::ETrackType_BoolArray:
+		{
+			return MakeShareable(new FSampleTrack<TArray<bool>>());
+		}
+		case FSampleTrackBase::ETrackType_Int32Array:
+		{
+			return MakeShareable(new FSampleTrack<TArray<int32>>());
+		}
+		case FSampleTrackBase::ETrackType_Uint32Array:
+		{
+			return MakeShareable(new FSampleTrack<TArray<uint32>>());
+		}
+		case FSampleTrackBase::ETrackType_FloatArray:
+		{
+			return MakeShareable(new FSampleTrack<TArray<float>>());
+		}
+		case FSampleTrackBase::ETrackType_NameArray:
+		{
+			return MakeShareable(new FSampleTrack<TArray<FName>>());
+		}
+		case FSampleTrackBase::ETrackType_StringArray:
+		{
+			return MakeShareable(new FSampleTrack<TArray<FString>>());
+		}
+		case FSampleTrackBase::ETrackType_Vector3fArray:
+		{
+			return MakeShareable(new FSampleTrack<TArray<FVector3f>>());
+		}
+		case FSampleTrackBase::ETrackType_QuatfArray:
+		{
+			return MakeShareable(new FSampleTrack<TArray<FQuat4f>>());
+		}
+		case FSampleTrackBase::ETrackType_TransformfArray:
+		{
+			return MakeShareable(new FSampleTrack<TArray<FTransform3f>>());
+		}
+		case FSampleTrackBase::ETrackType_LinearColorArray:
+		{
+			return MakeShareable(new FSampleTrack<TArray<FLinearColor>>());
+		}
+		case FSampleTrackBase::ETrackType_ElementKeyArray:
+		{
+			return MakeShareable(new FSampleTrack<TArray<FRigElementKey>>());
+		}
+		case FSampleTrackBase::ETrackType_ComponentKeyArray:
+		{
+			return MakeShareable(new FSampleTrack<TArray<FRigComponentKey>>());
+		}
+		case FSampleTrackBase::ETrackType_StructArray:
+		{
+			return MakeShareable(new FSampleTrack<TArray<FInstancedStruct>>());
+		}
+
+		default:
+		{
+			checkNoEntry();
+			break;
+		}
+	}
+	return nullptr;
+}
+
+void FSampleTrackContainer::RemoveInvalidTracks(bool bUpdateNameToIndexMap)
+{
+	const int32 NumRemoved = Tracks.RemoveAll([](const TSharedPtr<FSampleTrackBase>& Track)
+	{
+		return !Track.IsValid() || !Track->IsValid();
+	});
+
+	if(NumRemoved > 0 && bUpdateNameToIndexMap)
+	{
+		UpdateNameToIndexMap();
+	}
+}
+
+void FSampleTrackContainer::RemoveObsoleteTimes()
+{
+	TArray<bool> RemoveTimeIndexArray;
+	RemoveTimeIndexArray.SetNumUninitialized(TimeSamples.Num());
+	int32 NumTimesToRemove = TimeSamples.Num();
+
+	for(int32 Index = 0; Index < RemoveTimeIndexArray.Num(); Index++)
+	{
+		RemoveTimeIndexArray[Index] = true;
+		for(const TSharedPtr<FSampleTrackBase>& Track : Tracks)
+		{
+			if(Track->StoresValueForTimeIndex(Index))
+			{
+				if(RemoveTimeIndexArray[Index])
+				{
+					NumTimesToRemove--;
+					RemoveTimeIndexArray[Index] = false;
+				}
+			}
+		}
+	}
+
+	if(NumTimesToRemove == 0)
+	{
+		return;
+	}
+
+	TArray<float> NewTimeSamples;
+	NewTimeSamples.Reserve(TimeSamples.Num() - NumTimesToRemove);
+	TArray<int32> OldIndexToNewIndex;
+	OldIndexToNewIndex.Reserve(TimeSamples.Num());
+	for(int32 Index = 0; Index < RemoveTimeIndexArray.Num(); Index++)
+	{
+		// don't allow to remove the first time sample
+		if(Index == 0)
+		{
+			verify(!RemoveTimeIndexArray[Index]);
+		}
+		
+		if(RemoveTimeIndexArray[Index])
+		{
+			const int32 LastIndex = OldIndexToNewIndex.Last();
+			OldIndexToNewIndex.Add(LastIndex);
+		}
+		else
+		{
+			NewTimeSamples.Add(TimeSamples[Index]);
+			if(OldIndexToNewIndex.IsEmpty())
+			{
+				OldIndexToNewIndex.Add(0);
+			}
+			else
+			{
+				const int32 LastIndex = OldIndexToNewIndex.Last();
+				OldIndexToNewIndex.Add(LastIndex + 1);
+			}
+		}
+	}
+	
+	TimeSamples = MoveTemp(NewTimeSamples);
+
+	for(const TSharedPtr<FSampleTrackBase>& Track : Tracks)
+	{
+		Track->RemoveObsoleteTimes(NumTimesToRemove, RemoveTimeIndexArray, OldIndexToNewIndex);
+	}
+}
+
+void FSampleTrackContainer::MergeDuplicateTracks(bool bUpdateNameToIndexMap, float InTolerance)
+{
+	struct FTrackInfo
+	{
+		FSampleTrackBase::ETrackType Type;
+		FSampleTrackBase::EMode Mode;
+		int32 Num;
+		int32 NumSamples;
+
+		bool operator!=(const FTrackInfo& InOther) const
+		{
+			return (Type != InOther.Type) || (Mode != InOther.Mode) || (Num != InOther.Num) || (NumSamples != InOther.NumSamples);
+		}
+	};
+	
+	TArray<FTrackInfo> TrackInfos;
+	TrackInfos.Reserve(Tracks.Num());
+	
+	for(const TSharedPtr<FSampleTrackBase>& Track : Tracks)
+	{
+		FTrackInfo Info;
+		if(!Track.IsValid() || !Track->IsValid())
+		{
+			Info.Mode = FSampleTrackBase::EMode_Invalid;
+			TrackInfos.Add(Info);
+			continue;
+		}
+		Info.Type = Track->GetTrackType();
+		Info.Mode = Track->GetMode();
+		Info.Num = Track->NumTimes();
+		Info.NumSamples = Track->NumSamples();
+		TrackInfos.Add(Info);
+	}
+
+	int32 NumTracksMerged = 0;
+	TArray<bool> MergedTracks;
+	MergedTracks.AddZeroed(Tracks.Num());
+	for(int32 IndexA = 0; IndexA < Tracks.Num() - 1; IndexA++)
+	{
+		if(MergedTracks[IndexA])
+		{
+			continue;
+		}
+		if(TrackInfos[IndexA].Mode == FSampleTrackBase::EMode_Invalid)
+		{
+			continue;
+		}
+		
+		TSharedPtr<FSampleTrackBase>& TrackA = Tracks[IndexA];
+
+		for(int32 IndexB = IndexA + 1; IndexB < Tracks.Num(); IndexB++)
+		{
+			if(MergedTracks[IndexB])
+			{
+				continue;
+			}
+			if(TrackInfos[IndexA].Mode == FSampleTrackBase::EMode_Invalid)
+			{
+				continue;
+			}
+			
+			if(TrackInfos[IndexA] != TrackInfos[IndexB])
+			{
+				continue;
+			}
+			
+			TSharedPtr<FSampleTrackBase>& TrackB = Tracks[IndexB];
+			if(!TrackA->IsIdentical(TrackB.Get(), InTolerance))
+			{
+				continue;
+			}
+
+			// merge the track
+			TrackA->Names.Append(TrackB->Names);
+			MergedTracks[IndexB] = true;
+			NumTracksMerged++;
+		}
+	}
+	
+	if(NumTracksMerged == 0)
+	{
+		return;
+	}
+
+	for(int32 Index = 0; Index < MergedTracks.Num(); Index++)
+	{
+		if(MergedTracks[Index])
+		{
+			Tracks[Index].Reset();
+		}
+	}
+
+	RemoveInvalidTracks(bUpdateNameToIndexMap);
+}
+
+void FSampleTrackContainer::MergeTypedTracks(bool bUpdateNameToIndexMap, float InTolerance)
+{
+	struct FTrackGroup
+	{
+		FTrackGroup()
+		: TrackType(FSampleTrackBase::ETrackType_Unknown)
+		, ScriptStruct(nullptr)
+		{
+		}
+
+		FTrackGroup(const TSharedPtr<FSampleTrackBase>& InTrack)
+		: TrackType(InTrack->GetTrackType())
+		, ScriptStruct(InTrack->GetScriptStruct())
+		{
+			Tracks.Add(InTrack);
+		}
+
+		bool MergeTrack(const TSharedPtr<FSampleTrackBase>& InTrack)
+		{
+			if(TrackType != InTrack->GetTrackType())
+			{
+				return false;
+			}
+			if(ScriptStruct != InTrack->GetScriptStruct())
+			{
+				return false;
+			}
+			Tracks.Add(InTrack); 
+			return true;
+		}
+		
+		FSampleTrackBase::ETrackType TrackType;
+		const UScriptStruct* ScriptStruct;
+		TArray<TSharedPtr<FSampleTrackBase>> Tracks;
+	};
+	
+	TArray<FTrackGroup> TrackGroups;
+	for(int32 TrackIndex = 0; TrackIndex < Tracks.Num(); TrackIndex++)
+	{
+		TSharedPtr<FSampleTrackBase>& Track = Tracks[TrackIndex];
+		if(!Track.IsValid() || !Track->IsValid() || Track->IsReferenced())
+		{
+			continue;
+		}
+
+		// make sure this track's index is up2date
+		Track->TrackIndex = TrackIndex;
+
+		bool bMergedTrack = false;
+		for(FTrackGroup& TrackGroup : TrackGroups)
+		{
+			if(TrackGroup.MergeTrack(Track))
+			{
+				bMergedTrack = true;
+				break;
+			}
+		}
+
+		if(!bMergedTrack)
+		{
+			TrackGroups.Emplace(Track);
+		}
+	}
+
+	TrackGroups.RemoveAll([](const FTrackGroup& TrackGroup)
+	{
+		return TrackGroup.Tracks.Num() <= 1;
+	});
+
+	for(FTrackGroup& TrackGroup : TrackGroups)
+	{
+		TSharedPtr<FSampleTrackBase> CombinedTrack = MakeTrack(TrackGroup.TrackType);
+		CombinedTrack->Names.Reset();
+		CombinedTrack->Mode = FSampleTrackBase::EMode_Raw;
+		CombinedTrack->ScriptStruct = TrackGroup.ScriptStruct;
+		AddTrack(CombinedTrack);
+
+		int32 CombinedNumSamples = 0;
+		int32 CombinedNumValues = 0;
+		for(const TSharedPtr<FSampleTrackBase>& Track : TrackGroup.Tracks)
+		{
+			CombinedNumSamples += Track->NumSamples();
+			CombinedNumValues += Track->NumStoredValues();
+		}
+		CombinedTrack->Reserve(CombinedNumSamples, CombinedNumValues);
+
+		for(TSharedPtr<FSampleTrackBase>& Track : TrackGroup.Tracks)
+		{
+			// remove the atlas - this gives us the opportunity
+			// to apply an atlas on the combined track
+			Track->RemoveAtlas();
+
+			// also give the track a chance to unroll its values again to save
+			// on time index list memory.
+			Track->ConvertToComplete();
+
+			// shrink the memory use of the track as it is right now
+			Track->Shrink();
+
+			// copy samples and values over into the combined track 
+			const int32 FirstTimeIndex = CombinedTrack->TimeIndicesStorage.Num();
+			CombinedTrack->TimeIndicesStorage.Append(Track->TimeIndicesStorage);
+			const int32 FirstValueIndex = CombinedTrack->AppendValuesFromTrack(Track.Get());
+
+			const TTuple<int32,int32> TimeIndicesRange = {FirstTimeIndex, Track->TimeIndicesStorage.Num()}; 
+			const TTuple<int32,int32> ValuesRange = {FirstValueIndex, Track->NumStoredValues()}; 
+
+			// set up the referencing track
+			Track->Empty();
+			Track->ReferencedTrackIndex = CombinedTrack->GetTrackIndex();
+			Track->ReferencedTimeIndicesRange = TimeIndicesRange;
+			Track->ReferencedAtlasRange = {INDEX_NONE, INDEX_NONE};
+			Track->ReferencedValuesRange = ValuesRange;
+			Track->UpdateArrayViews();
+		}
+
+		// create an atlas for the merged track. if that's successful we'll have to upgrade our tracks' atlas range
+		if(CombinedTrack->AddAtlas(false, InTolerance))
+		{
+			for(TSharedPtr<FSampleTrackBase>& Track : TrackGroup.Tracks)
+			{
+				// the size of used atlas is going to be size of the values
+				Track->ReferencedAtlasRange = Track->ReferencedValuesRange;
+
+				// for the values we'll use the whole value array (since the atlas indexes into that anyway) 
+				Track->ReferencedValuesRange = {0, CombinedTrack->NumStoredValues()};
+				Track->UpdateArrayViews();
+			}
+		}
+	}
+
+	if(!TrackGroups.IsEmpty() && bUpdateNameToIndexMap)
+	{
+		UpdateNameToIndexMap();
+	}
+}
+
+void FSampleTrackContainer::EnableTrackAtlas(float InTolerance)
+{
+	for(const TSharedPtr<FSampleTrackBase>& Track : Tracks)
+	{
+		(void)Track->AddAtlas(false, InTolerance);
+	}
+}
+
+void FSampleTrackContainer::ConvertTracksToComplete()
+{
+	for(const TSharedPtr<FSampleTrackBase>& Track : Tracks)
+	{
+		(void)Track->ConvertToComplete();
+	}
+}
+
+void FSampleTrackContainer::ConvertTracksToSampled(float InTolerance)
+{
+	for(const TSharedPtr<FSampleTrackBase>& Track : Tracks)
+	{
+		(void)Track->ConvertToSampled(false, InTolerance);
+	}
+}
+
+void FSampleTrackContainer::UpdateNameToIndexMap()
+{
+	NameToIndex.Reset();
+
+	for(int32 TrackIndex = 0; TrackIndex < Tracks.Num(); TrackIndex++)
+	{
+		Tracks[TrackIndex]->TrackIndex = TrackIndex;
+		for(const FName& Name : Tracks[TrackIndex]->GetAllNames())
+		{
+			check(!NameToIndex.Contains(Name));
+			NameToIndex.Add(Name, TrackIndex);
+		}
+	}
+}
+
+int32 FSampleTrackContainer::AddTimeSample(float InTimeSample)
+{
+	return TimeSamples.Add(InTimeSample);
+}
+
+bool FSampleTrackContainer::IsEditable() const
+{
+	for(const TSharedPtr<FSampleTrackBase>& Track : Tracks)
+	{
+		if(Track->IsReferenced())
+		{
+			return false;
+		}
+		if(Track->UsesAtlas())
+		{
+			return false;
+		}
+		if(Track->GetMode() != FSampleTrackBase::EMode_Complete &&
+			Track->GetMode() != FSampleTrackBase::EMode_Singleton)
+		{
+			return false;
+		}
+	}
+	return true;
+}
+
+bool FSampleTrackContainer::MakeEditable()
+{
+	if(IsEditable())
+	{
+		return true;
+	}
+
+	bool bChangedSomething = false;
+	TArray<int32> TracksToRemove;
+	for(int32 TrackIndex = 0; TrackIndex < Tracks.Num(); TrackIndex++)
+	{
+		if(Tracks[TrackIndex]->GetMode() == FSampleTrackBase::EMode_Raw)
+		{
+			TracksToRemove.Add(TrackIndex);
+		}
+
+		Tracks[TrackIndex]->UpdateArrayViews();
+		if(Tracks[TrackIndex]->LocalizeValues())
+		{
+			bChangedSomething = true;
+		}
+		if(Tracks[TrackIndex]->RemoveAtlas())
+		{
+			bChangedSomething = true;
+		}
+		if(Tracks[TrackIndex]->ConvertToComplete(true))
+		{
+			bChangedSomething = true;
+		}
+	}
+
+	Algo::Reverse(TracksToRemove);
+	for(const int32& IndexToRemove : TracksToRemove)
+	{
+		Tracks.RemoveAt(IndexToRemove);
+		bChangedSomething = true;
+	}
+
+	return bChangedSomething;
+}
+
+
diff --git a/Engine/Plugins/Animation/ControlRig/Source/ControlRig/Public/ControlRigStepSampleTrack.h b/Engine/Plugins/Animation/ControlRig/Source/ControlRig/Public/ControlRigStepSampleTrack.h
deleted file mode 100644
diff --git a/Engine/Plugins/Animation/ControlRig/Source/ControlRig/Public/Tracks/SampleTrack.h b/Engine/Plugins/Animation/ControlRig/Source/ControlRig/Public/Tracks/SampleTrack.h
new file mode 100644
--- /dev/null
+++ b/Engine/Plugins/Animation/ControlRig/Source/ControlRig/Public/Tracks/SampleTrack.h
@@ -0,0 +1,797 @@
+// Copyright Epic Games, Inc. All Rights Reserved.
+
+#pragma once
+
+#include "SampleTrackBase.h"
+#include "Math/UnrealMath.h"
+#include "Math/UnrealMathUtility.h"
+
+template<typename T>
+class FSampleTrack : public FSampleTrackBase
+{
+public:
+
+	FSampleTrack()
+	{
+	}
+	
+	FSampleTrack(const FName& InName)
+	: FSampleTrackBase(InName)
+	{
+	}
+
+	virtual ~FSampleTrack() override {}
+
+	// the number of stored values - which can be less than NumSamples if there is an atlas present
+	virtual int32 NumStoredValues() const override
+	{
+		return Values.Num();
+	}
+
+	// returns the value of this track at a given time index
+	const T& GetValueAtTimeIndex(int32 InTimeIndex, int32& InOutSampleIndex) const
+	{
+		verify(GetSampleIndexForTimeIndex(InTimeIndex, InOutSampleIndex));
+		return GetValueAtSampleIndex(InOutSampleIndex);
+	}
+	
+	// returns the value of this track at a given time index (overload for instanced structs only)
+	template <
+		typename StructType,
+		typename TEnableIf<TIsInstancedStructCompliant<T, StructType>::Value, StructType>::Type* = nullptr
+	>
+	StructType GetValueAtTimeIndex(int32 InTimeIndex, int32& InOutSampleIndex) const
+	{
+		check(StructType::StaticStruct() == ScriptStruct);
+		const FInstancedStruct& InstancedStruct = GetValueAtTimeIndex(InTimeIndex, InOutSampleIndex);
+		StructType Result;
+		check(InstancedStruct.GetScriptStruct() == ScriptStruct);
+		ScriptStruct->CopyScriptStruct(&Result, InstancedStruct.GetMemory(), 1);
+		return Result;
+	}
+
+	// returns the value of this track at a given time index (overload for instanced struct arrays only)
+	template <
+		typename StructArrayType,
+		typename TEnableIf<TIsInstancedStructArrayCompliant<T, StructArrayType>::Value, StructArrayType>::Type* = nullptr
+	>
+	StructArrayType GetValueAtTimeIndex(int32 InTimeIndex, int32& InOutSampleIndex) const
+	{
+		check(StructArrayType::ElementType::StaticStruct() == ScriptStruct);
+		const TArray<FInstancedStruct>& InstancedStructArray = GetValueAtTimeIndex(InTimeIndex, InOutSampleIndex);
+		StructArrayType Result;
+		Result.SetNum(InstancedStructArray.Num());
+		for(int32 ElementIndex = 0; ElementIndex < InstancedStructArray.Num(); ElementIndex++)
+		{
+			check(InstancedStructArray[ElementIndex].GetScriptStruct() == ScriptStruct);
+			ScriptStruct->CopyScriptStruct(&Result[ElementIndex], InstancedStructArray[ElementIndex].GetMemory(), 1);
+		}
+		return Result;
+	}
+
+	// returns the value of this track at a given time index (overload for instanced struct arrays only)
+	template <
+		typename StructType,
+		typename TEnableIf<TIsInstancedStructArrayCompliant<T, TArray<StructType>>::Value, StructType>::Type* = nullptr
+	>
+	TArray<StructType> GetArrayValueAtTimeIndex(int32 InTimeIndex, int32& InOutSampleIndex) const
+	{
+		return GetValueAtTimeIndex<TArray<StructType>>(InTimeIndex, InOutSampleIndex);
+	}
+
+	// returns the value of this track at a given sample index
+	const T& GetValueAtSampleIndex(int32 InSampleIndex) const
+	{
+		if(UsesAtlas(InSampleIndex))
+		{
+			return Values[Atlas[InSampleIndex]];
+		}
+		return Values[InSampleIndex];
+	}
+
+	// returns the value of this track at a given sample index (overload for instanced structs only)
+	template <
+		typename StructType,
+		typename TEnableIf<TIsInstancedStructCompliant<T, StructType>::Value, StructType>::Type* = nullptr
+	>
+	StructType GetValueAtSampleIndex(int32 InSampleIndex) const
+	{
+		check(StructType::StaticStruct() == ScriptStruct);
+		const FInstancedStruct& InstancedStruct = GetValueAtSampleIndex(InSampleIndex);
+		StructType Result;
+		check(InstancedStruct.GetScriptStruct() == ScriptStruct);
+		ScriptStruct->CopyScriptStruct(&Result, InstancedStruct.GetMemory(), 1);
+		return Result;
+	}
+
+	// returns the value of this track at a given sample index (overload for instanced struct arrays only)
+	template <
+		typename StructType,
+		typename TEnableIf<TIsInstancedStructArrayCompliant<T, TArray<StructType>>::Value, StructType>::Type* = nullptr
+	>
+	TArray<StructType> GetValueAtSampleIndex(int32 InSampleIndex) const
+	{
+		return GetValueAtSampleIndex<TArray<StructType>>(InSampleIndex);
+	}
+
+	// returns the value of this track at a given sample index (overload for instanced struct arrays only)
+	template <
+		typename StructArrayType,
+		typename TEnableIf<TIsInstancedStructArrayCompliant<T, StructArrayType>::Value, StructArrayType>::Type* = nullptr
+	>
+	StructArrayType GetValueAtSampleIndex(int32 InSampleIndex) const
+	{
+		check(StructArrayType::ElementType::StaticStruct() == ScriptStruct);
+		const TArray<FInstancedStruct>& InstancedStructArray = GetValueAtSampleIndex(InSampleIndex);
+		StructArrayType Result;
+		Result.SetNum(InstancedStructArray.Num());
+		for(int32 ElementIndex = 0; ElementIndex < InstancedStructArray.Num(); ElementIndex++)
+		{
+			check(InstancedStructArray[ElementIndex].GetScriptStruct() == ScriptStruct);
+			ScriptStruct->CopyScriptStruct(&Result[ElementIndex], InstancedStructArray[ElementIndex].GetMemory(), 1);
+		}
+		return Result;
+	}
+
+	// adds a new sample to this track
+	void AddSample(const T& InValue, float InTolerance = KINDA_SMALL_NUMBER)
+	{
+		check(!IsReferenced());
+		
+		if(ValuesStorage.IsEmpty())
+		{
+			AddTimeIndex(false);
+			ValuesStorage.Add(InValue);
+			UpdateValueArrayView();
+		}
+		else if(Equals(ValuesStorage.Last(), InValue, InTolerance))
+		{
+			AddTimeIndex(true);
+		}
+		else
+		{
+			AddTimeIndex(false);
+			ValuesStorage.Add(InValue);
+			UpdateValueArrayView();
+		}
+	}
+
+	// adds a new sample to this track (overload for instanced struct only) 
+	template <
+		typename StructType,
+		typename TEnableIf<TIsInstancedStructCompliant<T, StructType>::Value, T>::Type* = nullptr
+	>
+	void AddSample(const StructType& InStructValue, float InTolerance = KINDA_SMALL_NUMBER)
+	{
+		check(StructType::StaticStruct() == ScriptStruct);
+		FInstancedStruct InstancedStruct;
+		InstancedStruct.InitializeAs(ScriptStruct, (uint8*)&InStructValue);
+		AddSample(InstancedStruct, InTolerance);
+	}
+
+	// adds a new sample to this track (overload for instanced struct array only) 
+	template <
+		typename StructArrayType,
+		typename TEnableIf<TIsInstancedStructArrayCompliant<T, StructArrayType>::Value, T>::Type* = nullptr
+	>
+	void AddSample(const StructArrayType& InStructArrayValue, float InTolerance = KINDA_SMALL_NUMBER)
+	{
+		check(StructArrayType::ElementType::StaticStruct() == ScriptStruct);
+		TArray<FInstancedStruct> InstancedStructArray;
+		InstancedStructArray.SetNum(InStructArrayValue.Num());
+		for(int32 ElementIndex = 0; ElementIndex < InStructArrayValue.Num(); ElementIndex++)
+		{
+			InstancedStructArray[ElementIndex].InitializeAs(ScriptStruct, (uint8*)&InStructArrayValue[ElementIndex]);
+		}
+		AddSample(InstancedStructArray, InTolerance);
+	}
+
+	// returns the type of this track
+	virtual ETrackType GetTrackType() const override
+	{
+		return ETrackType_Unknown;
+	}
+
+	virtual void Reset() override
+	{
+		FSampleTrackBase::Reset();
+		ValuesStorage.Reset();
+		UpdateValueArrayView();
+	}
+
+	virtual void Reserve(int32 InSampleCount, int32 InValueCount) override
+	{
+		FSampleTrackBase::Reserve(InSampleCount, InValueCount);
+		ValuesStorage.Reserve(InValueCount);
+		UpdateValueArrayView();
+	}
+
+	virtual void Shrink() override
+	{
+		FSampleTrackBase::Shrink();
+		ValuesStorage.Shrink();
+		UpdateValueArrayView();
+	}
+
+	virtual void Empty() override
+	{
+		FSampleTrackBase::Empty();
+		ValuesStorage.Empty();
+		UpdateValueArrayView();
+	}
+	
+	// returns true if this track is identical to another one
+	virtual bool IsIdentical(const FSampleTrackBase* InOther, float InTolerance = KINDA_SMALL_NUMBER) override
+	{
+		check(InOther);
+
+		if(!FSampleTrackBase::IsIdentical(InOther, InTolerance))
+		{
+			return false;
+		}
+
+		const FSampleTrack<T>* Other = static_cast<const FSampleTrack<T>*>(InOther);
+		if(NumSamples() != Other->NumSamples())
+		{
+			return false;
+		}
+
+		for(int32 SampleIndex = 0; SampleIndex < NumSamples(); SampleIndex++)
+		{
+			if(!Equals(GetValueAtSampleIndex(SampleIndex), Other->GetValueAtSampleIndex(SampleIndex), InTolerance))
+			{
+				return false;
+			}
+		}
+
+		return true;
+	}
+
+	virtual void Serialize(FArchive& InArchive) override
+	{
+		FSampleTrackBase::Serialize(InArchive);
+		InArchive << ValuesStorage;
+		UpdateValueArrayView();
+	}
+
+	const FSampleTrack* GetCastReferencedTrack() const
+	{
+		return static_cast<const FSampleTrack*>(GetReferencedTrack());
+	}
+	
+	FSampleTrack* GetCastReferencedTrack()
+	{
+		return static_cast<FSampleTrack*>(GetReferencedTrack());
+	}
+
+protected:
+
+	/**
+	 * Removes samples from the track which don't present any value anywhere.
+	 * This may only shift time indices and not actually remove any value.
+	 * Data passed in here is slightly redundant, but passed for performance reasons.
+	 * @param InNumTimesToRemove The number of samples to remove
+	 * @param InRemoveTimeIndexArray A bool array with a true value for each time index to remove (.Count(true) == InNumTimesToRemove
+	 * @param InOldTimeIndexToNewTimeIndex An array mapping the time indices from old to new
+	 */
+	virtual int32 RemoveObsoleteTimes(int32 InNumTimesToRemove, const TArray<bool>& InRemoveTimeIndexArray, const TArray<int32>& InOldTimeIndexToNewTimeIndex) override
+	{
+		UpdateArrayViews();
+		
+		// remember the previous indices
+		TArray<int32> OldTimeIndices;
+		if(GetMode() == EMode_Sampled)
+		{
+			OldTimeIndices = TimeIndices;
+		}
+		
+		const int32 NumRemoved = FSampleTrackBase::RemoveObsoleteTimes(InNumTimesToRemove, InRemoveTimeIndexArray, InOldTimeIndexToNewTimeIndex);
+		if(NumRemoved > 0)
+		{
+			TArray<T> NewValues;
+			if(GetMode() == EMode_Complete)
+			{
+				check(NumTimesInContainer == NumSamples());
+
+				NewValues.Reserve(NumSamples() - NumRemoved);
+				for(int32 SampleIndex = 0; SampleIndex < NumSamples(); SampleIndex++)
+				{
+					if(!InRemoveTimeIndexArray[SampleIndex])
+					{
+						NewValues.Add(Values[SampleIndex]);
+					}
+				}
+			}
+			else
+			{
+				TArray<bool> PerSampleRemoveIndex;
+				PerSampleRemoveIndex.Reserve(OldTimeIndices.Num());
+				for(const int32 OldTimeIndex : OldTimeIndices)
+				{
+					PerSampleRemoveIndex.Add(InRemoveTimeIndexArray[OldTimeIndex]);
+				}
+
+				NewValues.Reserve(OldTimeIndices.Num() - NumRemoved);
+				for(int32 Index = 0; Index < OldTimeIndices.Num(); Index++)
+				{
+					if(!PerSampleRemoveIndex[Index])
+					{
+						NewValues.Add(Values[Index]);
+					}
+				}
+			}
+			ValuesStorage = MoveTemp(NewValues);
+			AtlasStorage.Reset();
+			UpdateArrayViews();
+		}
+		return NumRemoved;
+	}
+
+public:
+	
+	static bool IsArrayValue()
+	{
+		return TIsArray<T>::Value;
+	}
+
+	int32 GetSizePerValue() const
+	{
+		return sizeof(T);
+	}
+
+	/**
+	 * Analyses the values stored and introduces the atlas index array,
+	 * an indirection to save memory. After this Values is going to store unique elements,
+	 * and the Atlas array is used to pick the right unique value for each sample index.
+	 */
+	virtual bool AddAtlas(bool bForce = false, float InTolerance = KINDA_SMALL_NUMBER) override
+	{
+		UpdateArrayViews();
+		
+		if(UsesAtlas() || !IsValid() || GetMode() == EMode_Singleton)
+		{
+			return false;
+		}
+
+		TArray<uint32> HashPerValue;
+		TMap<uint32, TTuple<int32, int32>> HashToValueIndex;
+		HashPerValue.Reserve(NumStoredValues());
+		HashToValueIndex.Reserve(NumStoredValues());
+
+		int32 NumNewValues = 0;
+		TArray<int32> OldValueIndexToNewValueIndex;
+		OldValueIndexToNewValueIndex.Reserve(NumStoredValues());
+
+		for(int32 SampleIndex = 0; SampleIndex < NumStoredValues(); SampleIndex++)
+		{
+			const uint32 Hash = GetValueHash(GetValueAtSampleIndex(SampleIndex));
+			HashPerValue.Add(Hash);
+
+			if(const TTuple<int32, int32>* IndexPair = HashToValueIndex.Find(Hash))
+			{
+				if(Equals(GetValueAtSampleIndex(SampleIndex), GetValueAtSampleIndex(IndexPair->Get<0>()), InTolerance))
+				{
+					OldValueIndexToNewValueIndex.Add(IndexPair->Get<1>());
+				}
+				else
+				{
+					OldValueIndexToNewValueIndex.Add(NumNewValues++);
+				}
+			}
+			else
+			{
+				const int32 NewValueIndex =  NumNewValues++;
+				HashToValueIndex.Add(Hash, TTuple<int32, int32>(SampleIndex, NewValueIndex));
+				OldValueIndexToNewValueIndex.Add(NewValueIndex);
+			}
+		}
+
+		// if each value is unique we don't need an atlas
+		if(NumNewValues == NumStoredValues())
+		{
+			return false;
+		}
+
+		if(!bForce)
+		{
+			// if using the atlas will increase the memory footprint - stick with the current representation
+			const int32 SizePerValue = GetSizePerValue();
+			const int32 CurrentSize = SizePerValue * NumStoredValues();
+			const int32 SizeWithAtlas = SizePerValue * HashToValueIndex.Num() + Atlas.GetTypeSize() * HashPerValue.Num();
+			if(CurrentSize <= SizeWithAtlas)
+			{
+				return false;
+			}
+		}
+
+		TArray<T> NewValues;
+		NewValues.Reserve(NumNewValues);
+		TArray<int32> NewAtlas;
+		NewAtlas.Reserve(NumStoredValues());
+
+		for(int32 Index = 0; Index < OldValueIndexToNewValueIndex.Num(); Index++)
+		{
+			const int32& NewValueIndex = OldValueIndexToNewValueIndex[Index];
+			if(!NewValues.IsValidIndex(NewValueIndex))
+			{
+				verify(NewValueIndex == NewValues.Add(GetValueAtSampleIndex(Index)));
+			}
+			NewAtlas.Add(NewValueIndex);
+		}
+
+		ValuesStorage = MoveTemp(NewValues);
+		AtlasStorage = MoveTemp(NewAtlas);
+		UpdateArrayViews();
+		return true;
+	}
+
+	/**
+	 * Removes an existing Atlas indirection by flattening out all unique values into a larger, non-unique Value array
+	 */
+	virtual bool RemoveAtlas() override
+	{
+		UpdateArrayViews();
+
+		if(!UsesAtlas())
+		{
+			return false;
+		}
+
+		TArray<T> NewValues;
+		NewValues.Reserve(NumSamples());
+		for(int32 SampleIndex = 0; SampleIndex < NumSamples(); SampleIndex++)
+		{
+			NewValues.Add(GetValueAtSampleIndex(SampleIndex));
+		}
+
+		ValuesStorage = MoveTemp(NewValues);
+		AtlasStorage.Reset();
+		UpdateArrayViews();
+		return true;
+	}
+
+	/**
+	 * Analyses the memory footprint of a sample track (just complete values vs time indices + values)
+	 * and converts the track back to a complete representation for memory efficiency
+	 */
+	virtual bool ConvertToComplete(bool bForce = false) override
+	{
+		UpdateArrayViews();
+
+		if(IsArrayValue() || GetMode() != EMode_Sampled || IsReferenced())
+		{
+			return false;
+		}
+
+		if(!bForce)
+		{
+			const int32 SizePerValue = GetSizePerValue();
+			const int32 CurrentSize = TimeIndices.GetTypeSize() * TimeIndices.Num() + SizePerValue * NumStoredValues();
+			const int32 CompleteSize = SizePerValue * NumTimesInContainer;
+			if(CurrentSize <= CompleteSize)
+			{
+				return false;
+			}
+		}
+
+		if(UsesAtlas())
+		{
+			TArray<int32> NewAtlas;
+			NewAtlas.Reserve(NumTimesInContainer);
+			int32 SampleIndex = 0;
+			for(int32 TimeIndex = 0; TimeIndex < NumTimesInContainer; TimeIndex++)
+			{
+				NewAtlas.Add(Atlas[GetSampleIndexForTimeIndex(TimeIndex, SampleIndex)]);
+			}
+			AtlasStorage = MoveTemp(NewAtlas);
+		}
+		else
+		{
+			TArray<T> NewValues;
+			NewValues.Reserve(NumTimesInContainer);
+			int32 SampleIndex = 0;
+			for(int32 TimeIndex = 0; TimeIndex < NumTimesInContainer; TimeIndex++)
+			{
+				NewValues.Add(GetValueAtTimeIndex(TimeIndex, SampleIndex));
+			}
+			ValuesStorage = MoveTemp(NewValues);
+		}
+		
+		Mode = EMode_Complete;
+		TimeIndicesStorage.Reset();
+		UpdateArrayViews();
+		return true;
+	}
+
+	/**
+	 * Adds all of the values back onto the track using a sampled representation
+	 */
+	virtual bool ConvertToSampled(bool bForce = false, float InTolerance = KINDA_SMALL_NUMBER) override
+	{
+		if(GetMode() != EMode_Complete || IsReferenced())
+		{
+			return false;
+		}
+
+		UpdateArrayViews();
+
+		if(!bForce)
+		{
+			int32 NumSampledValues = 1;
+			int32 LastAddedIndex = 0;
+			for(int32 SampleIndex = 1; SampleIndex < NumSamples(); SampleIndex++)
+			{
+				if(!Equals(GetValueAtSampleIndex(LastAddedIndex), GetValueAtSampleIndex(SampleIndex), InTolerance))
+				{
+					LastAddedIndex = SampleIndex;
+					NumSampledValues++;
+				}
+			}
+			
+			const int32 SizePerValue = GetSizePerValue();
+			const int32 CompleteSize = SizePerValue * NumTimesInContainer;
+			const int32 SampledSizeSize = TimeIndices.GetTypeSize() * NumSampledValues + SizePerValue * NumSampledValues;
+			if(CompleteSize <= SampledSizeSize)
+			{
+				return false;
+			}
+		}
+
+		RemoveAtlas();
+		
+		TArray<T> CopyOfValues;
+		Swap(CopyOfValues, ValuesStorage);
+		Reset();
+		Mode = EMode_Invalid;
+		UpdateArrayViews();
+
+		// add the values back
+		for(const T& Value : CopyOfValues)
+		{
+			AddSample(Value, InTolerance);
+		}
+
+		Shrink();
+
+		return true;
+	}
+
+	/**
+	 * Localizes values if this track is referenced
+	 */
+	virtual bool LocalizeValues() override
+	{
+		if(FSampleTrackBase::LocalizeValues())
+		{
+			ValuesStorage = Values;
+			UpdateValueArrayView();
+
+			// remove and add the atlas back to optimize the size of the values
+			if(RemoveAtlas())
+			{
+				AddAtlas();
+			}
+			Shrink();
+			return true;
+		}
+		return false;
+	}
+
+protected:
+	
+	virtual int32 AppendValuesFromTrack(const FSampleTrackBase* InOtherTrack) override
+	{
+		const int32 FirstValueIndex = ValuesStorage.Num(); 
+		ValuesStorage.Append(static_cast<const FSampleTrack*>(InOtherTrack)->ValuesStorage);
+		UpdateValueArrayView();
+		return FirstValueIndex;
+	}
+
+	virtual void UpdateArrayViews() override
+	{
+		FSampleTrackBase::UpdateArrayViews();
+		UpdateValueArrayView();
+	}
+
+	void UpdateValueArrayView()
+	{
+		if(IsReferenced())
+		{
+			FSampleTrack* ReferencedTrack = GetCastReferencedTrack();
+			if(ReferencedTrack == nullptr)
+			{
+				Values = TArrayView<T>();
+				return;
+			}
+
+			if(ReferencedValuesRange.Get<1>() <= 0)
+			{
+				Values = TArrayView<T>();
+			}
+			else
+			{
+				check(ReferencedTrack->ValuesStorage.IsValidIndex(ReferencedValuesRange.Get<0>()));
+				check(ReferencedTrack->ValuesStorage.IsValidIndex(ReferencedValuesRange.Get<0>() + ReferencedValuesRange.Get<1>() - 1));
+			
+				Values = TArrayView<T>(
+					ReferencedTrack->ValuesStorage.GetData() + ReferencedValuesRange.Get<0>(),
+					ReferencedValuesRange.Get<1>());
+			}
+		}
+		else
+		{
+			Values = TArrayView<T>(ValuesStorage.GetData(), ValuesStorage.Num());
+		}
+	}
+
+	TArray<T> ValuesStorage;
+	TArrayView<T> Values;
+	
+	friend class FSampleTrackContainer;
+	friend class FControlRigBasicsSampleTrackTest;
+};
+
+template<>
+inline FSampleTrackBase::ETrackType FSampleTrack<bool>::GetTrackType() const
+{
+	return ETrackType_Bool;
+}
+
+template<>
+inline FSampleTrackBase::ETrackType FSampleTrack<int32>::GetTrackType() const
+{
+	return ETrackType_Int32;
+}
+
+template<>
+inline FSampleTrackBase::ETrackType FSampleTrack<uint32>::GetTrackType() const
+{
+	return ETrackType_Uint32;
+}
+
+template<>
+inline FSampleTrackBase::ETrackType FSampleTrack<float>::GetTrackType() const
+{
+	return ETrackType_Float;
+}
+
+template<>
+inline FSampleTrackBase::ETrackType FSampleTrack<FName>::GetTrackType() const
+{
+	return ETrackType_Name;
+}
+
+template<>
+inline FSampleTrackBase::ETrackType FSampleTrack<FString>::GetTrackType() const
+{
+	return ETrackType_String;
+}
+
+template<>
+inline FSampleTrackBase::ETrackType FSampleTrack<FVector3f>::GetTrackType() const
+{
+	return ETrackType_Vector3f;
+}
+
+template<>
+inline FSampleTrackBase::ETrackType FSampleTrack<FQuat4f>::GetTrackType() const
+{
+	return ETrackType_Quatf;
+}
+
+template<>
+inline FSampleTrackBase::ETrackType FSampleTrack<FTransform3f>::GetTrackType() const
+{
+	return ETrackType_Transformf;
+}
+
+template<>
+inline FSampleTrackBase::ETrackType FSampleTrack<FLinearColor>::GetTrackType() const
+{
+	return ETrackType_LinearColor;
+}
+
+template<>
+inline FSampleTrackBase::ETrackType FSampleTrack<FRigElementKey>::GetTrackType() const
+{
+	return ETrackType_ElementKey;
+}
+
+template<>
+inline FSampleTrackBase::ETrackType FSampleTrack<FRigComponentKey>::GetTrackType() const
+{
+	return ETrackType_ComponentKey;
+}
+
+template<>
+inline FSampleTrackBase::ETrackType FSampleTrack<FInstancedStruct>::GetTrackType() const
+{
+	return ETrackType_Struct;
+}
+
+template<>
+inline FSampleTrackBase::ETrackType FSampleTrack<TArray<bool>>::GetTrackType() const
+{
+	return ETrackType_BoolArray;
+}
+
+template<>
+inline FSampleTrackBase::ETrackType FSampleTrack<TArray<int32>>::GetTrackType() const
+{
+	return ETrackType_Int32Array;
+}
+
+template<>
+inline FSampleTrackBase::ETrackType FSampleTrack<TArray<uint32>>::GetTrackType() const
+{
+	return ETrackType_Uint32Array;
+}
+
+template<>
+inline FSampleTrackBase::ETrackType FSampleTrack<TArray<float>>::GetTrackType() const
+{
+	return ETrackType_FloatArray;
+}
+
+template<>
+inline FSampleTrackBase::ETrackType FSampleTrack<TArray<FName>>::GetTrackType() const
+{
+	return ETrackType_NameArray;
+}
+
+template<>
+inline FSampleTrackBase::ETrackType FSampleTrack<TArray<FString>>::GetTrackType() const
+{
+	return ETrackType_StringArray;
+}
+
+template<>
+inline FSampleTrackBase::ETrackType FSampleTrack<TArray<FVector3f>>::GetTrackType() const
+{
+	return ETrackType_Vector3fArray;
+}
+
+template<>
+inline FSampleTrackBase::ETrackType FSampleTrack<TArray<FQuat4f>>::GetTrackType() const
+{
+	return ETrackType_QuatfArray;
+}
+
+template<>
+inline FSampleTrackBase::ETrackType FSampleTrack<TArray<FTransform3f>>::GetTrackType() const
+{
+	return ETrackType_TransformfArray;
+}
+
+template<>
+inline FSampleTrackBase::ETrackType FSampleTrack<TArray<FLinearColor>>::GetTrackType() const
+{
+	return ETrackType_LinearColorArray;
+}
+
+template<>
+inline FSampleTrackBase::ETrackType FSampleTrack<TArray<FRigElementKey>>::GetTrackType() const
+{
+	return ETrackType_ElementKeyArray;
+}
+
+template<>
+inline FSampleTrackBase::ETrackType FSampleTrack<TArray<FRigComponentKey>>::GetTrackType() const
+{
+	return ETrackType_ComponentKeyArray;
+}
+
+template<>
+inline FSampleTrackBase::ETrackType FSampleTrack<TArray<FInstancedStruct>>::GetTrackType() const
+{
+	return ETrackType_StructArray;
+}
+
+template<>
+CONTROLRIG_API void FSampleTrack<FInstancedStruct>::Serialize(FArchive& InArchive);
+
+template<>
+CONTROLRIG_API void FSampleTrack<TArray<FInstancedStruct>>::Serialize(FArchive& InArchive);
+
+template<>
+CONTROLRIG_API int32 FSampleTrack<FInstancedStruct>::GetSizePerValue() const;
diff --git a/Engine/Plugins/Animation/ControlRig/Source/ControlRig/Public/Tracks/SampleTrackArchive.h b/Engine/Plugins/Animation/ControlRig/Source/ControlRig/Public/Tracks/SampleTrackArchive.h
new file mode 100644
--- /dev/null
+++ b/Engine/Plugins/Animation/ControlRig/Source/ControlRig/Public/Tracks/SampleTrackArchive.h
@@ -0,0 +1,68 @@
+// Copyright Epic Games, Inc. All Rights Reserved.
+
+#pragma once
+
+#include "CoreMinimal.h"
+#include "Serialization/MemoryReader.h"
+#include "Serialization/MemoryWriter.h"
+#include "SampleTrackArchive.generated.h"
+
+USTRUCT()
+struct CONTROLRIG_API FSampleTrackMemoryData
+{
+	GENERATED_USTRUCT_BODY()
+
+	UPROPERTY()
+	TArray<uint8> Buffer;
+
+	UPROPERTY()
+	TArray<FName> Names;
+
+	UPROPERTY()
+	TArray<FString> ObjectPaths;
+	
+	void Serialize(FArchive& Ar);
+	friend FArchive& operator<<(FArchive& Ar, FSampleTrackMemoryData& P)
+	{
+		P.Serialize(Ar);
+		return Ar;
+	}
+};
+
+class CONTROLRIG_API FSampleTrackMemoryWriter : public FMemoryWriter
+{
+public:
+	FSampleTrackMemoryWriter() = delete;
+	FSampleTrackMemoryWriter(FSampleTrackMemoryData& InOutData, bool bIsPersistent = false);
+	
+	using FMemoryWriter::operator<<; // For visibility of the overloads we don't override
+
+	virtual FArchive& operator<<(FName& Value) override;
+	virtual FArchive& operator<<(FText& Value) override;
+	virtual FArchive& operator<<(class UObject*& Value) override;
+	virtual FArchive& operator<<(struct FObjectPtr& Value) override;
+
+protected:
+	FSampleTrackMemoryData& Data;
+	TMap<FName,int32> NameToIndex;
+	TMap<UObject*, int32> ObjectToIndex;
+};
+
+class CONTROLRIG_API FSampleTrackMemoryReader : public FMemoryReader
+{
+public:
+	FSampleTrackMemoryReader() = delete;
+	FSampleTrackMemoryReader(FSampleTrackMemoryData& InOutData, bool bIsPersistent = false);
+	
+	using FMemoryReader::operator<<; // For visibility of the overloads we don't override
+	
+	virtual FArchive& operator<<(FName& Value) override;
+	virtual FArchive& operator<<(FText& Value) override;
+	virtual FArchive& operator<<(class UObject*& Value) override;
+	virtual FArchive& operator<<(struct FObjectPtr& Value) override;
+
+protected:
+	
+	FSampleTrackMemoryData& Data;
+	TArray<UObject*> Objects;
+};
diff --git a/Engine/Plugins/Animation/ControlRig/Source/ControlRig/Public/Tracks/SampleTrackBase.h b/Engine/Plugins/Animation/ControlRig/Source/ControlRig/Public/Tracks/SampleTrackBase.h
new file mode 100644
--- /dev/null
+++ b/Engine/Plugins/Animation/ControlRig/Source/ControlRig/Public/Tracks/SampleTrackBase.h
@@ -0,0 +1,417 @@
+// Copyright Epic Games, Inc. All Rights Reserved.
+
+#pragma once
+
+#include "CoreMinimal.h"
+#include "Templates/SharedPointer.h"
+#include "StructUtils/InstancedStruct.h"
+#include "Rigs/RigHierarchyDefines.h"
+
+class FSampleTrackContainer;
+
+class CONTROLRIG_API FSampleTrackBase : public TSharedFromThis<FSampleTrackBase>
+{
+public:
+
+	enum ETrackType : uint8
+	{
+		ETrackType_Unknown,
+		ETrackType_Bool,
+		ETrackType_Int32,
+		ETrackType_Uint32,
+		ETrackType_Float,
+		ETrackType_Name,
+		ETrackType_String,
+		ETrackType_Vector3f,
+		ETrackType_Quatf,
+		ETrackType_Transformf,
+		ETrackType_LinearColor,
+		ETrackType_ElementKey,
+		ETrackType_ComponentKey,
+		ETrackType_Struct,
+		ETrackType_BoolArray,
+		ETrackType_Int32Array,
+		ETrackType_Uint32Array,
+		ETrackType_FloatArray,
+		ETrackType_NameArray,
+		ETrackType_StringArray,
+		ETrackType_Vector3fArray,
+		ETrackType_QuatfArray,
+		ETrackType_TransformfArray,
+		ETrackType_LinearColorArray,
+		ETrackType_ElementKeyArray,
+		ETrackType_ComponentKeyArray,
+		ETrackType_StructArray
+	};
+	
+	enum EMode : uint8
+	{
+		EMode_Invalid, // no data whatsoever
+		EMode_Singleton, // a single, constant value
+		EMode_Sampled, // values at certain time indices, each addressed time index is stored in TimeIndices
+		EMode_Complete, // a value for each time index, TimeIndices is empty
+		EMode_Raw, // unstructured data - used for referencing into from other tracks
+	};
+	
+	FSampleTrackBase();
+	FSampleTrackBase(const FName& InName);
+	virtual ~FSampleTrackBase() {}
+
+	bool IsValid() const { return Mode != EMode_Invalid; }
+
+	// the first name of this track
+	const FName& GetName() const { return Names[0]; }
+
+	// all names this track is using / this track represents
+	const TArray<FName>& GetAllNames() const { return Names; }
+
+	// the index of the track within the container
+	int32 GetTrackIndex() const { return TrackIndex; }
+
+	// the script struct the track is using if it is storing values of FInstancedStruct
+	const UScriptStruct* GetScriptStruct() const { return ScriptStruct; }
+
+	// the type of track
+	virtual ETrackType GetTrackType() const { return ETrackType_Unknown; }
+
+	// the mode in which the track is storing its data
+	EMode GetMode() const { return Mode; }
+
+	// the number of global time indices within the container
+	int32 NumTimes() const;
+
+	// the number of samples stored within this track
+	int32 NumSamples() const;
+
+	// the number of stored values - which can be less than NumSamples if there is an atlas present
+	virtual int32 NumStoredValues() const = 0;
+
+	// returns the sample index within the track given a time index
+	bool GetSampleIndexForTimeIndex(int32 InTimeIndex, int32& InOutSampleIndex) const;
+
+	// returns true if the track stores a value specifically at a given time index
+	bool StoresValueForTimeIndex(int32 InTimeIndex) const;
+
+	// returns true if this track is identical to another one
+	virtual bool IsIdentical(const FSampleTrackBase* InOther, float InTolerance = KINDA_SMALL_NUMBER);
+
+	virtual void Reserve(int32 InSampleCount, int32 InValueCount);
+	virtual void Reset();
+	virtual void Shrink();
+	virtual void Empty();
+	virtual void Serialize(FArchive& InArchive);
+
+	template<typename T>
+	static bool Equals(const T& InA, const T& InB, float InTolerance)
+	{
+		return InA == InB;
+	}
+
+	template<typename T>
+	static bool ArrayEquals(const TArray<T>& A, const TArray<T>& B, float InTolerance)
+	{
+		if(A.Num() != B.Num())
+		{
+			return false;
+		}
+		for(int32 Index = 0; Index < A.Num(); Index++)
+		{
+			if(!Equals<T>(A[Index], B[Index], InTolerance))
+			{
+				return false;
+			}
+		}
+		return true;
+	}
+
+	// returns true if this track is referencing another track
+	bool IsReferenced() const { return ReferencedTrackIndex != INDEX_NONE; }
+
+	// returns the index the track this track is referencing
+	int32 GetReferencedTrackIndex() const { return ReferencedTrackIndex; }
+
+protected:
+
+	const FSampleTrackBase* GetReferencedTrack() const;
+	FSampleTrackBase* GetReferencedTrack();
+	virtual void UpdateArrayViews();
+	void UpdateTimeAndAtlasArrayViews();
+	bool UsesAtlas(int32 InAtlasIndex = INDEX_NONE) const;
+	void AddTimeIndex(bool bOnlyIncreaseUpperBound);
+	virtual int32 AppendValuesFromTrack(const FSampleTrackBase* InOtherTrack) = 0;
+
+	/**
+	 * Removes samples from the track which don't present any value anywhere.
+	 * This may only shift time indices and not actually remove any value.
+	 * Data passed in here is slightly redundant, but passed for performance reasons.
+	 * @param InNumTimesToRemove The number of samples to remove
+	 * @param InRemoveTimeIndexArray A bool array with a true value for each time index to remove (.Count(true) == InNumTimesToRemove
+	 * @param InOldTimeIndexToNewTimeIndex An array mapping the time indices from old to new
+	 */
+	virtual int32 RemoveObsoleteTimes(int32 InNumTimesToRemove, const TArray<bool>& InRemoveTimeIndexArray, const TArray<int32>& InOldTimeIndexToNewTimeIndex);
+
+public:
+	
+	/**
+	 * Analyses the values stored and introduces the atlas index array,
+	 * an indirection to save memory. After this Values is going to store unique elements,
+	 * and the Atlas array is used to pick the right unique value for each sample index.
+	 */
+	virtual bool AddAtlas(bool bForce = false, float InTolerance = KINDA_SMALL_NUMBER) = 0;
+
+	/**
+	 * Removes an existing Atlas indirection by flattening out all unique values into a larger, non-unique Value array
+	 */
+	virtual bool RemoveAtlas() = 0;
+
+	/**
+	 * Analyses the memory footprint of a sample track (just complete values vs time indices + values)
+	 * and converts the track back to a complete representation for memory efficiency
+	 */
+	virtual bool ConvertToComplete(bool bForce = false) = 0;
+
+	/**
+	 * Adds all of the values back onto the track using a sampled representation
+	 */
+	virtual bool ConvertToSampled(bool bForce = false, float InTolerance = KINDA_SMALL_NUMBER) = 0;
+
+	/**
+	 * Localizes values if this track is referenced
+	 */
+	virtual bool LocalizeValues();
+	
+protected:
+
+	struct CIsStruct
+	{
+		template <typename T>
+		auto Requires(const UScriptStruct*& Val) -> decltype(
+			Val = T::StaticStruct()
+		);
+	};
+	
+	struct CIsStructArray
+	{
+		template <typename T>
+		auto Requires(const UScriptStruct*& Val) -> decltype(
+			Val = T::ElementType::StaticStruct()
+		);
+	};
+
+	template <typename T>
+	struct TIsInstancedStruct
+	{
+		enum { Value = false };
+	};
+
+	template <typename T>
+	struct TIsInstancedStructArray
+	{
+		enum { Value = false };
+	};
+
+	template <typename ValueType, typename StructType>
+	struct TIsInstancedStructCompliant
+	{
+		enum
+		{
+			Value =
+				TAnd<
+					TIsInstancedStruct<ValueType>,
+					TAnd<
+						TModels<CIsStruct, StructType>,
+						TNot<TIsInstancedStruct<StructType>>
+					>
+				>::Value
+		};
+	};
+
+	template <typename ValueType, typename StructArrayType>
+	struct TIsInstancedStructArrayCompliant
+	{
+		enum
+		{
+			Value =
+				TAnd<
+					TIsInstancedStructArray<ValueType>,
+					TAnd<
+						TModels<CIsStructArray, StructArrayType>,
+						TNot<TIsInstancedStructArray<StructArrayType>>
+					>
+				>::Value
+		};
+	};
+
+	template<typename T>
+	static uint32 GetSingleValueHash(const T& InValue)
+	{
+		return GetTypeHash(InValue);
+	}
+
+	template<typename T>
+	static uint32 GetValueHash(const T& InValue)
+	{
+		return GetSingleValueHash<T>(InValue);
+	}
+
+	template<typename T>
+	static uint32 GetValueHash(const TArray<T>& InValue)
+	{
+		uint32 Hash = GetTypeHash(InValue.Num());
+		for(const T& Element : InValue)
+		{
+			Hash = HashCombine(GetSingleValueHash<T>(Element));
+		}
+		return Hash;
+	}
+
+	TArray<FName> Names;
+	int32 TrackIndex;
+	int32 ReferencedTrackIndex;
+	TTuple<int32, int32> ReferencedTimeIndicesRange;
+	TTuple<int32, int32> ReferencedAtlasRange;
+	TTuple<int32, int32> ReferencedValuesRange;
+	const UScriptStruct* ScriptStruct;
+	EMode Mode;
+	TArray<int32> TimeIndicesStorage;
+	TArrayView<int32> TimeIndices;
+	TArray<int32> AtlasStorage;
+	TArrayView<int32> Atlas;
+	int32 NumTimesInContainer;
+	FSampleTrackContainer* Container;
+
+	friend class FSampleTrackContainer;
+};
+
+template <> struct FSampleTrackBase::TIsInstancedStruct<FInstancedStruct> { enum { Value = true }; };
+template <> struct FSampleTrackBase::TIsInstancedStructArray<TArray<FInstancedStruct>> { enum { Value = true }; };
+
+template<>
+inline bool FSampleTrackBase::Equals<float>(const float& InA, const float& InB, float InTolerance)
+{
+	return FMath::IsNearlyEqual(InA, InB, InTolerance);
+}
+
+template<>
+inline bool FSampleTrackBase::Equals<FVector3f>(const FVector3f& InA, const FVector3f& InB, float InTolerance)
+{
+	return InA.Equals(InB, InTolerance);
+}
+
+template<>
+inline bool FSampleTrackBase::Equals<FQuat4f>(const FQuat4f& InA, const FQuat4f& InB, float InTolerance)
+{
+	return InA.Equals(InB, InTolerance);
+}
+
+template<>
+inline bool FSampleTrackBase::Equals<FTransform3f>(const FTransform3f& InA, const FTransform3f& InB, float InTolerance)
+{
+	return InA.Equals(InB, InTolerance);
+}
+
+template<>
+inline bool FSampleTrackBase::Equals<FLinearColor>(const FLinearColor& InA, const FLinearColor& InB, float InTolerance)
+{
+	return InA.Equals(InB, InTolerance);
+}
+
+template<>
+inline bool FSampleTrackBase::Equals<FInstancedStruct>(const FInstancedStruct& InA, const FInstancedStruct& InB, float InTolerance)
+{
+	return InA.Identical(&InB, PPF_None);
+}
+
+template<>
+inline bool FSampleTrackBase::Equals<TArray<bool>>(const TArray<bool>& InA, const TArray<bool>& InB, float InTolerance)
+{
+	return ArrayEquals<bool>(InA, InB, InTolerance);
+}
+
+template<>
+inline bool FSampleTrackBase::Equals<TArray<int32>>(const TArray<int32>& InA, const TArray<int32>& InB, float InTolerance)
+{
+	return ArrayEquals<int32>(InA, InB, InTolerance);
+}
+
+template<>
+inline bool FSampleTrackBase::Equals<TArray<uint32>>(const TArray<uint32>& InA, const TArray<uint32>& InB, float InTolerance)
+{
+	return ArrayEquals<uint32>(InA, InB, InTolerance);
+}
+
+template<>
+inline bool FSampleTrackBase::Equals<TArray<float>>(const TArray<float>& InA, const TArray<float>& InB, float InTolerance)
+{
+	return ArrayEquals<float>(InA, InB, InTolerance);
+}
+
+template<>
+inline bool FSampleTrackBase::Equals<TArray<FName>>(const TArray<FName>& InA, const TArray<FName>& InB, float InTolerance)
+{
+	return ArrayEquals<FName>(InA, InB, InTolerance);
+}
+
+template<>
+inline bool FSampleTrackBase::Equals<TArray<FString>>(const TArray<FString>& InA, const TArray<FString>& InB, float InTolerance)
+{
+	return ArrayEquals<FString>(InA, InB, InTolerance);
+}
+
+template<>
+inline bool FSampleTrackBase::Equals<TArray<FVector3f>>(const TArray<FVector3f>& InA, const TArray<FVector3f>& InB, float InTolerance)
+{
+	return ArrayEquals<FVector3f>(InA, InB, InTolerance);
+}
+
+template<>
+inline bool FSampleTrackBase::Equals<TArray<FQuat4f>>(const TArray<FQuat4f>& InA, const TArray<FQuat4f>& InB, float InTolerance)
+{
+	return ArrayEquals<FQuat4f>(InA, InB, InTolerance);
+}
+
+template<>
+inline bool FSampleTrackBase::Equals<TArray<FTransform3f>>(const TArray<FTransform3f>& InA, const TArray<FTransform3f>& InB, float InTolerance)
+{
+	return ArrayEquals<FTransform3f>(InA, InB, InTolerance);
+}
+
+template<>
+inline bool FSampleTrackBase::Equals<TArray<FLinearColor>>(const TArray<FLinearColor>& InA, const TArray<FLinearColor>& InB, float InTolerance)
+{
+	return ArrayEquals<FLinearColor>(InA, InB, InTolerance);
+}
+
+template<>
+inline bool FSampleTrackBase::Equals<TArray<FRigElementKey>>(const TArray<FRigElementKey>& InA, const TArray<FRigElementKey>& InB, float InTolerance)
+{
+	return ArrayEquals<FRigElementKey>(InA, InB, InTolerance);
+}
+
+template<>
+inline bool FSampleTrackBase::Equals<TArray<FRigComponentKey>>(const TArray<FRigComponentKey>& InA, const TArray<FRigComponentKey>& InB, float InTolerance)
+{
+	return ArrayEquals<FRigComponentKey>(InA, InB, InTolerance);
+}
+
+template<>
+inline bool FSampleTrackBase::Equals<TArray<FInstancedStruct>>(const TArray<FInstancedStruct>& InA, const TArray<FInstancedStruct>& InB, float InTolerance)
+{
+	return ArrayEquals<FInstancedStruct>(InA, InB, InTolerance);
+}
+
+template<>
+inline uint32 FSampleTrackBase::GetSingleValueHash<FInstancedStruct>(const FInstancedStruct& InValue)
+{
+	if(InValue.IsValid())
+	{
+		uint32 Hash = GetTypeHash(InValue.GetScriptStruct());
+		const int32 Size = InValue.GetScriptStruct()->GetStructureSize();
+		for(int32 Index = 0; Index < Size; Index++)
+		{
+			Hash = HashCombine(Hash, GetTypeHash(InValue.GetMemory()[Index]));
+		}
+		return Hash;
+	}
+	return 0;
+}
diff --git a/Engine/Plugins/Animation/ControlRig/Source/ControlRig/Public/Tracks/SampleTrackContainer.h b/Engine/Plugins/Animation/ControlRig/Source/ControlRig/Public/Tracks/SampleTrackContainer.h
new file mode 100644
--- /dev/null
+++ b/Engine/Plugins/Animation/ControlRig/Source/ControlRig/Public/Tracks/SampleTrackContainer.h
@@ -0,0 +1,243 @@
+// Copyright Epic Games, Inc. All Rights Reserved.
+
+#pragma once
+
+#include "SampleTrack.h"
+#include "Rigs/RigHierarchyDefines.h"
+#include "SampleTrackContainer.generated.h"
+
+class FSampleTrackContainer;
+
+USTRUCT()
+struct CONTROLRIG_API FSampleTrackHost
+{
+public:
+	GENERATED_BODY();
+
+	FSampleTrackHost();
+	FSampleTrackHost(const FSampleTrackHost& InOther);
+	virtual ~FSampleTrackHost() {}
+
+	FSampleTrackHost& operator =(const FSampleTrackHost& InOther);
+
+	virtual void Serialize(FArchive& InArchive);
+
+	const FSampleTrackContainer* GetContainer() const;
+	FSampleTrackContainer* GetContainer();
+	
+private:
+	TUniquePtr<FSampleTrackContainer> Container;
+};
+
+class CONTROLRIG_API FSampleTrackContainer : public TSharedFromThis<FSampleTrackContainer>
+{
+public:
+	FSampleTrackContainer();
+	~FSampleTrackContainer();
+
+	void Reset();
+	void Shrink();
+	void Compact(float InTolerance = KINDA_SMALL_NUMBER);
+	void Reserve(int32 InNum);
+	void Serialize(FArchive& InArchive);
+	void SetForceToUseCompression(bool InForce = true);
+
+	TSharedPtr<FSampleTrack<bool>> AddBoolTrack(const FName& InName);
+	TSharedPtr<FSampleTrack<int32>> AddInt32Track(const FName& InName);
+	TSharedPtr<FSampleTrack<uint32>> AddUint32Track(const FName& InName);
+	TSharedPtr<FSampleTrack<float>> AddFloatTrack(const FName& InName);
+	TSharedPtr<FSampleTrack<FName>> AddNameTrack(const FName& InName);
+	TSharedPtr<FSampleTrack<FString>> AddStringTrack(const FName& InName);
+	TSharedPtr<FSampleTrack<FVector3f>> AddVectorTrack(const FName& InName);
+	TSharedPtr<FSampleTrack<FQuat4f>> AddQuatTrack(const FName& InName);
+	TSharedPtr<FSampleTrack<FTransform3f>> AddTransformTrack(const FName& InName);
+	TSharedPtr<FSampleTrack<FLinearColor>> AddLinearColorTrack(const FName& InName);
+	TSharedPtr<FSampleTrack<FRigElementKey>> AddRigElementKeyTrack(const FName& InName);
+	TSharedPtr<FSampleTrack<FRigComponentKey>> AddRigComponentKeyTrack(const FName& InName);
+	TSharedPtr<FSampleTrack<FInstancedStruct>> AddStructTrack(const FName& InName, const UScriptStruct* InScriptStruct);
+	TSharedPtr<FSampleTrack<TArray<bool>>> AddBoolArrayTrack(const FName& InName);
+	TSharedPtr<FSampleTrack<TArray<int32>>> AddInt32ArrayTrack(const FName& InName);
+	TSharedPtr<FSampleTrack<TArray<uint32>>> AddUint32ArrayTrack(const FName& InName);
+	TSharedPtr<FSampleTrack<TArray<float>>> AddFloatArrayTrack(const FName& InName);
+	TSharedPtr<FSampleTrack<TArray<FName>>> AddNameArrayTrack(const FName& InName);
+	TSharedPtr<FSampleTrack<TArray<FString>>> AddStringArrayTrack(const FName& InName);
+	TSharedPtr<FSampleTrack<TArray<FVector3f>>> AddVectorArrayTrack(const FName& InName);
+	TSharedPtr<FSampleTrack<TArray<FQuat4f>>> AddQuatArrayTrack(const FName& InName);
+	TSharedPtr<FSampleTrack<TArray<FTransform3f>>> AddTransformArrayTrack(const FName& InName);
+	TSharedPtr<FSampleTrack<TArray<FLinearColor>>> AddLinearColorArrayTrack(const FName& InName);
+	TSharedPtr<FSampleTrack<TArray<FRigElementKey>>> AddRigElementKeyArrayTrack(const FName& InName);
+	TSharedPtr<FSampleTrack<TArray<FRigComponentKey>>> AddRigComponentKeyArrayTrack(const FName& InName);
+	TSharedPtr<FSampleTrack<TArray<FInstancedStruct>>> AddStructArrayTrack(const FName& InName, const UScriptStruct* InScriptStruct);
+
+	int32 NumTracks() const
+	{
+		return Tracks.Num();
+	}
+
+	const TArray<float>& GetTimeSamples() const { return TimeSamples; }
+
+	int32 GetTrackIndex(const FName& InName) const;
+	TSharedPtr<const FSampleTrackBase> GetTrack(int32 InIndex) const;
+	TSharedPtr<FSampleTrackBase> GetTrack(int32 InIndex);
+	TSharedPtr<const FSampleTrackBase> FindTrack(const FName& InName) const;
+	TSharedPtr<FSampleTrackBase> FindTrack(const FName& InName);
+	
+	template<typename T>
+	TSharedPtr<const FSampleTrack<T>> GetTrack(int32 InIndex) const
+	{
+		const TSharedPtr<const FSampleTrackBase> Track = GetTrack(InIndex);
+#if WITH_EDITOR
+		static const FSampleTrack<T> TypeTestTrack;
+		check(TypeTestTrack.GetTrackType() == Track->GetTrackType());
+#endif
+		return StaticCastSharedPtr<const FSampleTrack<T>, const FSampleTrackBase>(Track);
+	}
+
+	template<typename T>
+	TSharedPtr<FSampleTrack<T>> GetTrack(int32 InIndex)
+	{
+		const TSharedPtr<FSampleTrackBase> Track = GetTrack(InIndex);
+#if WITH_EDITOR
+		static const FSampleTrack<T> TypeTestTrack;
+		check(TypeTestTrack.GetTrackType() == Track->GetTrackType());
+#endif
+		return StaticCastSharedPtr<FSampleTrack<T>, FSampleTrackBase>(Track);
+	}
+
+	template<typename T>
+	TSharedPtr<const FSampleTrack<TArray<T>>> GetArrayTrack(int32 InIndex) const
+	{
+		return GetTrack<TArray<T>>(InIndex);
+	}
+
+	template<typename T>
+	TSharedPtr<FSampleTrack<TArray<T>>> GetArrayTrack(int32 InIndex)
+	{
+		return GetTrack<TArray<T>>(InIndex);
+	}
+
+	TSharedPtr<const FSampleTrack<FInstancedStruct>> GetStructTrack(int32 InIndex) const
+	{
+		return GetTrack<FInstancedStruct>(InIndex);
+	}
+
+	TSharedPtr<FSampleTrack<FInstancedStruct>> GetStructTrack(int32 InIndex)
+	{
+		return GetTrack<FInstancedStruct>(InIndex);
+	}
+
+	TSharedPtr<const FSampleTrack<TArray<FInstancedStruct>>> GetStructArrayTrack(int32 InIndex) const
+	{
+		return GetArrayTrack<FInstancedStruct>(InIndex);
+	}
+
+	TSharedPtr<FSampleTrack<TArray<FInstancedStruct>>> GetStructArrayTrack(int32 InIndex)
+	{
+		return GetArrayTrack<FInstancedStruct>(InIndex);
+	}
+
+	template<typename T>
+	TSharedPtr<const FSampleTrack<T>> FindTrack(const FName& InName) const
+	{
+		const int32 Index = GetTrackIndex(InName);
+		if(Tracks.IsValidIndex(Index))
+		{
+			return GetTrack<T>(Index);
+		}
+		return nullptr;
+	}
+
+	template<typename T>
+	TSharedPtr<FSampleTrack<T>> FindTrack(const FName& InName)
+	{
+		const int32 Index = GetTrackIndex(InName);
+		if(Tracks.IsValidIndex(Index))
+		{
+			return GetTrack<T>(Index);
+		}
+		return nullptr;
+	}
+
+	template<typename T>
+	TSharedPtr<const FSampleTrack<TArray<T>>> FindArrayTrack(const FName& InName) const
+	{
+		return FindTrack<TArray<T>>(InName);
+	}
+
+	template<typename T>
+	TSharedPtr<FSampleTrack<TArray<T>>> FindArrayTrack(const FName& InName)
+	{
+		return FindTrack<TArray<T>>(InName);
+	}
+
+	TSharedPtr<const FSampleTrack<FInstancedStruct>> FindStructTrack(const FName& InName) const
+	{
+		return FindTrack<FInstancedStruct>(InName);
+	}
+
+	TSharedPtr<FSampleTrack<FInstancedStruct>> FindStructTrack(const FName& InName)
+	{
+		return FindTrack<FInstancedStruct>(InName);
+	}
+
+	TSharedPtr<const FSampleTrack<TArray<FInstancedStruct>>> FindStructArrayTrack(const FName& InName) const
+	{
+		return FindArrayTrack<FInstancedStruct>(InName);
+	}
+
+	TSharedPtr<FSampleTrack<TArray<FInstancedStruct>>> FindStructArrayTrack(const FName& InName)
+	{
+		return FindArrayTrack<FInstancedStruct>(InName);
+	}
+	
+	int32 AddTimeSample(float InTimeSample);
+
+	// returns true if all tracks are complete / singleton tracks
+	// without references and without atlases
+	bool IsEditable() const;
+
+	// makes this container editable, essentially the opposite of making it compact
+	bool MakeEditable();
+
+private:
+
+	void AddTrack(TSharedPtr<FSampleTrackBase> InTrack);
+	static TSharedPtr<FSampleTrackBase> MakeTrack(FSampleTrackBase::ETrackType InTrackType);
+
+	// removes all tracks with no data in them or nullptr tracks
+	void RemoveInvalidTracks(bool bUpdateNameToIndexMap = true);
+
+	// removes the time indices from all tracks at which we haven't stored any data 
+	void RemoveObsoleteTimes();
+
+	// combines tracks that have the exact same data in them
+	void MergeDuplicateTracks(bool bUpdateNameToIndexMap = true, float InTolerance = KINDA_SMALL_NUMBER);
+
+	// combines tracks of the same type into a single, longer track and
+	// spawns tracks referencing the larger one's sections
+	void MergeTypedTracks(bool bUpdateNameToIndexMap = true, float InTolerance = KINDA_SMALL_NUMBER);
+
+	// introduces a memory optimization to all tracks where necessary, by
+	// which the value storage is moved to a unique store and an atlas index
+	// array is stored (per sample index) to look up the unique value per sample.
+	void EnableTrackAtlas(float InTolerance = KINDA_SMALL_NUMBER);
+
+	/**
+	 * Analyses the memory footprint of a sample track (just complete values vs time indices + values)
+	 * and converts the track back to a complete representation for memory efficiency
+	 */
+	void ConvertTracksToComplete();
+
+	/**
+	 * After a track has been made editable - we can convert it back to sampled to save memory.
+	 */
+	void ConvertTracksToSampled(float InTolerance = KINDA_SMALL_NUMBER);
+
+	// Updates the name lookup after larger changes within the container
+	void UpdateNameToIndexMap();
+
+	TMap<FName, int32> NameToIndex;
+	TArray<float> TimeSamples;
+	TArray<TSharedPtr<FSampleTrackBase>> Tracks;
+	bool bForceToUseCompression;
+};
