ControlRig : Extracted IO Mapping, Pose Adapter, Variable Mapping and Hierarchy Mappings from AnimNode to external helpers, so these can be reused at AnimNext

diff --git a/Engine/Plugins/Animation/ControlRig/Source/ControlRig/Private/AnimNode_ControlRig.cpp b/Engine/Plugins/Animation/ControlRig/Source/ControlRig/Private/AnimNode_ControlRig.cpp
--- a/Engine/Plugins/Animation/ControlRig/Source/ControlRig/Private/AnimNode_ControlRig.cpp
+++ b/Engine/Plugins/Animation/ControlRig/Source/ControlRig/Private/AnimNode_ControlRig.cpp
@@ -14,6 +14,7 @@
 #include "Animation/AnimCurveUtils.h"
 #include "Animation/AnimStats.h"
 #include "Units/Execution/RigUnit_PrepareForExecution.h"
+#include "Tools/ControlRigVariableMappings.h"
 
 #include UE_INLINE_GENERATED_CPP_BY_NAME(AnimNode_ControlRig)
 
@@ -30,7 +31,6 @@ FAnimNode_ControlRig::FAnimNode_ControlRig()
 	, bSetRefPoseFromSkeleton(false)
 	, AlphaCurveName(NAME_None)
 	, LODThreshold(INDEX_NONE)
-	, RefPoseSetterHash()
 {
 }
 
@@ -44,7 +44,7 @@ FAnimNode_ControlRig::~FAnimNode_ControlRig()
 
 void FAnimNode_ControlRig::HandleOnInitialized_AnyThread(URigVMHost*, const FName&)
 {
-	RefPoseSetterHash.Reset();
+	ControlRigHierarchyMappings.ResetRefPoseSetterHash();
 }
 
 void FAnimNode_ControlRig::OnInitializeAnimInstance(const FAnimInstanceProxy* InProxy, const UAnimInstance* InAnimInstance)
@@ -59,7 +59,12 @@ void FAnimNode_ControlRig::OnInitializeAnimInstance(const FAnimInstanceProxy* In
 	
 	if(UpdateControlRigIfNeeded(InAnimInstance, InAnimInstance->GetRequiredBones()))
 	{
-		UpdateControlRigRefPoseIfNeeded(InProxy);
+		ControlRigHierarchyMappings.UpdateControlRigRefPoseIfNeeded(ControlRig
+			, InProxy->GetAnimInstanceObject()
+			, InProxy->GetSkelMeshComponent()
+			, InProxy->GetRequiredBones()
+			, bSetRefPoseFromSkeleton
+			, /*bIncludePoseInHash*/ false);
 	}
 
 	FAnimNode_ControlRigBase::OnInitializeAnimInstance(InProxy, InAnimInstance);
@@ -119,7 +124,13 @@ void FAnimNode_ControlRig::Update_AnyThread(const FAnimationUpdateContext& Conte
 		(void)UpdateControlRigIfNeeded(AnimInstance, Context.AnimInstanceProxy->GetRequiredBones());
 	}
 
-	UpdateControlRigRefPoseIfNeeded(Context.AnimInstanceProxy);
+	ControlRigHierarchyMappings.UpdateControlRigRefPoseIfNeeded(ControlRig
+		, Context.AnimInstanceProxy->GetAnimInstanceObject()
+		, Context.AnimInstanceProxy->GetSkelMeshComponent()
+		, Context.AnimInstanceProxy->GetRequiredBones()
+		, bSetRefPoseFromSkeleton
+		, /*bIncludePoseInHash*/ false);
+
 	FAnimNode_ControlRigBase::Update_AnyThread(Context);
 
 	TRACE_ANIM_NODE_VALUE(Context, TEXT("Class"), *GetNameSafe(ControlRigClass.Get()));
@@ -166,40 +177,12 @@ void FAnimNode_ControlRig::CacheBones_AnyThread(const FAnimationCacheBonesContex
 
 	FAnimNode_ControlRigBase::CacheBones_AnyThread(Context);
 
-	FBoneContainer& RequiredBones = Context.AnimInstanceProxy->GetRequiredBones();
-	InputToControlIndex.Reset();
+	ControlRigVariableMappings.ResetCurvesInputToControlCache();
 
+	FBoneContainer& RequiredBones = Context.AnimInstanceProxy->GetRequiredBones();
 	if(RequiredBones.IsValid())
 	{
-		RefPoseSetterHash.Reset();
-
-		auto CacheMapping = [&](const TMap<FName, FName>& InMapping, FCurveMappings& InCurveMappings,
-			const FAnimationCacheBonesContext& InContext, URigHierarchy* InHierarchy)
-		{
-			for (auto Iter = InMapping.CreateConstIterator(); Iter; ++Iter)
-			{
-				// we need to have list of variables using pin
-				const FName SourcePath = Iter.Key();
-				const FName TargetPath = Iter.Value();
-
-				if (SourcePath != NAME_None && TargetPath != NAME_None)
-				{
-					InCurveMappings.Add(SourcePath, TargetPath);
-
-					if(InHierarchy)
-					{
-						const FRigElementKey Key(TargetPath, ERigElementType::Control);
-						if(const FRigControlElement* ControlElement = InHierarchy->Find<FRigControlElement>(Key))
-						{
-							InputToControlIndex.Add(TargetPath, ControlElement->GetIndex());
-							continue;
-						}
-					}
-				}
-
-				// @todo: should we clear the item if not found?
-			}
-		};
+		ControlRigHierarchyMappings.ResetRefPoseSetterHash();
 
 		URigHierarchy* Hierarchy = nullptr;
 		if(UControlRig* CurrentControlRig = GetControlRig())
@@ -207,8 +190,7 @@ void FAnimNode_ControlRig::CacheBones_AnyThread(const FAnimationCacheBonesContex
 			Hierarchy = CurrentControlRig->GetHierarchy();
 		}
 
-		CacheMapping(InputMapping, InputCurveMappings, Context, Hierarchy);
-		CacheMapping(OutputMapping, OutputCurveMappings, Context, Hierarchy);
+		ControlRigVariableMappings.CacheCurveMappings(InputMapping, OutputMapping, Hierarchy);
 	}
 }
 
@@ -241,23 +223,7 @@ void FAnimNode_ControlRig::UpdateInput(UControlRig* InControlRig, FPoseContext&
 
 	FAnimNode_ControlRigBase::UpdateInput(InControlRig, InOutput);
 
-	// now go through variable mapping table and see if anything is mapping through input
-	if (InputMapping.Num() > 0 && InControlRig)
-	{
-		UE::Anim::FCurveUtils::BulkGet(InOutput.Curve, InputCurveMappings, 
-			[&InControlRig](const FControlRigCurveMapping& InBulkElement, float InValue)
-			{
-				FRigVMExternalVariable Variable = InControlRig->GetPublicVariableByName(InBulkElement.SourceName);
-				if (!Variable.bIsReadOnly && Variable.TypeName == TEXT("float"))
-				{
-					Variable.SetValue<float>(InValue);
-				}
-				else
-				{
-					UE_LOG(LogAnimation, Warning, TEXT("[%s] Missing Input Variable [%s]"), *GetNameSafe(InControlRig->GetClass()), *InBulkElement.SourceName.ToString());
-				}
-			});
-	}
+	ControlRigVariableMappings.UpdateCurveInputs(InControlRig, InputMapping, InOutput.Curve);
 }
 
 void FAnimNode_ControlRig::UpdateOutput(UControlRig* InControlRig, FPoseContext& InOutput)
@@ -266,24 +232,7 @@ void FAnimNode_ControlRig::UpdateOutput(UControlRig* InControlRig, FPoseContext&
 
 	FAnimNode_ControlRigBase::UpdateOutput(InControlRig, InOutput);
 
-	if (OutputMapping.Num() > 0 && InControlRig)
-	{
-		UE::Anim::FCurveUtils::BulkSet(InOutput.Curve, OutputCurveMappings, 
-			[&InControlRig](const FControlRigCurveMapping& InBulkElement) -> float
-			{
-				FRigVMExternalVariable Variable = InControlRig->GetPublicVariableByName(InBulkElement.SourceName);
-				if (Variable.TypeName == TEXT("float"))
-				{
-					return Variable.GetValue<float>();
-				}
-				else
-				{
-					UE_LOG(LogAnimation, Warning, TEXT("[%s] Missing Output Variable [%s]"), *GetNameSafe(InControlRig->GetClass()), *InBulkElement.SourceName.ToString());
-				}
-				
-				return 0.0f;
-			});
-	}
+	ControlRigVariableMappings.UpdateCurveOutputs(InControlRig, OutputMapping, InOutput.Curve);
 }
 
 void FAnimNode_ControlRig::SetControlRigClass(TSubclassOf<UControlRig> InControlRigClass)
@@ -488,441 +437,32 @@ bool FAnimNode_ControlRig::UpdateControlRigIfNeeded(const UAnimInstance* InAnimI
 			ControlRig->Initialize(true);
 			ControlRig->RequestInit();
 		}
-		RefPoseSetterHash.Reset();
+
+		ControlRigHierarchyMappings.ResetRefPoseSetterHash();
+
 		ControlRig->OnInitialized_AnyThread().AddRaw(this, &FAnimNode_ControlRig::HandleOnInitialized_AnyThread);
 
-		UpdateInputOutputMappingIfRequired(ControlRig, InRequiredBones);
+		ControlRigHierarchyMappings.UpdateInputOutputMappingIfRequired(ControlRig
+			, ControlRig->GetHierarchy()
+			, InRequiredBones
+			, InputBonesToTransfer
+			, OutputBonesToTransfer
+			, NodeMappingContainer
+			, bTransferPoseInGlobalSpace
+			, bResetInputPoseToInitial);
 
 		return true;
 	}
 	return false;
 }
 
-void FAnimNode_ControlRig::UpdateControlRigRefPoseIfNeeded(const FAnimInstanceProxy* InProxy, bool bIncludePoseInHash)
-{
-	if(!bSetRefPoseFromSkeleton)
-	{
-		return;
-	}
-
-	int32 ExpectedHash = 0;
-	ExpectedHash = HashCombine(ExpectedHash, (int32)reinterpret_cast<uintptr_t>(InProxy->GetAnimInstanceObject()));
-	ExpectedHash = HashCombine(ExpectedHash, (int32)reinterpret_cast<uintptr_t>(InProxy->GetSkelMeshComponent()));
-
-	if(InProxy->GetSkelMeshComponent())
-	{
-		ExpectedHash = HashCombine(ExpectedHash, (int32)reinterpret_cast<uintptr_t>(InProxy->GetSkelMeshComponent()->GetSkeletalMeshAsset()));
-	}
-
-	if(bIncludePoseInHash)
-	{
-		FMemMark Mark(FMemStack::Get());
-		FCompactPose RefPose;
-		RefPose.ResetToRefPose(InProxy->GetRequiredBones());
-
-		for(const FCompactPoseBoneIndex& BoneIndex : RefPose.ForEachBoneIndex())
-		{
-			const FTransform& Transform = RefPose.GetRefPose(BoneIndex);
-			const FQuat Rotation = Transform.GetRotation();
-
-			ExpectedHash = HashCombine(ExpectedHash, GetTypeHash(Transform.GetTranslation()));
-			ExpectedHash = HashCombine(ExpectedHash, GetTypeHash(FVector4(Rotation.X, Rotation.Y, Rotation.Z, Rotation.W)));
-			ExpectedHash = HashCombine(ExpectedHash, GetTypeHash(Transform.GetScale3D()));
-		}
-
-		if(RefPoseSetterHash.IsSet() && (ExpectedHash == RefPoseSetterHash.GetValue()))
-		{
-			return;
-		}
-
-		ControlRig->SetBoneInitialTransformsFromCompactPose(&RefPose);
-	}
-	else
-	{
-		if(RefPoseSetterHash.IsSet() && (ExpectedHash == RefPoseSetterHash.GetValue()))
-		{
-			return;
-		}
-
-		ControlRig->SetBoneInitialTransformsFromAnimInstanceProxy(InProxy);
-	}
-	
-	RefPoseSetterHash = ExpectedHash;
-}
-
-void FAnimNode_ControlRig::SetIOMapping(bool bInput, const FName& SourceProperty, const FName& TargetCurve)
-{
-	DECLARE_SCOPE_HIERARCHICAL_COUNTER_FUNC()
-
-	UClass* TargetClass = GetTargetClass();
-	if (TargetClass)
-	{
-		UControlRig* CDO = TargetClass->GetDefaultObject<UControlRig>();
-		if (CDO)
-		{
-			TMap<FName, FName>& MappingData = (bInput) ? InputMapping : OutputMapping;
-
-			// if it's valid as of now, we add it
-			bool bIsReadOnly = CDO->GetPublicVariableByName(SourceProperty).bIsReadOnly;
-			if (!bInput || !bIsReadOnly)
-			{
-				if (TargetCurve == NAME_None)
-				{
-					MappingData.Remove(SourceProperty);
-				}
-				else
-				{
-					MappingData.FindOrAdd(SourceProperty) = TargetCurve;
-				}
-			}
-		}
-	}
-}
-
-FName FAnimNode_ControlRig::GetIOMapping(bool bInput, const FName& SourceProperty) const
-{
-	DECLARE_SCOPE_HIERARCHICAL_COUNTER_FUNC()
-
-	const TMap<FName, FName>& MappingData = (bInput) ? InputMapping : OutputMapping;
-	if (const FName* NameFound = MappingData.Find(SourceProperty))
-	{
-		return *NameFound;
-	}
-
-	return NAME_None;
-}
-
-namespace AnimNodeLocals
-{
-
-const TCHAR* UseFunctionsName = TEXT("ControlRig.AnimNode.UseFunctions");
-
-static bool bUseFunctions = false;
-static FAutoConsoleVariableRef CVarUseCachedFunctions(
-   UseFunctionsName,
-   bUseFunctions,
-   TEXT("Cache and use propagate functions.")
-   );
-
-template<typename ValueType, typename HierarchyValueType = ValueType>
-void SetControlValue(FRigControlElement* InControlElement, FProperty* InSourceProperty, URigHierarchy* InTargetHierarchy, const UObject* InSourceInstance)
-{
-	const ValueType* ValuePtr = InSourceProperty->ContainerPtrToValuePtr<ValueType>(InSourceInstance);
-	const FRigControlValue Value = FRigControlValue::Make<HierarchyValueType>(*ValuePtr);
-	InTargetHierarchy->SetControlValue(InControlElement, Value, ERigControlValueType::Current);
-}
-
-template<typename ValueType, typename VariableValueType = ValueType>
-void SetVariableValue(FRigVMExternalVariable& InVariable, FProperty* InSourceProperty, const UObject* InSourceInstance)
-{
-	const ValueType* ValuePtr = InSourceProperty->ContainerPtrToValuePtr<ValueType>(InSourceInstance);
-	InVariable.SetValue<VariableValueType>(*ValuePtr);
-}
-
-}
-
-void FAnimNode_ControlRig::AddUpdateFunction(PropertyUpdateFunction&& InFunction)
-{
-	UpdateFunctions.Add( MoveTemp(InFunction) );
-}
-
-void FAnimNode_ControlRig::AddControlFunction(FRigControlElement* InControlElement, FProperty* InSourceProperty, URigHierarchy* InTargetHierarchy)
-{
-	using namespace AnimNodeLocals;
-	
-	switch(InControlElement->Settings.ControlType)
-	{
-		case ERigControlType::Bool:
-		{
-			if (ensure(CastField<FBoolProperty>(InSourceProperty)))
-			{
-				return AddUpdateFunction([InSourceProperty, InTargetHierarchy, InControlElement](const UObject* InSourceInstance)
-				{
-					SetControlValue<bool>(InControlElement, InSourceProperty, InTargetHierarchy, InSourceInstance);
-				});
-			}
-			break;
-		}
-		case ERigControlType::Float:
-		case ERigControlType::ScaleFloat:
-		{
-			if(ensure(CastField<FFloatProperty>(InSourceProperty)))
-			{
-				return AddUpdateFunction([InSourceProperty, InTargetHierarchy, InControlElement](const UObject* InSourceInstance)
-				{
-					SetControlValue<float>(InControlElement, InSourceProperty, InTargetHierarchy, InSourceInstance);
-				});
-			}
-			break;
-		}
-		case ERigControlType::Integer:
-		{
-			if(ensure(CastField<FIntProperty>(InSourceProperty)))
-			{
-				return AddUpdateFunction([InSourceProperty, InTargetHierarchy, InControlElement](const UObject* InSourceInstance)
-				{
-					SetControlValue<int32>(InControlElement, InSourceProperty, InTargetHierarchy, InSourceInstance);
-				});
-			}
-			break;
-		}
-		case ERigControlType::Vector2D:
-		{
-			FStructProperty* StructProperty = CastField<FStructProperty>(InSourceProperty);
-			if(ensure(StructProperty))
-			{
-				if(ensure(StructProperty->Struct == TBaseStructure<FVector2D>::Get()))
-				{
-					return AddUpdateFunction([InSourceProperty, InTargetHierarchy, InControlElement](const UObject* InSourceInstance)
-					{
-						SetControlValue<FVector2D>(InControlElement, InSourceProperty, InTargetHierarchy, InSourceInstance);
-					});
-				}
-			}
-			break;
-		}
-		case ERigControlType::Position:
-		case ERigControlType::Scale:
-		{
-			FStructProperty* StructProperty = CastField<FStructProperty>(InSourceProperty);
-			if(ensure(StructProperty))
-			{
-				if(ensure(StructProperty->Struct == TBaseStructure<FVector>::Get()))
-				{
-					return AddUpdateFunction([InSourceProperty, InTargetHierarchy, InControlElement](const UObject* InSourceInstance)
-					{
-						SetControlValue<FVector>(InControlElement, InSourceProperty, InTargetHierarchy, InSourceInstance);
-					});
-				}
-			}
-			break;
-		}
-		case ERigControlType::Rotator:
-		{
-			FStructProperty* StructProperty = CastField<FStructProperty>(InSourceProperty);
-			if(ensure(StructProperty))
-			{
-				if(ensure(StructProperty->Struct == TBaseStructure<FRotator>::Get()))
-				{
-					return AddUpdateFunction([InSourceProperty, InTargetHierarchy, InControlElement](const UObject* InSourceInstance)
-					{
-						SetControlValue<FRotator>(InControlElement, InSourceProperty, InTargetHierarchy, InSourceInstance);
-					});
-				}
-			}
-			break;
-		}
-		case ERigControlType::Transform:
-		{
-			FStructProperty* StructProperty = CastField<FStructProperty>(InSourceProperty);
-			if(ensure(StructProperty))
-			{
-				if(ensure(StructProperty->Struct == TBaseStructure<FTransform>::Get()))
-				{
-					return AddUpdateFunction([InSourceProperty, InTargetHierarchy, InControlElement](const UObject* InSourceInstance)
-	                {
-						SetControlValue<FTransform>(InControlElement, InSourceProperty, InTargetHierarchy, InSourceInstance);
-	                });
-				}
-			}
-			break;
-		}
-		case ERigControlType::TransformNoScale:
-		{
-			FStructProperty* StructProperty = CastField<FStructProperty>(InSourceProperty);
-			if(ensure(StructProperty))
-			{
-				if(ensure(StructProperty->Struct == TBaseStructure<FTransform>::Get()))
-				{
-					return AddUpdateFunction([InSourceProperty, InTargetHierarchy, InControlElement](const UObject* InSourceInstance)
-					{
-						SetControlValue<FTransform, FTransformNoScale>(InControlElement, InSourceProperty, InTargetHierarchy, InSourceInstance);
-					});
-				}
-			}
-			break;
-		}
-		case ERigControlType::EulerTransform:
-		{
-			FStructProperty* StructProperty = CastField<FStructProperty>(InSourceProperty);
-			if(ensure(StructProperty))
-			{
-				if(ensure(StructProperty->Struct == TBaseStructure<FTransform>::Get()))
-				{
-					return AddUpdateFunction([InSourceProperty, InTargetHierarchy, InControlElement](const UObject* InSourceInstance)
-					{
-						const FTransform* ValuePtr = InSourceProperty->ContainerPtrToValuePtr<FTransform>(InSourceInstance);
-						const FRigControlValue Value = FRigControlValue::Make<FEulerTransform>(FEulerTransform(*ValuePtr));
-						InTargetHierarchy->SetControlValue(InControlElement, Value, ERigControlValueType::Current);
-					});
-				}
-			}
-			break;
-		}
-		default:
-		{
-			checkNoEntry();
-		}
-	}
-}
-
-void FAnimNode_ControlRig::AddVariableFunction(const FRigVMExternalVariable& InVariable, FProperty* InSourceProperty)
-{
-	using namespace AnimNodeLocals;
-	
-	const int32 VariableIndex = Variables.Add(InVariable);
-	
-	if (CastField<FBoolProperty>(InSourceProperty) && InVariable.TypeName == RigVMTypeUtils::BoolTypeName)
-	{
-		return AddUpdateFunction([InSourceProperty, this, VariableIndex](const UObject* InSourceInstance)
-		{
-			SetVariableValue<bool>(Variables[VariableIndex], InSourceProperty, InSourceInstance);
-		});
-	}
-	else if (CastField<FFloatProperty>(InSourceProperty) && (InVariable.TypeName == RigVMTypeUtils::FloatTypeName || InVariable.TypeName == RigVMTypeUtils::DoubleTypeName))
-	{
-		if(InVariable.TypeName == RigVMTypeUtils::FloatTypeName)
-		{
-			return AddUpdateFunction([InSourceProperty, this, VariableIndex](const UObject* InSourceInstance)
-			{
-				SetVariableValue<float>(Variables[VariableIndex], InSourceProperty, InSourceInstance);
-			});
-		}
-		else
-		{
-			return AddUpdateFunction([InSourceProperty, this, VariableIndex](const UObject* InSourceInstance)
-			{
-				SetVariableValue<float, double>(Variables[VariableIndex], InSourceProperty, InSourceInstance);
-			});
-		}
-	}
-	else if (CastField<FDoubleProperty>(InSourceProperty) && (InVariable.TypeName == RigVMTypeUtils::FloatTypeName || InVariable.TypeName == RigVMTypeUtils::DoubleTypeName))
-	{
-		if(InVariable.TypeName == RigVMTypeUtils::FloatTypeName)
-		{
-			return AddUpdateFunction([InSourceProperty, this, VariableIndex](const UObject* InSourceInstance)
-			{
-				SetVariableValue<double, float>(Variables[VariableIndex], InSourceProperty, InSourceInstance);
-			});
-		}
-		else
-		{
-			return AddUpdateFunction([InSourceProperty, this, VariableIndex](const UObject* InSourceInstance)
-			{
-				SetVariableValue<double>(Variables[VariableIndex], InSourceProperty, InSourceInstance);
-			});
-		}
-	}
-	else if (CastField<FIntProperty>(InSourceProperty) && InVariable.TypeName == RigVMTypeUtils::Int32TypeName)
-	{
-		return AddUpdateFunction([InSourceProperty, this, VariableIndex](const UObject* InSourceInstance)
-		{
-			SetVariableValue<int32>(Variables[VariableIndex], InSourceProperty, InSourceInstance);
-		});
-	}
-	else if (CastField<FNameProperty>(InSourceProperty) && InVariable.TypeName == RigVMTypeUtils::FNameTypeName)
-	{
-		return AddUpdateFunction([InSourceProperty, this, VariableIndex](const UObject* InSourceInstance)
-		{
-			SetVariableValue<FName>(Variables[VariableIndex], InSourceProperty, InSourceInstance);
-		});
-	}
-	else if (CastField<FNameProperty>(InSourceProperty) && InVariable.TypeName == RigVMTypeUtils::FStringTypeName)
-	{
-		return AddUpdateFunction([InSourceProperty, this, VariableIndex](const UObject* InSourceInstance)
-		{
-			SetVariableValue<FString>(Variables[VariableIndex], InSourceProperty, InSourceInstance);
-		});
-	}
-	else if (FStructProperty* StructProperty = CastField<FStructProperty>(InSourceProperty))
-	{
-		return AddUpdateFunction([StructProperty, this, VariableIndex](const UObject* InSourceInstance)
-		{
-			const uint8* SrcPtr = StructProperty->ContainerPtrToValuePtr<uint8>(InSourceInstance);
-			StructProperty->Struct->CopyScriptStruct(Variables[VariableIndex].Memory, SrcPtr, 1);
-		});
-	}
-	else if(FArrayProperty* ArrayProperty = CastField<FArrayProperty>(InSourceProperty))
-	{
-		if(ensure(ArrayProperty->SameType(InVariable.Property)))
-		{
-			return AddUpdateFunction([ArrayProperty, this, VariableIndex](const UObject* InSourceInstance)
-			{
-				const uint8* SrcPtr = ArrayProperty->ContainerPtrToValuePtr<uint8>(InSourceInstance);
-				ArrayProperty->CopyCompleteValue(Variables[VariableIndex].Memory, SrcPtr);
-			});
-		}
-	}
-	else if(FObjectProperty* ObjectProperty = CastField<FObjectProperty>(InSourceProperty))
-	{
-		if(ensure(ObjectProperty->SameType(InVariable.Property)))
-		{
-			return AddUpdateFunction([ObjectProperty, this, VariableIndex](const UObject* InSourceInstance)
-			{
-				const uint8* SrcPtr = ObjectProperty->ContainerPtrToValuePtr<uint8>(InSourceInstance);
-				ObjectProperty->CopyCompleteValue(Variables[VariableIndex].Memory, SrcPtr);
-			});
-		}
-	}
-	else if(FEnumProperty* EnumProperty = CastField<FEnumProperty>(InSourceProperty))
-	{
-		if(ensure(EnumProperty->SameType(InVariable.Property)))
-		{
-			return AddUpdateFunction([EnumProperty, this, VariableIndex](const UObject* InSourceInstance)
-			{
-				const uint8* SrcPtr = EnumProperty->ContainerPtrToValuePtr<uint8>(InSourceInstance);
-				EnumProperty->CopyCompleteValue(Variables[VariableIndex].Memory, SrcPtr);
-			});
-		}
-	}
-	else
-	{
-		ensureMsgf(false, TEXT("Property %s type %s not recognized"), *InSourceProperty->GetName(), *InSourceProperty->GetCPPType());
-	}
-}
-
 void FAnimNode_ControlRig::InitializeProperties(const UObject* InSourceInstance, UClass* InTargetClass)
 {
-	// Build property lists
-	SourceProperties.Reset(SourcePropertyNames.Num());
-	DestProperties.Reset(SourcePropertyNames.Num());
-	UpdateFunctions.Reset(SourcePropertyNames.Num());
-	Variables.Reset();
-
-	check(SourcePropertyNames.Num() == DestPropertyNames.Num());
-
-	UControlRig* TargetControlRig = Cast<UControlRig>(TargetInstance);
-	URigHierarchy* TargetHierarchy = TargetControlRig ? TargetControlRig->GetHierarchy() : nullptr;
+	// Full base override, no Super calling
 
+	// Call helper to setup data
 	const UClass* SourceClass = InSourceInstance->GetClass();
-	for (int32 Idx = 0; Idx < SourcePropertyNames.Num(); ++Idx)
-	{
-		const FName& SourceName = SourcePropertyNames[Idx];
-
-		FProperty* SourceProperty = FindFProperty<FProperty>(SourceClass, SourceName);
-		SourceProperties.Add(SourceProperty);
-		DestProperties.Add(nullptr);
-
-		if (AnimNodeLocals::bUseFunctions)
-        {
-			if (TargetHierarchy && SourceProperty)
-			{
-				const FName& DestPropertyName = DestPropertyNames[Idx];
-				if(FRigControlElement* ControlElement = TargetControlRig->FindControl(DestPropertyName))
-				{
-					AddControlFunction(ControlElement, SourceProperty, TargetHierarchy);
-				}
-				else
-				{
-					const FRigVMExternalVariable Variable = TargetControlRig->GetPublicVariableByName(DestPropertyName);
-					if (Variable.IsValid() && !Variable.bIsReadOnly)
-					{
-						AddVariableFunction(Variable, SourceProperty);
-					}
-				}
-			}
-        }
-	}
+	ControlRigVariableMappings.InitializeProperties(SourceClass, TargetInstance, InTargetClass, SourcePropertyNames, DestPropertyNames);
 }
 
 void FAnimNode_ControlRig::PropagateInputProperties(const UObject* InSourceInstance)
@@ -934,265 +474,14 @@ void FAnimNode_ControlRig::PropagateInputProperties(const UObject* InSourceInsta
 	{
 		return;
 	}
-	
+
 	UControlRig* TargetControlRig = Cast<UControlRig>(TargetInstance);
 	if (!TargetControlRig)
 	{
 		return;
 	}
-	
-	URigHierarchy* TargetHierarchy = TargetControlRig->GetHierarchy();
-	if (!TargetHierarchy)
-	{
-		return;
-	}
-
-	if (AnimNodeLocals::bUseFunctions && !UpdateFunctions.IsEmpty())
-	{
-		Algo::ForEach(UpdateFunctions, [InSourceInstance](const PropertyUpdateFunction& InFunc)
-		{
-			InFunc(InSourceInstance);
-		});
-	}
-	else
-	{
-		PropagateInputPropertiesNoCache(InSourceInstance, TargetControlRig, TargetHierarchy);
-	}
-}
-
-void FAnimNode_ControlRig::PropagateInputPropertiesNoCache(
-	const UObject* InSourceInstance,
-	const UControlRig* InTargetControlRig,
-	URigHierarchy* InTargetHierarchy)
-{
-	check(SourceProperties.Num() == DestProperties.Num());
-	
-	for (int32 PropIdx = 0; PropIdx < SourceProperties.Num(); ++PropIdx)
-	{
-		FProperty* CallerProperty = SourceProperties[PropIdx];
-
-		if(FRigControlElement* ControlElement = InTargetControlRig->FindControl(DestPropertyNames[PropIdx]))
-		{
-			const uint8* SrcPtr = CallerProperty->ContainerPtrToValuePtr<uint8>(InSourceInstance);
-
-			bool bIsValid = false;
-			FRigControlValue Value;
-			switch(ControlElement->Settings.ControlType)
-			{
-				case ERigControlType::Bool:
-				{
-					if(ensure(CastField<FBoolProperty>(CallerProperty)))
-					{
-						Value = FRigControlValue::Make<bool>(*(bool*)SrcPtr);
-						bIsValid = true;
-					}
-					break;
-				}
-				case ERigControlType::Float:
-				case ERigControlType::ScaleFloat:
-				{
-					if(ensure(CastField<FFloatProperty>(CallerProperty)))
-					{
-						Value = FRigControlValue::Make<float>(*(float*)SrcPtr);
-						bIsValid = true;
-					}
-					break;
-				}
-				case ERigControlType::Integer:
-				{
-					if(ensure(CastField<FIntProperty>(CallerProperty)))
-					{
-						Value = FRigControlValue::Make<int32>(*(int32*)SrcPtr);
-						bIsValid = true;
-					}
-					break;
-				}
-				case ERigControlType::Vector2D:
-				{
-					FStructProperty* StructProperty = CastField<FStructProperty>(CallerProperty);
-					if(ensure(StructProperty))
-					{
-						if(ensure(StructProperty->Struct == TBaseStructure<FVector2D>::Get()))
-						{
-							const FVector2D& SrcVector = *(FVector2D*)SrcPtr;
-							Value = FRigControlValue::Make<FVector2D>(SrcVector);
-							bIsValid = true;
-						}
-					}
-					break;
-				}
-				case ERigControlType::Position:
-				case ERigControlType::Scale:
-				{
-					FStructProperty* StructProperty = CastField<FStructProperty>(CallerProperty);
-					if(ensure(StructProperty))
-					{
-						if(ensure(StructProperty->Struct == TBaseStructure<FVector>::Get()))
-						{
-							const FVector& SrcVector = *(FVector*)SrcPtr;  
-							Value = FRigControlValue::Make<FVector>(SrcVector);
-							bIsValid = true;
-						}
-					}
-					break;
-				}
-				case ERigControlType::Rotator:
-				{
-					FStructProperty* StructProperty = CastField<FStructProperty>(CallerProperty);
-					if(ensure(StructProperty))
-					{
-						if(ensure(StructProperty->Struct == TBaseStructure<FRotator>::Get()))
-						{
-							const FRotator& SrcRotator = *(FRotator*)SrcPtr;
-							Value = FRigControlValue::Make<FRotator>(SrcRotator);
-							bIsValid = true;
-						}
-					}
-					break;
-				}
-				case ERigControlType::Transform:
-				{
-					FStructProperty* StructProperty = CastField<FStructProperty>(CallerProperty);
-					if(ensure(StructProperty))
-					{
-						if(ensure(StructProperty->Struct == TBaseStructure<FTransform>::Get()))
-						{
-							const FTransform& SrcTransform = *(FTransform*)SrcPtr;  
-							Value = FRigControlValue::Make<FTransform>(SrcTransform);
-							bIsValid = true;
-						}
-					}
-					break;
-				}
-				case ERigControlType::TransformNoScale:
-				{
-					FStructProperty* StructProperty = CastField<FStructProperty>(CallerProperty);
-					if(ensure(StructProperty))
-					{
-						if(ensure(StructProperty->Struct == TBaseStructure<FTransform>::Get()))
-						{
-							const FTransform& SrcTransform = *(FTransform*)SrcPtr;  
-							Value = FRigControlValue::Make<FTransformNoScale>(SrcTransform);
-							bIsValid = true;
-						}
-					}
-					break;
-				}
-				case ERigControlType::EulerTransform:
-				{
-					FStructProperty* StructProperty = CastField<FStructProperty>(CallerProperty);
-					if(ensure(StructProperty))
-					{
-						if(ensure(StructProperty->Struct == TBaseStructure<FTransform>::Get()))
-						{
-							const FTransform& SrcTransform = *(FTransform*)SrcPtr;  
-							Value = FRigControlValue::Make<FEulerTransform>(FEulerTransform(SrcTransform));
-							bIsValid = true;
-						}
-					}
-					break;
-				}
-				default:
-				{
-					checkNoEntry();
-				}
-			}
 
-			if(bIsValid)
-			{
-				InTargetHierarchy->SetControlValue(ControlElement, Value, ERigControlValueType::Current);
-			}
-			continue;
-		}
-
-		FRigVMExternalVariable Variable = InTargetControlRig->GetPublicVariableByName(DestPropertyNames[PropIdx]);
-		if(Variable.IsValid())
-		{
-			if (Variable.bIsReadOnly)
-			{
-				continue;
-			}
-
-			const uint8* SrcPtr = CallerProperty->ContainerPtrToValuePtr<uint8>(InSourceInstance);
-
-			if (CastField<FBoolProperty>(CallerProperty) != nullptr && Variable.TypeName == RigVMTypeUtils::BoolTypeName)
-			{
-				const bool Value = *(const bool*)SrcPtr;
-				Variable.SetValue<bool>(Value);
-			}
-			else if (CastField<FFloatProperty>(CallerProperty) != nullptr && (Variable.TypeName == RigVMTypeUtils::FloatTypeName ||  Variable.TypeName == RigVMTypeUtils::DoubleTypeName))
-			{
-				const float Value = *(const float*)SrcPtr;
-				if(Variable.TypeName == RigVMTypeUtils::FloatTypeName)
-				{
-					Variable.SetValue<float>(Value);
-				}
-				else
-				{
-					Variable.SetValue<double>(Value);
-				}
-			}
-			else if (CastField<FDoubleProperty>(CallerProperty) != nullptr && (Variable.TypeName == RigVMTypeUtils::FloatTypeName ||  Variable.TypeName == RigVMTypeUtils::DoubleTypeName))
-			{
-				const double Value = *(const double*)SrcPtr;
-				if(Variable.TypeName == RigVMTypeUtils::FloatTypeName)
-				{
-					Variable.SetValue<float>((float)Value);
-				}
-				else
-				{
-					Variable.SetValue<double>(Value);
-				}
-			}
-			else if (CastField<FIntProperty>(CallerProperty) != nullptr && Variable.TypeName == RigVMTypeUtils::Int32TypeName)
-			{
-				const int32 Value = *(const int32*)SrcPtr;
-				Variable.SetValue<int32>(Value);
-			}
-			else if (CastField<FNameProperty>(CallerProperty) != nullptr && Variable.TypeName == RigVMTypeUtils::FNameTypeName)
-			{
-				const FName Value = *(const FName*)SrcPtr;
-				Variable.SetValue<FName>(Value);
-			}
-			else if (CastField<FNameProperty>(CallerProperty) != nullptr && Variable.TypeName == RigVMTypeUtils::FStringTypeName)
-			{
-				const FString Value = *(const FString*)SrcPtr;
-				Variable.SetValue<FString>(Value);
-			}
-			else if (FStructProperty* StructProperty = CastField<FStructProperty>(CallerProperty))
-			{
-				if (StructProperty->Struct == Variable.TypeObject)
-				{
-					StructProperty->Struct->CopyScriptStruct(Variable.Memory, SrcPtr, 1);
-				}
-			}
-			else if(FArrayProperty* ArrayProperty = CastField<FArrayProperty>(CallerProperty))
-			{
-				if(ensure(ArrayProperty->SameType(Variable.Property)))
-				{
-					ArrayProperty->CopyCompleteValue(Variable.Memory, SrcPtr);
-				}
-			}
-			else if(FObjectProperty* ObjectProperty = CastField<FObjectProperty>(CallerProperty))
-			{
-				if(ensure(ObjectProperty->SameType(Variable.Property)))
-				{
-					ObjectProperty->CopyCompleteValue(Variable.Memory, SrcPtr);
-				}
-			}
-			else if(FEnumProperty* EnumProperty = CastField<FEnumProperty>(CallerProperty))
-			{
-				if(ensure(EnumProperty->SameType(Variable.Property)))
-				{
-					EnumProperty->CopyCompleteValue(Variable.Memory, SrcPtr);
-				}
-			}
-			else
-			{
-				ensureMsgf(false, TEXT("Property %s type %s not recognized"), *CallerProperty->GetName(), *CallerProperty->GetCPPType());
-			}
-		}
-	}
+	ControlRigVariableMappings.PropagateInputProperties(InSourceInstance, TargetControlRig, DestPropertyNames);
 }
 
 #if WITH_EDITOR
diff --git a/Engine/Plugins/Animation/ControlRig/Source/ControlRig/Private/AnimNode_ControlRigBase.cpp b/Engine/Plugins/Animation/ControlRig/Source/ControlRig/Private/AnimNode_ControlRigBase.cpp
--- a/Engine/Plugins/Animation/ControlRig/Source/ControlRig/Private/AnimNode_ControlRigBase.cpp
+++ b/Engine/Plugins/Animation/ControlRig/Source/ControlRig/Private/AnimNode_ControlRigBase.cpp
@@ -12,6 +12,8 @@
 #include "Units/Execution/RigUnit_PrepareForExecution.h"
 #include "Algo/Transform.h"
 #include "Animation/AnimCurveUtils.h"
+#include "Tools/ControlRigPoseAdapter.h"
+
 
 #include UE_INLINE_GENERATED_CPP_BY_NAME(AnimNode_ControlRigBase)
 
@@ -27,343 +29,6 @@ TAutoConsoleVariable<int32> CVarAnimNodeControlRigDebug(TEXT("a.AnimNode.Control
 // CVar to disable control rig execution within an anim node
 static TAutoConsoleVariable<int32> CVarControlRigDisableExecutionAnimNode(TEXT("ControlRig.DisableExecutionInAnimNode"), 0, TEXT("if nonzero we disable the execution of Control Rigs inside an anim node."));
 
-// CVar to enable performance optimizations within the anim node pose exchange
-static TAutoConsoleVariable<int32> CVarControlRigEnableAnimNodePerformanceOptimizations(TEXT("ControlRig.EnableAnimNodePerformanceOptimizations"), 0, TEXT("if nonzero we enable the (experimental) execution performance optimizations of Control Rig AnimNodes."));
-
-void FAnimNode_ControlRig_PoseAdapter::PostLinked(URigHierarchy* InHierarchy)
-{
-	FRigHierarchyPoseAdapter::PostLinked(InHierarchy);
-
-	Hierarchy = InHierarchy;
-
-	// 1. sort the hierarchies storage so that we can make sure initial and local is grouped correctly
-	SortHierarchyStorage();
-
-	// 2. make sure to compute all local transforms in initial and local
-	const TArray<FRigTransformElement*> TransformElements = Hierarchy->GetElementsOfType<FRigTransformElement>(false);
-	for(FRigTransformElement* TransformElement : TransformElements)
-	{
-		(void)Hierarchy->GetTransform(TransformElement, ERigTransformType::InitialLocal);
-		(void)Hierarchy->GetTransform(TransformElement, ERigTransformType::CurrentLocal);
-	}
-	
-	// 3. set up a list of dependents in the hierarchy of the rig to be reset to initial on execute
-	TArray<FRigTransformElement*> DependentTransformElements;
-
-	struct Local
-	{
-		static bool ProcessTransformElement(FAnimNode_ControlRig_PoseAdapter* InPoseAdapter, TArray<FRigTransformElement*>& InDependentTransformElements, FRigTransformElement* InTransformElement)
-		{
-			if(InTransformElement == nullptr)
-			{
-				return false;
-			}
-			
-			if(InPoseAdapter->ElementIndexToPoseIndex.Contains((uint16)InTransformElement->GetIndex()))
-			{
-				return true;
-			}
-
-			if(InDependentTransformElements.Contains(InTransformElement))
-			{
-				return true;
-			}
-
-			const FRigBaseElementParentArray& ParentElements = InPoseAdapter->Hierarchy->GetParents(InTransformElement);
-			for(const FRigBaseElement* ParentElement : ParentElements)
-			{
-				if(ProcessTransformElement(InPoseAdapter, InDependentTransformElements, Cast<FRigTransformElement>(const_cast<FRigBaseElement*>(ParentElement))))
-				{
-					InDependentTransformElements.AddUnique(InTransformElement);
-					return true;
-				}
-			}
-
-			return false;
-		}
-	};
-	
-	for(FRigTransformElement* TransformElement : TransformElements)
-	{
-		Local::ProcessTransformElement(this, DependentTransformElements, TransformElement);
-	}
-	
-	Dependents.Reset();
-	Dependents.Reserve(DependentTransformElements.Num() * 3);
-	
-	for(FRigTransformElement* DependentTransformElement : DependentTransformElements)
-	{
-		// skip bones - since they are taking care of by BonesToResetToInitial
-		if(const FRigBoneElement* BoneElement = Cast<FRigBoneElement>(DependentTransformElement))
-		{
-			// skip bones if they are not user defined - and if the parent of the bone is not a dependent as well
-			if(BoneElement->BoneType == ERigBoneType::Imported &&
-				!DependentTransformElements.Contains(BoneElement->ParentElement))
-			{
-				continue;
-			}
-		}
-
-		FRigControlElement* DependentControlElement = Cast<FRigControlElement>(DependentTransformElement);
-		if(DependentControlElement && DependentControlElement->IsAnimationChannel())
-		{
-			continue;
-		}
-		
-		Dependents.Emplace(DependentTransformElement->GetKeyAndIndex(), ERigTransformType::CurrentGlobal, ERigTransformStorageType::Pose, &DependentTransformElement->GetDirtyState().Current);
-		if(DependentControlElement)
-		{
-			Dependents.Emplace(DependentTransformElement->GetKeyAndIndex(), ERigTransformType::CurrentGlobal, ERigTransformStorageType::Offset, &DependentControlElement->GetOffsetDirtyState().Current);
-			Dependents.Emplace(DependentTransformElement->GetKeyAndIndex(), ERigTransformType::CurrentGlobal, ERigTransformStorageType::Shape, &DependentControlElement->GetShapeDirtyState().Current);
-		}
-	}
-
-	// 4. relink the storage for the transforms (local or global or both) and dirty states to our local storage
-	static constexpr bool bLocalIsPrimary = true;
-	UpdateDirtyStates(bLocalIsPrimary);
-	for(int32 PoseIndex = 0; PoseIndex < PoseIndexToElementIndex.Num(); PoseIndex++)
-	{
-		const int32& TransformElementIndex = PoseIndexToElementIndex[PoseIndex];
-		if(TransformElementIndex != INDEX_NONE)
-		{
-			const FRigElementKeyAndIndex& KeyAndIndex = Hierarchy->GetKeyAndIndex(TransformElementIndex);
-			RelinkTransformStorage(KeyAndIndex, ERigTransformType::CurrentLocal, ERigTransformStorageType::Pose, &LocalPose[PoseIndex], &LocalPoseIsDirty[PoseIndex]);
-			RelinkTransformStorage(KeyAndIndex, ERigTransformType::CurrentGlobal, ERigTransformStorageType::Pose, &GlobalPose[PoseIndex], &GlobalPoseIsDirty[PoseIndex]);
-		}
-	}
-	UpdateDirtyStates(bLocalIsPrimary); // do this again to make sure local is flagged as clean and global as dirty
-
-	// 5. Shrink the storage on the hierarchy now that we've relinked it
-	ShrinkHierarchyStorage();
-}
-
-void FAnimNode_ControlRig_PoseAdapter::PreUnlinked(URigHierarchy* InHierarchy)
-{
-	UnlinkTransformStorage();
-	UpdateHierarchyStorage();
-	
-	Hierarchy = nullptr;
-	FRigHierarchyPoseAdapter::PreUnlinked(InHierarchy);
-}
-
-bool FAnimNode_ControlRig_PoseAdapter::IsUpdateToDate(const URigHierarchy* InHierarchy) const
-{
-	if(!FRigHierarchyPoseAdapter::IsUpdateToDate(InHierarchy))
-	{
-		return false;
-	}
-	return !LocalPose.IsEmpty() && !GlobalPose.IsEmpty(); 
-}
-
-void FAnimNode_ControlRig_PoseAdapter::UnlinkTransformStorage()
-{
-	check(Hierarchy);
-
-	for(int32 PoseIndex = 0; PoseIndex < PoseIndexToElementIndex.Num(); PoseIndex++)
-	{
-		const int32& TransformElementIndex = PoseIndexToElementIndex[PoseIndex];
-		if(TransformElementIndex != INDEX_NONE)
-		{
-			const FRigElementKeyAndIndex& KeyAndIndex = Hierarchy->GetKeyAndIndex(TransformElementIndex);
-			RestoreTransformStorage(KeyAndIndex, ERigTransformType::CurrentLocal, ERigTransformStorageType::Pose, false);
-			RestoreTransformStorage(KeyAndIndex, ERigTransformType::CurrentGlobal, ERigTransformStorageType::Pose, false);
-		}
-	}
-
-	ElementIndexToPoseIndex.Reset();
-	PoseIndexToElementIndex.Reset();
-}
-
-void FAnimNode_ControlRig_PoseAdapter::ConvertToLocalPose()
-{
-	check(ParentPoseIndices.Num() == GlobalPose.Num());
-	LocalPose.SetNum(GlobalPose.Num());
-
-	if(Hierarchy)
-	{
-		for(int32 Index = 0; Index < ParentPoseIndices.Num(); Index++)
-		{
-			(void)GetLocalTransform(Index);
-		}
-	}
-}
-void FAnimNode_ControlRig_PoseAdapter::ConvertToGlobalPose()
-{
-	check(ParentPoseIndices.Num() == LocalPose.Num());
-	GlobalPose.SetNum(LocalPose.Num());
-
-	if(Hierarchy)
-	{
-		for(int32 Index = 0; Index < ParentPoseIndices.Num(); Index++)
-		{
-			(void)GetGlobalTransform(Index);
-		}
-	}
-}
-
-const FTransform& FAnimNode_ControlRig_PoseAdapter::GetLocalTransform(int32 InIndex)
-{
-	check(LocalPose.IsValidIndex(InIndex));
-	check(LocalPoseIsDirty.IsValidIndex(InIndex));
-	check(ParentPoseIndices.IsValidIndex(InIndex));
-	check(RequiresHierarchyForSpaceConversion.IsValidIndex(InIndex));
-
-	check(Hierarchy);
-
-	if(LocalPoseIsDirty[InIndex])
-	{
-		check(GlobalPoseIsDirty[InIndex] == false);
-		if(RequiresHierarchyForSpaceConversion[InIndex] && PoseIndexToElementIndex.IsValidIndex(InIndex))
-		{
-			LocalPose[InIndex] = Hierarchy->GetLocalTransformByIndex(PoseIndexToElementIndex[InIndex]);
-		}
-		else
-		{
-			const int32 ParentIndex = ParentPoseIndices[InIndex];
-			if(ParentIndex == INDEX_NONE)
-			{
-				LocalPose[InIndex] = GetGlobalTransform(InIndex);
-			}
-			else
-			{
-				LocalPose[InIndex] = GetGlobalTransform(InIndex).GetRelativeTransform(GetGlobalTransform(ParentIndex));
-			}
-			LocalPose[InIndex].NormalizeRotation();
-		}
-		LocalPoseIsDirty[InIndex] = false;
-	}
-	return LocalPose[InIndex];
-}
-
-const FTransform& FAnimNode_ControlRig_PoseAdapter::GetGlobalTransform(int32 InIndex)
-{
-	check(GlobalPose.IsValidIndex(InIndex));
-	check(GlobalPoseIsDirty.IsValidIndex(InIndex));
-	check(ParentPoseIndices.IsValidIndex(InIndex));
-	check(RequiresHierarchyForSpaceConversion.IsValidIndex(InIndex));
-
-	check(Hierarchy);
-
-	if(GlobalPoseIsDirty[InIndex])
-	{
-		check(LocalPoseIsDirty[InIndex] == false);
-		if(RequiresHierarchyForSpaceConversion[InIndex] && PoseIndexToElementIndex.IsValidIndex(InIndex))
-		{
-			GlobalPose[InIndex] = Hierarchy->GetGlobalTransformByIndex(PoseIndexToElementIndex[InIndex]);
-		}
-		else
-		{
-			const int32 ParentIndex = ParentPoseIndices[InIndex];
-			if(ParentIndex == INDEX_NONE)
-			{
-				GlobalPose[InIndex] = GetLocalTransform(InIndex);
-			}
-			else
-			{
-				GlobalPose[InIndex] = GetLocalTransform(InIndex) * GetGlobalTransform(ParentIndex);
-			}
-			GlobalPose[InIndex].NormalizeRotation();
-		}
-		GlobalPoseIsDirty[InIndex] = false;
-	}
-	return GlobalPose[InIndex];
-}
-
-void FAnimNode_ControlRig_PoseAdapter::UpdateDirtyStates(const TOptional<bool> InLocalIsPrimary)
-{
-	bool bLocalIsPrimary = bTransferInLocalSpace;
-	if(InLocalIsPrimary.IsSet())
-	{
-		bLocalIsPrimary = InLocalIsPrimary.GetValue();
-	}
-	
-	if(LocalPoseIsDirty.Num() != LocalPose.Num())
-	{
-		LocalPoseIsDirty.Reset();
-		LocalPoseIsDirty.AddZeroed(LocalPose.Num());
-	}
-	else if(bLocalIsPrimary)
-	{
-		FMemory::Memzero(LocalPoseIsDirty.GetData(), LocalPoseIsDirty.GetAllocatedSize());
-	}
-	if(!bLocalIsPrimary)
-	{
-		for(bool& Flag : LocalPoseIsDirty)
-		{
-			Flag = true;
-		}
-	}
-
-	if(GlobalPoseIsDirty.Num() != GlobalPose.Num())
-	{
-		GlobalPoseIsDirty.Reset();
-		GlobalPoseIsDirty.AddZeroed(GlobalPose.Num());
-	}
-	else if(!bLocalIsPrimary)
-	{
-		FMemory::Memzero(GlobalPoseIsDirty.GetData(), GlobalPoseIsDirty.GetAllocatedSize());
-	}
-	if(bLocalIsPrimary)
-	{
-		for(bool& Flag : GlobalPoseIsDirty)
-		{
-			Flag = true;
-		}
-	}
-}
-
-void FAnimNode_ControlRig_PoseAdapter::ComputeDependentTransforms()
-{
-	check(Hierarchy);
-
-	// ensure to compute all local transforms
-	for(const FDependentTransform& Dependent : Dependents)
-	{
-		if(FRigTransformElement* TransformElement = Hierarchy->Get<FRigTransformElement>(Dependent.KeyAndIndex.Index))
-		{
-			switch(Dependent.StorageType)
-			{
-				case ERigTransformStorageType::Pose:
-				{
-					(void)Hierarchy->GetTransform(TransformElement, MakeLocal(Dependent.TransformType));
-					break;
-				}
-				case ERigTransformStorageType::Offset:
-				{
-					FRigControlElement* ControlElement = CastChecked<FRigControlElement>(TransformElement);
-					(void)Hierarchy->GetControlOffsetTransform(ControlElement, MakeLocal(Dependent.TransformType));
-					break;
-				}
-				case ERigTransformStorageType::Shape:
-				{
-					FRigControlElement* ControlElement = CastChecked<FRigControlElement>(TransformElement);
-					(void)Hierarchy->GetControlShapeTransform(ControlElement, MakeLocal(Dependent.TransformType));
-					break;
-				}
-				default:
-				{
-					break;
-				}
-			}
-			check(Dependent.DirtyState->Local.Get() == false);
-		}
-	}
-}
-
-void FAnimNode_ControlRig_PoseAdapter::MarkDependentsDirty()
-{
-	ComputeDependentTransforms();
-
-	check(Hierarchy);
-
-	// mark the global dependent as dirty
-	for(const FDependentTransform& Dependent : Dependents)
-	{
-		check(Dependent.DirtyState->Local.Get() == false);
-		Dependent.DirtyState->Global.Set(true);
-	}
-}
-
 FAnimNode_ControlRigBase::FAnimNode_ControlRigBase()
 	: FAnimNode_CustomProperty()
 	, bResetInputPoseToInitial(true) 
@@ -375,7 +40,6 @@ FAnimNode_ControlRigBase::FAnimNode_ControlRigBase()
 	, bExecute(true)
 	, InternalBlendAlpha (1.f)
 	, bControlRigRequiresInitialization(true)
-	, bEnablePoseAdapter(CVarControlRigEnableAnimNodePerformanceOptimizations->GetInt() != 0)
 	, LastBonesSerialNumberForCacheBones(0)
 {
 }
@@ -386,13 +50,9 @@ void FAnimNode_ControlRigBase::OnInitializeAnimInstance(const FAnimInstanceProxy
 
 	FAnimNode_CustomProperty::OnInitializeAnimInstance(InProxy, InAnimInstance);
 	
-	bEnablePoseAdapter = (CVarControlRigEnableAnimNodePerformanceOptimizations->GetInt() != 0);
+	ControlRigHierarchyMappings.InitializeInstance();
 
 	WeakAnimInstanceObject = TWeakObjectPtr<const UAnimInstance>(InAnimInstance);
-	if(bEnablePoseAdapter)
-	{
-		PoseAdapter = MakeShareable(new FAnimNode_ControlRig_PoseAdapter());
-	}
 
 	USkeletalMeshComponent* Component = InAnimInstance->GetOwningComponent();
 	UControlRig* ControlRig = GetControlRig();
@@ -453,7 +113,7 @@ bool FAnimNode_ControlRigBase::CanExecute()
 		return false;
 	}
 
-	if(bEnablePoseAdapter && !PoseAdapter.IsValid())
+	if(!ControlRigHierarchyMappings.CanExecute())
 	{
 		return false;
 	}
@@ -475,212 +135,18 @@ void FAnimNode_ControlRigBase::UpdateInput(UControlRig* ControlRig, FPoseContext
 		return;
 	}
 
-	URigHierarchy* Hierarchy = ControlRig->GetHierarchy();
-	if(Hierarchy == nullptr)
-	{
-		return;
-	}
-
-#if WITH_EDITOR
-	// if we are recording any change - let's clear the undo stack
-	if(bExecute && Hierarchy->IsTracingChanges())
-	{
-		Hierarchy->ResetTransformStack();
-	}
-#endif
-
 	DECLARE_SCOPE_HIERARCHICAL_COUNTER_FUNC()
 
-	if(bEnablePoseAdapter)
-	{
-		if (PoseAdapter && InputSettings.bUpdatePose)
-		{
-			if(InputSettings.bUpdatePose)
-			{
-				// reset transforms here to prevent additive transforms from accumulating to INF
-				// we only update transforms from the mesh pose for bones in the current LOD, 
-				// so the reset here ensures excluded bones are also reset
-				if(!PoseAdapter->BonesToResetToInitial.IsEmpty())
-				{
-					FRigHierarchyValidityBracket ValidityBracket(Hierarchy);
-					{
-#if WITH_EDITOR
-						// make sure transient controls don't get reset
-						UControlRig::FTransientControlPoseScope PoseScope(ControlRig);
-#endif
-						for(const int32& BoneElementIndex : PoseAdapter->BonesToResetToInitial)
-						{
-							if(FRigTransformElement* BoneElement = Hierarchy->Get<FRigBoneElement>(BoneElementIndex))
-							{
-								const FTransform InitialLocalTransform = Hierarchy->GetTransform(BoneElement, ERigTransformType::InitialLocal);
-								BoneElement->GetTransform().Set(ERigTransformType::CurrentLocal, InitialLocalTransform);
-								BoneElement->GetDirtyState().MarkClean(ERigTransformType::CurrentLocal);
-								BoneElement->GetDirtyState().MarkDirty(ERigTransformType::CurrentGlobal);
-							}
-						}
-					}
-				}
-				
-				if(bTransferInputPose && InOutput.Pose.GetNumBones() == PoseAdapter->LocalPose.Num())
-				{
-					PoseAdapter->MarkDependentsDirty();
-					InOutput.Pose.CopyBonesTo(PoseAdapter->LocalPose);
-					PoseAdapter->UpdateDirtyStates(true);
-				}
-			}
-			
-			if(InputSettings.bUpdateCurves || OutputSettings.bUpdateCurves)
-			{
-				Hierarchy->UnsetCurveValues();
-				
-				const TArray<FRigBaseElement*> HierarchyCurves = Hierarchy->GetCurvesFast();
-				if(PoseAdapter->HierarchyCurveLookup.Num() != HierarchyCurves.Num())
-				{
-					PoseAdapter->HierarchyCurveLookup.Reset();
-					PoseAdapter->HierarchyCurveLookup.Reserve(HierarchyCurves.Num());
-					for(int32 Index = 0; Index < HierarchyCurves.Num(); Index++)
-					{
-						PoseAdapter->HierarchyCurveLookup.Add(HierarchyCurves[Index]->GetFName(), Index);
-					}
-				}
-
-				PoseAdapter->PoseCurveToHierarchyCurve.SetNumUninitialized(InOutput.Curve.Num());
-
-				int32 CurveIndex = 0;
-				InOutput.Curve.ForEachElement([this, &CurveIndex, &HierarchyCurves](const UE::Anim::FCurveElement& InCurveElement)
-				{
-					PoseAdapter->PoseCurveToHierarchyCurve[CurveIndex] = INDEX_NONE;
-
-					// the index stored here is the sub index of the curve (the index of the curve within the list of curves)
-					if(const int32* IndexPtr = PoseAdapter->HierarchyCurveLookup.Find(InCurveElement.Name))
-					{
-						const int32& Index = *IndexPtr;
-						FRigCurveElement* HierarchyCurve = CastChecked<FRigCurveElement>(HierarchyCurves[Index]);
-
-						// when setting the curve we need to mark it as "value set", otherwise the copy
-						// pose may reset it to unset - thus we'll loose the value that was just copied in.
-						HierarchyCurve->Set(InCurveElement.Value, true);
-						PoseAdapter->PoseCurveToHierarchyCurve[CurveIndex] = Index;
-					}
-					CurveIndex++;
-				});
-			}
-		}
-	}
-	else
-	{
-		if (InputSettings.bUpdatePose && bTransferInputPose)
-		{
-			const FBoneContainer& RequiredBones = InOutput.Pose.GetBoneContainer();
-
-			// reset transforms here to prevent additive transforms from accumulating to INF
-			// we only update transforms from the mesh pose for bones in the current LOD, 
-			// so the reset here ensures excluded bones are also reset
-			if(!ControlRigBoneInputMappingByName.IsEmpty() || bResetInputPoseToInitial)
-			{
-				FRigHierarchyValidityBracket ValidityBracket(Hierarchy);
-
-				{
-	#if WITH_EDITOR
-					// make sure transient controls don't get reset
-					UControlRig::FTransientControlPoseScope PoseScope(ControlRig);
-	#endif 
-					Hierarchy->ResetPoseToInitial(ERigElementType::Bone);
-				}
-			}
-
-			if(bTransferPoseInGlobalSpace || NodeMappingContainer.IsValid())
-			{
-				// get component pose from control rig
-				FCSPose<FCompactPose> MeshPoses;
-				// first I need to convert to local pose
-				MeshPoses.InitPose(InOutput.Pose);
-
-				if(!ControlRigBoneInputMappingByIndex.IsEmpty())
-				{
-					for (const TPair<uint16, uint16>& Pair : ControlRigBoneInputMappingByIndex)
-					{
-						const uint16 ControlRigIndex = Pair.Key;
-						const uint16 SkeletonIndex = Pair.Value;
-						
-						FCompactPoseBoneIndex CompactPoseIndex(SkeletonIndex);
-						const FTransform& ComponentTransform = MeshPoses.GetComponentSpaceTransform(CompactPoseIndex);
-						Hierarchy->SetGlobalTransformByIndex(ControlRigIndex, ComponentTransform, false);
-					}
-				}
-				else
-				{
-					for (auto Iter = ControlRigBoneInputMappingByName.CreateConstIterator(); Iter; ++Iter)
-					{
-						const FName& Name = Iter.Key();
-						const uint16 Index = Iter.Value();
-						const FRigElementKey Key(Name, ERigElementType::Bone);
-
-						FCompactPoseBoneIndex CompactPoseIndex(Index);
-
-						const FTransform& ComponentTransform = MeshPoses.GetComponentSpaceTransform(CompactPoseIndex);
-						if (NodeMappingContainer.IsValid())
-						{
-							const FTransform& RelativeTransformReverse = NodeMappingContainer->GetSourceToTargetTransform(Name).GetRelativeTransformReverse(ComponentTransform);
-							Hierarchy->SetGlobalTransform(Key, RelativeTransformReverse, false);
-						}
-						else
-						{
-							Hierarchy->SetGlobalTransform(Key, ComponentTransform, false);
-						}
-						
-					}
-				}
-			}
-			else
-			{
-				if(!ControlRigBoneInputMappingByIndex.IsEmpty())
-				{
-					for (const TPair<uint16, uint16>& Pair : ControlRigBoneInputMappingByIndex)
-					{
-						const uint16 ControlRigIndex = Pair.Key;
-						const uint16 SkeletonIndex = Pair.Value;
-						
-						FCompactPoseBoneIndex CompactPoseIndex(SkeletonIndex);
-						const FTransform& LocalTransform = InOutput.Pose[CompactPoseIndex];
-						Hierarchy->SetLocalTransformByIndex(ControlRigIndex, LocalTransform, false);
-					}
-				}
-				else
-				{
-					for (auto Iter = ControlRigBoneInputMappingByName.CreateConstIterator(); Iter; ++Iter)
-					{
-						const FName& Name = Iter.Key();
-						const uint16 SkeletonIndex = Iter.Value();
-						const FRigElementKey Key(Name, ERigElementType::Bone);
-
-						FCompactPoseBoneIndex CompactPoseIndex(SkeletonIndex);
-						const FTransform& LocalTransform = InOutput.Pose[CompactPoseIndex];
-						Hierarchy->SetLocalTransform(Key, LocalTransform, false);
-					}
-				}
-			}
-			
-		}
-
-		if (InputSettings.bUpdateCurves && bTransferInputCurves)
-		{
-			Hierarchy->UnsetCurveValues();
-			
-			InOutput.Curve.ForEachElement([Hierarchy](const UE::Anim::FCurveElement& InCurveElement)
-			{
-				const FRigElementKey Key(InCurveElement.Name, ERigElementType::Curve);
-				Hierarchy->SetCurveValue(Key, InCurveElement.Value);
-			});
-		}		
-	}
-
-#if WITH_EDITOR
-	if(bExecute && Hierarchy->IsTracingChanges())
-	{
-		Hierarchy->StorePoseForTrace(TEXT("FAnimNode_ControlRigBase::UpdateInput"));
-	}
-#endif
+	ControlRigHierarchyMappings.UpdateInput(ControlRig
+		, InOutput
+		, InputSettings
+		, OutputSettings
+		, NodeMappingContainer
+		, bExecute
+		, bTransferInputPose
+		, bResetInputPoseToInitial
+		, bTransferPoseInGlobalSpace
+		, bTransferInputCurves);
 }
 
 void FAnimNode_ControlRigBase::UpdateOutput(UControlRig* ControlRig, FPoseContext& InOutput)
@@ -692,206 +158,14 @@ void FAnimNode_ControlRigBase::UpdateOutput(UControlRig* ControlRig, FPoseContex
 		return;
 	}
 
-	URigHierarchy* Hierarchy = ControlRig->GetHierarchy();
-	if(Hierarchy == nullptr)
-	{
-		return;
-	}
-
 	DECLARE_SCOPE_HIERARCHICAL_COUNTER_FUNC()
 
-	if(bEnablePoseAdapter)
-	{
-		if (OutputSettings.bUpdatePose)
-		{
-			if(PoseAdapter)
-			{
-				// make sure the local respectively global transforms are all ready
-				if(PoseAdapter->bTransferInLocalSpace)
-				{
-					PoseAdapter->ConvertToLocalPose();
-				}
-				else
-				{
-					PoseAdapter->ConvertToGlobalPose();
-				}
-
-				// if we transfered in global - let's convert back to local
-				if(!PoseAdapter->bTransferInLocalSpace)
-				{
-					if(PoseAdapter->GlobalPose.Num() == InOutput.Pose.GetNumBones())
-					{
-						PoseAdapter->ConvertToLocalPose();
-					}
-				}
-				
-				InOutput.Pose.CopyBonesFrom(PoseAdapter->LocalPose);
-			}
-		}
-
-		if(OutputSettings.bUpdateCurves)
-		{
-			const TArray<int32>& ChangedCurveIndices = Hierarchy->GetChangedCurveIndices();
-			if(ChangedCurveIndices.Num() > 0)
-			{
-				const TArray<FRigBaseElement*> HierarchyCurves = Hierarchy->GetCurvesFast();
-
-				int32 CurveIndex = 0;
-				int32 CurvesCopied = 0;
-				InOutput.Curve.ForEachElement([this, &CurveIndex, &CurvesCopied, &HierarchyCurves](const UE::Anim::FCurveElement& InCurveElement)
-				{
-					// the index stored here is the sub index of the curve (the index of the curve within the list of curves)
-					const int32& HierarchyIndex = PoseAdapter->PoseCurveToHierarchyCurve[CurveIndex];
-					if(HierarchyIndex != INDEX_NONE)
-					{
-						const FRigCurveElement* HierarchyCurve = CastChecked<FRigCurveElement>(HierarchyCurves[HierarchyIndex]);
-						if(HierarchyCurve->IsValueSet())
-						{
-							const_cast<UE::Anim::FCurveElement*>(&InCurveElement)->Value = HierarchyCurve->Get();
-						}
-						CurvesCopied++;
-					}
-					CurveIndex++;
-				});
-
-				if(CurvesCopied < HierarchyCurves.Num())
-				{
-					PoseAdapter->HierarchyCurveCopied.SetNumUninitialized(Hierarchy->Num());
-					FMemory::Memzero(PoseAdapter->HierarchyCurveCopied.GetData(), PoseAdapter->HierarchyCurveCopied.GetAllocatedSize());
-					
-					FBlendedCurve ControlRigCurves;
-					ControlRigCurves.Reserve(ChangedCurveIndices.Num());
-					for(const int32& ChangedCurveIndex : ChangedCurveIndices)
-					{
-						if(!PoseAdapter->HierarchyCurveCopied[ChangedCurveIndex])
-						{
-							if(const FRigCurveElement* HierarchyCurve = Hierarchy->Get<FRigCurveElement>(ChangedCurveIndex))
-							{
-								if(HierarchyCurve->IsValueSet())
-								{
-									ControlRigCurves.Add(HierarchyCurve->GetFName(), HierarchyCurve->Get());
-								}
-							}
-							PoseAdapter->HierarchyCurveCopied[ChangedCurveIndex] = true;
-						}
-					}
-
-					InOutput.Curve.Combine(ControlRigCurves);
-				}
-			}
-		}
-	}
-	else
-	{
-		if (OutputSettings.bUpdatePose)
-		{
-			// copy output of the rig
-			const FBoneContainer& RequiredBones = InOutput.Pose.GetBoneContainer();
-
-			TMap<FName, uint16>& NameBasedMapping = ControlRigBoneOutputMappingByName;
-			TArray<TPair<uint16, uint16>>& IndexBasedMapping = ControlRigBoneOutputMappingByIndex;
-
-			// if we don't have a different mapping for outputs, use the input mapping
-			if(NameBasedMapping.IsEmpty() && IndexBasedMapping.IsEmpty())
-			{
-				NameBasedMapping = ControlRigBoneInputMappingByName;
-				IndexBasedMapping = ControlRigBoneInputMappingByIndex;
-			}
-
-			if(bTransferPoseInGlobalSpace || NodeMappingContainer.IsValid())
-			{
-				// get component pose from control rig
-				FCSPose<FCompactPose> MeshPoses;
-				MeshPoses.InitPose(InOutput.Pose);
-
-				if(!IndexBasedMapping.IsEmpty())
-				{
-					for (const TPair<uint16, uint16>& Pair : IndexBasedMapping)
-					{
-						const uint16 ControlRigIndex = Pair.Key;
-						const uint16 SkeletonIndex = Pair.Value;
-
-						FCompactPoseBoneIndex CompactPoseIndex(SkeletonIndex);
-						FTransform ComponentTransform = Hierarchy->GetGlobalTransformByIndex(ControlRigIndex);
-						MeshPoses.SetComponentSpaceTransform(CompactPoseIndex, ComponentTransform);
-					}
-				}
-				else
-				{
-					for (auto Iter = NameBasedMapping.CreateConstIterator(); Iter; ++Iter)
-					{
-						const FName& Name = Iter.Key();
-						const uint16 SkeletonIndex = Iter.Value();
-						const FRigElementKey Key(Name, ERigElementType::Bone);
-
-						FCompactPoseBoneIndex CompactPoseIndex(SkeletonIndex);
-						FTransform ComponentTransform = Hierarchy->GetGlobalTransform(Key);
-						if (NodeMappingContainer.IsValid())
-						{
-							ComponentTransform = NodeMappingContainer->GetSourceToTargetTransform(Name) * ComponentTransform;
-						}
-
-						MeshPoses.SetComponentSpaceTransform(CompactPoseIndex, ComponentTransform);
-					}
-				}
-
-				FCSPose<FCompactPose>::ConvertComponentPosesToLocalPosesSafe(MeshPoses, InOutput.Pose);
-				InOutput.Pose.NormalizeRotations();
-			}
-			else
-			{
-				if(!IndexBasedMapping.IsEmpty())
-				{
-					for (const TPair<uint16, uint16>& Pair : IndexBasedMapping)
-					{
-						const uint16 ControlRigIndex = Pair.Key;
-						const uint16 SkeletonIndex = Pair.Value;
-
-						FCompactPoseBoneIndex CompactPoseIndex(SkeletonIndex);
-						FTransform LocalTransform = Hierarchy->GetLocalTransformByIndex(ControlRigIndex);
-						InOutput.Pose[CompactPoseIndex] = LocalTransform;
-					}
-				}
-				else
-				{
-					for (auto Iter = NameBasedMapping.CreateConstIterator(); Iter; ++Iter)
-					{
-						const FName& Name = Iter.Key();
-						const uint16 Index = Iter.Value();
-						const FRigElementKey Key(Name, ERigElementType::Bone);
-
-						FCompactPoseBoneIndex CompactPoseIndex(Index);
-						FTransform LocalTransform = Hierarchy->GetLocalTransform(Key);
-						InOutput.Pose[CompactPoseIndex] = LocalTransform;
-					}
-				}
-			}
-		}
-
-		if (OutputSettings.bUpdateCurves)
-		{
-			FBlendedCurve ControlRigCurves;
-			ControlRigCurves.Reserve(Hierarchy->Num(ERigElementType::Curve));
-			Hierarchy->ForEach<FRigCurveElement>([&ControlRigCurves](const FRigCurveElement* InElement)
-			{
-				if(InElement->IsValueSet())
-				{
-					ControlRigCurves.Add(InElement->GetFName(), InElement->Get());
-				}
-				return true;
-			});
-
-			InOutput.Curve.Combine(ControlRigCurves);
-		}
-	}
-
-#if WITH_EDITOR
-	if(bExecute && Hierarchy->IsTracingChanges())
-	{
-		Hierarchy->StorePoseForTrace(TEXT("FAnimNode_ControlRigBase::UpdateOutput"));
-		Hierarchy->DumpTransformStackToFile();
-	}
-#endif
+	ControlRigHierarchyMappings.UpdateOutput(ControlRig
+		, InOutput
+		, OutputSettings
+		, NodeMappingContainer
+		, bExecute
+		, bTransferPoseInGlobalSpace);
 }
 
 void FAnimNode_ControlRigBase::Evaluate_AnyThread(FPoseContext& Output)
@@ -910,7 +184,7 @@ void FAnimNode_ControlRigBase::Evaluate_AnyThread(FPoseContext& Output)
 		SourcePose.ResetToRefPose();
 	}
 
-	if(bEnablePoseAdapter && !PoseAdapter.IsValid())
+	if(!ControlRigHierarchyMappings.CheckPoseAdapter())
 	{
 		Output = SourcePose;
 		ensureMsgf(false, TEXT("At this point the pose adapter needs to be set!"));
@@ -973,11 +247,14 @@ void FAnimNode_ControlRigBase::ExecuteControlRig(FPoseContext& InOutput)
 
 		if(bEnablePoseAdapter)
 		{
-			if(!PoseAdapter->IsUpdateToDate(Hierarchy))
+			if(!ControlRigHierarchyMappings.IsUpdateToDate(Hierarchy))
 			{
-				Hierarchy->UnlinkPoseAdapter();
-				UpdateInputOutputMappingIfRequired(ControlRig, InOutput.Pose.GetBoneContainer());
-				Hierarchy->LinkPoseAdapter(PoseAdapter);
+				ControlRigHierarchyMappings.PerformUpdateToDate(ControlRig
+					, Hierarchy
+					, InOutput.Pose.GetBoneContainer()
+					, NodeMappingContainer
+					, bTransferPoseInGlobalSpace
+					, bResetInputPoseToInitial);
 			}
 		}
 
@@ -1107,7 +384,16 @@ void FAnimNode_ControlRigBase::CacheBones_AnyThread(const FAnimationCacheBonesCo
 			}
 		}
 
-		UpdateInputOutputMappingIfRequired(ControlRig, RequiredBones);
+		URigHierarchy* Hierarchy = ControlRig->GetHierarchy();
+
+		ControlRigHierarchyMappings.UpdateInputOutputMappingIfRequired(ControlRig
+			, Hierarchy
+			, RequiredBones
+			, InputBonesToTransfer
+			, OutputBonesToTransfer
+			, NodeMappingContainer
+			, bTransferPoseInGlobalSpace
+			, bResetInputPoseToInitial);
 
 		if(bControlRigRequiresInitialization)
 		{
@@ -1122,379 +408,7 @@ void FAnimNode_ControlRigBase::CacheBones_AnyThread(const FAnimationCacheBonesCo
 		
 		LastBonesSerialNumberForCacheBones = BonesSerialNumber;
 
-		if(bEnablePoseAdapter)
-		{
-			if(URigHierarchy* Hierarchy = ControlRig->GetHierarchy())
-			{
-				if(!PoseAdapter.IsValid())
-				{
-					PoseAdapter = MakeShareable(new FAnimNode_ControlRig_PoseAdapter());
-				}
-				Hierarchy->LinkPoseAdapter(PoseAdapter);
-			}
-		}
-	}
-}
-
-void FAnimNode_ControlRigBase::UpdateInputOutputMappingIfRequired(UControlRig* InControlRig, const FBoneContainer& InRequiredBones)
-{
-	URigHierarchy* Hierarchy = InControlRig->GetHierarchy();
-	if(Hierarchy == nullptr)
-	{
-		return;
-	}
-
-	if(bEnablePoseAdapter)
-	{
-		Hierarchy->UnlinkPoseAdapter();
-
-		if(!PoseAdapter.IsValid())
-		{
-			return;
-		}
-
-		PoseAdapter->ParentPoseIndices.Reset();
-		PoseAdapter->RequiresHierarchyForSpaceConversion.Reset();
-		PoseAdapter->ElementIndexToPoseIndex.Reset();
-
-		const int32 NumBonesInPose = InRequiredBones.GetCompactPoseNumBones();
-
-		PoseAdapter->GlobalPose.Reset();
-		PoseAdapter->LocalPose.Reset();
-		PoseAdapter->ParentPoseIndices.Reserve(NumBonesInPose);
-		PoseAdapter->RequiresHierarchyForSpaceConversion.Reserve(NumBonesInPose);
-		PoseAdapter->GlobalPose.AddDefaulted(NumBonesInPose);
-		PoseAdapter->LocalPose.AddDefaulted(NumBonesInPose);
-		PoseAdapter->HierarchyCurveLookup.Reset();
-
-		PoseAdapter->bTransferInLocalSpace = !(bTransferPoseInGlobalSpace || NodeMappingContainer.IsValid());
-
-		for(int32 Index = 0; Index < NumBonesInPose; Index++)
-		{
-			PoseAdapter->ParentPoseIndices.Add(INDEX_NONE);
-
-			const FCompactPoseBoneIndex ParentBoneIndex = InRequiredBones.GetParentBoneIndex(FCompactPoseBoneIndex(Index));
-			if(ParentBoneIndex.IsValid())
-			{
-				PoseAdapter->ParentPoseIndices[Index] = ParentBoneIndex.GetInt();
-			}
-			PoseAdapter->RequiresHierarchyForSpaceConversion.Add(false);
-		}
-		PoseAdapter->UpdateDirtyStates();
-
-		TArray<int32> MappedBoneElementIndices;
-		if(InRequiredBones.IsValid())
-		{
-			PoseAdapter->ElementIndexToPoseIndex.Reserve(NumBonesInPose);
-			PoseAdapter->PoseIndexToElementIndex.Reserve(NumBonesInPose);
-
-			const FReferenceSkeleton* RefSkeleton = &InRequiredBones.GetReferenceSkeleton();
-			if(const USkeleton* Skeleton = InRequiredBones.GetSkeletonAsset())
-			{
-				RefSkeleton = &Skeleton->GetReferenceSkeleton();
-			}			
-
-			// @todo: thread-safe? probably not in editor, but it may not be a big issue in editor
-			if (NodeMappingContainer.IsValid())
-			{
-				// get target to source mapping table - this is reversed mapping table
-				TMap<FName, FName> TargetToSourceMappingTable;
-				NodeMappingContainer->GetTargetToSourceMappingTable(TargetToSourceMappingTable);
-
-				// now fill up node name
-				for (uint16 Index = 0; Index < NumBonesInPose; ++Index)
-				{
-					// get bone name, and find reverse mapping
-					const FSkeletonPoseBoneIndex BoneIndex = InRequiredBones.GetSkeletonPoseIndexFromCompactPoseIndex(FCompactPoseBoneIndex(Index));
-					if(BoneIndex.IsValid())
-					{
-						FName TargetNodeName = RefSkeleton->GetBoneName(BoneIndex.GetInt());
-						FName* SourceName = TargetToSourceMappingTable.Find(TargetNodeName);
-						if (SourceName)
-						{
-							const int32 ElementIndex = Hierarchy->GetIndex({*SourceName, ERigElementType::Bone});
-							if(ElementIndex != INDEX_NONE)
-							{
-								MappedBoneElementIndices.Add(ElementIndex);
-								PoseAdapter->ElementIndexToPoseIndex.Add(static_cast<uint16>(ElementIndex), Index);
-								PoseAdapter->PoseIndexToElementIndex.Add(ElementIndex);
-								PoseAdapter->LocalPose[Index] = Hierarchy->GetLocalTransform(ElementIndex);
-								PoseAdapter->GlobalPose[Index] = Hierarchy->GetGlobalTransform(ElementIndex);
-								continue;
-							}
-						}
-					}
-					PoseAdapter->PoseIndexToElementIndex.Add(INDEX_NONE);
-				}
-			}
-			else
-			{
-				TArray<FName> NodeNames;
-				TArray<FNodeItem> NodeItems;
-				InControlRig->GetMappableNodeData(NodeNames, NodeItems);
-
-				// even if not mapped, we map only node that exists in the controlrig
-				for (uint16 Index = 0; Index < NumBonesInPose; ++Index)
-				{
-					const FSkeletonPoseBoneIndex BoneIndex = InRequiredBones.GetSkeletonPoseIndexFromCompactPoseIndex(FCompactPoseBoneIndex(Index));
-					if(BoneIndex.IsValid())
-					{
-						const FName& BoneName = RefSkeleton->GetBoneName(BoneIndex.GetInt());
-						if (NodeNames.Contains(BoneName))
-						{
-							const int32 ElementIndex = Hierarchy->GetIndex({BoneName, ERigElementType::Bone});
-							if(ElementIndex != INDEX_NONE)
-							{
-								MappedBoneElementIndices.Add(ElementIndex);
-								PoseAdapter->ElementIndexToPoseIndex.Add(static_cast<uint16>(ElementIndex), Index);
-								PoseAdapter->PoseIndexToElementIndex.Add(ElementIndex);
-								PoseAdapter->LocalPose[Index] = Hierarchy->GetLocalTransform(ElementIndex);
-								PoseAdapter->GlobalPose[Index] = Hierarchy->GetGlobalTransform(ElementIndex);
-								continue;
-							}
-						}
-					}
-					PoseAdapter->PoseIndexToElementIndex.Add(INDEX_NONE);
-				}
-			}
-
-			// once we know all of the bones we are going to transfer - we can check if any of these bones has a different parenting
-			// relationship in the skeleton used in the anim graph vs the hierarchy in the rig. in that case we have to transfer in global
-			if(PoseAdapter->bTransferInLocalSpace)
-			{
-				for(const int32& BoneElementIndex : MappedBoneElementIndices)
-				{
-					const int32 HierarchyParentIndex = Hierarchy->GetFirstParent(BoneElementIndex);
-					const int16 PoseIndex = PoseAdapter->ElementIndexToPoseIndex.FindChecked(BoneElementIndex); 
-					const FCompactPoseBoneIndex CompactPoseParentIndex(PoseAdapter->ParentPoseIndices[PoseIndex]);
-
-					FName HierarchyParentName(NAME_None);
-					FName PoseParentName(NAME_None);
-				
-					if(HierarchyParentIndex != INDEX_NONE)
-					{
-						HierarchyParentName = Hierarchy->Get(HierarchyParentIndex)->GetFName();
-					}
-					if(CompactPoseParentIndex.IsValid())
-					{
-						const FSkeletonPoseBoneIndex SkeletonIndex = InRequiredBones.GetSkeletonPoseIndexFromCompactPoseIndex(CompactPoseParentIndex);
-						PoseParentName = RefSkeleton->GetBoneName(SkeletonIndex.GetInt());
-					}
-
-					if(HierarchyParentName.IsEqual(PoseParentName, ENameCase::CaseSensitive))
-					{
-						continue;
-					}
-				
-					PoseAdapter->RequiresHierarchyForSpaceConversion[PoseIndex] = true;
-					check(PoseAdapter->PoseIndexToElementIndex[PoseIndex] != INDEX_NONE);
-					PoseAdapter->bTransferInLocalSpace = false;
-				}
-			}
-
-			// only reset the full pose if we are not mapping all bones
-			const TArray<FRigBaseElement*>& HierarchyBones = Hierarchy->GetBonesFast();
-			const bool bMapsAllBones = MappedBoneElementIndices.Num() == HierarchyBones.Num();
-			PoseAdapter->BonesToResetToInitial.Reset();
-			PoseAdapter->bRequiresResetPoseToInitial = bResetInputPoseToInitial && !bMapsAllBones;
-
-			if(PoseAdapter->bRequiresResetPoseToInitial)
-			{
-				PoseAdapter->BonesToResetToInitial.Reserve(HierarchyBones.Num() - MappedBoneElementIndices.Num());
-
-				// bone is mapped stores sub indices (bone index within the list of bones)
-				TArray<bool> BoneIsMapped;
-				BoneIsMapped.AddZeroed(HierarchyBones.Num());
-				for(const int32& MappedTransformIndex : MappedBoneElementIndices)
-				{
-					const FRigBaseElement* MappedElement = Hierarchy->Get(MappedTransformIndex);
-					check(MappedElement);
-					BoneIsMapped[MappedElement->GetSubIndex()] = true;
-				}
-
-				// when we want to know which bones to reset we want to convert back to a global index
-				for(int32 UnmappedBoneIndex = 0; UnmappedBoneIndex < BoneIsMapped.Num(); UnmappedBoneIndex++)
-				{
-					if(!BoneIsMapped[UnmappedBoneIndex])
-					{
-						PoseAdapter->BonesToResetToInitial.Add(HierarchyBones[UnmappedBoneIndex]->GetIndex());
-					}
-				}
-			}
-		}
-	}
-	else
-	{
-		ControlRigBoneInputMappingByIndex.Reset();
-		ControlRigBoneOutputMappingByIndex.Reset();
-		ControlRigCurveMappingByIndex.Reset();
-		ControlRigBoneInputMappingByName.Reset();
-		ControlRigBoneOutputMappingByName.Reset();
-		ControlRigCurveMappingByName.Reset();
-
-		if(InRequiredBones.IsValid())
-		{
-			const TArray<FBoneIndexType>& RequiredBonesArray = InRequiredBones.GetBoneIndicesArray();
-			const int32 NumBones = RequiredBonesArray.Num();
-
-			const FReferenceSkeleton& RefSkeleton = InRequiredBones.GetReferenceSkeleton();
-
-			// @todo: thread-safe? probably not in editor, but it may not be a big issue in editor
-			if (NodeMappingContainer.IsValid())
-			{
-				// get target to source mapping table - this is reversed mapping table
-				TMap<FName, FName> TargetToSourceMappingTable;
-				NodeMappingContainer->GetTargetToSourceMappingTable(TargetToSourceMappingTable);
-
-				// now fill up node name
-				for (uint16 Index = 0; Index < NumBones; ++Index)
-				{
-					// get bone name, and find reverse mapping
-					FName TargetNodeName = RefSkeleton.GetBoneName(RequiredBonesArray[Index]);
-					FName* SourceName = TargetToSourceMappingTable.Find(TargetNodeName);
-					if (SourceName)
-					{
-						ControlRigBoneInputMappingByName.Add(*SourceName, Index);
-					}
-				}
-			}
-			else
-			{
-				TArray<FName> NodeNames;
-				TArray<FNodeItem> NodeItems;
-				InControlRig->GetMappableNodeData(NodeNames, NodeItems);
-
-				// even if not mapped, we map only node that exists in the controlrig
-				for (uint16 Index = 0; Index < NumBones; ++Index)
-				{
-					const FName& BoneName = RefSkeleton.GetBoneName(RequiredBonesArray[Index]);
-					if (NodeNames.Contains(BoneName))
-					{
-						ControlRigBoneInputMappingByName.Add(BoneName, Index);
-					}
-				}
-			}
-
-			auto UpdatingMappingFromSpecificTransferList = [] (
-				TArray<FBoneReference>& InTransferList,
-				const TWeakObjectPtr<UNodeMappingContainer>& InMappingContainer,
-				const FBoneContainer& InRequiredBones,
-				const FReferenceSkeleton& InRefSkeleton,
-				const TArray<FBoneIndexType>& InRequiredBonesArray,
-				const UControlRig* InControlRig,
-				TMap<FName, uint16>& OutMapping
-			) {
-				OutMapping.Reset();
-				
-				if (InMappingContainer.IsValid())
-				{
-					// get target to source mapping table - this is reversed mapping table
-					TMap<FName, FName> TargetToSourceMappingTable;
-					InMappingContainer->GetTargetToSourceMappingTable(TargetToSourceMappingTable);
-
-					for(FBoneReference& InputBoneToTransfer : InTransferList)
-					{
-						if(!InputBoneToTransfer.Initialize(InRequiredBones))
-						{
-							continue;
-						}
-						const FName TargetNodeName = InRefSkeleton.GetBoneName(InputBoneToTransfer.BoneIndex);
-						if (const FName* SourceName = TargetToSourceMappingTable.Find(TargetNodeName))
-						{
-							OutMapping.Add(*SourceName, InputBoneToTransfer.BoneIndex);
-						}
-					}
-				}
-				else
-				{
-					TArray<FName> NodeNames;
-					TArray<FNodeItem> NodeItems;
-					InControlRig->GetMappableNodeData(NodeNames, NodeItems);
-
-					for(FBoneReference& InputBoneToTransfer : InTransferList)
-					{
-						if(!InputBoneToTransfer.Initialize(InRequiredBones))
-						{
-							continue;
-						}
-						if (InRequiredBonesArray.IsValidIndex(InputBoneToTransfer.BoneIndex))
-						{
-							const FName& BoneName = InRefSkeleton.GetBoneName(InRequiredBonesArray[InputBoneToTransfer.BoneIndex]);
-							if (NodeNames.Contains(BoneName))
-							{
-								OutMapping.Add(BoneName, InputBoneToTransfer.BoneIndex);
-							}
-						}
-					}
-				}
-			};
-			
-			if(!InputBonesToTransfer.IsEmpty())
-			{
-				ControlRigBoneOutputMappingByName = ControlRigBoneInputMappingByName;
-
-				UpdatingMappingFromSpecificTransferList(
-					InputBonesToTransfer,
-					NodeMappingContainer,
-					InRequiredBones,
-					RefSkeleton,
-					RequiredBonesArray,
-					InControlRig,
-					ControlRigBoneInputMappingByName);
-			}
-
-			if(!OutputBonesToTransfer.IsEmpty())
-			{
-				UpdatingMappingFromSpecificTransferList(
-					OutputBonesToTransfer,
-					NodeMappingContainer,
-					InRequiredBones,
-					RefSkeleton,
-					RequiredBonesArray,
-					InControlRig,
-					ControlRigBoneOutputMappingByName);
-			}
-
-			// check if we can switch the bones to an index based mapping.
-			// we can only do that if there is no node mapping container set.
-			if(!NodeMappingContainer.IsValid())
-			{
-				for(int32 InputOutput = 0; InputOutput < 2; InputOutput++)
-				{
-					bool bIsMappingByIndex = true;
-					TMap<FName, uint16>& NameBasedMapping = InputOutput == 0 ? ControlRigBoneInputMappingByName : ControlRigBoneOutputMappingByName;
-					if(NameBasedMapping.IsEmpty())
-					{
-						continue;
-					}
-					
-					TArray<TPair<uint16, uint16>>& IndexBasedMapping = InputOutput == 0 ? ControlRigBoneInputMappingByIndex : ControlRigBoneOutputMappingByIndex;
-					
-					for (auto Iter = NameBasedMapping.CreateConstIterator(); Iter; ++Iter)
-					{
-						const uint16 SkeletonIndex = Iter.Value();
-						const int32 ControlRigIndex = Hierarchy->GetIndex(FRigElementKey(Iter.Key(), ERigElementType::Bone));
-						if(ControlRigIndex != INDEX_NONE)
-						{
-							IndexBasedMapping.Add(TPair<uint16, uint16>((uint16)ControlRigIndex, SkeletonIndex));
-						}
-						else
-						{
-							bIsMappingByIndex = false;
-						}
-					}
-
-					if(bIsMappingByIndex)
-					{
-						NameBasedMapping.Reset();
-					}
-					else
-					{
-						IndexBasedMapping.Reset();
-					}
-				}
-			}
-		}		
+		ControlRigHierarchyMappings.LinkToHierarchy(Hierarchy);
 	}
 }
 
diff --git a/Engine/Plugins/Animation/ControlRig/Source/ControlRig/Private/ControlRigObjectBinding.cpp b/Engine/Plugins/Animation/ControlRig/Source/ControlRig/Private/ControlRigObjectBinding.cpp
--- a/Engine/Plugins/Animation/ControlRig/Source/ControlRig/Private/ControlRigObjectBinding.cpp
+++ b/Engine/Plugins/Animation/ControlRig/Source/ControlRig/Private/ControlRigObjectBinding.cpp
@@ -69,6 +69,10 @@ UObject* FControlRigObjectBinding::GetBindableObject(UObject* InObject)
 	{
 		return Skeleton;
 	}
+	else if (UActorComponent* ActorComponent = Cast<UActorComponent>(InObject))
+	{
+		return GetBindableObject(ActorComponent->GetOuter());
+	}
 
 	return nullptr;
 }
diff --git a/Engine/Plugins/Animation/ControlRig/Source/ControlRig/Private/Tools/ControlRigHierarchyMappings.cpp b/Engine/Plugins/Animation/ControlRig/Source/ControlRig/Private/Tools/ControlRigHierarchyMappings.cpp
new file mode 100644
--- /dev/null
+++ b/Engine/Plugins/Animation/ControlRig/Source/ControlRig/Private/Tools/ControlRigHierarchyMappings.cpp
@@ -0,0 +1,774 @@
+// Copyright Epic Games, Inc. All Rights Reserved.
+
+#include "Tools/ControlRigHierarchyMappings.h"
+
+#include "ControlRig.h"
+#include "BoneContainer.h"
+#include "Rigs/RigHierarchy.h"
+#include "Animation/NodeMappingContainer.h"
+#include "Animation/AnimNodeBase.h"
+#include "Animation/AnimCurveTypes.h"
+#include "HAL/IConsoleManager.h"
+#include "Tools/ControlRigIOSettings.h"
+#include "Tools/ControlRigPoseAdapter.h"
+#include "Components/SkeletalMeshComponent.h"
+
+// CVar to enable performance optimizations within the anim node pose exchange
+TAutoConsoleVariable<int32> CVarControlRigEnableAnimNodePerformanceOptimizations(TEXT("ControlRig.EnableAnimNodePerformanceOptimizations"), 0, TEXT("if nonzero we enable the (experimental) execution performance optimizations of Control Rig AnimNodes."));
+
+
+void FControlRigHierarchyMappings::InitializeInstance()
+{
+	SetEnablePoseAdapter(CVarControlRigEnableAnimNodePerformanceOptimizations->GetInt() != 0);
+}
+
+void FControlRigHierarchyMappings::LinkToHierarchy(URigHierarchy* InHierarchy)
+{
+	if (bEnablePoseAdapter)
+	{
+		if (InHierarchy)
+		{
+			if (!PoseAdapter.IsValid())
+			{
+				PoseAdapter = MakeShareable(new FControlRigPoseAdapter());
+			}
+			InHierarchy->LinkPoseAdapter(PoseAdapter);
+		}
+	}
+}
+
+void FControlRigHierarchyMappings::UpdateControlRigRefPoseIfNeeded(UControlRig* ControlRig
+	, const UObject* InstanceObject
+	, const USkeletalMeshComponent* SkeletalMeshComponent
+	, const FBoneContainer& RequiredBones
+	, bool bInSetRefPoseFromSkeleton
+	, bool bIncludePoseInHash)
+{
+	check(ControlRig);
+
+	if (!bInSetRefPoseFromSkeleton)
+	{
+		return;
+	}
+
+	int32 ExpectedHash = 0;
+	ExpectedHash = HashCombine(ExpectedHash, (int32)reinterpret_cast<uintptr_t>(InstanceObject));
+	ExpectedHash = HashCombine(ExpectedHash, (int32)reinterpret_cast<uintptr_t>(SkeletalMeshComponent));
+
+	if (SkeletalMeshComponent)
+	{
+		ExpectedHash = HashCombine(ExpectedHash, (int32)reinterpret_cast<uintptr_t>(SkeletalMeshComponent->GetSkeletalMeshAsset()));
+	}
+
+	FMemMark Mark(FMemStack::Get());
+	FCompactPose RefPose;
+	RefPose.ResetToRefPose(RequiredBones);
+
+	if (bIncludePoseInHash)
+	{
+		for (const FCompactPoseBoneIndex& BoneIndex : RefPose.ForEachBoneIndex())
+		{
+			const FTransform& Transform = RefPose.GetRefPose(BoneIndex);
+			const FQuat Rotation = Transform.GetRotation();
+
+			ExpectedHash = HashCombine(ExpectedHash, GetTypeHash(Transform.GetTranslation()));
+			ExpectedHash = HashCombine(ExpectedHash, GetTypeHash(FVector4(Rotation.X, Rotation.Y, Rotation.Z, Rotation.W)));
+			ExpectedHash = HashCombine(ExpectedHash, GetTypeHash(Transform.GetScale3D()));
+		}
+	}
+
+	if (RefPoseSetterHash.IsSet() && (ExpectedHash == RefPoseSetterHash.GetValue()))
+	{
+		return;
+	}
+	ControlRig->SetBoneInitialTransformsFromCompactPose(&RefPose);
+
+	RefPoseSetterHash = ExpectedHash;
+}
+
+void FControlRigHierarchyMappings::UpdateInputOutputMappingIfRequired(UControlRig* InControlRig
+	, URigHierarchy* InHierarchy
+	, const FBoneContainer& InRequiredBones
+	, TArray<FBoneReference>& InInputBonesToTransfer
+	, TArray<FBoneReference>& InOutputBonesToTransfer
+	, const TWeakObjectPtr<UNodeMappingContainer>& InNodeMappingContainer
+	, bool bInTransferPoseInGlobalSpace
+	, bool bInResetInputPoseToInitial)
+{
+	DECLARE_SCOPE_HIERARCHICAL_COUNTER_FUNC()
+
+	if (InHierarchy == nullptr)
+	{
+		return;
+	}
+
+	if (bEnablePoseAdapter)
+	{
+		if (!PoseAdapter.IsValid())
+		{
+			return;
+		}
+
+		PoseAdapter->UpdateInputOutputMappingIfRequired(InControlRig
+			, InHierarchy
+			, InRequiredBones
+			, InNodeMappingContainer
+			, bInTransferPoseInGlobalSpace
+			, bInResetInputPoseToInitial);
+	}
+	else
+	{
+		UpdateInputOutputMappingIfRequiredImpl(InControlRig
+			, InHierarchy
+			, InRequiredBones
+			, InInputBonesToTransfer
+			, InOutputBonesToTransfer
+			, InNodeMappingContainer);
+	}
+}
+
+void FControlRigHierarchyMappings::UpdateInput(UControlRig* ControlRig
+	, FPoseContext& InOutput
+	, const FControlRigIOSettings& InInputSettings
+	, const FControlRigIOSettings& InOutputSettings
+	, TWeakObjectPtr<UNodeMappingContainer>& InNodeMappingContainer
+	, bool bInExecute
+	, bool bInTransferInputPose
+	, bool bInResetInputPoseToInitial
+	, bool bInTransferPoseInGlobalSpace
+	, bool bInTransferInputCurves)
+{
+	DECLARE_SCOPE_HIERARCHICAL_COUNTER_FUNC()
+
+	URigHierarchy* Hierarchy = ControlRig->GetHierarchy();
+	if (Hierarchy == nullptr)
+	{
+		return;
+	}
+
+#if WITH_EDITOR
+	// if we are recording any change - let's clear the undo stack
+	if (bInExecute && Hierarchy->IsTracingChanges())
+	{
+		Hierarchy->ResetTransformStack();
+	}
+#endif
+
+	if (bEnablePoseAdapter)
+	{
+		if (PoseAdapter && InInputSettings.bUpdatePose)
+		{
+			if (InInputSettings.bUpdatePose)
+			{
+				// reset transforms here to prevent additive transforms from accumulating to INF
+				// we only update transforms from the mesh pose for bones in the current LOD, 
+				// so the reset here ensures excluded bones are also reset
+				if (!PoseAdapter->GetBonesToResetToInitial().IsEmpty())
+				{
+					FRigHierarchyValidityBracket ValidityBracket(Hierarchy);
+					{
+#if WITH_EDITOR
+						// make sure transient controls don't get reset
+						UControlRig::FTransientControlPoseScope PoseScope(ControlRig);
+#endif
+						for (const int32& BoneElementIndex : PoseAdapter->GetBonesToResetToInitial())
+						{
+							if (FRigTransformElement* BoneElement = Hierarchy->Get<FRigBoneElement>(BoneElementIndex))
+							{
+								const FTransform InitialLocalTransform = Hierarchy->GetTransform(BoneElement, ERigTransformType::InitialLocal);
+								BoneElement->GetTransform().Set(ERigTransformType::CurrentLocal, InitialLocalTransform);
+								BoneElement->GetDirtyState().MarkClean(ERigTransformType::CurrentLocal);
+								BoneElement->GetDirtyState().MarkDirty(ERigTransformType::CurrentGlobal);
+							}
+						}
+					}
+				}
+
+				if (bInTransferInputPose && InOutput.Pose.GetNumBones() == PoseAdapter->GetLocalPose().Num())
+				{
+					PoseAdapter->MarkDependentsDirty();
+					PoseAdapter->CopyBonesFrom(InOutput.Pose);
+					PoseAdapter->UpdateDirtyStates(true);
+				}
+			}
+
+			if (InInputSettings.bUpdateCurves || InOutputSettings.bUpdateCurves)
+			{
+				Hierarchy->UnsetCurveValues();
+
+				const TArray<FRigBaseElement*> HierarchyCurves = Hierarchy->GetCurvesFast();
+				PoseAdapter->SetHierarchyCurvesLookup(HierarchyCurves);
+				PoseAdapter->SetPoseCurveToHierarchyCurve(HierarchyCurves, InOutput.Curve);
+			}
+		}
+	}
+	else
+	{
+		if (InInputSettings.bUpdatePose && bInTransferInputPose)
+		{
+			const FBoneContainer& RequiredBones = InOutput.Pose.GetBoneContainer();
+
+			// reset transforms here to prevent additive transforms from accumulating to INF
+			// we only update transforms from the mesh pose for bones in the current LOD, 
+			// so the reset here ensures excluded bones are also reset
+			if (!GetControlRigBoneInputMappingByName().IsEmpty() || bInResetInputPoseToInitial)
+			{
+				FRigHierarchyValidityBracket ValidityBracket(Hierarchy);
+
+				{
+#if WITH_EDITOR
+					// make sure transient controls don't get reset
+					UControlRig::FTransientControlPoseScope PoseScope(ControlRig);
+#endif 
+					Hierarchy->ResetPoseToInitial(ERigElementType::Bone);
+				}
+			}
+
+			if (bInTransferPoseInGlobalSpace || InNodeMappingContainer.IsValid())
+			{
+				// get component pose from control rig
+				FCSPose<FCompactPose> MeshPoses;
+				// first I need to convert to local pose
+				MeshPoses.InitPose(InOutput.Pose);
+
+				if (!GetControlRigBoneInputMappingByIndex().IsEmpty())
+				{
+					for (const TPair<uint16, uint16>& Pair : GetControlRigBoneInputMappingByIndex())
+					{
+						const uint16 ControlRigIndex = Pair.Key;
+						const uint16 SkeletonIndex = Pair.Value;
+
+						FCompactPoseBoneIndex CompactPoseIndex(SkeletonIndex);
+						const FTransform& ComponentTransform = MeshPoses.GetComponentSpaceTransform(CompactPoseIndex);
+						Hierarchy->SetGlobalTransformByIndex(ControlRigIndex, ComponentTransform, false);
+					}
+				}
+				else
+				{
+					for (auto Iter = GetControlRigBoneInputMappingByName().CreateConstIterator(); Iter; ++Iter)
+					{
+						const FName& Name = Iter.Key();
+						const uint16 Index = Iter.Value();
+						const FRigElementKey Key(Name, ERigElementType::Bone);
+
+						FCompactPoseBoneIndex CompactPoseIndex(Index);
+
+						const FTransform& ComponentTransform = MeshPoses.GetComponentSpaceTransform(CompactPoseIndex);
+						if (InNodeMappingContainer.IsValid())
+						{
+							const FTransform& RelativeTransformReverse = InNodeMappingContainer->GetSourceToTargetTransform(Name).GetRelativeTransformReverse(ComponentTransform);
+							Hierarchy->SetGlobalTransform(Key, RelativeTransformReverse, false);
+						}
+						else
+						{
+							Hierarchy->SetGlobalTransform(Key, ComponentTransform, false);
+						}
+
+					}
+				}
+			}
+			else
+			{
+				if (!GetControlRigBoneInputMappingByIndex().IsEmpty())
+				{
+					for (const TPair<uint16, uint16>& Pair : GetControlRigBoneInputMappingByIndex())
+					{
+						const uint16 ControlRigIndex = Pair.Key;
+						const uint16 SkeletonIndex = Pair.Value;
+
+						FCompactPoseBoneIndex CompactPoseIndex(SkeletonIndex);
+						const FTransform& LocalTransform = InOutput.Pose[CompactPoseIndex];
+						Hierarchy->SetLocalTransformByIndex(ControlRigIndex, LocalTransform, false);
+					}
+				}
+				else
+				{
+					for (auto Iter = GetControlRigBoneInputMappingByName().CreateConstIterator(); Iter; ++Iter)
+					{
+						const FName& Name = Iter.Key();
+						const uint16 SkeletonIndex = Iter.Value();
+						const FRigElementKey Key(Name, ERigElementType::Bone);
+
+						FCompactPoseBoneIndex CompactPoseIndex(SkeletonIndex);
+						const FTransform& LocalTransform = InOutput.Pose[CompactPoseIndex];
+						Hierarchy->SetLocalTransform(Key, LocalTransform, false);
+					}
+				}
+			}
+
+		}
+
+		if (InInputSettings.bUpdateCurves && bInTransferInputCurves)
+		{
+			Hierarchy->UnsetCurveValues();
+
+			InOutput.Curve.ForEachElement([Hierarchy](const UE::Anim::FCurveElement& InCurveElement)
+				{
+					const FRigElementKey Key(InCurveElement.Name, ERigElementType::Curve);
+					Hierarchy->SetCurveValue(Key, InCurveElement.Value);
+				});
+		}
+	}
+
+#if WITH_EDITOR
+	if (bInExecute && Hierarchy->IsTracingChanges())
+	{
+		Hierarchy->StorePoseForTrace(TEXT("FAnimNode_ControlRigBase::UpdateInput"));
+	}
+#endif
+}
+
+void FControlRigHierarchyMappings::UpdateOutput(UControlRig* ControlRig
+	, FPoseContext& InOutput
+	, const FControlRigIOSettings& InOutputSettings
+	, TWeakObjectPtr<UNodeMappingContainer>& InNodeMappingContainer
+	, bool bInExecute
+	, bool bInTransferPoseInGlobalSpace)
+{
+	DECLARE_SCOPE_HIERARCHICAL_COUNTER_FUNC()
+
+	URigHierarchy* Hierarchy = ControlRig->GetHierarchy();
+	if (Hierarchy == nullptr)
+	{
+		return;
+	}
+
+	if (bEnablePoseAdapter)
+	{
+		if (InOutputSettings.bUpdatePose)
+		{
+			if (PoseAdapter)
+			{
+				// make sure the local respectively global transforms are all ready
+				if (PoseAdapter->GetTransferInLocalSpace())
+				{
+					PoseAdapter->ConvertToLocalPose();
+				}
+				else
+				{
+					PoseAdapter->ConvertToGlobalPose();
+				}
+
+				// if we transfered in global - let's convert back to local
+				if (!PoseAdapter->GetTransferInLocalSpace())
+				{
+					if (PoseAdapter->GetGlobalPose().Num() == InOutput.Pose.GetNumBones())
+					{
+						PoseAdapter->ConvertToLocalPose();
+					}
+				}
+
+				InOutput.Pose.CopyBonesFrom(PoseAdapter->GetLocalPose());
+			}
+		}
+
+		if (InOutputSettings.bUpdateCurves)
+		{
+			const TArray<int32>& ChangedCurveIndices = Hierarchy->GetChangedCurveIndices();
+			if (ChangedCurveIndices.Num() > 0)
+			{
+				const TArray<FRigBaseElement*> HierarchyCurves = Hierarchy->GetCurvesFast();
+
+				int32 CurveIndex = 0;
+				int32 CurvesCopied = 0;
+				InOutput.Curve.ForEachElement([this, &CurveIndex, &CurvesCopied, &HierarchyCurves](const UE::Anim::FCurveElement& InCurveElement)
+					{
+						// the index stored here is the sub index of the curve (the index of the curve within the list of curves)
+						const int32& HierarchyIndex = PoseAdapter->GetPoseCurveToHierarchyCurve()[CurveIndex];
+						if (HierarchyIndex != INDEX_NONE)
+						{
+							const FRigCurveElement* HierarchyCurve = CastChecked<FRigCurveElement>(HierarchyCurves[HierarchyIndex]);
+							if (HierarchyCurve->IsValueSet())
+							{
+								const_cast<UE::Anim::FCurveElement*>(&InCurveElement)->Value = HierarchyCurve->Get();
+							}
+							CurvesCopied++;
+						}
+						CurveIndex++;
+					});
+
+				if (CurvesCopied < HierarchyCurves.Num())
+				{
+					HierarchyCurveCopied.SetNumUninitialized(Hierarchy->Num());
+					FMemory::Memzero(HierarchyCurveCopied.GetData(), HierarchyCurveCopied.GetAllocatedSize());
+
+					FBlendedCurve ControlRigCurves;
+					ControlRigCurves.Reserve(ChangedCurveIndices.Num());
+					for (const int32& ChangedCurveIndex : ChangedCurveIndices)
+					{
+						if (!HierarchyCurveCopied[ChangedCurveIndex])
+						{
+							if (const FRigCurveElement* HierarchyCurve = Hierarchy->Get<FRigCurveElement>(ChangedCurveIndex))
+							{
+								if (HierarchyCurve->IsValueSet())
+								{
+									ControlRigCurves.Add(HierarchyCurve->GetFName(), HierarchyCurve->Get());
+								}
+							}
+							HierarchyCurveCopied[ChangedCurveIndex] = true;
+						}
+					}
+
+					InOutput.Curve.Combine(ControlRigCurves);
+				}
+			}
+		}
+	}
+	else
+	{
+		if (InOutputSettings.bUpdatePose)
+		{
+			// copy output of the rig
+			const FBoneContainer& RequiredBones = InOutput.Pose.GetBoneContainer();
+
+			TMap<FName, uint16>& NameBasedMapping = GetControlRigBoneOutputMappingByName();
+			TArray<TPair<uint16, uint16>>& IndexBasedMapping = GetControlRigBoneOutputMappingByIndex();
+
+			// if we don't have a different mapping for outputs, use the input mapping
+			if (NameBasedMapping.IsEmpty() && IndexBasedMapping.IsEmpty())
+			{
+				NameBasedMapping = GetControlRigBoneInputMappingByName();
+				IndexBasedMapping = GetControlRigBoneInputMappingByIndex();
+			}
+
+			if (bInTransferPoseInGlobalSpace || InNodeMappingContainer.IsValid())
+			{
+				// get component pose from control rig
+				FCSPose<FCompactPose> MeshPoses;
+				MeshPoses.InitPose(InOutput.Pose);
+
+				if (!IndexBasedMapping.IsEmpty())
+				{
+					for (const TPair<uint16, uint16>& Pair : IndexBasedMapping)
+					{
+						const uint16 ControlRigIndex = Pair.Key;
+						const uint16 SkeletonIndex = Pair.Value;
+
+						FCompactPoseBoneIndex CompactPoseIndex(SkeletonIndex);
+						FTransform ComponentTransform = Hierarchy->GetGlobalTransformByIndex(ControlRigIndex);
+						MeshPoses.SetComponentSpaceTransform(CompactPoseIndex, ComponentTransform);
+					}
+				}
+				else
+				{
+					for (auto Iter = NameBasedMapping.CreateConstIterator(); Iter; ++Iter)
+					{
+						const FName& Name = Iter.Key();
+						const uint16 SkeletonIndex = Iter.Value();
+						const FRigElementKey Key(Name, ERigElementType::Bone);
+
+						FCompactPoseBoneIndex CompactPoseIndex(SkeletonIndex);
+						FTransform ComponentTransform = Hierarchy->GetGlobalTransform(Key);
+						if (InNodeMappingContainer.IsValid())
+						{
+							ComponentTransform = InNodeMappingContainer->GetSourceToTargetTransform(Name) * ComponentTransform;
+						}
+
+						MeshPoses.SetComponentSpaceTransform(CompactPoseIndex, ComponentTransform);
+					}
+				}
+
+				FCSPose<FCompactPose>::ConvertComponentPosesToLocalPosesSafe(MeshPoses, InOutput.Pose);
+				InOutput.Pose.NormalizeRotations();
+			}
+			else
+			{
+				if (!IndexBasedMapping.IsEmpty())
+				{
+					for (const TPair<uint16, uint16>& Pair : IndexBasedMapping)
+					{
+						const uint16 ControlRigIndex = Pair.Key;
+						const uint16 SkeletonIndex = Pair.Value;
+
+						FCompactPoseBoneIndex CompactPoseIndex(SkeletonIndex);
+						FTransform LocalTransform = Hierarchy->GetLocalTransformByIndex(ControlRigIndex);
+						InOutput.Pose[CompactPoseIndex] = LocalTransform;
+					}
+				}
+				else
+				{
+					for (auto Iter = NameBasedMapping.CreateConstIterator(); Iter; ++Iter)
+					{
+						const FName& Name = Iter.Key();
+						const uint16 Index = Iter.Value();
+						const FRigElementKey Key(Name, ERigElementType::Bone);
+
+						FCompactPoseBoneIndex CompactPoseIndex(Index);
+						FTransform LocalTransform = Hierarchy->GetLocalTransform(Key);
+						InOutput.Pose[CompactPoseIndex] = LocalTransform;
+					}
+				}
+			}
+		}
+
+		if (InOutputSettings.bUpdateCurves)
+		{
+			FBlendedCurve ControlRigCurves;
+			ControlRigCurves.Reserve(Hierarchy->Num(ERigElementType::Curve));
+			Hierarchy->ForEach<FRigCurveElement>([&ControlRigCurves](const FRigCurveElement* InElement)
+				{
+					if (InElement->IsValueSet())
+					{
+						ControlRigCurves.Add(InElement->GetFName(), InElement->Get());
+					}
+					return true;
+				});
+
+			InOutput.Curve.Combine(ControlRigCurves);
+		}
+	}
+
+#if WITH_EDITOR
+	if (bInExecute && Hierarchy->IsTracingChanges())
+	{
+		Hierarchy->StorePoseForTrace(TEXT("FAnimNode_ControlRigBase::UpdateOutput"));
+		Hierarchy->DumpTransformStackToFile();
+	}
+#endif
+}
+
+void FControlRigHierarchyMappings::UpdateInputOutputMappingIfRequiredImpl(UControlRig* InControlRig
+	, URigHierarchy* InHierarchy
+	, const FBoneContainer& InRequiredBones
+	, TArray<FBoneReference>& InInputBonesToTransfer
+	, TArray<FBoneReference>& InOutputBonesToTransfer
+	, const TWeakObjectPtr<UNodeMappingContainer>& InNodeMappingContainer)
+{
+	DECLARE_SCOPE_HIERARCHICAL_COUNTER_FUNC()
+
+	ControlRigBoneInputMappingByIndex.Reset();
+	ControlRigBoneOutputMappingByIndex.Reset();
+	ControlRigCurveMappingByIndex.Reset();
+	ControlRigBoneInputMappingByName.Reset();
+	ControlRigBoneOutputMappingByName.Reset();
+	ControlRigCurveMappingByName.Reset();
+
+	if (InRequiredBones.IsValid())
+	{
+		const TArray<FBoneIndexType>& RequiredBonesArray = InRequiredBones.GetBoneIndicesArray();
+		const int32 NumBones = RequiredBonesArray.Num();
+
+		const FReferenceSkeleton& RefSkeleton = InRequiredBones.GetReferenceSkeleton();
+
+		// @todo: thread-safe? probably not in editor, but it may not be a big issue in editor
+		if (InNodeMappingContainer.IsValid())
+		{
+			// get target to source mapping table - this is reversed mapping table
+			TMap<FName, FName> TargetToSourceMappingTable;
+			InNodeMappingContainer->GetTargetToSourceMappingTable(TargetToSourceMappingTable);
+
+			// now fill up node name
+			for (uint16 Index = 0; Index < NumBones; ++Index)
+			{
+				// get bone name, and find reverse mapping
+				FName TargetNodeName = RefSkeleton.GetBoneName(RequiredBonesArray[Index]);
+				FName* SourceName = TargetToSourceMappingTable.Find(TargetNodeName);
+				if (SourceName)
+				{
+					ControlRigBoneInputMappingByName.Add(*SourceName, Index);
+				}
+			}
+		}
+		else
+		{
+			TArray<FName> NodeNames;
+			TArray<FNodeItem> NodeItems;
+			InControlRig->GetMappableNodeData(NodeNames, NodeItems);
+
+			// even if not mapped, we map only node that exists in the controlrig
+			for (uint16 Index = 0; Index < NumBones; ++Index)
+			{
+				const FName& BoneName = RefSkeleton.GetBoneName(RequiredBonesArray[Index]);
+				if (NodeNames.Contains(BoneName))
+				{
+					ControlRigBoneInputMappingByName.Add(BoneName, Index);
+				}
+			}
+		}
+
+		auto UpdatingMappingFromSpecificTransferList = [](
+			TArray<FBoneReference>& InTransferList,
+			const TWeakObjectPtr<UNodeMappingContainer>& InMappingContainer,
+			const FBoneContainer& InRequiredBones,
+			const FReferenceSkeleton& InRefSkeleton,
+			const TArray<FBoneIndexType>& InRequiredBonesArray,
+			const UControlRig* InControlRig,
+			TMap<FName, uint16>& OutMapping
+			) {
+				OutMapping.Reset();
+
+				if (InMappingContainer.IsValid())
+				{
+					// get target to source mapping table - this is reversed mapping table
+					TMap<FName, FName> TargetToSourceMappingTable;
+					InMappingContainer->GetTargetToSourceMappingTable(TargetToSourceMappingTable);
+
+					for (FBoneReference& InputBoneToTransfer : InTransferList)
+					{
+						if (!InputBoneToTransfer.Initialize(InRequiredBones))
+						{
+							continue;
+						}
+						const FName TargetNodeName = InRefSkeleton.GetBoneName(InputBoneToTransfer.BoneIndex);
+						if (const FName* SourceName = TargetToSourceMappingTable.Find(TargetNodeName))
+						{
+							OutMapping.Add(*SourceName, InputBoneToTransfer.BoneIndex);
+						}
+					}
+				}
+				else
+				{
+					TArray<FName> NodeNames;
+					TArray<FNodeItem> NodeItems;
+					InControlRig->GetMappableNodeData(NodeNames, NodeItems);
+
+					for (FBoneReference& InputBoneToTransfer : InTransferList)
+					{
+						if (!InputBoneToTransfer.Initialize(InRequiredBones))
+						{
+							continue;
+						}
+						if (InRequiredBonesArray.IsValidIndex(InputBoneToTransfer.BoneIndex))
+						{
+							const FName& BoneName = InRefSkeleton.GetBoneName(InRequiredBonesArray[InputBoneToTransfer.BoneIndex]);
+							if (NodeNames.Contains(BoneName))
+							{
+								OutMapping.Add(BoneName, InputBoneToTransfer.BoneIndex);
+							}
+						}
+					}
+				}
+			};
+
+		if (!InInputBonesToTransfer.IsEmpty())
+		{
+			ControlRigBoneOutputMappingByName = ControlRigBoneInputMappingByName;
+
+			UpdatingMappingFromSpecificTransferList(
+				InInputBonesToTransfer,
+				InNodeMappingContainer,
+				InRequiredBones,
+				RefSkeleton,
+				RequiredBonesArray,
+				InControlRig,
+				ControlRigBoneInputMappingByName);
+		}
+
+		if (!InOutputBonesToTransfer.IsEmpty())
+		{
+			UpdatingMappingFromSpecificTransferList(
+				InOutputBonesToTransfer,
+				InNodeMappingContainer,
+				InRequiredBones,
+				RefSkeleton,
+				RequiredBonesArray,
+				InControlRig,
+				ControlRigBoneOutputMappingByName);
+		}
+
+		// check if we can switch the bones to an index based mapping.
+		// we can only do that if there is no node mapping container set.
+		if (!InNodeMappingContainer.IsValid())
+		{
+			for (int32 InputOutput = 0; InputOutput < 2; InputOutput++)
+			{
+				bool bIsMappingByIndex = true;
+				TMap<FName, uint16>& NameBasedMapping = InputOutput == 0 ? ControlRigBoneInputMappingByName : ControlRigBoneOutputMappingByName;
+				if (NameBasedMapping.IsEmpty())
+				{
+					continue;
+				}
+
+				TArray<TPair<uint16, uint16>>& IndexBasedMapping = InputOutput == 0 ? ControlRigBoneInputMappingByIndex : ControlRigBoneOutputMappingByIndex;
+
+				for (auto Iter = NameBasedMapping.CreateConstIterator(); Iter; ++Iter)
+				{
+					const uint16 SkeletonIndex = Iter.Value();
+					const int32 ControlRigIndex = InHierarchy->GetIndex(FRigElementKey(Iter.Key(), ERigElementType::Bone));
+					if (ControlRigIndex != INDEX_NONE)
+					{
+						IndexBasedMapping.Add(TPair<uint16, uint16>((uint16)ControlRigIndex, SkeletonIndex));
+					}
+					else
+					{
+						bIsMappingByIndex = false;
+					}
+				}
+
+				if (bIsMappingByIndex)
+				{
+					NameBasedMapping.Reset();
+				}
+				else
+				{
+					IndexBasedMapping.Reset();
+				}
+			}
+		}
+	}
+}
+
+bool FControlRigHierarchyMappings::CheckPoseAdapter() const
+{
+	DECLARE_SCOPE_HIERARCHICAL_COUNTER_FUNC()
+
+	if (bEnablePoseAdapter && !PoseAdapter.IsValid())
+	{
+		return false;
+	}
+
+	return true;
+}
+
+bool FControlRigHierarchyMappings::IsUpdateToDate(const URigHierarchy* InHierarchy) const
+{
+	DECLARE_SCOPE_HIERARCHICAL_COUNTER_FUNC()
+
+	if (bEnablePoseAdapter)
+	{
+		if (!PoseAdapter->IsUpdateToDate(InHierarchy))
+		{
+			return false;
+		}
+	}
+	return true;
+}
+
+void FControlRigHierarchyMappings::PerformUpdateToDate(UControlRig* ControlRig
+	, URigHierarchy* InHierarchy
+	, const FBoneContainer& InRequiredBones
+	, const TWeakObjectPtr<UNodeMappingContainer>& InNodeMappingContainer
+	, bool bInTransferPoseInGlobalSpace
+	, bool bInResetInputPoseToInitial)
+{
+	DECLARE_SCOPE_HIERARCHICAL_COUNTER_FUNC()
+
+	if (bEnablePoseAdapter)
+	{
+		if (!PoseAdapter->IsUpdateToDate(InHierarchy))
+		{
+			InHierarchy->UnlinkPoseAdapter();
+			PoseAdapter->UpdateInputOutputMappingIfRequired(ControlRig, InHierarchy, InRequiredBones, InNodeMappingContainer, bInTransferPoseInGlobalSpace, bInResetInputPoseToInitial);
+			InHierarchy->LinkPoseAdapter(PoseAdapter);
+		}
+	}
+}
+
+void FControlRigHierarchyMappings::SetEnablePoseAdapter(bool bInEnablePoseAdapter)
+{
+	DECLARE_SCOPE_HIERARCHICAL_COUNTER_FUNC()
+
+	if (bEnablePoseAdapter != bInEnablePoseAdapter)
+	{
+		bEnablePoseAdapter = bInEnablePoseAdapter;
+
+		if (bEnablePoseAdapter)
+		{
+			PoseAdapter = MakeShareable(new FControlRigPoseAdapter());
+		}
+		else
+		{
+			PoseAdapter.Reset();
+		}
+	}
+}
diff --git a/Engine/Plugins/Animation/ControlRig/Source/ControlRig/Private/Tools/ControlRigPoseAdapter.cpp b/Engine/Plugins/Animation/ControlRig/Source/ControlRig/Private/Tools/ControlRigPoseAdapter.cpp
new file mode 100644
--- /dev/null
+++ b/Engine/Plugins/Animation/ControlRig/Source/ControlRig/Private/Tools/ControlRigPoseAdapter.cpp
@@ -0,0 +1,588 @@
+// Copyright Epic Games, Inc. All Rights Reserved.
+
+#include "Tools/ControlRigPoseAdapter.h"
+#include "Rigs/RigHierarchy.h"
+#include "ControlRig.h"
+#include "BoneContainer.h"
+#include "Animation/NodeMappingContainer.h"
+
+void FControlRigPoseAdapter::PostLinked(URigHierarchy* InHierarchy)
+{
+	DECLARE_SCOPE_HIERARCHICAL_COUNTER_FUNC()
+
+	FRigHierarchyPoseAdapter::PostLinked(InHierarchy);
+
+	Hierarchy = InHierarchy;
+
+	// 1. sort the hierarchies storage so that we can make sure initial and local is grouped correctly
+	SortHierarchyStorage();
+
+	// 2. make sure to compute all local transforms in initial and local
+	const TArray<FRigTransformElement*> TransformElements = Hierarchy->GetElementsOfType<FRigTransformElement>(false);
+	for (FRigTransformElement* TransformElement : TransformElements)
+	{
+		(void)Hierarchy->GetTransform(TransformElement, ERigTransformType::InitialLocal);
+		(void)Hierarchy->GetTransform(TransformElement, ERigTransformType::CurrentLocal);
+	}
+
+	// 3. set up a list of dependents in the hierarchy of the rig to be reset to initial on execute
+	TArray<FRigTransformElement*> DependentTransformElements;
+
+	struct Local
+	{
+		static bool ProcessTransformElement(FControlRigPoseAdapter* InPoseAdapter, TArray<FRigTransformElement*>& InDependentTransformElements, FRigTransformElement* InTransformElement)
+		{
+			if (InTransformElement == nullptr)
+			{
+				return false;
+			}
+
+			if (InPoseAdapter->ElementIndexToPoseIndex.Contains((uint16)InTransformElement->GetIndex()))
+			{
+				return true;
+			}
+
+			if (InDependentTransformElements.Contains(InTransformElement))
+			{
+				return true;
+			}
+
+			const FRigBaseElementParentArray& ParentElements = InPoseAdapter->Hierarchy->GetParents(InTransformElement);
+			for (const FRigBaseElement* ParentElement : ParentElements)
+			{
+				if (ProcessTransformElement(InPoseAdapter, InDependentTransformElements, Cast<FRigTransformElement>(const_cast<FRigBaseElement*>(ParentElement))))
+				{
+					InDependentTransformElements.AddUnique(InTransformElement);
+					return true;
+				}
+			}
+
+			return false;
+		}
+	};
+
+	for (FRigTransformElement* TransformElement : TransformElements)
+	{
+		Local::ProcessTransformElement(this, DependentTransformElements, TransformElement);
+	}
+
+	Dependents.Reset();
+	Dependents.Reserve(DependentTransformElements.Num() * 3);
+
+	for (FRigTransformElement* DependentTransformElement : DependentTransformElements)
+	{
+		// skip bones - since they are taking care of by BonesToResetToInitial
+		if (const FRigBoneElement* BoneElement = Cast<FRigBoneElement>(DependentTransformElement))
+		{
+			// skip bones if they are not user defined - and if the parent of the bone is not a dependent as well
+			if (BoneElement->BoneType == ERigBoneType::Imported &&
+				!DependentTransformElements.Contains(BoneElement->ParentElement))
+			{
+				continue;
+			}
+		}
+
+		FRigControlElement* DependentControlElement = Cast<FRigControlElement>(DependentTransformElement);
+		if (DependentControlElement && DependentControlElement->IsAnimationChannel())
+		{
+			continue;
+		}
+
+		Dependents.Emplace(DependentTransformElement->GetKeyAndIndex(), ERigTransformType::CurrentGlobal, ERigTransformStorageType::Pose, &DependentTransformElement->GetDirtyState().Current);
+		if (DependentControlElement)
+		{
+			Dependents.Emplace(DependentTransformElement->GetKeyAndIndex(), ERigTransformType::CurrentGlobal, ERigTransformStorageType::Offset, &DependentControlElement->GetOffsetDirtyState().Current);
+			Dependents.Emplace(DependentTransformElement->GetKeyAndIndex(), ERigTransformType::CurrentGlobal, ERigTransformStorageType::Shape, &DependentControlElement->GetShapeDirtyState().Current);
+		}
+	}
+
+	// 4. relink the storage for the transforms (local or global or both) and dirty states to our local storage
+	static constexpr bool bLocalIsPrimary = true;
+	UpdateDirtyStates(bLocalIsPrimary);
+	for (int32 PoseIndex = 0; PoseIndex < PoseIndexToElementIndex.Num(); PoseIndex++)
+	{
+		const int32& TransformElementIndex = PoseIndexToElementIndex[PoseIndex];
+		if (TransformElementIndex != INDEX_NONE)
+		{
+			const FRigElementKeyAndIndex& KeyAndIndex = Hierarchy->GetKeyAndIndex(TransformElementIndex);
+			RelinkTransformStorage(KeyAndIndex, ERigTransformType::CurrentLocal, ERigTransformStorageType::Pose, &LocalPose[PoseIndex], &LocalPoseIsDirty[PoseIndex]);
+			RelinkTransformStorage(KeyAndIndex, ERigTransformType::CurrentGlobal, ERigTransformStorageType::Pose, &GlobalPose[PoseIndex], &GlobalPoseIsDirty[PoseIndex]);
+		}
+	}
+	UpdateDirtyStates(bLocalIsPrimary); // do this again to make sure local is flagged as clean and global as dirty
+
+	// 5. Shrink the storage on the hierarchy now that we've relinked it
+	ShrinkHierarchyStorage();
+}
+
+void FControlRigPoseAdapter::PreUnlinked(URigHierarchy* InHierarchy)
+{
+	DECLARE_SCOPE_HIERARCHICAL_COUNTER_FUNC()
+
+	UnlinkTransformStorage();
+	UpdateHierarchyStorage();
+
+	Hierarchy = nullptr;
+	FRigHierarchyPoseAdapter::PreUnlinked(InHierarchy);
+}
+
+bool FControlRigPoseAdapter::IsUpdateToDate(const URigHierarchy* InHierarchy) const
+{
+	DECLARE_SCOPE_HIERARCHICAL_COUNTER_FUNC()
+
+	if (!FRigHierarchyPoseAdapter::IsUpdateToDate(InHierarchy))
+	{
+		return false;
+	}
+	return !LocalPose.IsEmpty() && !GlobalPose.IsEmpty();
+}
+
+void FControlRigPoseAdapter::UpdateInputOutputMappingIfRequired(UControlRig* InControlRig
+	, URigHierarchy* InHierarchy
+	, const FBoneContainer& InRequiredBones
+	, const TWeakObjectPtr<UNodeMappingContainer>& InNodeMappingContainer
+	, bool bInTransferPoseInGlobalSpace
+	, bool bInResetInputPoseToInitial)
+{
+	DECLARE_SCOPE_HIERARCHICAL_COUNTER_FUNC()
+
+	if (Hierarchy == nullptr)
+	{
+		return;
+	}
+
+	Hierarchy->UnlinkPoseAdapter();
+
+	ParentPoseIndices.Reset();
+	RequiresHierarchyForSpaceConversion.Reset();
+	ElementIndexToPoseIndex.Reset();
+
+	const int32 NumBonesInPose = InRequiredBones.GetCompactPoseNumBones();
+
+	GlobalPose.Reset();
+	LocalPose.Reset();
+	ParentPoseIndices.Reserve(NumBonesInPose);
+	RequiresHierarchyForSpaceConversion.Reserve(NumBonesInPose);
+	GlobalPose.AddDefaulted(NumBonesInPose);
+	LocalPose.AddDefaulted(NumBonesInPose);
+	HierarchyCurveLookup.Reset();
+
+	bTransferInLocalSpace = !(bInTransferPoseInGlobalSpace || InNodeMappingContainer.IsValid());
+
+	for (int32 Index = 0; Index < NumBonesInPose; Index++)
+	{
+		ParentPoseIndices.Add(INDEX_NONE);
+
+		const FCompactPoseBoneIndex ParentBoneIndex = InRequiredBones.GetParentBoneIndex(FCompactPoseBoneIndex(Index));
+		if (ParentBoneIndex.IsValid())
+		{
+			ParentPoseIndices[Index] = ParentBoneIndex.GetInt();
+		}
+		RequiresHierarchyForSpaceConversion.Add(false);
+	}
+	UpdateDirtyStates();
+
+	TArray<int32> MappedBoneElementIndices;
+	if (InRequiredBones.IsValid())
+	{
+		ElementIndexToPoseIndex.Reserve(NumBonesInPose);
+		PoseIndexToElementIndex.Reserve(NumBonesInPose);
+
+		const FReferenceSkeleton* RefSkeleton = &InRequiredBones.GetReferenceSkeleton();
+		if (const USkeleton* Skeleton = InRequiredBones.GetSkeletonAsset())
+		{
+			RefSkeleton = &Skeleton->GetReferenceSkeleton();
+		}
+
+		// @todo: thread-safe? probably not in editor, but it may not be a big issue in editor
+		if (InNodeMappingContainer.IsValid())
+		{
+			// get target to source mapping table - this is reversed mapping table
+			TMap<FName, FName> TargetToSourceMappingTable;
+			InNodeMappingContainer->GetTargetToSourceMappingTable(TargetToSourceMappingTable);
+
+			// now fill up node name
+			for (uint16 Index = 0; Index < NumBonesInPose; ++Index)
+			{
+				// get bone name, and find reverse mapping
+				const FSkeletonPoseBoneIndex BoneIndex = InRequiredBones.GetSkeletonPoseIndexFromCompactPoseIndex(FCompactPoseBoneIndex(Index));
+				if (BoneIndex.IsValid())
+				{
+					FName TargetNodeName = RefSkeleton->GetBoneName(BoneIndex.GetInt());
+					FName* SourceName = TargetToSourceMappingTable.Find(TargetNodeName);
+					if (SourceName)
+					{
+						const int32 ElementIndex = Hierarchy->GetIndex({ *SourceName, ERigElementType::Bone });
+						if (ElementIndex != INDEX_NONE)
+						{
+							MappedBoneElementIndices.Add(ElementIndex);
+							ElementIndexToPoseIndex.Add(static_cast<uint16>(ElementIndex), Index);
+							PoseIndexToElementIndex.Add(ElementIndex);
+							LocalPose[Index] = Hierarchy->GetLocalTransform(ElementIndex);
+							GlobalPose[Index] = Hierarchy->GetGlobalTransform(ElementIndex);
+							continue;
+						}
+					}
+				}
+				PoseIndexToElementIndex.Add(INDEX_NONE);
+			}
+		}
+		else
+		{
+			TArray<FName> NodeNames;
+			TArray<FNodeItem> NodeItems;
+			InControlRig->GetMappableNodeData(NodeNames, NodeItems);
+
+			// even if not mapped, we map only node that exists in the controlrig
+			for (uint16 Index = 0; Index < NumBonesInPose; ++Index)
+			{
+				const FSkeletonPoseBoneIndex BoneIndex = InRequiredBones.GetSkeletonPoseIndexFromCompactPoseIndex(FCompactPoseBoneIndex(Index));
+				if (BoneIndex.IsValid())
+				{
+					const FName& BoneName = RefSkeleton->GetBoneName(BoneIndex.GetInt());
+					if (NodeNames.Contains(BoneName))
+					{
+						const int32 ElementIndex = Hierarchy->GetIndex({ BoneName, ERigElementType::Bone });
+						if (ElementIndex != INDEX_NONE)
+						{
+							MappedBoneElementIndices.Add(ElementIndex);
+							ElementIndexToPoseIndex.Add(static_cast<uint16>(ElementIndex), Index);
+							PoseIndexToElementIndex.Add(ElementIndex);
+							LocalPose[Index] = Hierarchy->GetLocalTransform(ElementIndex);
+							GlobalPose[Index] = Hierarchy->GetGlobalTransform(ElementIndex);
+							continue;
+						}
+					}
+				}
+				PoseIndexToElementIndex.Add(INDEX_NONE);
+			}
+		}
+
+		// once we know all of the bones we are going to transfer - we can check if any of these bones has a different parenting
+		// relationship in the skeleton used in the anim graph vs the hierarchy in the rig. in that case we have to transfer in global
+		if (bTransferInLocalSpace)
+		{
+			for (const int32& BoneElementIndex : MappedBoneElementIndices)
+			{
+				const int32 HierarchyParentIndex = Hierarchy->GetFirstParent(BoneElementIndex);
+				const int16 PoseIndex = ElementIndexToPoseIndex.FindChecked(BoneElementIndex);
+				const FCompactPoseBoneIndex CompactPoseParentIndex(ParentPoseIndices[PoseIndex]);
+
+				FName HierarchyParentName(NAME_None);
+				FName PoseParentName(NAME_None);
+
+				if (HierarchyParentIndex != INDEX_NONE)
+				{
+					HierarchyParentName = Hierarchy->Get(HierarchyParentIndex)->GetFName();
+				}
+				if (CompactPoseParentIndex.IsValid())
+				{
+					const FSkeletonPoseBoneIndex SkeletonIndex = InRequiredBones.GetSkeletonPoseIndexFromCompactPoseIndex(CompactPoseParentIndex);
+					PoseParentName = RefSkeleton->GetBoneName(SkeletonIndex.GetInt());
+				}
+
+				if (HierarchyParentName.IsEqual(PoseParentName, ENameCase::CaseSensitive))
+				{
+					continue;
+				}
+
+				RequiresHierarchyForSpaceConversion[PoseIndex] = true;
+				check(PoseIndexToElementIndex[PoseIndex] != INDEX_NONE);
+				bTransferInLocalSpace = false;
+			}
+		}
+
+		// only reset the full pose if we are not mapping all bones
+		const TArray<FRigBaseElement*>& HierarchyBones = Hierarchy->GetBonesFast();
+		const bool bMapsAllBones = MappedBoneElementIndices.Num() == HierarchyBones.Num();
+		BonesToResetToInitial.Reset();
+		bRequiresResetPoseToInitial = bInResetInputPoseToInitial && !bMapsAllBones;
+
+		if (bRequiresResetPoseToInitial)
+		{
+			BonesToResetToInitial.Reserve(HierarchyBones.Num() - MappedBoneElementIndices.Num());
+
+			// bone is mapped stores sub indices (bone index within the list of bones)
+			TArray<bool> BoneIsMapped;
+			BoneIsMapped.AddZeroed(HierarchyBones.Num());
+			for (const int32& MappedTransformIndex : MappedBoneElementIndices)
+			{
+				const FRigBaseElement* MappedElement = Hierarchy->Get(MappedTransformIndex);
+				check(MappedElement);
+				BoneIsMapped[MappedElement->GetSubIndex()] = true;
+			}
+
+			// when we want to know which bones to reset we want to convert back to a global index
+			for (int32 UnmappedBoneIndex = 0; UnmappedBoneIndex < BoneIsMapped.Num(); UnmappedBoneIndex++)
+			{
+				if (!BoneIsMapped[UnmappedBoneIndex])
+				{
+					BonesToResetToInitial.Add(HierarchyBones[UnmappedBoneIndex]->GetIndex());
+				}
+			}
+		}
+	}
+}
+
+void FControlRigPoseAdapter::SetHierarchyCurvesLookup(const TArray<FRigBaseElement*>& InHierarchyCurves)
+{
+	const int32 NumHierarchyCurves = InHierarchyCurves.Num();
+	if (GetHierarchyCurveLookup().Num() != NumHierarchyCurves)
+	{
+		HierarchyCurveLookup.Reset();
+		HierarchyCurveLookup.Reserve(NumHierarchyCurves);
+		for (int32 Index = 0; Index < NumHierarchyCurves; Index++)
+		{
+			HierarchyCurveLookup.Add(InHierarchyCurves[Index]->GetFName(), Index);
+		}
+	}
+}
+
+void FControlRigPoseAdapter::SetPoseCurveToHierarchyCurve(const TArray<FRigBaseElement*> InHierarchyCurves, const FBlendedCurve& InCurve)
+{
+	PoseCurveToHierarchyCurve.SetNumUninitialized(InCurve.Num());
+
+	int32 CurveIndex = 0;
+	InCurve.ForEachElement([this, &CurveIndex, &InHierarchyCurves](const UE::Anim::FCurveElement& InCurveElement)
+		{
+			PoseCurveToHierarchyCurve[CurveIndex] = INDEX_NONE;
+
+			// the index stored here is the sub index of the curve (the index of the curve within the list of curves)
+			if (const int32* IndexPtr = HierarchyCurveLookup.Find(InCurveElement.Name))
+			{
+				const int32& Index = *IndexPtr;
+				FRigCurveElement* HierarchyCurve = CastChecked<FRigCurveElement>(InHierarchyCurves[Index]);
+
+				// when setting the curve we need to mark it as "value set", otherwise the copy
+				// pose may reset it to unset - thus we'll loose the value that was just copied in.
+				HierarchyCurve->Set(InCurveElement.Value, true);
+				PoseCurveToHierarchyCurve[CurveIndex] = Index;
+			}
+			CurveIndex++;
+		});
+}
+
+void FControlRigPoseAdapter::UnlinkTransformStorage()
+{
+	DECLARE_SCOPE_HIERARCHICAL_COUNTER_FUNC()
+
+	check(Hierarchy);
+
+	for (int32 PoseIndex = 0; PoseIndex < PoseIndexToElementIndex.Num(); PoseIndex++)
+	{
+		const int32& TransformElementIndex = PoseIndexToElementIndex[PoseIndex];
+		if (TransformElementIndex != INDEX_NONE)
+		{
+			const FRigElementKeyAndIndex& KeyAndIndex = Hierarchy->GetKeyAndIndex(TransformElementIndex);
+			RestoreTransformStorage(KeyAndIndex, ERigTransformType::CurrentLocal, ERigTransformStorageType::Pose, false);
+			RestoreTransformStorage(KeyAndIndex, ERigTransformType::CurrentGlobal, ERigTransformStorageType::Pose, false);
+		}
+	}
+
+	ElementIndexToPoseIndex.Reset();
+	PoseIndexToElementIndex.Reset();
+}
+
+void FControlRigPoseAdapter::ConvertToLocalPose()
+{
+	DECLARE_SCOPE_HIERARCHICAL_COUNTER_FUNC()
+
+	check(ParentPoseIndices.Num() == GlobalPose.Num());
+	LocalPose.SetNum(GlobalPose.Num());
+
+	if (Hierarchy)
+	{
+		for (int32 Index = 0; Index < ParentPoseIndices.Num(); Index++)
+		{
+			(void)GetLocalTransform(Index);
+		}
+	}
+}
+void FControlRigPoseAdapter::ConvertToGlobalPose()
+{
+	DECLARE_SCOPE_HIERARCHICAL_COUNTER_FUNC()
+
+	check(ParentPoseIndices.Num() == LocalPose.Num());
+	GlobalPose.SetNum(LocalPose.Num());
+
+	if (Hierarchy)
+	{
+		for (int32 Index = 0; Index < ParentPoseIndices.Num(); Index++)
+		{
+			(void)GetGlobalTransform(Index);
+		}
+	}
+}
+
+const FTransform& FControlRigPoseAdapter::GetLocalTransform(int32 InIndex)
+{
+	DECLARE_SCOPE_HIERARCHICAL_COUNTER_FUNC()
+
+	check(LocalPose.IsValidIndex(InIndex));
+	check(LocalPoseIsDirty.IsValidIndex(InIndex));
+	check(ParentPoseIndices.IsValidIndex(InIndex));
+	check(RequiresHierarchyForSpaceConversion.IsValidIndex(InIndex));
+
+	check(Hierarchy);
+
+	if (LocalPoseIsDirty[InIndex])
+	{
+		check(GlobalPoseIsDirty[InIndex] == false);
+		if (RequiresHierarchyForSpaceConversion[InIndex] && PoseIndexToElementIndex.IsValidIndex(InIndex))
+		{
+			LocalPose[InIndex] = Hierarchy->GetLocalTransformByIndex(PoseIndexToElementIndex[InIndex]);
+		}
+		else
+		{
+			const int32 ParentIndex = ParentPoseIndices[InIndex];
+			if (ParentIndex == INDEX_NONE)
+			{
+				LocalPose[InIndex] = GetGlobalTransform(InIndex);
+			}
+			else
+			{
+				LocalPose[InIndex] = GetGlobalTransform(InIndex).GetRelativeTransform(GetGlobalTransform(ParentIndex));
+			}
+			LocalPose[InIndex].NormalizeRotation();
+		}
+		LocalPoseIsDirty[InIndex] = false;
+	}
+	return LocalPose[InIndex];
+}
+
+const FTransform& FControlRigPoseAdapter::GetGlobalTransform(int32 InIndex)
+{
+	DECLARE_SCOPE_HIERARCHICAL_COUNTER_FUNC()
+
+	check(GlobalPose.IsValidIndex(InIndex));
+	check(GlobalPoseIsDirty.IsValidIndex(InIndex));
+	check(ParentPoseIndices.IsValidIndex(InIndex));
+	check(RequiresHierarchyForSpaceConversion.IsValidIndex(InIndex));
+
+	check(Hierarchy);
+
+	if (GlobalPoseIsDirty[InIndex])
+	{
+		check(LocalPoseIsDirty[InIndex] == false);
+		if (RequiresHierarchyForSpaceConversion[InIndex] && PoseIndexToElementIndex.IsValidIndex(InIndex))
+		{
+			GlobalPose[InIndex] = Hierarchy->GetGlobalTransformByIndex(PoseIndexToElementIndex[InIndex]);
+		}
+		else
+		{
+			const int32 ParentIndex = ParentPoseIndices[InIndex];
+			if (ParentIndex == INDEX_NONE)
+			{
+				GlobalPose[InIndex] = GetLocalTransform(InIndex);
+			}
+			else
+			{
+				GlobalPose[InIndex] = GetLocalTransform(InIndex) * GetGlobalTransform(ParentIndex);
+			}
+			GlobalPose[InIndex].NormalizeRotation();
+		}
+		GlobalPoseIsDirty[InIndex] = false;
+	}
+	return GlobalPose[InIndex];
+}
+
+void FControlRigPoseAdapter::UpdateDirtyStates(const TOptional<bool> InLocalIsPrimary)
+{
+	DECLARE_SCOPE_HIERARCHICAL_COUNTER_FUNC()
+
+	bool bLocalIsPrimary = bTransferInLocalSpace;
+	if (InLocalIsPrimary.IsSet())
+	{
+		bLocalIsPrimary = InLocalIsPrimary.GetValue();
+	}
+	
+	if (LocalPoseIsDirty.Num() != LocalPose.Num())
+	{
+		LocalPoseIsDirty.Reset();
+		LocalPoseIsDirty.AddZeroed(LocalPose.Num());
+	}
+	else if (bLocalIsPrimary)
+	{
+		FMemory::Memzero(LocalPoseIsDirty.GetData(), LocalPoseIsDirty.GetAllocatedSize());
+	}
+	if (!bLocalIsPrimary)
+	{
+		for (bool& Flag : LocalPoseIsDirty)
+		{
+			Flag = true;
+		}
+	}
+
+	if (GlobalPoseIsDirty.Num() != GlobalPose.Num())
+	{
+		GlobalPoseIsDirty.Reset();
+		GlobalPoseIsDirty.AddZeroed(GlobalPose.Num());
+	}
+	else if (!bLocalIsPrimary)
+	{
+		FMemory::Memzero(GlobalPoseIsDirty.GetData(), GlobalPoseIsDirty.GetAllocatedSize());
+	}
+	if (bLocalIsPrimary)
+	{
+		for (bool& Flag : GlobalPoseIsDirty)
+		{
+			Flag = true;
+		}
+	}
+}
+
+void FControlRigPoseAdapter::ComputeDependentTransforms()
+{
+	DECLARE_SCOPE_HIERARCHICAL_COUNTER_FUNC()
+
+	check(Hierarchy);
+
+	// ensure to compute all local transforms
+	for (const FDependentTransform& Dependent : Dependents)
+	{
+		if (FRigTransformElement* TransformElement = Hierarchy->Get<FRigTransformElement>(Dependent.KeyAndIndex.Index))
+		{
+			switch(Dependent.StorageType)
+			{
+				case ERigTransformStorageType::Pose:
+				{
+					(void)Hierarchy->GetTransform(TransformElement, MakeLocal(Dependent.TransformType));
+					break;
+				}
+				case ERigTransformStorageType::Offset:
+				{
+					FRigControlElement* ControlElement = CastChecked<FRigControlElement>(TransformElement);
+					(void)Hierarchy->GetControlOffsetTransform(ControlElement, MakeLocal(Dependent.TransformType));
+					break;
+				}
+				case ERigTransformStorageType::Shape:
+				{
+					FRigControlElement* ControlElement = CastChecked<FRigControlElement>(TransformElement);
+					(void)Hierarchy->GetControlShapeTransform(ControlElement, MakeLocal(Dependent.TransformType));
+					break;
+				}
+				default:
+				{
+					break;
+				}
+			}
+			check(Dependent.DirtyState->Local.Get() == false);
+		}
+	}
+}
+
+void FControlRigPoseAdapter::MarkDependentsDirty()
+{
+	DECLARE_SCOPE_HIERARCHICAL_COUNTER_FUNC()
+
+	ComputeDependentTransforms();
+
+	check(Hierarchy);
+
+	// mark the global dependent as dirty
+	for (const FDependentTransform& Dependent : Dependents)
+	{
+		check(Dependent.DirtyState->Local.Get() == false);
+		Dependent.DirtyState->Global.Set(true);
+	}
+}
\ No newline at end of file
diff --git a/Engine/Plugins/Animation/ControlRig/Source/ControlRig/Private/Tools/ControlRigVariableMappings.cpp b/Engine/Plugins/Animation/ControlRig/Source/ControlRig/Private/Tools/ControlRigVariableMappings.cpp
new file mode 100644
--- /dev/null
+++ b/Engine/Plugins/Animation/ControlRig/Source/ControlRig/Private/Tools/ControlRigVariableMappings.cpp
@@ -0,0 +1,693 @@
+// Copyright Epic Games, Inc. All Rights Reserved.
+
+#include "Tools/ControlRigVariableMappings.h"
+#include "ControlRig.h"
+#include "Components/SkeletalMeshComponent.h"
+#include "Animation/AnimCurveUtils.h"
+#include "Algo/ForEach.h"
+
+namespace AnimNodeLocals
+{
+
+const TCHAR* UseFunctionsName = TEXT("ControlRig.AnimNode.UseFunctions");
+
+static bool bUseFunctions = false;
+static FAutoConsoleVariableRef CVarUseCachedFunctions(
+   UseFunctionsName,
+   bUseFunctions,
+   TEXT("Cache and use propagate functions.")
+   );
+
+template<typename ValueType, typename HierarchyValueType = ValueType>
+void SetControlValue(FRigControlElement* InControlElement, FProperty* InSourceProperty, URigHierarchy* InTargetHierarchy, const UObject* InSourceInstance)
+{
+	const ValueType* ValuePtr = InSourceProperty->ContainerPtrToValuePtr<ValueType>(InSourceInstance);
+	const FRigControlValue Value = FRigControlValue::Make<HierarchyValueType>(*ValuePtr);
+	InTargetHierarchy->SetControlValue(InControlElement, Value, ERigControlValueType::Current);
+}
+
+template<typename ValueType, typename VariableValueType = ValueType>
+void SetVariableValue(FRigVMExternalVariable& InVariable, FProperty* InSourceProperty, const UObject* InSourceInstance)
+{
+	const ValueType* ValuePtr = InSourceProperty->ContainerPtrToValuePtr<ValueType>(InSourceInstance);
+	InVariable.SetValue<VariableValueType>(*ValuePtr);
+}
+
+} // end namespace AnimNodeLocals
+
+
+void FControlRigVariableMappings::InitializeProperties(const UClass* InSourceClass
+	, UObject* TargetInstance
+	, UClass* InTargetClass
+	, const TArray<FName>& InSourcePropertyNames
+	, TArray<FName>& InDestPropertyNames)
+{
+	const bool bUseFunctions = AnimNodeLocals::bUseFunctions;
+
+	// Build property lists
+	SourceProperties.Reset(InSourcePropertyNames.Num());
+	DestProperties.Reset(InSourcePropertyNames.Num());
+	UpdateFunctions.Reset(InSourcePropertyNames.Num());
+	Variables.Reset();
+
+	check(InSourcePropertyNames.Num() == InDestPropertyNames.Num());
+
+	UControlRig* TargetControlRig = Cast<UControlRig>(TargetInstance);
+	URigHierarchy* TargetHierarchy = TargetControlRig ? TargetControlRig->GetHierarchy() : nullptr;
+
+	
+	for (int32 Idx = 0; Idx < InSourcePropertyNames.Num(); ++Idx)
+	{
+		const FName& SourceName = InSourcePropertyNames[Idx];
+
+		FProperty* SourceProperty = FindFProperty<FProperty>(InSourceClass, SourceName);
+		SourceProperties.Add(SourceProperty);
+		DestProperties.Add(nullptr);
+
+		if (bUseFunctions)
+		{
+			if (TargetHierarchy && SourceProperty)
+			{
+				const FName& DestPropertyName = InDestPropertyNames[Idx];
+				if (FRigControlElement* ControlElement = TargetControlRig->FindControl(DestPropertyName))
+				{
+					AddControlFunction(ControlElement, SourceProperty, TargetHierarchy);
+				}
+				else
+				{
+					const FRigVMExternalVariable Variable = TargetControlRig->GetPublicVariableByName(DestPropertyName);
+					if (Variable.IsValid() && !Variable.bIsReadOnly)
+					{
+						AddVariableFunction(Variable, SourceProperty);
+					}
+				}
+			}
+		}
+	}
+}
+
+void FControlRigVariableMappings::PropagateInputProperties(const UObject* InSourceInstance, UControlRig* InTargetControlRig, TArray<FName>& InDestPropertyNames)
+{
+	DECLARE_SCOPE_HIERARCHICAL_COUNTER_FUNC()
+	QUICK_SCOPE_CYCLE_COUNTER(FControlRigVariableMappings_PropagateInputProperties);
+
+	if (!InSourceInstance)
+	{
+		return;
+	}
+
+	if (!InTargetControlRig)
+	{
+		return;
+	}
+
+	URigHierarchy* TargetHierarchy = InTargetControlRig->GetHierarchy();
+	if (!TargetHierarchy)
+	{
+		return;
+	}
+
+	if (AnimNodeLocals::bUseFunctions && !UpdateFunctions.IsEmpty())
+	{
+		Algo::ForEach(UpdateFunctions, [InSourceInstance](const PropertyUpdateFunction& InFunc)
+			{
+				InFunc(InSourceInstance);
+			});
+	}
+	else
+	{
+		PropagateInputPropertiesNoCache(InSourceInstance, InTargetControlRig, TargetHierarchy, InDestPropertyNames);
+	}
+}
+
+void FControlRigVariableMappings::UpdateCurveInputs(UControlRig* InControlRig, const TMap<FName, FName> &InputMapping, const FBlendedCurve& InCurveData)
+{
+	DECLARE_SCOPE_HIERARCHICAL_COUNTER_FUNC()
+
+	// now go through variable mapping table and see if anything is mapping through input
+	if (InputMapping.Num() > 0 && InControlRig)
+	{
+		UE::Anim::FCurveUtils::BulkGet(InCurveData, InputCurveMappings,
+			[&InControlRig](const FControlRigCurveMapping& InBulkElement, float InValue)
+			{
+				FRigVMExternalVariable Variable = InControlRig->GetPublicVariableByName(InBulkElement.SourceName);
+				if (!Variable.bIsReadOnly && Variable.TypeName == TEXT("float"))
+				{
+					Variable.SetValue<float>(InValue);
+				}
+				else
+				{
+					UE_LOG(LogAnimation, Warning, TEXT("[%s] Missing Input Variable [%s]"), *GetNameSafe(InControlRig->GetClass()), *InBulkElement.SourceName.ToString());
+				}
+			});
+	}
+}
+
+void FControlRigVariableMappings::UpdateCurveOutputs(UControlRig* InControlRig, const TMap<FName, FName> &OutputMapping, FBlendedCurve& InCurveData)
+{
+	DECLARE_SCOPE_HIERARCHICAL_COUNTER_FUNC()
+
+	if (OutputMapping.Num() > 0 && InControlRig)
+	{
+		UE::Anim::FCurveUtils::BulkSet(InCurveData, OutputCurveMappings,
+			[&InControlRig](const FControlRigCurveMapping& InBulkElement) -> float
+			{
+				FRigVMExternalVariable Variable = InControlRig->GetPublicVariableByName(InBulkElement.SourceName);
+				if (Variable.TypeName == TEXT("float"))
+				{
+					return Variable.GetValue<float>();
+				}
+				else
+				{
+					UE_LOG(LogAnimation, Warning, TEXT("[%s] Missing Output Variable [%s]"), *GetNameSafe(InControlRig->GetClass()), *InBulkElement.SourceName.ToString());
+				}
+
+				return 0.0f;
+			});
+	}
+}
+
+void FControlRigVariableMappings::PropagateInputPropertiesNoCache(const UObject* InSourceInstance, const UControlRig* InTargetControlRig, URigHierarchy* InTargetHierarchy, TArray<FName>& InDestPropertyNames)
+{
+	check(SourceProperties.Num() == DestProperties.Num());
+
+	for (int32 PropIdx = 0; PropIdx < SourceProperties.Num(); ++PropIdx)
+	{
+		FProperty* CallerProperty = SourceProperties[PropIdx];
+
+		if (FRigControlElement* ControlElement = InTargetControlRig->FindControl(InDestPropertyNames[PropIdx]))
+		{
+			const uint8* SrcPtr = CallerProperty->ContainerPtrToValuePtr<uint8>(InSourceInstance);
+
+			bool bIsValid = false;
+			FRigControlValue Value;
+			switch (ControlElement->Settings.ControlType)
+			{
+			case ERigControlType::Bool:
+			{
+				if (ensure(CastField<FBoolProperty>(CallerProperty)))
+				{
+					Value = FRigControlValue::Make<bool>(*(bool*)SrcPtr);
+					bIsValid = true;
+				}
+				break;
+			}
+			case ERigControlType::Float:
+			case ERigControlType::ScaleFloat:
+			{
+				if (ensure(CastField<FFloatProperty>(CallerProperty)))
+				{
+					Value = FRigControlValue::Make<float>(*(float*)SrcPtr);
+					bIsValid = true;
+				}
+				break;
+			}
+			case ERigControlType::Integer:
+			{
+				if (ensure(CastField<FIntProperty>(CallerProperty)))
+				{
+					Value = FRigControlValue::Make<int32>(*(int32*)SrcPtr);
+					bIsValid = true;
+				}
+				break;
+			}
+			case ERigControlType::Vector2D:
+			{
+				FStructProperty* StructProperty = CastField<FStructProperty>(CallerProperty);
+				if (ensure(StructProperty))
+				{
+					if (ensure(StructProperty->Struct == TBaseStructure<FVector2D>::Get()))
+					{
+						const FVector2D& SrcVector = *(FVector2D*)SrcPtr;
+						Value = FRigControlValue::Make<FVector2D>(SrcVector);
+						bIsValid = true;
+					}
+				}
+				break;
+			}
+			case ERigControlType::Position:
+			case ERigControlType::Scale:
+			{
+				FStructProperty* StructProperty = CastField<FStructProperty>(CallerProperty);
+				if (ensure(StructProperty))
+				{
+					if (ensure(StructProperty->Struct == TBaseStructure<FVector>::Get()))
+					{
+						const FVector& SrcVector = *(FVector*)SrcPtr;
+						Value = FRigControlValue::Make<FVector>(SrcVector);
+						bIsValid = true;
+					}
+				}
+				break;
+			}
+			case ERigControlType::Rotator:
+			{
+				FStructProperty* StructProperty = CastField<FStructProperty>(CallerProperty);
+				if (ensure(StructProperty))
+				{
+					if (ensure(StructProperty->Struct == TBaseStructure<FRotator>::Get()))
+					{
+						const FRotator& SrcRotator = *(FRotator*)SrcPtr;
+						Value = FRigControlValue::Make<FRotator>(SrcRotator);
+						bIsValid = true;
+					}
+				}
+				break;
+			}
+			case ERigControlType::Transform:
+			{
+				FStructProperty* StructProperty = CastField<FStructProperty>(CallerProperty);
+				if (ensure(StructProperty))
+				{
+					if (ensure(StructProperty->Struct == TBaseStructure<FTransform>::Get()))
+					{
+						const FTransform& SrcTransform = *(FTransform*)SrcPtr;
+						Value = FRigControlValue::Make<FTransform>(SrcTransform);
+						bIsValid = true;
+					}
+				}
+				break;
+			}
+			case ERigControlType::TransformNoScale:
+			{
+				FStructProperty* StructProperty = CastField<FStructProperty>(CallerProperty);
+				if (ensure(StructProperty))
+				{
+					if (ensure(StructProperty->Struct == TBaseStructure<FTransform>::Get()))
+					{
+						const FTransform& SrcTransform = *(FTransform*)SrcPtr;
+						Value = FRigControlValue::Make<FTransformNoScale>(SrcTransform);
+						bIsValid = true;
+					}
+				}
+				break;
+			}
+			case ERigControlType::EulerTransform:
+			{
+				FStructProperty* StructProperty = CastField<FStructProperty>(CallerProperty);
+				if (ensure(StructProperty))
+				{
+					if (ensure(StructProperty->Struct == TBaseStructure<FTransform>::Get()))
+					{
+						const FTransform& SrcTransform = *(FTransform*)SrcPtr;
+						Value = FRigControlValue::Make<FEulerTransform>(FEulerTransform(SrcTransform));
+						bIsValid = true;
+					}
+				}
+				break;
+			}
+			default:
+			{
+				checkNoEntry();
+			}
+			}
+
+			if (bIsValid)
+			{
+				InTargetHierarchy->SetControlValue(ControlElement, Value, ERigControlValueType::Current);
+			}
+			continue;
+		}
+
+		FRigVMExternalVariable Variable = InTargetControlRig->GetPublicVariableByName(InDestPropertyNames[PropIdx]);
+		if (Variable.IsValid())
+		{
+			if (Variable.bIsReadOnly)
+			{
+				continue;
+			}
+
+			const uint8* SrcPtr = CallerProperty->ContainerPtrToValuePtr<uint8>(InSourceInstance);
+
+			if (CastField<FBoolProperty>(CallerProperty) != nullptr && Variable.TypeName == RigVMTypeUtils::BoolTypeName)
+			{
+				const bool Value = *(const bool*)SrcPtr;
+				Variable.SetValue<bool>(Value);
+			}
+			else if (CastField<FFloatProperty>(CallerProperty) != nullptr && (Variable.TypeName == RigVMTypeUtils::FloatTypeName || Variable.TypeName == RigVMTypeUtils::DoubleTypeName))
+			{
+				const float Value = *(const float*)SrcPtr;
+				if (Variable.TypeName == RigVMTypeUtils::FloatTypeName)
+				{
+					Variable.SetValue<float>(Value);
+				}
+				else
+				{
+					Variable.SetValue<double>(Value);
+				}
+			}
+			else if (CastField<FDoubleProperty>(CallerProperty) != nullptr && (Variable.TypeName == RigVMTypeUtils::FloatTypeName || Variable.TypeName == RigVMTypeUtils::DoubleTypeName))
+			{
+				const double Value = *(const double*)SrcPtr;
+				if (Variable.TypeName == RigVMTypeUtils::FloatTypeName)
+				{
+					Variable.SetValue<float>((float)Value);
+				}
+				else
+				{
+					Variable.SetValue<double>(Value);
+				}
+			}
+			else if (CastField<FIntProperty>(CallerProperty) != nullptr && Variable.TypeName == RigVMTypeUtils::Int32TypeName)
+			{
+				const int32 Value = *(const int32*)SrcPtr;
+				Variable.SetValue<int32>(Value);
+			}
+			else if (CastField<FNameProperty>(CallerProperty) != nullptr && Variable.TypeName == RigVMTypeUtils::FNameTypeName)
+			{
+				const FName Value = *(const FName*)SrcPtr;
+				Variable.SetValue<FName>(Value);
+			}
+			else if (CastField<FNameProperty>(CallerProperty) != nullptr && Variable.TypeName == RigVMTypeUtils::FStringTypeName)
+			{
+				const FString Value = *(const FString*)SrcPtr;
+				Variable.SetValue<FString>(Value);
+			}
+			else if (FStructProperty* StructProperty = CastField<FStructProperty>(CallerProperty))
+			{
+				if (StructProperty->Struct == Variable.TypeObject)
+				{
+					StructProperty->Struct->CopyScriptStruct(Variable.Memory, SrcPtr, 1);
+				}
+			}
+			else if (FArrayProperty* ArrayProperty = CastField<FArrayProperty>(CallerProperty))
+			{
+				if (ensure(ArrayProperty->SameType(Variable.Property)))
+				{
+					ArrayProperty->CopyCompleteValue(Variable.Memory, SrcPtr);
+				}
+			}
+			else if (FObjectProperty* ObjectProperty = CastField<FObjectProperty>(CallerProperty))
+			{
+				if (ensure(ObjectProperty->SameType(Variable.Property)))
+				{
+					ObjectProperty->CopyCompleteValue(Variable.Memory, SrcPtr);
+				}
+			}
+			else if (FEnumProperty* EnumProperty = CastField<FEnumProperty>(CallerProperty))
+			{
+				if (ensure(EnumProperty->SameType(Variable.Property)))
+				{
+					EnumProperty->CopyCompleteValue(Variable.Memory, SrcPtr);
+				}
+			}
+			else
+			{
+				ensureMsgf(false, TEXT("Property %s type %s not recognized"), *CallerProperty->GetName(), *CallerProperty->GetCPPType());
+			}
+		}
+	}
+}
+
+void FControlRigVariableMappings::AddControlFunction(FRigControlElement* InControlElement, FProperty* InSourceProperty, URigHierarchy* InTargetHierarchy)
+{
+	using namespace AnimNodeLocals;
+
+	switch (InControlElement->Settings.ControlType)
+	{
+		case ERigControlType::Bool:
+		{
+			if (ensure(CastField<FBoolProperty>(InSourceProperty)))
+			{
+				return AddUpdateFunction([InSourceProperty, InTargetHierarchy, InControlElement](const UObject* InSourceInstance)
+					{
+						SetControlValue<bool>(InControlElement, InSourceProperty, InTargetHierarchy, InSourceInstance);
+					});
+			}
+			break;
+		}
+		case ERigControlType::Float:
+		case ERigControlType::ScaleFloat:
+		{
+			if (ensure(CastField<FFloatProperty>(InSourceProperty)))
+			{
+				return AddUpdateFunction([InSourceProperty, InTargetHierarchy, InControlElement](const UObject* InSourceInstance)
+					{
+						SetControlValue<float>(InControlElement, InSourceProperty, InTargetHierarchy, InSourceInstance);
+					});
+			}
+			break;
+		}
+		case ERigControlType::Integer:
+		{
+			if (ensure(CastField<FIntProperty>(InSourceProperty)))
+			{
+				return AddUpdateFunction([InSourceProperty, InTargetHierarchy, InControlElement](const UObject* InSourceInstance)
+					{
+						SetControlValue<int32>(InControlElement, InSourceProperty, InTargetHierarchy, InSourceInstance);
+					});
+			}
+			break;
+		}
+		case ERigControlType::Vector2D:
+		{
+			FStructProperty* StructProperty = CastField<FStructProperty>(InSourceProperty);
+			if (ensure(StructProperty))
+			{
+				if (ensure(StructProperty->Struct == TBaseStructure<FVector2D>::Get()))
+				{
+					return AddUpdateFunction([InSourceProperty, InTargetHierarchy, InControlElement](const UObject* InSourceInstance)
+						{
+							SetControlValue<FVector2D>(InControlElement, InSourceProperty, InTargetHierarchy, InSourceInstance);
+						});
+				}
+			}
+			break;
+		}
+		case ERigControlType::Position:
+		case ERigControlType::Scale:
+		{
+			FStructProperty* StructProperty = CastField<FStructProperty>(InSourceProperty);
+			if (ensure(StructProperty))
+			{
+				if (ensure(StructProperty->Struct == TBaseStructure<FVector>::Get()))
+				{
+					return AddUpdateFunction([InSourceProperty, InTargetHierarchy, InControlElement](const UObject* InSourceInstance)
+						{
+							SetControlValue<FVector>(InControlElement, InSourceProperty, InTargetHierarchy, InSourceInstance);
+						});
+				}
+			}
+			break;
+		}
+		case ERigControlType::Rotator:
+		{
+			FStructProperty* StructProperty = CastField<FStructProperty>(InSourceProperty);
+			if (ensure(StructProperty))
+			{
+				if (ensure(StructProperty->Struct == TBaseStructure<FRotator>::Get()))
+				{
+					return AddUpdateFunction([InSourceProperty, InTargetHierarchy, InControlElement](const UObject* InSourceInstance)
+						{
+							SetControlValue<FRotator>(InControlElement, InSourceProperty, InTargetHierarchy, InSourceInstance);
+						});
+				}
+			}
+			break;
+		}
+		case ERigControlType::Transform:
+		{
+			FStructProperty* StructProperty = CastField<FStructProperty>(InSourceProperty);
+			if (ensure(StructProperty))
+			{
+				if (ensure(StructProperty->Struct == TBaseStructure<FTransform>::Get()))
+				{
+					return AddUpdateFunction([InSourceProperty, InTargetHierarchy, InControlElement](const UObject* InSourceInstance)
+						{
+							SetControlValue<FTransform>(InControlElement, InSourceProperty, InTargetHierarchy, InSourceInstance);
+						});
+				}
+			}
+			break;
+		}
+		case ERigControlType::TransformNoScale:
+		{
+			FStructProperty* StructProperty = CastField<FStructProperty>(InSourceProperty);
+			if (ensure(StructProperty))
+			{
+				if (ensure(StructProperty->Struct == TBaseStructure<FTransform>::Get()))
+				{
+					return AddUpdateFunction([InSourceProperty, InTargetHierarchy, InControlElement](const UObject* InSourceInstance)
+						{
+							SetControlValue<FTransform, FTransformNoScale>(InControlElement, InSourceProperty, InTargetHierarchy, InSourceInstance);
+						});
+				}
+			}
+			break;
+		}
+		case ERigControlType::EulerTransform:
+		{
+			FStructProperty* StructProperty = CastField<FStructProperty>(InSourceProperty);
+			if (ensure(StructProperty))
+			{
+				if (ensure(StructProperty->Struct == TBaseStructure<FTransform>::Get()))
+				{
+					return AddUpdateFunction([InSourceProperty, InTargetHierarchy, InControlElement](const UObject* InSourceInstance)
+						{
+							const FTransform* ValuePtr = InSourceProperty->ContainerPtrToValuePtr<FTransform>(InSourceInstance);
+							const FRigControlValue Value = FRigControlValue::Make<FEulerTransform>(FEulerTransform(*ValuePtr));
+							InTargetHierarchy->SetControlValue(InControlElement, Value, ERigControlValueType::Current);
+						});
+				}
+			}
+			break;
+		}
+		default:
+		{
+			checkNoEntry();
+		}
+	}
+}
+
+void FControlRigVariableMappings::AddUpdateFunction(PropertyUpdateFunction&& InFunction)
+{
+	UpdateFunctions.Add(MoveTemp(InFunction));
+}
+
+void FControlRigVariableMappings::AddVariableFunction(const FRigVMExternalVariable& InVariable, FProperty* InSourceProperty)
+{
+	using namespace AnimNodeLocals;
+	
+	const int32 VariableIndex = Variables.Add(InVariable);
+	
+	if (CastField<FBoolProperty>(InSourceProperty) && InVariable.TypeName == RigVMTypeUtils::BoolTypeName)
+	{
+		return AddUpdateFunction([InSourceProperty, this, VariableIndex](const UObject* InSourceInstance)
+		{
+			SetVariableValue<bool>(Variables[VariableIndex], InSourceProperty, InSourceInstance);
+		});
+	}
+	else if (CastField<FFloatProperty>(InSourceProperty) && (InVariable.TypeName == RigVMTypeUtils::FloatTypeName || InVariable.TypeName == RigVMTypeUtils::DoubleTypeName))
+	{
+		if(InVariable.TypeName == RigVMTypeUtils::FloatTypeName)
+		{
+			return AddUpdateFunction([InSourceProperty, this, VariableIndex](const UObject* InSourceInstance)
+			{
+				SetVariableValue<float>(Variables[VariableIndex], InSourceProperty, InSourceInstance);
+			});
+		}
+		else
+		{
+			return AddUpdateFunction([InSourceProperty, this, VariableIndex](const UObject* InSourceInstance)
+			{
+				SetVariableValue<float, double>(Variables[VariableIndex], InSourceProperty, InSourceInstance);
+			});
+		}
+	}
+	else if (CastField<FDoubleProperty>(InSourceProperty) && (InVariable.TypeName == RigVMTypeUtils::FloatTypeName || InVariable.TypeName == RigVMTypeUtils::DoubleTypeName))
+	{
+		if(InVariable.TypeName == RigVMTypeUtils::FloatTypeName)
+		{
+			return AddUpdateFunction([InSourceProperty, this, VariableIndex](const UObject* InSourceInstance)
+			{
+				SetVariableValue<double, float>(Variables[VariableIndex], InSourceProperty, InSourceInstance);
+			});
+		}
+		else
+		{
+			return AddUpdateFunction([InSourceProperty, this, VariableIndex](const UObject* InSourceInstance)
+			{
+				SetVariableValue<double>(Variables[VariableIndex], InSourceProperty, InSourceInstance);
+			});
+		}
+	}
+	else if (CastField<FIntProperty>(InSourceProperty) && InVariable.TypeName == RigVMTypeUtils::Int32TypeName)
+	{
+		return AddUpdateFunction([InSourceProperty, this, VariableIndex](const UObject* InSourceInstance)
+		{
+			SetVariableValue<int32>(Variables[VariableIndex], InSourceProperty, InSourceInstance);
+		});
+	}
+	else if (CastField<FNameProperty>(InSourceProperty) && InVariable.TypeName == RigVMTypeUtils::FNameTypeName)
+	{
+		return AddUpdateFunction([InSourceProperty, this, VariableIndex](const UObject* InSourceInstance)
+		{
+			SetVariableValue<FName>(Variables[VariableIndex], InSourceProperty, InSourceInstance);
+		});
+	}
+	else if (CastField<FNameProperty>(InSourceProperty) && InVariable.TypeName == RigVMTypeUtils::FStringTypeName)
+	{
+		return AddUpdateFunction([InSourceProperty, this, VariableIndex](const UObject* InSourceInstance)
+		{
+			SetVariableValue<FString>(Variables[VariableIndex], InSourceProperty, InSourceInstance);
+		});
+	}
+	else if (FStructProperty* StructProperty = CastField<FStructProperty>(InSourceProperty))
+	{
+		return AddUpdateFunction([StructProperty, this, VariableIndex](const UObject* InSourceInstance)
+		{
+			const uint8* SrcPtr = StructProperty->ContainerPtrToValuePtr<uint8>(InSourceInstance);
+			StructProperty->Struct->CopyScriptStruct(Variables[VariableIndex].Memory, SrcPtr, 1);
+		});
+	}
+	else if(FArrayProperty* ArrayProperty = CastField<FArrayProperty>(InSourceProperty))
+	{
+		if(ensure(ArrayProperty->SameType(InVariable.Property)))
+		{
+			return AddUpdateFunction([ArrayProperty, this, VariableIndex](const UObject* InSourceInstance)
+			{
+				const uint8* SrcPtr = ArrayProperty->ContainerPtrToValuePtr<uint8>(InSourceInstance);
+				ArrayProperty->CopyCompleteValue(Variables[VariableIndex].Memory, SrcPtr);
+			});
+		}
+	}
+	else if(FObjectProperty* ObjectProperty = CastField<FObjectProperty>(InSourceProperty))
+	{
+		if(ensure(ObjectProperty->SameType(InVariable.Property)))
+		{
+			return AddUpdateFunction([ObjectProperty, this, VariableIndex](const UObject* InSourceInstance)
+			{
+				const uint8* SrcPtr = ObjectProperty->ContainerPtrToValuePtr<uint8>(InSourceInstance);
+				ObjectProperty->CopyCompleteValue(Variables[VariableIndex].Memory, SrcPtr);
+			});
+		}
+	}
+	else if(FEnumProperty* EnumProperty = CastField<FEnumProperty>(InSourceProperty))
+	{
+		if(ensure(EnumProperty->SameType(InVariable.Property)))
+		{
+			return AddUpdateFunction([EnumProperty, this, VariableIndex](const UObject* InSourceInstance)
+			{
+				const uint8* SrcPtr = EnumProperty->ContainerPtrToValuePtr<uint8>(InSourceInstance);
+				EnumProperty->CopyCompleteValue(Variables[VariableIndex].Memory, SrcPtr);
+			});
+		}
+	}
+	else
+	{
+		ensureMsgf(false, TEXT("Property %s type %s not recognized"), *InSourceProperty->GetName(), *InSourceProperty->GetCPPType());
+	}
+}
+
+void FControlRigVariableMappings::CacheCurveMappings(const TMap<FName, FName>& InInputMapping, const TMap<FName, FName>& InOutputMapping, URigHierarchy* InHierarchy)
+{
+	CacheCurveMappings(InInputMapping, InputCurveMappings, InHierarchy);
+	CacheCurveMappings(InOutputMapping, OutputCurveMappings, InHierarchy);
+}
+
+void FControlRigVariableMappings::CacheCurveMappings(const TMap<FName, FName>& InMapping, FControlRigVariableMappings::FCurveMappings& InCurveMappings,	URigHierarchy* InHierarchy)
+{
+	for (auto Iter = InMapping.CreateConstIterator(); Iter; ++Iter)
+	{
+		// we need to have list of variables using pin
+		const FName SourcePath = Iter.Key();
+		const FName TargetPath = Iter.Value();
+
+		if (SourcePath != NAME_None && TargetPath != NAME_None)
+		{
+			InCurveMappings.Add(SourcePath, TargetPath);
+
+			if (InHierarchy)
+			{
+				const FRigElementKey Key(TargetPath, ERigElementType::Control);
+				if (const FRigControlElement* ControlElement = InHierarchy->Find<FRigControlElement>(Key))
+				{
+					CurveInputToControlIndex.Add(TargetPath, ControlElement->GetIndex());
+					continue;
+				}
+			}
+		}
+
+		// @todo: should we clear the item if not found?
+	}
+}
diff --git a/Engine/Plugins/Animation/ControlRig/Source/ControlRig/Public/AnimNode_ControlRig.h b/Engine/Plugins/Animation/ControlRig/Source/ControlRig/Public/AnimNode_ControlRig.h
--- a/Engine/Plugins/Animation/ControlRig/Source/ControlRig/Public/AnimNode_ControlRig.h
+++ b/Engine/Plugins/Animation/ControlRig/Source/ControlRig/Public/AnimNode_ControlRig.h
@@ -6,6 +6,7 @@
 #include "Animation/InputScaleBias.h"
 #include "AnimNode_ControlRigBase.h"
 #include "Animation/AnimBulkCurves.h"
+#include "Tools/ControlRigVariableMappings.h"
 #include "AnimNode_ControlRig.generated.h"
 
 class UNodeMappingContainer;
@@ -34,8 +35,9 @@ public:
 	virtual void CacheBones_AnyThread(const FAnimationCacheBonesContext& Context) override;
 	virtual void Evaluate_AnyThread(FPoseContext & Output) override;
 	virtual int32 GetLODThreshold() const override { return LODThreshold; }
-	void SetIOMapping(bool bInput, const FName& SourceProperty, const FName& TargetCurve);
-	FName GetIOMapping(bool bInput, const FName& SourceProperty) const;
+
+	//void SetIOMapping(bool bInput, const FName& SourceProperty, const FName& TargetCurve);
+	//FName GetIOMapping(bool bInput, const FName& SourceProperty) const;
 
 	virtual void InitializeProperties(const UObject* InSourceInstance, UClass* InTargetClass) override;
 	virtual void PropagateInputProperties(const UObject* InSourceInstance) override;
@@ -99,10 +101,6 @@ private:
 	UPROPERTY()
 	TMap<FName, FName> OutputMapping;
 
-	TMap<FName, FName> InputTypes;
-	TMap<FName, FName> OutputTypes;
-	TArray<uint8*> DestParameters;
-
 	/*
 	 * Max LOD that this node is allowed to run
 	 * For example if you have LODThreshold to be 2, it will run until LOD 2 (based on 0 index)
@@ -112,43 +110,6 @@ private:
 	UPROPERTY(EditAnywhere, Category = Performance, meta = (PinHiddenByDefault, DisplayName = "LOD Threshold"))
 	int32 LODThreshold;
 
-private:
-	struct FControlRigCurveMapping
-	{
-		FControlRigCurveMapping() = default;
-		
-		FControlRigCurveMapping(FName InSourceName, FName InTargetName)
-			: Name(InTargetName)
-			, SourceName(InSourceName)
-		{}
-
-		FName Name = NAME_None;
-		FName SourceName = NAME_None;
-	};
-
-	using FCurveMappings = UE::Anim::TNamedValueArray<FDefaultAllocator, FControlRigCurveMapping>;
-
-	// Bulk curves for I/O
-	FCurveMappings InputCurveMappings;
-	FCurveMappings OutputCurveMappings;
-
-	/* This function is used to copy the property values going thru the hierarchy without any caching mechanism. */
-	void PropagateInputPropertiesNoCache(const UObject* InSourceInstance, const UControlRig* InControlRig, URigHierarchy* InHierarchy);
-	
-	/*
-	 * Cached functions used to copy property values to avoid looking for them for every PropagateInputProperties calls.
-	 * They are stored in InitializeProperties and basically represent 'set value' functions on controls and variables.
-	 */
-	void AddControlFunction(FRigControlElement* InControlElement, FProperty* InSourceProperty, URigHierarchy* InTargetHierarchy);
-	void AddVariableFunction(const FRigVMExternalVariable& InVariable, FProperty* InSourceProperty);
-	
-	using PropertyUpdateFunction = TFunction<void(const UObject*)>;
-
-	void AddUpdateFunction(PropertyUpdateFunction&& InFunction);
-
-	TArray<PropertyUpdateFunction> UpdateFunctions;
-	TArray<FRigVMExternalVariable> Variables;
-
 protected:
 	virtual UClass* GetTargetClass() const override;
 	virtual void UpdateInput(UControlRig* InControlRig, FPoseContext& InOutput) override;
@@ -157,11 +118,7 @@ protected:
 	void SetControlRigClass(TSubclassOf<UControlRig> InControlRigClass);
 	bool UpdateControlRigIfNeeded(const UAnimInstance* InAnimInstance, const FBoneContainer& InRequiredBones);
 
-	// Helper function to update the initial ref pose within the Control Rig if needed
-	void UpdateControlRigRefPoseIfNeeded(const FAnimInstanceProxy* InProxy, bool bIncludePoseInHash = false);
-	
-	// A hash to encode the pointer to anim instance, an
-	TOptional<int32> RefPoseSetterHash;
+	FControlRigVariableMappings ControlRigVariableMappings;
 
 public:
 
diff --git a/Engine/Plugins/Animation/ControlRig/Source/ControlRig/Public/AnimNode_ControlRigBase.h b/Engine/Plugins/Animation/ControlRig/Source/ControlRig/Public/AnimNode_ControlRigBase.h
--- a/Engine/Plugins/Animation/ControlRig/Source/ControlRig/Public/AnimNode_ControlRigBase.h
+++ b/Engine/Plugins/Animation/ControlRig/Source/ControlRig/Public/AnimNode_ControlRigBase.h
@@ -6,40 +6,13 @@
 #include "Animation/AnimNode_CustomProperty.h"
 #include "Rigs/RigHierarchyPoseAdapter.h"
 #include "Rigs/RigHierarchy.h"
+#include "Tools/ControlRigHierarchyMappings.h"
+#include "Tools/ControlRigIOSettings.h"
 #include "AnimNode_ControlRigBase.generated.h"
 
 class UNodeMappingContainer;
 class UControlRig;
-
-/** Struct defining the settings to override when driving a control rig */
-USTRUCT()
-struct CONTROLRIG_API FControlRigIOSettings
-{
-	GENERATED_BODY()
-
-	FControlRigIOSettings()
-		: bUpdatePose(true)
-		, bUpdateCurves(true)
-	{}
-
-	static FControlRigIOSettings MakeEnabled()
-	{
-		return FControlRigIOSettings();
-	}
-
-	static FControlRigIOSettings MakeDisabled()
-	{
-		FControlRigIOSettings Settings;
-		Settings.bUpdatePose = Settings.bUpdateCurves = false;
-		return Settings;
-	}
-
-	UPROPERTY()
-	bool bUpdatePose;
-
-	UPROPERTY()
-	bool bUpdateCurves;
-};
+class FControlRigPoseAdapter;
 
 USTRUCT()
 struct CONTROLRIG_API FControlRigAnimNodeEventName
@@ -54,80 +27,6 @@ struct CONTROLRIG_API FControlRigAnimNodeEventName
 	FName EventName;	
 };
 
-class CONTROLRIG_API FAnimNode_ControlRig_PoseAdapter : public FRigHierarchyPoseAdapter
-{
-public:
-
-	FAnimNode_ControlRig_PoseAdapter()
-	: bTransferInLocalSpace(true)
-	, bRequiresResetPoseToInitial(true)
-	, Hierarchy(nullptr)
-	, bEnabled(true)
-	{
-	}
-	
-protected:
-
-	virtual void PostLinked(URigHierarchy* InHierarchy) override;
-	virtual void PreUnlinked(URigHierarchy* InHierarchy) override;
-	virtual bool IsUpdateToDate(const URigHierarchy* InHierarchy) const override;
-	
-	void ConvertToLocalPose();
-	void ConvertToGlobalPose();
-	const FTransform& GetLocalTransform(int32 InIndex);
-	const FTransform& GetGlobalTransform(int32 InIndex);
-	
-	void UpdateDirtyStates(const TOptional<bool> InLocalIsPrimary = TOptional<bool>());
-	void ComputeDependentTransforms();
-	void MarkDependentsDirty();
-	
-	void UnlinkTransformStorage();
-
-	bool bTransferInLocalSpace;
-
-	TArray<int32> ParentPoseIndices;
-	TArray<bool> RequiresHierarchyForSpaceConversion;
-	TArray<FTransform> LocalPose;
-	TArray<FTransform> GlobalPose;
-	TArray<bool> LocalPoseIsDirty;
-	TArray<bool> GlobalPoseIsDirty;
-	TArray<int32> PoseCurveToHierarchyCurve;
-	TArray<bool> HierarchyCurveCopied;
-	TMap<FName, int32> HierarchyCurveLookup;
-	bool bRequiresResetPoseToInitial;
-	TArray<int32> BonesToResetToInitial;
-
-	TMap<uint16, uint16> ElementIndexToPoseIndex;
-	TArray<int32> PoseIndexToElementIndex;
-
-	struct FDependentTransform
-	{
-		FDependentTransform()
-			: KeyAndIndex()
-			, TransformType(ERigTransformType::InitialLocal)
-			, StorageType(ERigTransformStorageType::Pose)
-			, DirtyState(nullptr)
-		{}
-		
-		FDependentTransform(const FRigElementKeyAndIndex& InKeyAndIndex, ERigTransformType::Type InTransformType, ERigTransformStorageType::Type InStorageType, FRigLocalAndGlobalDirtyState* InDirtyState)
-			: KeyAndIndex(InKeyAndIndex)
-			, TransformType(InTransformType)
-			, StorageType(InStorageType)
-			, DirtyState(InDirtyState)
-		{}
-
-		FRigElementKeyAndIndex KeyAndIndex;
-		ERigTransformType::Type TransformType;
-		ERigTransformStorageType::Type StorageType;
-		FRigLocalAndGlobalDirtyState* DirtyState;
-	};
-	TArray<FDependentTransform> Dependents;
-
-	URigHierarchy* Hierarchy;
-	bool bEnabled;;
-
-	friend struct FAnimNode_ControlRigBase;
-};
 
 /**
  * Animation node that allows animation ControlRig output to be used in an animation graph
@@ -154,8 +53,6 @@ struct CONTROLRIG_API FAnimNode_ControlRigBase : public FAnimNode_CustomProperty
 
 protected:
 
-	void UpdateInputOutputMappingIfRequired(UControlRig* InControlRig, const FBoneContainer& RequiredBones); 
-
 	UPROPERTY(EditAnywhere, Category = Links)
 	FPoseLink Source;
 
@@ -209,8 +106,6 @@ protected:
 	UPROPERTY(EditAnywhere, AdvancedDisplay, Category = Settings)
 	TArray<TObjectPtr<UAssetUserData>> AssetUserData;
 
-	TMap<FName, int32> InputToControlIndex;
-	
 	/** Node Mapping Container */
 	UPROPERTY(transient)
 	TWeakObjectPtr<UNodeMappingContainer> NodeMappingContainer;
@@ -221,21 +116,22 @@ protected:
 	UPROPERTY(transient)
 	FControlRigIOSettings OutputSettings;
 
-	/** Complete mapping from skeleton to control rig bone index */
-	TArray<TPair<uint16, uint16>> ControlRigBoneInputMappingByIndex;
-	TArray<TPair<uint16, uint16>> ControlRigBoneOutputMappingByIndex;
+	///** Complete mapping from skeleton to control rig bone index */
+	//TArray<TPair<uint16, uint16>> ControlRigBoneInputMappingByIndex;
+	//TArray<TPair<uint16, uint16>> ControlRigBoneOutputMappingByIndex;
 
-	/** Complete mapping from skeleton to curve name */
-	TArray<TPair<uint16, FName>> ControlRigCurveMappingByIndex;
+	///** Complete mapping from skeleton to curve name */
+	//TArray<TPair<uint16, FName>> ControlRigCurveMappingByIndex;
 
-	/** Rig Hierarchy bone name to required array index mapping */
-	TMap<FName, uint16> ControlRigBoneInputMappingByName;
-	TMap<FName, uint16> ControlRigBoneOutputMappingByName;
+	///** Rig Hierarchy bone name to required array index mapping */
+	//TMap<FName, uint16> ControlRigBoneInputMappingByName;
+	//TMap<FName, uint16> ControlRigBoneOutputMappingByName;
 
-	/** Rig Curve name to Curve mapping */
-	TMap<FName, FName> ControlRigCurveMappingByName;
+	///** Rig Curve name to Curve mapping */
+	//TMap<FName, FName> ControlRigCurveMappingByName;
 	
-	TSharedPtr<FAnimNode_ControlRig_PoseAdapter> PoseAdapter;
+	//TSharedPtr<FAnimNode_ControlRig_PoseAdapter> PoseAdapter;
+	//TSharedPtr<FControlRigPoseAdapter> PoseAdapter;
 
 	UPROPERTY(transient)
 	bool bExecute;
@@ -266,6 +162,9 @@ protected:
 	bool bControlRigRequiresInitialization;
 	bool bEnablePoseAdapter;
 	uint16 LastBonesSerialNumberForCacheBones;
+
+	FControlRigHierarchyMappings ControlRigHierarchyMappings;
+
 	TWeakObjectPtr<const UAnimInstance> WeakAnimInstanceObject;
 
 	friend struct FControlRigSequencerAnimInstanceProxy;
diff --git a/Engine/Plugins/Animation/ControlRig/Source/ControlRig/Public/ControlRig.h b/Engine/Plugins/Animation/ControlRig/Source/ControlRig/Public/ControlRig.h
--- a/Engine/Plugins/Animation/ControlRig/Source/ControlRig/Public/ControlRig.h
+++ b/Engine/Plugins/Animation/ControlRig/Source/ControlRig/Public/ControlRig.h
@@ -932,6 +932,7 @@ private:
 	friend class FControlRigEditorModule;
 	friend class UModularRig;
 	friend class UModularRigController;
+	friend struct FControlRigHierarchyMappings;
 };
 
 class CONTROLRIG_API FControlRigBracketScope
diff --git a/Engine/Plugins/Animation/ControlRig/Source/ControlRig/Public/Tools/ControlRigHierarchyMappings.h b/Engine/Plugins/Animation/ControlRig/Source/ControlRig/Public/Tools/ControlRigHierarchyMappings.h
new file mode 100644
--- /dev/null
+++ b/Engine/Plugins/Animation/ControlRig/Source/ControlRig/Public/Tools/ControlRigHierarchyMappings.h
@@ -0,0 +1,137 @@
+// Copyright Epic Games, Inc. All Rights Reserved.
+
+#pragma once
+
+#include "CoreMinimal.h"
+
+class UControlRig;
+class URigHierarchy;
+class UNodeMappingContainer;
+class FControlRigPoseAdapter;
+class USkeletalMeshComponent;
+struct FBoneContainer;
+struct FBoneReference;
+struct FPoseContext;
+struct FControlRigIOSettings;
+struct FCompactPose;
+
+template <class T>
+class TAutoConsoleVariable;
+
+extern TAutoConsoleVariable<int32> CVarControlRigEnableAnimNodePerformanceOptimizations;
+
+struct FControlRigHierarchyMappings
+{
+	FControlRigHierarchyMappings() = default;
+
+	void InitializeInstance();
+
+	void LinkToHierarchy(URigHierarchy* InHierarchy);
+
+	bool CanExecute() const
+	{
+		return !bEnablePoseAdapter || PoseAdapter.IsValid();
+	}
+
+	void ResetRefPoseSetterHash()
+	{
+		RefPoseSetterHash.Reset();
+	}
+
+	void UpdateControlRigRefPoseIfNeeded(UControlRig* ControlRig
+		, const UObject* InstanceObject
+		, const USkeletalMeshComponent* SkeletalMeshComponent
+		, const FBoneContainer& RequiredBones
+		, bool bInSetRefPoseFromSkeleton
+		, bool bIncludePoseInHash);
+
+	void UpdateInputOutputMappingIfRequired(UControlRig* InControlRig
+		, URigHierarchy* InHierarchy
+		, const FBoneContainer& InRequiredBones
+		, TArray<FBoneReference>& InInputBonesToTransfer
+		, TArray<FBoneReference>& InOutputBonesToTransfer
+		, const TWeakObjectPtr<UNodeMappingContainer>& InNodeMappingContainer
+		, bool bInTransferPoseInGlobalSpace
+		, bool bResetInputPoseToInitial);
+
+	void UpdateInput(UControlRig* ControlRig
+		, FPoseContext& InOutput
+		, const FControlRigIOSettings& InInputSettings
+		, const FControlRigIOSettings& InOutputSettings
+		, TWeakObjectPtr<UNodeMappingContainer>& InNodeMappingContainer
+		, bool bInExecute
+		, bool bInTransferInputPose
+		, bool bInResetInputPoseToInitial
+		, bool bInTransferPoseInGlobalSpace
+		, bool bInTransferInputCurves);
+
+	void UpdateOutput(UControlRig* ControlRig
+		, FPoseContext& InOutput
+		, const FControlRigIOSettings& InOutputSettings
+		, TWeakObjectPtr<UNodeMappingContainer>& InNodeMappingContainer
+		, bool bInExecute
+		, bool bInTransferPoseInGlobalSpace);
+
+	const TArray<TPair<uint16, uint16>>& GetControlRigBoneInputMappingByIndex() const
+	{
+		return ControlRigBoneInputMappingByIndex;
+	}
+
+	TArray<TPair<uint16, uint16>>& GetControlRigBoneOutputMappingByIndex()
+	{
+		return ControlRigBoneOutputMappingByIndex;
+	}
+
+	const TMap<FName, uint16>& GetControlRigBoneInputMappingByName() const
+	{
+		return ControlRigBoneInputMappingByName;
+	}
+
+	TMap<FName, uint16>& GetControlRigBoneOutputMappingByName()
+	{
+		return ControlRigBoneOutputMappingByName;
+	}
+
+	bool CheckPoseAdapter() const;
+	bool IsUpdateToDate(const URigHierarchy* InHierarchy) const;
+	void PerformUpdateToDate(UControlRig* ControlRig
+		, URigHierarchy* InHierarchy
+		, const FBoneContainer& InRequiredBones
+		, const TWeakObjectPtr<UNodeMappingContainer>& InNodeMappingContainer
+		, bool bInTransferPoseInGlobalSpace
+		, bool bInResetInputPoseToInitial);
+
+private:
+	void SetEnablePoseAdapter(bool bInEnablePoseAdapter);
+
+	// Update IO whithout PoseAdapter
+	void UpdateInputOutputMappingIfRequiredImpl(UControlRig* InControlRig
+		, URigHierarchy* InHierarchy
+		, const FBoneContainer& InRequiredBones
+		, TArray<FBoneReference>& InInputBonesToTransfer
+		, TArray<FBoneReference>& InOutputBonesToTransfer
+		, const TWeakObjectPtr<UNodeMappingContainer>& InNodeMappingContainer);
+
+	/** Complete mapping from skeleton to control rig bone index */
+	TArray<TPair<uint16, uint16>> ControlRigBoneInputMappingByIndex;
+	TArray<TPair<uint16, uint16>> ControlRigBoneOutputMappingByIndex;
+
+	/** Complete mapping from skeleton to curve name */
+	TArray<TPair<uint16, FName>> ControlRigCurveMappingByIndex;
+
+	/** Rig Hierarchy bone name to required array index mapping */
+	TMap<FName, uint16> ControlRigBoneInputMappingByName;
+	TMap<FName, uint16> ControlRigBoneOutputMappingByName;
+
+	/** Rig Curve name to Curve mapping */
+	TMap<FName, FName> ControlRigCurveMappingByName;
+
+	TArray<bool> HierarchyCurveCopied;
+
+	TSharedPtr<FControlRigPoseAdapter> PoseAdapter;
+
+	// A hash to encode the pointer to anim instance
+	TOptional<int32> RefPoseSetterHash;
+
+	bool bEnablePoseAdapter = false;
+};
diff --git a/Engine/Plugins/Animation/ControlRig/Source/ControlRig/Public/Tools/ControlRigIOSettings.h b/Engine/Plugins/Animation/ControlRig/Source/ControlRig/Public/Tools/ControlRigIOSettings.h
new file mode 100644
--- /dev/null
+++ b/Engine/Plugins/Animation/ControlRig/Source/ControlRig/Public/Tools/ControlRigIOSettings.h
@@ -0,0 +1,35 @@
+// Copyright Epic Games, Inc. All Rights Reserved.
+
+#pragma once
+
+#include "ControlRigIOSettings.generated.h"
+
+/** Struct defining the settings to override when driving a control rig */
+USTRUCT()
+struct CONTROLRIG_API FControlRigIOSettings
+{
+	GENERATED_BODY()
+
+	FControlRigIOSettings()
+		: bUpdatePose(true)
+		, bUpdateCurves(true)
+	{}
+
+	static FControlRigIOSettings MakeEnabled()
+	{
+		return FControlRigIOSettings();
+	}
+
+	static FControlRigIOSettings MakeDisabled()
+	{
+		FControlRigIOSettings Settings;
+		Settings.bUpdatePose = Settings.bUpdateCurves = false;
+		return Settings;
+	}
+
+	UPROPERTY()
+	bool bUpdatePose;
+
+	UPROPERTY()
+	bool bUpdateCurves;
+};
diff --git a/Engine/Plugins/Animation/ControlRig/Source/ControlRig/Public/Tools/ControlRigPoseAdapter.h b/Engine/Plugins/Animation/ControlRig/Source/ControlRig/Public/Tools/ControlRigPoseAdapter.h
new file mode 100644
--- /dev/null
+++ b/Engine/Plugins/Animation/ControlRig/Source/ControlRig/Public/Tools/ControlRigPoseAdapter.h
@@ -0,0 +1,131 @@
+// Copyright Epic Games, Inc. All Rights Reserved.
+
+#pragma once
+
+#include "CoreMinimal.h"
+#include "Rigs/RigHierarchyPoseAdapter.h"
+#include "Rigs/RigHierarchyDefines.h"
+#include "BonePose.h"
+
+class UControlRig;
+class URigHierarchy;
+class UNodeMappingContainer;
+struct FBoneContainer;
+struct FBlendedCurve;
+struct FRigElementKeyAndIndex;
+
+class CONTROLRIG_API FControlRigPoseAdapter : public FRigHierarchyPoseAdapter
+{
+public:
+	FControlRigPoseAdapter() = default;
+
+	// --- FRigHierarchyPoseAdapter overrides start ---
+	virtual void PostLinked(URigHierarchy* InHierarchy) override;
+	virtual void PreUnlinked(URigHierarchy* InHierarchy) override;
+	virtual bool IsUpdateToDate(const URigHierarchy* InHierarchy) const override;
+	// --- FRigHierarchyPoseAdapter overrides end ---
+
+	const TArray<int32>& GetBonesToResetToInitial() const
+	{
+		return BonesToResetToInitial;
+	}
+
+	const TArray<FTransform>& GetLocalPose() const
+	{
+		return LocalPose;
+	}
+
+	const TArray<FTransform> GetGlobalPose() const
+	{
+		return LocalPose;
+	}
+
+	void CopyBonesFrom(const FCompactPose& InPose)
+	{
+		InPose.CopyBonesTo(LocalPose);
+	}
+
+	const TMap<FName, int32> GetHierarchyCurveLookup() const
+	{
+		return HierarchyCurveLookup;
+	}
+
+	bool GetTransferInLocalSpace() const
+	{
+		return bTransferInLocalSpace;
+	}
+
+	void SetHierarchyCurvesLookup(const TArray<FRigBaseElement*>& InHierarchyCurves);
+
+	const TArray<int32>& GetPoseCurveToHierarchyCurve() const
+	{
+		return PoseCurveToHierarchyCurve;
+	}
+
+	void SetPoseCurveToHierarchyCurve(const TArray<FRigBaseElement*> InHierarchyCurves, const FBlendedCurve& InCurve);
+
+	void UpdateInputOutputMappingIfRequired(UControlRig* InControlRig
+		, URigHierarchy* InHierarchy
+		, const FBoneContainer& InRequiredBones
+		, const TWeakObjectPtr<UNodeMappingContainer>& InNodeMappingContainer
+		, bool bInTransferPoseInGlobalSpace
+		, bool bInResetInputPoseToInitial);
+
+	void ConvertToLocalPose();
+	void ConvertToGlobalPose();
+	const FTransform& GetLocalTransform(int32 InIndex);
+	const FTransform& GetGlobalTransform(int32 InIndex);
+
+	void UpdateDirtyStates(const TOptional<bool> InLocalIsPrimary = TOptional<bool>());
+	void ComputeDependentTransforms();
+	void MarkDependentsDirty();
+
+	void UnlinkTransformStorage();
+
+protected:
+	struct FDependentTransform
+	{
+		FDependentTransform()
+			: KeyAndIndex()
+			, TransformType(ERigTransformType::InitialLocal)
+			, StorageType(ERigTransformStorageType::Pose)
+			, DirtyState(nullptr)
+		{}
+
+		FDependentTransform(const FRigElementKeyAndIndex& InKeyAndIndex, ERigTransformType::Type InTransformType, ERigTransformStorageType::Type InStorageType, FRigLocalAndGlobalDirtyState* InDirtyState)
+			: KeyAndIndex(InKeyAndIndex)
+			, TransformType(InTransformType)
+			, StorageType(InStorageType)
+			, DirtyState(InDirtyState)
+		{}
+
+		FRigElementKeyAndIndex KeyAndIndex;
+		ERigTransformType::Type TransformType;
+		ERigTransformStorageType::Type StorageType;
+		FRigLocalAndGlobalDirtyState* DirtyState;
+	};
+
+	TArray<int32> ParentPoseIndices;
+	TArray<bool> RequiresHierarchyForSpaceConversion;
+	TArray<FTransform> LocalPose;
+	TArray<FTransform> GlobalPose;
+	TArray<bool> LocalPoseIsDirty;
+	TArray<bool> GlobalPoseIsDirty;
+	TArray<int32> PoseCurveToHierarchyCurve;
+	
+	TMap<FName, int32> HierarchyCurveLookup;
+	TArray<int32> BonesToResetToInitial;
+
+	TMap<uint16, uint16> ElementIndexToPoseIndex;
+	TArray<int32> PoseIndexToElementIndex;
+
+	TArray<FDependentTransform> Dependents;
+
+	URigHierarchy* Hierarchy = nullptr;
+
+	bool bTransferInLocalSpace = true;
+	bool bRequiresResetPoseToInitial = true;
+	bool bEnabled = true;
+
+	friend struct FAnimNode_ControlRigBase;
+};
diff --git a/Engine/Plugins/Animation/ControlRig/Source/ControlRig/Public/Tools/ControlRigVariableMappings.h b/Engine/Plugins/Animation/ControlRig/Source/ControlRig/Public/Tools/ControlRigVariableMappings.h
new file mode 100644
--- /dev/null
+++ b/Engine/Plugins/Animation/ControlRig/Source/ControlRig/Public/Tools/ControlRigVariableMappings.h
@@ -0,0 +1,92 @@
+// Copyright Epic Games, Inc. All Rights Reserved.
+
+#pragma once
+
+#include "CoreMinimal.h"
+#include "Misc/Optional.h"
+#include "Animation/NamedValueArray.h"
+
+class UControlRig;
+class USkeletalMeshComponent;
+class FProperty;
+class URigHierarchy;
+struct FBlendedCurve;
+struct FCompactPose;
+struct FBoneContainer;
+struct FRigVMExternalVariable;
+struct FRigControlElement;
+
+struct FControlRigVariableMappings
+{
+	struct FControlRigCurveMapping
+	{
+		FControlRigCurveMapping() = default;
+
+		FControlRigCurveMapping(FName InSourceName, FName InTargetName)
+			: Name(InTargetName)
+			, SourceName(InSourceName)
+		{}
+
+		FName Name = NAME_None;
+		FName SourceName = NAME_None;
+	};
+
+	using FCurveMappings = UE::Anim::TNamedValueArray<FDefaultAllocator, FControlRigCurveMapping>;
+
+	void InitializeProperties(const UClass* InSourceClass
+		, UObject* TargetInstance
+		, UClass* InTargetClass
+		, const TArray<FName>& SourcePropertyNames
+		, TArray<FName>& DestPropertyNames);
+
+	void PropagateInputProperties(const UObject* InSourceInstance, UControlRig* InTargetInstance, TArray<FName>& InDestPropertyNames);
+
+	void UpdateCurveInputs(UControlRig* InControlRig, const TMap<FName, FName>& InputMapping, const FBlendedCurve& CurveData);
+
+	void UpdateCurveOutputs(UControlRig* InControlRig, const TMap<FName, FName>& OutputMapping, FBlendedCurve& InCurveData);
+
+	void ResetCurvesInputToControlCache()
+	{
+		CurveInputToControlIndex.Reset();
+	}
+
+	void CacheCurveMappings(const TMap<FName, FName>& InInputMappings, const TMap<FName, FName>& InOutputMappings, URigHierarchy* InHierarchy);
+
+	const FCurveMappings& GetInputCurveMappings()
+	{
+		return InputCurveMappings;
+	}
+
+	const FCurveMappings& GetOutputCurveMappings()
+	{
+		return OutputCurveMappings;
+	}
+
+private:
+	using PropertyUpdateFunction = TFunction<void(const UObject*)>;
+
+	/*
+	 * Cached functions used to copy property values to avoid looking for them for every PropagateInputProperties calls.
+	 * They are stored in InitializeProperties and basically represent 'set value' functions on controls and variables.
+	 */
+	void AddControlFunction(FRigControlElement* InControlElement, FProperty* InSourceProperty, URigHierarchy* InTargetHierarchy);
+	void AddUpdateFunction(PropertyUpdateFunction&& InFunction);
+	void AddVariableFunction(const FRigVMExternalVariable& InVariable, FProperty* InSourceProperty);
+	
+	/* This function is used to copy the property values going thru the hierarchy without any caching mechanism. */
+	void PropagateInputPropertiesNoCache(const UObject* InSourceInstance, const UControlRig* InTargetControlRig, URigHierarchy* InTargetHierarchy, TArray<FName>& InDestPropertyNames);
+
+	void CacheCurveMappings(const TMap<FName, FName>& InMapping, FControlRigVariableMappings::FCurveMappings& InCurveMappings, URigHierarchy* InHierarchy);
+
+	TArray<FProperty*> SourceProperties;
+	TArray<FProperty*> DestProperties;
+
+	TArray<PropertyUpdateFunction> UpdateFunctions;
+	TArray<FRigVMExternalVariable> Variables;
+
+	// Bulk curves for I/O
+	FCurveMappings InputCurveMappings;
+	FCurveMappings OutputCurveMappings;
+
+	TMap<FName, int32> CurveInputToControlIndex;
+};
diff --git a/Engine/Plugins/Animation/ControlRig/Source/ControlRigDeveloper/Private/AnimGraphNode_ControlRig.cpp b/Engine/Plugins/Animation/ControlRig/Source/ControlRigDeveloper/Private/AnimGraphNode_ControlRig.cpp
--- a/Engine/Plugins/Animation/ControlRig/Source/ControlRigDeveloper/Private/AnimGraphNode_ControlRig.cpp
+++ b/Engine/Plugins/Animation/ControlRig/Source/ControlRigDeveloper/Private/AnimGraphNode_ControlRig.cpp
@@ -7,9 +7,6 @@
 #include "Widgets/Text/STextBlock.h"
 #include "Widgets/Input/SCheckBox.h"
 #include "DetailWidgetRow.h"
-#if WITH_EDITOR
-#include "Widgets/SRigVMVariableMappingWidget.h"
-#endif
 #include "ScopedTransaction.h"
 #include "Kismet2/BlueprintEditorUtils.h"
 #include "AnimationGraphSchema.h"
@@ -24,7 +21,12 @@
 
 UAnimGraphNode_ControlRig::UAnimGraphNode_ControlRig(const FObjectInitializer& ObjectInitializer)
 	: Super(ObjectInitializer)
+	, ControlRigIOMapping(MakeShared<FControlRigIOMapping>(Node.InputMapping, Node.OutputMapping, CustomPinProperties))
 {
+	ControlRigIOMapping->GetOnPinCheckStateChangedDelegate().BindUObject(this, &UAnimGraphNode_ControlRig::OnPropertyExposeCheckboxChanged);
+	ControlRigIOMapping->GetOnVariableMappingChanged().BindUObject(this, &UAnimGraphNode_ControlRig::OnVariableMappingChanged);
+	ControlRigIOMapping->GetOnGetTargetSkeletonDelegate().BindUObject(this, &UAnimGraphNode_ControlRig::GetTargetSkeleton);
+	ControlRigIOMapping->GetOnGetTargetClassDelegate().BindUObject(this, &UAnimGraphNode_ControlRig::GetTargetClass);
 }
 
 FText UAnimGraphNode_ControlRig::GetNodeTitle(ENodeTitleType::Type TitleType) const
@@ -44,7 +46,7 @@ void UAnimGraphNode_ControlRig::CreateCustomPins(TArray<UEdGraphPin*>* OldPins)
 	DECLARE_SCOPE_HIERARCHICAL_COUNTER_FUNC()
 	
 	// we do this to refresh input variables
-	RebuildExposedProperties();
+	ControlRigIOMapping->RebuildExposedProperties();
 	
 	// Grab the SKELETON class here as when we are reconstructed during during BP compilation
 	// the full generated class is not yet present built.
@@ -138,14 +140,14 @@ void UAnimGraphNode_ControlRig::CreateCustomPins(TArray<UEdGraphPin*>* OldPins)
 		}
 	}
 
-	for (TPair<FName, FRigVMExternalVariable>& VariablePair : InputVariables)
+	for (const TPair<FName, FRigVMExternalVariable>& VariablePair : ControlRigIOMapping->GetInputVariables())
 	{
 		FName PropertyName = VariablePair.Key;
 		BeginExposableNames.Remove(PropertyName);
 
 		if (CurrentlyExposedNames.Contains(PropertyName))
 		{
-			FRigVMExternalVariable& Variable = VariablePair.Value;
+			const FRigVMExternalVariable& Variable = VariablePair.Value;
 
 			FEdGraphPinType PinType = RigVMTypeUtils::PinTypeFromExternalVariable(Variable);
 			if (!PinType.PinCategory.IsValid())
@@ -199,8 +201,8 @@ void UAnimGraphNode_ControlRig::CreateCustomPins(TArray<UEdGraphPin*>* OldPins)
 		}
 	}
 
-	TArray<FControlsInfo>& Controls = GetControls();
-	for (FControlsInfo& ControlInfo : Controls)
+	const TArray<FControlRigIOMapping::FControlsInfo>& Controls = ControlRigIOMapping->GetControls();
+	for (const FControlRigIOMapping::FControlsInfo& ControlInfo : Controls)
 	{
 		const FName ControlName = ControlInfo.Name;
 		BeginExposableNames.Remove(ControlName);
@@ -287,225 +289,17 @@ void UAnimGraphNode_ControlRig::ValidateAnimNodeDuringCompilation(USkeleton* For
 	Super::ValidateAnimNodeDuringCompilation(ForSkeleton, MessageLog);
 }
 
-void UAnimGraphNode_ControlRig::RebuildExposedProperties()
-{
-	DECLARE_SCOPE_HIERARCHICAL_COUNTER_FUNC()
-
-	if (URigVMBlueprintGeneratedClass* TargetClass = Cast<URigVMBlueprintGeneratedClass>(GetTargetClass()))
-	{
-		if (UControlRigBlueprint* RigBlueprint = Cast<UControlRigBlueprint>(TargetClass->ClassGeneratedBy))
-		{
-			if (RigBlueprint->HasAllFlags(RF_NeedPostLoad))
-			{
-				return;
-			}
-		}
-	}
-
-	TSet<FName> OldOptionalPinNames;
-	TSet<FName> OldExposedPinNames;
-	for(const FOptionalPinFromProperty& OptionalPin : CustomPinProperties)
-	{
-		OldOptionalPinNames.Add(OptionalPin.PropertyName);
-
-		if(OptionalPin.bShowPin)
-		{
-			OldExposedPinNames.Add(OptionalPin.PropertyName);
-		}
-	}
-	CustomPinProperties.Empty();
-
-	// go through exposed properties, and clean up
-	GetVariables(true, InputVariables);
-	// we still update OUtputvariables
-	// we don't want output to be exposed
-	GetVariables(false, OutputVariables);
-
-	// clear IO variables that don't exist anymore
-	auto ClearInvalidMapping = [](const TMap<FName, FRigVMExternalVariable>& InVariables, TMap<FName, FName>& InOutMapping)
-	{
-		TArray<FName> KeyArray;
-		InOutMapping.GenerateKeyArray(KeyArray);
-
-		for (int32 Index=0; Index<KeyArray.Num(); ++Index)
-		{
-			// if this input doesn't exist anymore
-			if (!InVariables.Contains(KeyArray[Index]))
-			{
-				InOutMapping.Remove(KeyArray[Index]);
-			}
-		}
-	};
-
-	ClearInvalidMapping(InputVariables, Node.InputMapping);
-	ClearInvalidMapping(OutputVariables, Node.OutputMapping);
-
-	auto MakeOptionalPin = [&OldExposedPinNames](FName InPinName)-> FOptionalPinFromProperty
-	{
-		FOptionalPinFromProperty OptionalPin;
-		OptionalPin.PropertyName = InPinName;
-		OptionalPin.bShowPin = OldExposedPinNames.Contains(InPinName);
-		OptionalPin.bCanToggleVisibility = true;
-		OptionalPin.bIsOverrideEnabled = false;
-		return OptionalPin;
-	};
-	
-	for (auto Iter = InputVariables.CreateConstIterator(); Iter; ++Iter)
-	{
-		CustomPinProperties.Add(MakeOptionalPin(Iter.Key()));
-	}
-
-	// also add all of the controls
-	TArray<FControlsInfo>& Controls = GetControls();
-	for (FControlsInfo& ControlInfo : Controls)
-	{
-		const FName ControlName = ControlInfo.Name;
-		CustomPinProperties.Add(MakeOptionalPin(ControlName));
-	}
-}
-
-bool UAnimGraphNode_ControlRig::IsInputProperty(const FName& PropertyName) const
-{
-	DECLARE_SCOPE_HIERARCHICAL_COUNTER_FUNC()
-
-	// this is true for both input variables and controls
-	return InputVariables.Contains(PropertyName) || !OutputVariables.Contains(PropertyName);
-}
-
-UAnimGraphNode_ControlRig::FControlsInfo* UAnimGraphNode_ControlRig::FindControlElement(const FName& InControlName) const
-{
-	TArray<FControlsInfo>& Controls = GetControls();
-	return Controls.FindByPredicate([InControlName](const FControlsInfo& Info)
-	{
-		return Info.Name == InControlName;
-	});
-}
-
-bool UAnimGraphNode_ControlRig::IsAvailableToMapToCurve(const FName& PropertyName, bool bInput) const
-{
-	DECLARE_SCOPE_HIERARCHICAL_COUNTER_FUNC()
-
-	// find if input or output
-	// ensure it could convert to float
-	const FRigVMExternalVariable* Variable = (bInput) ? InputVariables.Find(PropertyName) : OutputVariables.Find(PropertyName);
-	if (Variable)
-	{
-		return Variable->TypeName == TEXT("float");
-	}
-
-	if(const FControlsInfo* ControlInfo = FindControlElement(PropertyName))
-	{
-		return (ControlInfo->ControlType == ERigControlType::Float || ControlInfo->ControlType == ERigControlType::ScaleFloat);
-	}
-
-	return ensure(false);
-}
-
-bool UAnimGraphNode_ControlRig::IsPropertyExposeEnabled(FName PropertyName) const
-{
-	DECLARE_SCOPE_HIERARCHICAL_COUNTER_FUNC()
-
-	// if known exposable, and and if it hasn't been exposed yet
-	if (CustomPinProperties.ContainsByPredicate([PropertyName](const FOptionalPinFromProperty& InOptionalPin){ return InOptionalPin.PropertyName == PropertyName; }))
-	{
-		return IsInputProperty(PropertyName);
-	}
-
-	return false;
-}
-
-ECheckBoxState UAnimGraphNode_ControlRig::IsPropertyExposed(FName PropertyName) const
-{
-	if (CustomPinProperties.ContainsByPredicate([PropertyName](const FOptionalPinFromProperty& InOptionalPin){ return InOptionalPin.bShowPin && InOptionalPin.PropertyName == PropertyName; }))
-	{
-		return ECheckBoxState::Checked;
-	}
-
-	return ECheckBoxState::Unchecked;
-}
-
-void UAnimGraphNode_ControlRig::OnPropertyExposeCheckboxChanged(ECheckBoxState NewState, FName PropertyName)
-{
-	DECLARE_SCOPE_HIERARCHICAL_COUNTER_FUNC()
-
-	FOptionalPinFromProperty* FoundPin = CustomPinProperties.FindByPredicate([PropertyName](const FOptionalPinFromProperty& InOptionalPin)
-	{
-		return InOptionalPin.PropertyName == PropertyName;
-	});
-	
-	if (FoundPin)
-	{
-		FoundPin->bShowPin = !FoundPin->bShowPin;
-
-		FGuardValue_Bitfield(bDisableOrphanPinSaving, true);
-		ReconstructNode();
-		
-		// see if any of my child has the mapping, and clear them
-		if (NewState == ECheckBoxState::Checked)
-		{
-			FScopedTransaction Transaction(LOCTEXT("PropertyExposedChanged", "Expose Property to Pin"));
-			Modify();
-
-			FBlueprintEditorUtils::MarkBlueprintAsStructurallyModified(GetBlueprint());
-
-			bool bInput = IsInputProperty(PropertyName);
-			// if checked, we clear mapping
-			// and unclear all children
-			Node.SetIOMapping(bInput, PropertyName, NAME_None); 
-		}
-	}
-}
-
 void UAnimGraphNode_ControlRig::CustomizeDetails(IDetailLayoutBuilder& DetailBuilder)
 {
 	DECLARE_SCOPE_HIERARCHICAL_COUNTER_FUNC()
 
 	Super::CustomizeDetails(DetailBuilder);
 
-	// We dont allow multi-select here
-	if (DetailBuilder.GetSelectedObjects().Num() > 1)
+	if (!ControlRigIOMapping->CreateVariableMappingWidget(DetailBuilder))
 	{
 		return;
 	}
 
-	// input/output exposure feature START
-	RebuildExposedProperties();
-
-	IDetailCategoryBuilder& InputCategoryBuilder = DetailBuilder.EditCategory(FName(TEXT("Input")));
-
-	FDetailWidgetRow& InputWidgetRow = InputCategoryBuilder.AddCustomRow(FText::FromString(TEXT("Input")));
-#if WITH_EDITOR
-	InputWidgetRow.WholeRowContent()
-	[
-		SNew(SRigVMVariableMappingWidget)
-		.OnVariableMappingChanged(FOnRigVMVariableMappingChanged::CreateUObject(this, &UAnimGraphNode_ControlRig::OnVariableMappingChanged, true))
-		.OnGetVariableMapping(FOnRigVMGetVariableMapping::CreateUObject(this, &UAnimGraphNode_ControlRig::GetVariableMapping, true))
-		.OnGetAvailableMapping(FOnRigVMGetAvailableMapping::CreateUObject(this, &UAnimGraphNode_ControlRig::GetAvailableMapping, true))
-		.OnCreateVariableMapping(FOnRigVMCreateVariableMapping::CreateUObject(this, &UAnimGraphNode_ControlRig::CreateVariableMapping, true))
-		.OnVariableOptionAvailable(FOnRigVMVarOptionAvailable::CreateUObject(this, &UAnimGraphNode_ControlRig::IsAvailableToMapToCurve, true))
-		.OnPinGetCheckState(FOnRigVMPinGetCheckState::CreateUObject(this, &UAnimGraphNode_ControlRig::IsPropertyExposed))
-		.OnPinCheckStateChanged(FOnRigVMPinCheckStateChanged::CreateUObject(this, &UAnimGraphNode_ControlRig::OnPropertyExposeCheckboxChanged))
-		.OnPinIsEnabledCheckState(FOnRigVMPinIsCheckEnabled::CreateUObject(this, &UAnimGraphNode_ControlRig::IsPropertyExposeEnabled))
-	];
-#endif
-
-	IDetailCategoryBuilder& OutputCategoryBuilder = DetailBuilder.EditCategory(FName(TEXT("Output")));
-	FDetailWidgetRow& OutputWidgetRow = OutputCategoryBuilder.AddCustomRow(FText::FromString(TEXT("Output")));
-#if WITH_EDITOR
-	OutputWidgetRow.WholeRowContent()
-	[
-		SNew(SRigVMVariableMappingWidget)
-		.OnVariableMappingChanged(FOnRigVMVariableMappingChanged::CreateUObject(this, &UAnimGraphNode_ControlRig::OnVariableMappingChanged, false))
-		.OnGetVariableMapping(FOnRigVMGetVariableMapping::CreateUObject(this, &UAnimGraphNode_ControlRig::GetVariableMapping, false))
-		.OnGetAvailableMapping(FOnRigVMGetAvailableMapping::CreateUObject(this, &UAnimGraphNode_ControlRig::GetAvailableMapping, false))
-		.OnCreateVariableMapping(FOnRigVMCreateVariableMapping::CreateUObject(this, &UAnimGraphNode_ControlRig::CreateVariableMapping, false))
-		.OnVariableOptionAvailable(FOnRigVMVarOptionAvailable::CreateUObject(this, &UAnimGraphNode_ControlRig::IsAvailableToMapToCurve, false))
-		.OnPinGetCheckState(FOnRigVMPinGetCheckState::CreateUObject(this, &UAnimGraphNode_ControlRig::IsPropertyExposed))
-		.OnPinCheckStateChanged(FOnRigVMPinCheckStateChanged::CreateUObject(this, &UAnimGraphNode_ControlRig::OnPropertyExposeCheckboxChanged))
-		.OnPinIsEnabledCheckState(FOnRigVMPinIsCheckEnabled::CreateUObject(this, &UAnimGraphNode_ControlRig::IsPropertyExposeEnabled))
-	];
-#endif
-
 	TSharedRef<IPropertyHandle> ClassHandle = DetailBuilder.GetProperty(TEXT("Node.ControlRigClass"), GetClass());
 	if (ClassHandle->IsValidHandle())
 	{
@@ -542,199 +336,6 @@ void UAnimGraphNode_ControlRig::CustomizeDetails(IDetailLayoutBuilder& DetailBui
 	// alpha property blending support END
 }
 
-TArray<UAnimGraphNode_ControlRig::FControlsInfo>& UAnimGraphNode_ControlRig::GetControls() const
-{
-	const UClass* ControlRigClass = GetTargetClass();
-
-	if (ControlsInfoClass != ControlRigClass)
-	{
-		ControlsInfo.Reset();
-		ControlsInfoClass = ControlRigClass;
-		
-		if (ControlRigClass)
-		{
-			if (USkeleton* Skeleton = GetAnimBlueprint()->TargetSkeleton)
-			{
-				UControlRig* TemplateRig = NewObject<UControlRig>(GetTransientPackage(), ControlRigClass, NAME_None, RF_Transient);
-				TemplateRig->SetObjectBinding(MakeShared<FControlRigObjectBinding>());
-				TemplateRig->GetObjectBinding()->BindToObject(Skeleton);
-				TemplateRig->GetDataSourceRegistry()->RegisterDataSource(UControlRig::OwnerComponent, TemplateRig->GetObjectBinding()->GetBoundObject());
-
-				TemplateRig->Initialize();
-				TemplateRig->SetBoneInitialTransformsFromRefSkeleton(Skeleton->GetReferenceSkeleton());
-				{
-					// Empty event queue to evaluate only construction
-					TGuardValue<TArray<FName>> EventQueueGuard(TemplateRig->EventQueue, {});
-					TemplateRig->Evaluate_AnyThread();
-				}
-
-				if(const URigHierarchy* Hierarchy = TemplateRig->GetHierarchy())
-				{
-					Hierarchy->ForEach<FRigControlElement>([&](FRigControlElement* ControlElement) -> bool
-					{
-						if (Hierarchy->IsAnimatable(ControlElement))
-						{
-							FControlsInfo Info;
-							Info.Name = ControlElement->GetFName();
-							Info.DisplayName = ControlElement->GetName();
-							Info.PinType = TemplateRig->GetHierarchy()->GetControlPinType(ControlElement);
-							Info.ControlType = ControlElement->Settings.ControlType;
-							Info.DefaultValue = TemplateRig->GetHierarchy()->GetControlPinDefaultValue(ControlElement, true);
-							ControlsInfo.Add(Info);
-						}
-						return true;
-					});
-				}
-
-				TemplateRig->MarkAsGarbage();
-			}
-		}
-	}
-
-	return ControlsInfo;
-}
-
-void UAnimGraphNode_ControlRig::GetVariables(bool bInput, TMap<FName, FRigVMExternalVariable>& OutVariables) const
-{
-	DECLARE_SCOPE_HIERARCHICAL_COUNTER_FUNC()
-
-	OutVariables.Reset();
-
-	if (URigVMBlueprintGeneratedClass* TargetClass = Cast<URigVMBlueprintGeneratedClass>(GetTargetClass()))
-	{
-		if (UControlRigBlueprint* RigBlueprint = Cast<UControlRigBlueprint>(TargetClass->ClassGeneratedBy))
-		{
-			//RigBlueprint->CleanupVariables();
-			UControlRig* ControlRig = TargetClass->GetDefaultObject<UControlRig>();
-			if (ControlRig)
-			{
-				const TArray<FRigVMExternalVariable>& PublicVariables = ControlRig->GetPublicVariables();
-				for (const FRigVMExternalVariable& PublicVariable : PublicVariables)
-				{
-					if (!bInput || !PublicVariable.bIsReadOnly)
-					{
-						OutVariables.Add(PublicVariable.Name, PublicVariable);
-					}
-				}
-			}
-		}
-	}
-}
-
-#if WITH_EDITOR
-
-void UAnimGraphNode_ControlRig::OnVariableMappingChanged(const FName& PathName, const FName& Curve, bool bInput)
-{
-	DECLARE_SCOPE_HIERARCHICAL_COUNTER_FUNC()
-
-	FScopedTransaction Transaction(LOCTEXT("VariableMappingChanged", "Change Variable Mapping"));
-	Modify();
-
-	FBlueprintEditorUtils::MarkBlueprintAsStructurallyModified(GetBlueprint());
-
-	// @todo: this is not enough when we start breaking down struct
-	Node.SetIOMapping(bInput, PathName, Curve);
-}
-
-FName UAnimGraphNode_ControlRig::GetVariableMapping(const FName& PathName, bool bInput)
-{
-	DECLARE_SCOPE_HIERARCHICAL_COUNTER_FUNC()
-
-	// @todo: this is not enough when we start breaking down struct
-	return Node.GetIOMapping(bInput, PathName);
-}
-
-void UAnimGraphNode_ControlRig::GetAvailableMapping(const FName& PathName, TArray<FName>& OutArray, bool bInput)
-{
-	DECLARE_SCOPE_HIERARCHICAL_COUNTER_FUNC()
-
-	UAnimBlueprint* AnimBP = CastChecked<UAnimBlueprint>(GetBlueprint());
-	USkeleton* TargetSkeleton = AnimBP->TargetSkeleton;
-	OutArray.Reset();
-	if (TargetSkeleton)
-	{
-		TargetSkeleton->GetCurveMetaDataNames(OutArray);
-		
-		// also add all controls
-		TArray<FControlsInfo>& Controls = GetControls();
-		for(FControlsInfo& ControlInfo : Controls)
-		{
-			OutArray.Add(ControlInfo.Name);
-		}
-
-		// we want to exclude anything that has been mapped already
-		TArray<FName> InputMapping, OutputMapping;
-		Node.InputMapping.GenerateValueArray(InputMapping);
-		Node.OutputMapping.GenerateValueArray(OutputMapping);
-
-		// I have to remove Input/Output Curves from OutArray
-		for (int32 Index = 0; Index < OutArray.Num(); ++Index)
-		{
-			const FName& Item = OutArray[Index];
-
-			if (InputMapping.Contains(Item))
-			{
-				OutArray.RemoveAt(Index);
-				--Index;
-			}
-			else if (OutputMapping.Contains(Item))
-			{
-				OutArray.RemoveAt(Index);
-				--Index;
-			}
-		}
-	}
-}
-
-void UAnimGraphNode_ControlRig::CreateVariableMapping(const FString& FilteredText, TArray< TSharedPtr<FRigVMVariableMappingInfo> >& OutArray, bool bInput)
-{
-	DECLARE_SCOPE_HIERARCHICAL_COUNTER_FUNC()
-
-	// should have latest
-	OutArray.Reset();
-
-	bool bDoFiltering = !FilteredText.IsEmpty();
-	const TMap<FName, FRigVMExternalVariable>& Variables = (bInput)? InputVariables : OutputVariables;
-	for (auto Iter = Variables.CreateConstIterator(); Iter; ++Iter)
-	{
-		const FName& Name = Iter.Key();
-		const FString& DisplayName = Name.ToString();
-
-		const FString MappedName = GetVariableMapping(Name, bInput).ToString();
-		// make sure it doesn't fit any of them
-		if (!bDoFiltering || 
-			(DisplayName.Contains(FilteredText) || MappedName.Contains(FilteredText)))
-		{
-			TSharedRef<FRigVMVariableMappingInfo> Item = FRigVMVariableMappingInfo::Make(Iter.Key());
-
-			FRigVMVariableMappingInfo& ItemRaw = Item.Get();
-			FString PathString = ItemRaw.GetPathName().ToString();
-			FString DisplayString = ItemRaw.GetDisplayName();
-
-			OutArray.Add(Item);
-		}
-	}
-
-	if(bInput)
-	{
-		// add all controls
-		TArray<FControlsInfo>& Controls = GetControls();
-		for (FControlsInfo& ControlInfo : Controls)
-		{
-			const FName ControlName = ControlInfo.Name;
-			const FString& DisplayName = ControlInfo.DisplayName;
-
-			if (!bDoFiltering || DisplayName.Contains(FilteredText))
-			{
-				TSharedRef<FRigVMVariableMappingInfo> Item = FRigVMVariableMappingInfo::Make(ControlName);
-				OutArray.Add(Item);
-			}
-		}
-	}
-}
-
-#endif
-
 void UAnimGraphNode_ControlRig::PostEditChangeProperty(FPropertyChangedEvent& PropertyChangedEvent)
 {
 	DECLARE_SCOPE_HIERARCHICAL_COUNTER_FUNC()
@@ -749,7 +350,7 @@ void UAnimGraphNode_ControlRig::PostEditChangeProperty(FPropertyChangedEvent& Pr
 		if (ChangedProperty->GetFName() == GET_MEMBER_NAME_CHECKED(FAnimNode_ControlRig, ControlRigClass))
 		{
 			bRequiresNodeReconstruct = true;
-			RebuildExposedProperties();
+			ControlRigIOMapping->RebuildExposedProperties();
 		}
 
 		if (ChangedProperty->GetFName() == GET_MEMBER_NAME_STRING_CHECKED(FAnimNode_ControlRig, AlphaInputType))
@@ -838,8 +439,8 @@ void UAnimGraphNode_ControlRig::PostReconstructNode()
 						FString DefaultValue = FixTransform.ToString();
 
 						const FRigVMExternalVariable* RigVMExternalVariable = (CurrentPin->Direction == EGPD_Input) ? 
-							InputVariables.Find(CurrentPin->GetFName()) : (CurrentPin->Direction == EGPD_Output) ? 
-							OutputVariables.Find(CurrentPin->GetFName()) : nullptr;
+							ControlRigIOMapping->GetInputVariables().Find(CurrentPin->GetFName()) : (CurrentPin->Direction == EGPD_Output) ?
+							ControlRigIOMapping->GetOutputVariables().Find(CurrentPin->GetFName()) : nullptr;
 
 						if (RigVMExternalVariable != nullptr && RigVMExternalVariable->IsValid())
 						{
@@ -901,5 +502,48 @@ void UAnimGraphNode_ControlRig::CustomizePinData(UEdGraphPin* Pin, FName SourceP
 		}
 	}
 }
+
+void UAnimGraphNode_ControlRig::OnPropertyExposeCheckboxChanged(ECheckBoxState NewState, FName PropertyName)
+{
+	DECLARE_SCOPE_HIERARCHICAL_COUNTER_FUNC()
+
+	FGuardValue_Bitfield(bDisableOrphanPinSaving, true);
+	ReconstructNode();
+
+	// see if any of my child has the mapping, and clear them
+	if (NewState == ECheckBoxState::Checked)
+	{
+		FScopedTransaction Transaction(LOCTEXT("PropertyExposedChanged", "Expose Property to Pin"));
+		Modify();
+
+		FBlueprintEditorUtils::MarkBlueprintAsStructurallyModified(GetBlueprint());
+
+		bool bInput = ControlRigIOMapping->IsInputProperty(PropertyName);
+		// if checked, we clear mapping
+		// and unclear all children
+		ControlRigIOMapping->SetIOMapping(bInput, PropertyName, NAME_None);
+	}
+}
+
+void UAnimGraphNode_ControlRig::OnVariableMappingChanged(const FName& PathName, const FName& Curve, bool bInput)
+{
+	FScopedTransaction Transaction(LOCTEXT("VariableMappingChanged", "Change Variable Mapping"));
+	Modify();
+
+	FBlueprintEditorUtils::MarkBlueprintAsStructurallyModified(GetBlueprint());
+
+	// @todo: this is not enough when we start breaking down struct
+	ControlRigIOMapping->SetIOMapping(bInput, PathName, Curve);
+}
+
+USkeleton* UAnimGraphNode_ControlRig::GetTargetSkeleton() const
+{
+	USkeleton* TargetSkeleton = nullptr;
+	if (UAnimBlueprint* AnimBP = CastChecked<UAnimBlueprint>(GetBlueprint()))
+	{
+		TargetSkeleton = AnimBP->TargetSkeleton;
+	}
+	return TargetSkeleton;
+}
 #undef LOCTEXT_NAMESPACE
 
diff --git a/Engine/Plugins/Animation/ControlRig/Source/ControlRigDeveloper/Private/ControlRigIOMapping.cpp b/Engine/Plugins/Animation/ControlRig/Source/ControlRigDeveloper/Private/ControlRigIOMapping.cpp
new file mode 100644
--- /dev/null
+++ b/Engine/Plugins/Animation/ControlRig/Source/ControlRigDeveloper/Private/ControlRigIOMapping.cpp
@@ -0,0 +1,487 @@
+// Copyright Epic Games, Inc. All Rights Reserved.
+
+#include "ControlRigIOMapping.h"
+#include "DetailLayoutBuilder.h"
+#include "ControlRigBlueprint.h"
+#include "ControlRig.h"
+#include "ControlRigObjectBinding.h"
+#include "DetailWidgetRow.h"
+#include "Widgets/SRigVMVariableMappingWidget.h"
+#include "Kismet2/BlueprintEditorUtils.h"
+#include "K2Node.h"
+#include "Units/Execution/RigUnit_PrepareForExecution.h"
+
+FControlRigIOMapping::FControlRigIOMapping(TMap<FName, FName>& InInputMapping, TMap<FName, FName>& InOutputMapping, TArray<FOptionalPinFromProperty>& InCustomPinProperties)
+	: InputMapping(InInputMapping)
+	, OutputMapping(InOutputMapping)
+	, CustomPinProperties(InCustomPinProperties)
+{
+}
+
+FControlRigIOMapping::~FControlRigIOMapping()
+{
+}
+
+bool FControlRigIOMapping::CreateVariableMappingWidget(IDetailLayoutBuilder& DetailBuilder)
+{
+	DECLARE_SCOPE_HIERARCHICAL_COUNTER_FUNC()
+
+	// We dont allow multi-select here
+	if (DetailBuilder.GetSelectedObjects().Num() > 1)
+	{
+		return false;
+	}
+	
+	if (GetTargetClass() == nullptr)
+	{
+		return false;
+	}
+
+	// input/output exposure feature START
+	RebuildExposedProperties();
+
+	IDetailCategoryBuilder& InputCategoryBuilder = DetailBuilder.EditCategory(FName(TEXT("Input")));
+
+	FDetailWidgetRow& InputWidgetRow = InputCategoryBuilder.AddCustomRow(FText::FromString(TEXT("Input")));
+#if WITH_EDITOR
+	InputWidgetRow.WholeRowContent()
+		[
+			SNew(SRigVMVariableMappingWidget)
+				.OnVariableMappingChanged(FOnRigVMVariableMappingChanged::CreateSP(this, &FControlRigIOMapping::OnVariableMappingChanged, true))
+				.OnGetVariableMapping(FOnRigVMGetVariableMapping::CreateSP(this, &FControlRigIOMapping::GetVariableMapping, true))
+				.OnGetAvailableMapping(FOnRigVMGetAvailableMapping::CreateSP(this, &FControlRigIOMapping::GetAvailableMapping, true))
+				.OnCreateVariableMapping(FOnRigVMCreateVariableMapping::CreateSP(this, &FControlRigIOMapping::CreateVariableMapping, true))
+				.OnVariableOptionAvailable(FOnRigVMVarOptionAvailable::CreateSP(this, &FControlRigIOMapping::IsAvailableToMapToCurve, true))
+				.OnPinGetCheckState(FOnRigVMPinGetCheckState::CreateSP(this, &FControlRigIOMapping::IsPropertyExposed))
+				.OnPinCheckStateChanged(FOnRigVMPinCheckStateChanged::CreateSP(this, &FControlRigIOMapping::OnPropertyExposeCheckboxChanged))
+				.OnPinIsEnabledCheckState(FOnRigVMPinIsCheckEnabled::CreateSP(this, &FControlRigIOMapping::IsPropertyExposeEnabled))
+		];
+#endif
+
+	IDetailCategoryBuilder& OutputCategoryBuilder = DetailBuilder.EditCategory(FName(TEXT("Output")));
+	FDetailWidgetRow& OutputWidgetRow = OutputCategoryBuilder.AddCustomRow(FText::FromString(TEXT("Output")));
+#if WITH_EDITOR
+	OutputWidgetRow.WholeRowContent()
+		[
+			SNew(SRigVMVariableMappingWidget)
+				.OnVariableMappingChanged(FOnRigVMVariableMappingChanged::CreateSP(this, &FControlRigIOMapping::OnVariableMappingChanged, false))
+				.OnGetVariableMapping(FOnRigVMGetVariableMapping::CreateSP(this, &FControlRigIOMapping::GetVariableMapping, false))
+				.OnGetAvailableMapping(FOnRigVMGetAvailableMapping::CreateSP(this, &FControlRigIOMapping::GetAvailableMapping, false))
+				.OnCreateVariableMapping(FOnRigVMCreateVariableMapping::CreateSP(this, &FControlRigIOMapping::CreateVariableMapping, false))
+				.OnVariableOptionAvailable(FOnRigVMVarOptionAvailable::CreateSP(this, &FControlRigIOMapping::IsAvailableToMapToCurve, false))
+				.OnPinGetCheckState(FOnRigVMPinGetCheckState::CreateSP(this, &FControlRigIOMapping::IsPropertyExposed))
+				.OnPinCheckStateChanged(FOnRigVMPinCheckStateChanged::CreateSP(this, &FControlRigIOMapping::OnPropertyExposeCheckboxChanged))
+				.OnPinIsEnabledCheckState(FOnRigVMPinIsCheckEnabled::CreateSP(this, &FControlRigIOMapping::IsPropertyExposeEnabled))
+		];
+#endif
+
+	return true;
+}
+
+void FControlRigIOMapping::RebuildExposedProperties()
+{
+	DECLARE_SCOPE_HIERARCHICAL_COUNTER_FUNC()
+
+	const URigVMBlueprintGeneratedClass* TargetClass = Cast<URigVMBlueprintGeneratedClass>(GetTargetClass());
+	if (TargetClass)
+	{
+		if (const UControlRigBlueprint* RigBlueprint = Cast<UControlRigBlueprint>(TargetClass->ClassGeneratedBy))
+		{
+			if (RigBlueprint->HasAllFlags(RF_NeedPostLoad))
+			{
+				return;
+			}
+		}
+	}
+	else
+	{
+		return;
+	}
+
+	TSet<FName> OldOptionalPinNames;
+	TSet<FName> OldExposedPinNames;
+	for (const FOptionalPinFromProperty& OptionalPin : CustomPinProperties)
+	{
+		OldOptionalPinNames.Add(OptionalPin.PropertyName);
+
+		if (OptionalPin.bShowPin)
+		{
+			OldExposedPinNames.Add(OptionalPin.PropertyName);
+		}
+	}
+	CustomPinProperties.Reset();
+
+	// go through exposed properties, and clean up
+	GetVariables(TargetClass, true, InputVariables);
+	// we still update OUtputvariables
+	// we don't want output to be exposed
+	GetVariables(TargetClass, false, OutputVariables);
+
+	// clear IO variables that don't exist anymore
+	auto ClearInvalidMapping = [](const TMap<FName, FRigVMExternalVariable>& InVariables, TMap<FName, FName>& InOutMapping)
+		{
+			TArray<FName> KeyArray;
+			InOutMapping.GenerateKeyArray(KeyArray);
+
+			for (int32 Index = 0; Index < KeyArray.Num(); ++Index)
+			{
+				// if this input doesn't exist anymore
+				if (!InVariables.Contains(KeyArray[Index]))
+				{
+					InOutMapping.Remove(KeyArray[Index]);
+				}
+			}
+		};
+
+	ClearInvalidMapping(InputVariables, InputMapping);
+	ClearInvalidMapping(OutputVariables, OutputMapping);
+
+	auto MakeOptionalPin = [&OldExposedPinNames](FName InPinName)-> FOptionalPinFromProperty
+		{
+			FOptionalPinFromProperty OptionalPin;
+			OptionalPin.PropertyName = InPinName;
+			OptionalPin.bShowPin = OldExposedPinNames.Contains(InPinName);
+			OptionalPin.bCanToggleVisibility = true;
+			OptionalPin.bIsOverrideEnabled = false;
+			return OptionalPin;
+		};
+
+	for (auto Iter = InputVariables.CreateConstIterator(); Iter; ++Iter)
+	{
+		CustomPinProperties.Add(MakeOptionalPin(Iter.Key()));
+	}
+
+	// also add all of the controls
+	const TArray<FControlsInfo>& Controls = GetControls();
+	for (const FControlsInfo& ControlInfo : Controls)
+	{
+		CustomPinProperties.Add(MakeOptionalPin(ControlInfo.Name));
+	}
+}
+
+bool FControlRigIOMapping::IsInputProperty(const FName& PropertyName) const
+{
+	DECLARE_SCOPE_HIERARCHICAL_COUNTER_FUNC()
+
+	// this is true for both input variables and controls
+	return InputVariables.Contains(PropertyName) || !OutputVariables.Contains(PropertyName);
+}
+
+void FControlRigIOMapping::SetIOMapping(bool bInput, const FName& SourceProperty, const FName& TargetCurve)
+{
+	DECLARE_SCOPE_HIERARCHICAL_COUNTER_FUNC()
+
+	const UClass* TargetClass = GetTargetClass();
+	if (TargetClass)
+	{
+		const UControlRig* CDO = TargetClass->GetDefaultObject<UControlRig>();
+		if (CDO)
+		{
+			TMap<FName, FName>& MappingData = (bInput) ? InputMapping : OutputMapping;
+
+			// if it's valid as of now, we add it
+			bool bIsReadOnly = CDO->GetPublicVariableByName(SourceProperty).bIsReadOnly;
+			if (!bInput || !bIsReadOnly)
+			{
+				if (TargetCurve == NAME_None)
+				{
+					MappingData.Remove(SourceProperty);
+				}
+				else
+				{
+					MappingData.FindOrAdd(SourceProperty) = TargetCurve;
+				}
+			}
+		}
+	}
+}
+
+FName FControlRigIOMapping::GetIOMapping(bool bInput, const FName& SourceProperty) const
+{
+	DECLARE_SCOPE_HIERARCHICAL_COUNTER_FUNC()
+
+	const TMap<FName, FName>& MappingData = (bInput) ? InputMapping : OutputMapping;
+	if (const FName* NameFound = MappingData.Find(SourceProperty))
+	{
+		return *NameFound;
+	}
+
+	return NAME_None;
+}
+
+const FControlRigIOMapping::FControlsInfo* FControlRigIOMapping::FindControlElement(const FName& InControlName) const
+{
+	DECLARE_SCOPE_HIERARCHICAL_COUNTER_FUNC()
+
+	const TArray<FControlsInfo>& Controls = GetControls();
+	return Controls.FindByPredicate([InControlName](const FControlsInfo& Info)
+		{
+			return Info.Name == InControlName;
+		});
+}
+
+bool FControlRigIOMapping::IsAvailableToMapToCurve(const FName& PropertyName, bool bInput) const
+{
+	DECLARE_SCOPE_HIERARCHICAL_COUNTER_FUNC()
+
+	// find if input or output
+	// ensure it could convert to float
+	const FRigVMExternalVariable* Variable = (bInput) ? InputVariables.Find(PropertyName) : OutputVariables.Find(PropertyName);
+	if (Variable)
+	{
+		return Variable->TypeName == TEXT("float");
+	}
+
+	if (const FControlsInfo* ControlInfo = FindControlElement(PropertyName))
+	{
+		return (ControlInfo->ControlType == ERigControlType::Float || ControlInfo->ControlType == ERigControlType::ScaleFloat);
+	}
+
+	return ensure(false);
+}
+
+bool FControlRigIOMapping::IsPropertyExposeEnabled(FName PropertyName) const
+{
+	DECLARE_SCOPE_HIERARCHICAL_COUNTER_FUNC()
+
+	// if known exposable, and and if it hasn't been exposed yet
+	if (CustomPinProperties.ContainsByPredicate([PropertyName](const FOptionalPinFromProperty& InOptionalPin) { return InOptionalPin.PropertyName == PropertyName; }))
+	{
+		return IsInputProperty(PropertyName);
+	}
+
+	return false;
+}
+
+ECheckBoxState FControlRigIOMapping::IsPropertyExposed(FName PropertyName) const
+{
+	DECLARE_SCOPE_HIERARCHICAL_COUNTER_FUNC()
+
+	if (CustomPinProperties.ContainsByPredicate([PropertyName](const FOptionalPinFromProperty& InOptionalPin) { return InOptionalPin.bShowPin && InOptionalPin.PropertyName == PropertyName; }))
+	{
+		return ECheckBoxState::Checked;
+	}
+
+	return ECheckBoxState::Unchecked;
+}
+
+void FControlRigIOMapping::OnPropertyExposeCheckboxChanged(ECheckBoxState NewState, FName PropertyName)
+{
+	DECLARE_SCOPE_HIERARCHICAL_COUNTER_FUNC()
+
+	FOptionalPinFromProperty* FoundPin = CustomPinProperties.FindByPredicate([PropertyName](const FOptionalPinFromProperty& InOptionalPin)
+		{
+			return InOptionalPin.PropertyName == PropertyName;
+		});
+
+	if (FoundPin)
+	{
+		FoundPin->bShowPin = !FoundPin->bShowPin;
+
+		if (OnPinCheckStateChangedDelegate.IsBound())
+		{
+			OnPinCheckStateChangedDelegate.Execute(NewState, PropertyName);
+		}
+	}
+}
+
+const TArray<FControlRigIOMapping::FControlsInfo>& FControlRigIOMapping::GetControls() const
+{
+	DECLARE_SCOPE_HIERARCHICAL_COUNTER_FUNC()
+
+	const UClass* ControlRigClass = GetTargetClass();
+
+	if (ControlsInfoClass != ControlRigClass)
+	{
+		ControlsInfo.Reset();
+		ControlsInfoClass = ControlRigClass;
+
+		if (ControlRigClass)
+		{
+			if (USkeleton* Skeleton = GetTargetSkeleton())
+			{
+				UControlRig* TemplateRig = NewObject<UControlRig>(GetTransientPackage(), ControlRigClass, NAME_None, RF_Transient);
+				TemplateRig->SetObjectBinding(MakeShared<FControlRigObjectBinding>());
+				TemplateRig->GetObjectBinding()->BindToObject(Skeleton);
+				TemplateRig->GetDataSourceRegistry()->RegisterDataSource(UControlRig::OwnerComponent, TemplateRig->GetObjectBinding()->GetBoundObject());
+
+				TemplateRig->Initialize();
+				TemplateRig->SetBoneInitialTransformsFromRefSkeleton(Skeleton->GetReferenceSkeleton());
+				{
+					// Empty event queue to evaluate only construction
+					TGuardValue<TArray<FName>> EventQueueGuard(TemplateRig->EventQueue, {});
+					TemplateRig->Evaluate_AnyThread();
+				}
+
+				if (const URigHierarchy* Hierarchy = TemplateRig->GetHierarchy())
+				{
+					Hierarchy->ForEach<FRigControlElement>([&](FRigControlElement* ControlElement) -> bool
+						{
+							if (Hierarchy->IsAnimatable(ControlElement))
+							{
+								FControlsInfo Info;
+								Info.Name = ControlElement->GetFName();
+								Info.DisplayName = ControlElement->GetName();
+								Info.PinType = TemplateRig->GetHierarchy()->GetControlPinType(ControlElement);
+								Info.ControlType = ControlElement->Settings.ControlType;
+								Info.DefaultValue = TemplateRig->GetHierarchy()->GetControlPinDefaultValue(ControlElement, true);
+								ControlsInfo.Add(Info);
+							}
+							return true;
+						});
+				}
+
+				TemplateRig->MarkAsGarbage();
+			}
+		}
+	}
+
+	return ControlsInfo;
+}
+
+void FControlRigIOMapping::GetVariables(const UClass* TargetClass, bool bInput, TMap<FName, FRigVMExternalVariable>& OutVariables)
+{
+	DECLARE_SCOPE_HIERARCHICAL_COUNTER_FUNC()
+
+	OutVariables.Reset();
+
+	if (TargetClass != nullptr)
+	{
+		if (UControlRigBlueprint* RigBlueprint = Cast<UControlRigBlueprint>(TargetClass->ClassGeneratedBy))
+		{
+			//RigBlueprint->CleanupVariables();
+			UControlRig* ControlRig = TargetClass->GetDefaultObject<UControlRig>();
+			if (ControlRig)
+			{
+				const TArray<FRigVMExternalVariable>& PublicVariables = ControlRig->GetPublicVariables();
+				for (const FRigVMExternalVariable& PublicVariable : PublicVariables)
+				{
+					if (!bInput || !PublicVariable.bIsReadOnly)
+					{
+						OutVariables.Add(PublicVariable.Name, PublicVariable);
+					}
+				}
+			}
+		}
+	}
+}
+
+#if WITH_EDITOR
+
+void FControlRigIOMapping::OnVariableMappingChanged(const FName& PathName, const FName& Curve, bool bInput)
+{
+	DECLARE_SCOPE_HIERARCHICAL_COUNTER_FUNC()
+
+	if (OnVariableMappingChangedDelegate.IsBound())
+	{
+		OnVariableMappingChangedDelegate.Execute(PathName, Curve, bInput);
+	}
+}
+
+FName FControlRigIOMapping::GetVariableMapping(const FName& PathName, bool bInput)
+{
+	DECLARE_SCOPE_HIERARCHICAL_COUNTER_FUNC()
+
+	// @todo: this is not enough when we start breaking down struct
+	return GetIOMapping(bInput, PathName);
+}
+
+void FControlRigIOMapping::GetAvailableMapping(const FName& PathName, TArray<FName>& OutArray, bool bInput)
+{
+	DECLARE_SCOPE_HIERARCHICAL_COUNTER_FUNC()
+
+	OutArray.Reset();
+
+	if (USkeleton * TargetSkeleton = GetTargetSkeleton())
+	{
+		TargetSkeleton->GetCurveMetaDataNames(OutArray);
+
+		// also add all controls
+		const TArray<FControlsInfo>& Controls = GetControls();
+		for (const FControlsInfo& ControlInfo : Controls)
+		{
+			OutArray.Add(ControlInfo.Name);
+		}
+
+		// we want to exclude anything that has been mapped already
+		TArray<FName> TmpInputMapping, TmpOutputMapping;
+		InputMapping.GenerateValueArray(TmpInputMapping);
+		OutputMapping.GenerateValueArray(TmpOutputMapping);
+
+		// I have to remove Input/Output Curves from OutArray
+		for (int32 Index = 0; Index < OutArray.Num(); ++Index)
+		{
+			const FName& Item = OutArray[Index];
+
+			if (TmpInputMapping.Contains(Item))
+			{
+				OutArray.RemoveAt(Index);
+				--Index;
+			}
+			else if (TmpOutputMapping.Contains(Item))
+			{
+				OutArray.RemoveAt(Index);
+				--Index;
+			}
+		}
+	}
+}
+
+void FControlRigIOMapping::CreateVariableMapping(const FString& FilteredText, TArray< TSharedPtr<FRigVMVariableMappingInfo> >& OutArray, bool bInput)
+{
+	DECLARE_SCOPE_HIERARCHICAL_COUNTER_FUNC()
+
+	// should have latest
+	OutArray.Reset();
+
+	bool bDoFiltering = !FilteredText.IsEmpty();
+	const TMap<FName, FRigVMExternalVariable>& Variables = (bInput) ? InputVariables : OutputVariables;
+	for (auto Iter = Variables.CreateConstIterator(); Iter; ++Iter)
+	{
+		const FName& Name = Iter.Key();
+		const FString& DisplayName = Name.ToString();
+
+		const FString MappedName = GetVariableMapping(Name, bInput).ToString();
+		// make sure it doesn't fit any of them
+		if (!bDoFiltering ||
+			(DisplayName.Contains(FilteredText) || MappedName.Contains(FilteredText)))
+		{
+			TSharedRef<FRigVMVariableMappingInfo> Item = FRigVMVariableMappingInfo::Make(Iter.Key());
+
+			FRigVMVariableMappingInfo& ItemRaw = Item.Get();
+			FString PathString = ItemRaw.GetPathName().ToString();
+			FString DisplayString = ItemRaw.GetDisplayName();
+
+			OutArray.Add(Item);
+		}
+	}
+
+	if (bInput)
+	{
+		// add all controls
+		const TArray<FControlsInfo>& Controls = GetControls();
+		for (const FControlsInfo& ControlInfo : Controls)
+		{
+			const FName ControlName = ControlInfo.Name;
+			const FString& DisplayName = ControlInfo.DisplayName;
+
+			if (!bDoFiltering || DisplayName.Contains(FilteredText))
+			{
+				TSharedRef<FRigVMVariableMappingInfo> Item = FRigVMVariableMappingInfo::Make(ControlName);
+				OutArray.Add(Item);
+			}
+		}
+	}
+}
+
+#endif
+
+const UClass* FControlRigIOMapping::GetTargetClass() const
+{
+	return (OnGetTargetClassDelegate.IsBound()) ? OnGetTargetClassDelegate.Execute() : nullptr;
+}
+
+USkeleton* FControlRigIOMapping::GetTargetSkeleton() const
+{
+	return (OnGetTargetSkeletonDelegate.IsBound()) ? OnGetTargetSkeletonDelegate.Execute() : nullptr;
+}
diff --git a/Engine/Plugins/Animation/ControlRig/Source/ControlRigDeveloper/Public/AnimGraphNode_ControlRig.h b/Engine/Plugins/Animation/ControlRig/Source/ControlRigDeveloper/Public/AnimGraphNode_ControlRig.h
--- a/Engine/Plugins/Animation/ControlRig/Source/ControlRigDeveloper/Public/AnimGraphNode_ControlRig.h
+++ b/Engine/Plugins/Animation/ControlRig/Source/ControlRigDeveloper/Public/AnimGraphNode_ControlRig.h
@@ -5,6 +5,7 @@
 #include "AnimGraphNode_CustomProperty.h"
 #include "AnimNode_ControlRig.h"
 #include "SSearchableComboBox.h"
+#include "ControlRigIOMapping.h"
 #include "AnimGraphNode_ControlRig.generated.h"
 
 struct FRigVMVariableMappingInfo;
@@ -30,43 +31,13 @@ private:
 	virtual FAnimNode_CustomProperty* GetCustomPropertyNode() override { return &Node; }
 	virtual const FAnimNode_CustomProperty* GetCustomPropertyNode() const override { return &Node; }
 
-	struct FControlsInfo
-	{
-		FName Name;
-		FString DisplayName;
-		FEdGraphPinType PinType;
-		ERigControlType ControlType;
-		FString DefaultValue;
-	};
-	mutable const UClass* ControlsInfoClass = nullptr;
-	mutable TArray<FControlsInfo> ControlsInfo;
-	TArray<FControlsInfo>& GetControls() const;
-
-	// property related things
-	void GetVariables(bool bInput, TMap<FName, FRigVMExternalVariable>& OutParameters) const;
-
-	TMap<FName, FRigVMExternalVariable> InputVariables;
-	TMap<FName, FRigVMExternalVariable> OutputVariables;
-
-	void RebuildExposedProperties();
 	virtual void CreateCustomPins(TArray<UEdGraphPin*>* OldPins) override;
 	virtual void ValidateAnimNodeDuringCompilation(USkeleton* ForSkeleton, FCompilerResultsLog& MessageLog) override;
-	
-	// pin option related
-	bool IsPropertyExposeEnabled(FName PropertyName) const;
-	ECheckBoxState IsPropertyExposed(FName PropertyName) const;
-	void OnPropertyExposeCheckboxChanged(ECheckBoxState NewState, FName PropertyName);
 
-#if WITH_EDITOR
-	// SRigVMVariableMappingWidget related
+	void OnPropertyExposeCheckboxChanged(ECheckBoxState NewState, FName PropertyName);
 	void OnVariableMappingChanged(const FName& PathName, const FName& Curve, bool bInput);
-	FName GetVariableMapping(const FName& PathName, bool bInput);
-	void GetAvailableMapping(const FName& PathName, TArray<FName>& OutArray, bool bInput);
-	void CreateVariableMapping(const FString& FilteredText, TArray< TSharedPtr<FRigVMVariableMappingInfo> >& OutArray, bool bInput);
-#endif
 
-	bool IsAvailableToMapToCurve(const FName& PropertyName, bool bInput) const;
-	bool IsInputProperty(const FName& PropertyName) const;
-	UAnimGraphNode_ControlRig::FControlsInfo* FindControlElement(const FName& InControlName) const;
-};
+	USkeleton* GetTargetSkeleton() const;
 
+	TSharedPtr <FControlRigIOMapping> ControlRigIOMapping;
+};
diff --git a/Engine/Plugins/Animation/ControlRig/Source/ControlRigDeveloper/Public/ControlRigIOMapping.h b/Engine/Plugins/Animation/ControlRig/Source/ControlRigDeveloper/Public/ControlRigIOMapping.h
new file mode 100644
--- /dev/null
+++ b/Engine/Plugins/Animation/ControlRig/Source/ControlRigDeveloper/Public/ControlRigIOMapping.h
@@ -0,0 +1,116 @@
+// Copyright Epic Games, Inc. All Rights Reserved.
+
+#pragma once
+
+#include "EdGraph/EdGraphPin.h"
+#include "Rigs/RigHierarchyDefines.h"
+#include "RigVMCore/RigVMExternalVariable.h"
+#include "Styling/SlateTypes.h"
+
+struct FRigVMVariableMappingInfo;
+struct FOptionalPinFromProperty;
+class IDetailLayoutBuilder;
+class UClass;
+class USkeleton;
+class UControlRig;
+
+
+struct CONTROLRIGDEVELOPER_API FControlRigIOMapping : public TSharedFromThis<FControlRigIOMapping>
+{
+	DECLARE_DELEGATE_RetVal(USkeleton*, FOnGetTargetSkeleton);
+	DECLARE_DELEGATE_RetVal(UClass*, FOnGetTargetClass);
+	DECLARE_DELEGATE_TwoParams(FOnPinCheckStateChanged, ECheckBoxState /*NewState*/, FName /** Property Name */);
+	DECLARE_DELEGATE_ThreeParams(FOnVariableMappingChanged, const FName& /*PathName*/, const FName& /*Curve*/, bool /*bInput*/);
+
+	FControlRigIOMapping() = delete;
+	FControlRigIOMapping(TMap<FName, FName>& InInputMapping, TMap<FName, FName>& InOutputMapping, TArray<FOptionalPinFromProperty>& InCustomPinProperties);
+	virtual ~FControlRigIOMapping();
+
+	struct FControlsInfo
+	{
+		FName Name;
+		FString DisplayName;
+		FEdGraphPinType PinType;
+		ERigControlType ControlType;
+		FString DefaultValue;
+	};
+	
+	bool CreateVariableMappingWidget(IDetailLayoutBuilder& DetailBuilder);
+
+	const TMap<FName, FRigVMExternalVariable>& GetInputVariables() const
+	{
+		return InputVariables;
+	}
+	const TMap<FName, FRigVMExternalVariable>& GetOutputVariables() const
+	{
+		return OutputVariables;
+	}
+
+	const TArray<FControlRigIOMapping::FControlsInfo>& GetControls() const;
+
+	void RebuildExposedProperties();
+
+	FOnGetTargetSkeleton& GetOnGetTargetSkeletonDelegate()
+	{
+		return OnGetTargetSkeletonDelegate;
+	}
+
+	FOnGetTargetClass& GetOnGetTargetClassDelegate()
+	{
+		return OnGetTargetClassDelegate;
+	}
+
+	FOnPinCheckStateChanged& GetOnPinCheckStateChangedDelegate()
+	{
+		return OnPinCheckStateChangedDelegate;
+	}
+
+	FOnVariableMappingChanged& GetOnVariableMappingChanged()
+	{
+		return OnVariableMappingChangedDelegate;
+	}
+
+	bool IsInputProperty(const FName& PropertyName) const;
+
+	void SetIOMapping(bool bInput, const FName& SourceProperty, const FName& TargetCurve);
+	FName GetIOMapping(bool bInput, const FName& SourceProperty) const;
+
+private:
+	FOnGetTargetSkeleton OnGetTargetSkeletonDelegate;
+	FOnGetTargetClass OnGetTargetClassDelegate;
+	FOnPinCheckStateChanged OnPinCheckStateChangedDelegate;
+	FOnVariableMappingChanged OnVariableMappingChangedDelegate;
+
+	TMap<FName, FName>& InputMapping;
+	TMap<FName, FName>& OutputMapping;
+	TArray<FOptionalPinFromProperty>& CustomPinProperties;
+
+	// Controls data
+	mutable const UClass* ControlsInfoClass = nullptr;
+	mutable TArray<FControlsInfo> ControlsInfo;
+
+	// property related things
+	static void GetVariables(const UClass* TargetClass, bool bInput, TMap<FName, FRigVMExternalVariable>& OutParameters);
+
+	TMap<FName, FRigVMExternalVariable> InputVariables;
+	TMap<FName, FRigVMExternalVariable> OutputVariables;
+
+	// pin option related
+	bool IsPropertyExposeEnabled(FName PropertyName) const;
+	ECheckBoxState IsPropertyExposed(FName PropertyName) const;
+	void OnPropertyExposeCheckboxChanged(ECheckBoxState NewState, FName PropertyName);
+
+#if WITH_EDITOR
+	// SRigVMVariableMappingWidget related
+	void OnVariableMappingChanged(const FName& PathName, const FName& Curve, bool bInput);
+	FName GetVariableMapping(const FName& PathName, bool bInput);
+	void GetAvailableMapping(const FName& PathName, TArray<FName>& OutArray, bool bInput);
+	void CreateVariableMapping(const FString& FilteredText, TArray< TSharedPtr<FRigVMVariableMappingInfo> >& OutArray, bool bInput);
+#endif
+
+	bool IsAvailableToMapToCurve(const FName& PropertyName, bool bInput) const;
+	const FControlsInfo* FindControlElement(const FName& InControlName) const;
+
+	const UClass* GetTargetClass() const;
+	USkeleton* GetTargetSkeleton() const;
+};
diff --git a/Engine/Source/Runtime/Engine/Public/BonePose.h b/Engine/Source/Runtime/Engine/Public/BonePose.h
--- a/Engine/Source/Runtime/Engine/Public/BonePose.h
+++ b/Engine/Source/Runtime/Engine/Public/BonePose.h
@@ -212,7 +212,7 @@ public:
 	}
 
 	template <typename OtherAllocator>
-	void CopyBonesTo(TArray<FTransform, OtherAllocator>& DestPoseBones)
+	void CopyBonesTo(TArray<FTransform, OtherAllocator>& DestPoseBones) const
 	{
 		// this won't work if you're copying to FBaseCompactPose without BoneContainer data
 		// you'll like to make CopyBonesTo(FBaseCompactPose<OtherAllocator>& DestPose) to fix this properly
