Constraints: added the ability to select the parent before creating the constraint instead of using the picker

diff --git a/Engine/Plugins/Animation/ControlRig/Source/ControlRigEditor/Private/EditMode/SControlRigEditModeTools.cpp b/Engine/Plugins/Animation/ControlRig/Source/ControlRigEditor/Private/EditMode/SControlRigEditModeTools.cpp
--- a/Engine/Plugins/Animation/ControlRig/Source/ControlRigEditor/Private/EditMode/SControlRigEditModeTools.cpp
+++ b/Engine/Plugins/Animation/ControlRig/Source/ControlRigEditor/Private/EditMode/SControlRigEditModeTools.cpp
@@ -382,7 +382,6 @@ void SControlRigEditModeTools::Construct(const FArguments& InArgs, TSharedPtr<FC
 				.AreaTitle(LOCTEXT("ConstraintsWidget", "Constraints"))
 				.AreaTitleFont(FAppStyle::GetFontStyle("DetailsView.CategoryFontStyle"))
 				.BorderBackgroundColor(FLinearColor(.6f, .6f, .6f))
-				.Padding(FMargin(8.f))
 				.HeaderContent()
 				[
 					SNew(SHorizontalBox)
@@ -1165,18 +1164,27 @@ FReply SControlRigEditModeTools::HandleAddConstraintClicked()
 	{
 		const TSharedRef<SConstraintMenuEntry> Entry =
 			SNew(SConstraintMenuEntry, InConstraintType)
-		.OnConstraintCreated_Lambda([this]()
-		{
-			// magic number to auto expand the widget when creating a new constraint. We keep that number below a reasonable
-			// threshold to avoid automatically creating a large number of items (this can be style done by the user) 
-			static constexpr int32 NumAutoExpand = 20;
-			const int32 NumItems = ConstraintsEditionWidget ? ConstraintsEditionWidget->RefreshConstraintList() : 0;	
-			if (ConstraintPickerExpander && NumItems < NumAutoExpand)
+			.OnConstraintCreated_Lambda([this]()
 			{
-				ConstraintPickerExpander->SetExpanded(true);
-			}
-		});
-		MenuBuilder.AddWidget(Entry, FText::GetEmpty(), true);
+				// magic number to auto expand the widget when creating a new constraint. We keep that number below a reasonable
+				// threshold to avoid automatically creating a large number of items (this can be style done by the user) 
+				static constexpr int32 NumAutoExpand = 20;
+				const int32 NumItems = ConstraintsEditionWidget ? ConstraintsEditionWidget->RefreshConstraintList() : 0;	
+				if (ConstraintPickerExpander && NumItems < NumAutoExpand)
+				{
+					ConstraintPickerExpander->SetExpanded(true);
+				}
+				if (ConstraintsEditionWidget)
+				{
+					ConstraintsEditionWidget->ResetParent(); 
+				}
+			})
+			.OnGetParent_Lambda([this]()
+			{
+				static const FConstrainable InvalidParent;
+				return ConstraintsEditionWidget ? ConstraintsEditionWidget->GetParent() : InvalidParent;
+			});
+			MenuBuilder.AddWidget(Entry, FText::GetEmpty(), true);
 	};
 	
 	MenuBuilder.BeginSection("CreateConstraint", LOCTEXT("CreateConstraintHeader", "Create New..."));
diff --git a/Engine/Plugins/Animation/ControlRig/Source/ControlRigEditor/Private/Editor/Constraints/SConstraintsWidget.cpp b/Engine/Plugins/Animation/ControlRig/Source/ControlRigEditor/Private/Editor/Constraints/SConstraintsWidget.cpp
--- a/Engine/Plugins/Animation/ControlRig/Source/ControlRigEditor/Private/Editor/Constraints/SConstraintsWidget.cpp
+++ b/Engine/Plugins/Animation/ControlRig/Source/ControlRigEditor/Private/Editor/Constraints/SConstraintsWidget.cpp
@@ -7,10 +7,8 @@
 #include "ConstraintsActor.h"
 #include "ControlRigEditorStyle.h"
 #include "DetailLayoutBuilder.h"
-#include "LevelEditorActions.h"
 #include "LevelEditorViewport.h"
 #include "PropertyCustomizationHelpers.h"
-#include "Selection.h"
 #include "SlateOptMacros.h"
 #include "LevelEditor.h"
 #include "SSocketChooser.h"
@@ -24,11 +22,11 @@
 #include "ScopedTransaction.h"
 #include "ISequencer.h"
 #include "Tools/BakingHelper.h"
-#include "MovieSceneToolHelpers.h"
 #include "ConstraintsManager.h"
 #include "Styling/SlateIconFinder.h"
 #include "Widgets/Input/NumericTypeInterface.h"
 #include "FrameNumberDetailsCustomization.h"
+#include "SAssetDropTarget.h"
 #include "Sequencer/ControlRigParameterTrackEditor.h"
 
 #define LOCTEXT_NAMESPACE "SConstraintsWidget"
@@ -124,49 +122,194 @@ UTickableTransformConstraint* FConstraintInfo::GetConfigurable(ETransformConstra
 	return nullptr;
 }
 
-namespace
+namespace UE::Private
 {
-
-UWorld* GetCurrentWorld()
-{
-	return GCurrentLevelEditingViewportClient ? GCurrentLevelEditingViewportClient->GetWorld() : nullptr;
-}
+	UWorld* GetCurrentWorld()
+	{
+		return GCurrentLevelEditingViewportClient ? GCurrentLevelEditingViewportClient->GetWorld() : nullptr;
+	}
 	
-// NOTE we use this function to get the current selection as control actors are Temporary Editor Actors so won't be
-// pushed added to the selection list.
-TArray<AActor*> GetCurrentSelection()
-{
-	const UWorld* World = GetCurrentWorld();
-	const ULevel* CurrentLevel = IsValid(World) ? World->GetCurrentLevel() : nullptr;
-	if (!CurrentLevel)
+	// NOTE we use this function to get the current selection as control actors are Temporary Editor Actors so won't be
+	// pushed added to the selection list.
+	TArray<AActor*> GetCurrentSelection()
 	{
-		static const TArray<AActor*> DummyArray;
-		return DummyArray;
+		const UWorld* World = GetCurrentWorld();
+		const ULevel* CurrentLevel = IsValid(World) ? World->GetCurrentLevel() : nullptr;
+		if (!CurrentLevel)
+		{
+			static const TArray<AActor*> DummyArray;
+			return DummyArray;
+		}
+
+		return ObjectPtrDecay(CurrentLevel->Actors).FilterByPredicate( [](const AActor* Actor)
+		{
+			return IsValid(Actor) && Actor->IsSelected();
+		});	
 	}
 
-	return ObjectPtrDecay(CurrentLevel->Actors).FilterByPredicate( [](const AActor* Actor)
+	static TWeakPtr<ISequencer> GetSequencerChecked()
 	{
-		return IsValid(Actor) && Actor->IsSelected();
-	});	
-}
+		const TWeakPtr<ISequencer> WeakSequencer = FBakingHelper::GetSequencer();
+		if (!WeakSequencer.IsValid() || !WeakSequencer.Pin()->GetFocusedMovieSceneSequence())
+		{
+			return nullptr;
+		}
+	
+		const UMovieScene* MovieScene = WeakSequencer.Pin()->GetFocusedMovieSceneSequence()->GetMovieScene();
+		if (!MovieScene)
+		{
+			return nullptr;
+		}
 
-static TWeakPtr<ISequencer> GetSequencerChecked()
-{
-	const TWeakPtr<ISequencer> WeakSequencer = FBakingHelper::GetSequencer();
-	if (!WeakSequencer.IsValid() || !WeakSequencer.Pin()->GetFocusedMovieSceneSequence())
-	{
-		return nullptr;
+		return WeakSequencer;
 	}
-	
-	const UMovieScene* MovieScene = WeakSequencer.Pin()->GetFocusedMovieSceneSequence()->GetMovieScene();
-	if (!MovieScene)
+
+	void GetConstrainable(AActor* InActor, FConstrainable& InOutConstrainable)
 	{
-		return nullptr;
+		InOutConstrainable = FConstrainable();
+
+		if (!InActor)
+		{
+			return;
+		}
+	
+		InOutConstrainable.Actor = InActor;
+		InOutConstrainable.Object = InActor;
+	
+		// gather sub components with sockets
+		const TInlineComponentArray<USceneComponent*> Components(InActor);
+		const TArray<USceneComponent*> ComponentsWithSockets = Components.FilterByPredicate([](const USceneComponent* Component)
+		{
+			return Component->HasAnySockets();
+		});
+
+		// if no component socket available then return the actor
+		const int32 NumComponentsWithSockets = ComponentsWithSockets.Num();
+		if (NumComponentsWithSockets == 0)
+		{
+			return;
+		}
+
+		// creates a menu encapsulating InContent  
+		auto CreateMenu = [](const TSharedRef<SWidget>& InContent, const FVector2D& InLocation)
+		{
+			const FLevelEditorModule& LevelEditorModule = FModuleManager::GetModuleChecked<FLevelEditorModule>("LevelEditor");
+			const TSharedPtr< ILevelEditor > LevelEditor = LevelEditorModule.GetFirstLevelEditor();
+			FSlateApplication::Get().PushMenu(
+				LevelEditor.ToSharedRef(),
+				FWidgetPath(),
+				InContent,
+			InLocation,
+			FPopupTransitionEffect( FPopupTransitionEffect::ContextMenu ) );
+		};
+
+		// creates a new socket chooser popup widget
+		auto GetSocketChooserWidget = [&InOutConstrainable](USceneComponent* Component)
+		{
+			return SNew(SSocketChooserPopup)
+			.SceneComponent( Component )
+			.OnSocketChosen_Lambda([&InOutConstrainable, Component](FName InSocketName)
+			{
+				InOutConstrainable.Object = Component;
+				InOutConstrainable.Socket = InSocketName;
+			});
+		};
+
+		// decal the menu 
+		const FVector2D SummonLocation = FSlateApplication::Get().GetCursorPos();
+	
+		// if one component with sockets then constrain the component selecting the socket
+		if (NumComponentsWithSockets == 1)
+		{
+			return CreateMenu( GetSocketChooserWidget(ComponentsWithSockets[0]), FSlateApplication::Get().GetCursorPos());
+		}
+
+		// if there are several of them, then build a component chooser first then the socket chooser
+		static constexpr bool CloseAfterSelection = true;
+		FMenuBuilder MenuBuilder(CloseAfterSelection, nullptr);
+		MenuBuilder.BeginSection("ChooseComp", LOCTEXT("ChooseComponentSection", "Choose Component"));
+		{
+			for (USceneComponent* Component: ComponentsWithSockets)
+			{
+				MenuBuilder.AddMenuEntry(FText::FromName(Component->GetFName()),
+				FText(),
+					FSlateIconFinder::FindIconForClass(Component->GetClass(), TEXT("SCS.Component")),
+					FUIAction(FExecuteAction::CreateLambda([CreateMenu, GetSocketChooserWidget, Component]()
+					{
+						CreateMenu(GetSocketChooserWidget(Component), FSlateApplication::Get().GetCursorPos());
+					})),
+					NAME_None,
+				   EUserInterfaceActionType::Button);
+			}
+		}
+		MenuBuilder.EndSection();
+		CreateMenu(MenuBuilder.MakeWidget(), SummonLocation);
 	}
 
-	return WeakSequencer;
-}
+	bool CreateConstraint(const TArray<AActor*>& InChildren, const FConstrainable& InParent, const ETransformConstraintType InConstraintType, const bool bUseDefault)
+	{
+		UObject* Parent = InParent.Object.Get();
+		UWorld* World = GetCurrentWorld();
+	
+		if (!Parent || !IsValid(World))
+		{
+			return false;
+		}
+	
+		if (InChildren.IsEmpty())
+		{
+			return false;
+		}
+
+		const TWeakPtr<ISequencer> WeakSequencer = GetSequencerChecked();
 	
+		bool bCreated = false;
+		for (AActor* Child: InChildren)
+		{
+			if (Child != InParent.Object)
+			{
+				FScopedTransaction Transaction(LOCTEXT("CreateConstraintKey", "Create Constraint Key"));
+				UTickableTransformConstraint* Constraint =
+					FTransformConstraintUtils::CreateAndAddFromObjects(World, Parent, InParent.Socket, Child, NAME_None, InConstraintType, true, bUseDefault);
+
+				if (Constraint)
+				{
+					bCreated = true;
+
+					if (WeakSequencer.IsValid())
+					{
+						if (bUseDefault)
+						{
+							Constraint->Evaluate();
+						}
+						FControlRigParameterTrackEditor::AddConstraintToSequencer(WeakSequencer.Pin(), Constraint);
+					}
+					else
+					{
+						FConstraintsManagerController& Controller = FConstraintsManagerController::Get(World);
+						Controller.StaticConstraintCreated(World, Constraint);
+					}		
+				}
+			}
+		}
+
+		return bCreated;
+	}
+
+	bool IsActorConstrainable(const AActor* InActor)
+	{
+		if (!InActor)
+		{
+			return false;
+		}
+
+		if (!InActor->GetRootComponent())
+		{
+			return false;
+		}
+		
+		return true;
+	}
 }
 
 /**
@@ -179,6 +322,7 @@ void SConstraintMenuEntry::Construct(
 {
 	ConstraintType = InType;
 	OnConstraintCreated = InArgs._OnConstraintCreated;
+	OnGetParent = InArgs._OnGetParent;
 	
 	const FButtonStyle& ButtonStyle = FAppStyle::Get().GetWidgetStyle<FButtonStyle>( "Menu.Button" );
 
@@ -263,7 +407,7 @@ FReply SConstraintMenuEntry::OnMouseButtonDown(const FGeometry& MyGeometry, cons
 	{
 		bIsPressed = true;
 		static constexpr bool bUseMutableDefault = false;
-		return CreateSelectionPicker(bUseMutableDefault);
+		return CreateFromSelectionOrPicker(bUseMutableDefault);
 	}
 
 	return FReply::Unhandled();
@@ -362,7 +506,7 @@ TSharedRef<SWidget> SConstraintMenuEntry::GenerateConstraintDefaultWidget()
 					.OnClicked_Lambda([this]()
 					{
 						static constexpr bool bUseMutableDefault = true;
-						return CreateSelectionPicker(bUseMutableDefault);
+						return CreateFromSelectionOrPicker(bUseMutableDefault);
 					})
 				]
 			]
@@ -415,6 +559,38 @@ void SConstraintMenuEntry::NotifyPostChange(const FPropertyChangedEvent& InPrope
 	}
 }
 
+FReply SConstraintMenuEntry::CreateFromSelectionOrPicker(const bool bUseDefault) const
+{
+	// hide the creation menu
+	FSlateApplication::Get().DismissAllMenus();
+	
+	// use the selected parent if any
+	if (OnGetParent.IsBound())
+	{
+		const FConstrainable Parent = OnGetParent.Execute();
+		if (Parent.Object.IsValid())
+		{
+			const TArray<AActor*> Children = UE::Private::GetCurrentSelection();
+			const bool bCreated = UE::Private::CreateConstraint(Children, Parent, ConstraintType, bUseDefault);
+			if (!bCreated)
+			{
+				return FReply::Unhandled();
+			}
+
+			// update list
+			if (OnConstraintCreated.IsBound())
+			{
+				OnConstraintCreated.Execute();
+			}
+
+			return FReply::Handled();
+		}
+	}
+
+	// pick the parent
+	return CreateSelectionPicker(bUseDefault);
+}
+
 FReply SConstraintMenuEntry::CreateSelectionPicker(const bool bUseDefault) const
 {
 	struct FCreationArgs
@@ -429,9 +605,6 @@ FReply SConstraintMenuEntry::CreateSelectionPicker(const bool bUseDefault) const
 	Args.ConstraintType = ConstraintType;
 	Args.bUseDefault = bUseDefault;
 	
-	// FIXME temp approach for selecting the parent
-	FSlateApplication::Get().DismissAllMenus();
-	
 	static const FActorPickerModeModule& ActorPickerMode = FModuleManager::Get().GetModuleChecked<FActorPickerModeModule>("ActorPickerMode");
 
 	ActorPickerMode.BeginActorPickingMode(
@@ -447,10 +620,12 @@ FReply SConstraintMenuEntry::CreateSelectionPicker(const bool bUseDefault) const
 
 void SConstraintMenuEntry::CreateConstraint(
 	AActor* InParent,
-	FOnConstraintCreated InCreationDelegate,
+	const FOnConstraintCreated& InCreationDelegate,
 	const ETransformConstraintType InConstraintType,
 	const bool bUseDefault)
 {
+	using namespace UE::Private;
+
 	if (!InParent)
 	{
 		return;	
@@ -466,125 +641,20 @@ void SConstraintMenuEntry::CreateConstraint(
 	const TArray<AActor*> Selection = GetCurrentSelection();
 	if (Selection.IsEmpty())
 	{
-		return;	
+		return;
 	}
 
-	// gather sub components with sockets
-	const TInlineComponentArray<USceneComponent*> Components(InParent);
-	const TArray<USceneComponent*> ComponentsWithSockets = Components.FilterByPredicate([](const USceneComponent* Component)
-	{
-		return Component->HasAnySockets();
-	});
+	FConstrainable ConstrainableParent;
+	GetConstrainable(InParent, ConstrainableParent);
 
-	const TWeakPtr<ISequencer> WeakSequencer = GetSequencerChecked();
-	
-	// create constraints
-	auto CreateConstraint = [InCreationDelegate, InConstraintType, WeakSequencer, bUseDefault](
-		const TArray<AActor*>& Selection, UObject* InParent, const FName& InSocketName)
+	if (ConstrainableParent.Object.IsValid())
 	{
-		UWorld* World = GetCurrentWorld();
-		if (!IsValid(World))
-		{
-			return;
-		}
-
-		bool bCreated = false;
-		for (AActor* Child: Selection)
-		{
-			if (Child != InParent)
-			{
-				FScopedTransaction Transaction(LOCTEXT("CreateConstraintKey", "Create Constraint Key"));
-				UTickableTransformConstraint* Constraint =
-					FTransformConstraintUtils::CreateAndAddFromObjects(World, InParent, InSocketName, Child, NAME_None, InConstraintType, true, bUseDefault);
-				if (Constraint)
-				{
-					bCreated = true;
-					if (WeakSequencer.IsValid())
-					{
-						if (bUseDefault)
-						{
-							Constraint->Evaluate();
-						}
-						FControlRigParameterTrackEditor::AddConstraintToSequencer(WeakSequencer.Pin(), Constraint);
-					}
-					else
-					{
-						FConstraintsManagerController& Controller = FConstraintsManagerController::Get(World);
-						Controller.StaticConstraintCreated(World, Constraint);
-					}		
-				}
-			}
-		}
-	
-		// update list
+		const bool bCreated = UE::Private::CreateConstraint(Selection, ConstrainableParent, InConstraintType, bUseDefault);
 		if (bCreated && InCreationDelegate.IsBound())
-		{
-			InCreationDelegate.Execute();
-		}
-	};
-
-	const int32 NumComponentsWithSockets = ComponentsWithSockets.Num();
-
-	// if no component socket available then constrain the whole actor 
-	if (NumComponentsWithSockets == 0)
-	{
-		return CreateConstraint(Selection, InParent, NAME_None);
+        {
+        	InCreationDelegate.Execute();
+        }
 	}
-
-	// creates a menu encapsulating InContent  
-	auto CreateMenu = [](const TSharedRef<SWidget>& InContent, const FVector2D& InLocation)
-	{
-		const FLevelEditorModule& LevelEditorModule = FModuleManager::GetModuleChecked<FLevelEditorModule>("LevelEditor");
-		const TSharedPtr< ILevelEditor > LevelEditor = LevelEditorModule.GetFirstLevelEditor();
-		FSlateApplication::Get().PushMenu(
-			LevelEditor.ToSharedRef(),
-			FWidgetPath(),
-			InContent,
-		InLocation,
-		FPopupTransitionEffect( FPopupTransitionEffect::ContextMenu ) );
-	};
-
-	// creates a new socket chooser popup widget
-	auto GetSocketChooserWidget = [CreateConstraint, Selection](USceneComponent* Component)
-	{
-		return SNew(SSocketChooserPopup)
-		.SceneComponent( Component )
-		.OnSocketChosen_Lambda([CreateConstraint, Selection, Component](FName InSocketName)
-		{
-			CreateConstraint(Selection, Component, InSocketName);
-		});
-	};
-
-	// decal the menu 
-	FVector2D SummonLocation = FSlateApplication::Get().GetCursorPos();
-	SummonLocation.Y += 4 * FSlateApplication::Get().GetCursorSize().Y;
-	
-	// if one component with sockets then constrain the component selecting the socket
-	if (NumComponentsWithSockets == 1)
-	{
-		return CreateMenu( GetSocketChooserWidget(ComponentsWithSockets[0]), SummonLocation);
-	}
-
-	// if there are several of them, then build a component chooser first then the socket chooser
-	static constexpr bool CloseAfterSelection = true;
-	FMenuBuilder MenuBuilder(CloseAfterSelection, nullptr);
-	MenuBuilder.BeginSection("ChooseComp", LOCTEXT("ChooseComponentSection", "Choose Component"));
-	{
-		for (USceneComponent* Component: ComponentsWithSockets)
-		{
-			MenuBuilder.AddMenuEntry(FText::FromName(Component->GetFName()),
-			FText(),
-				FSlateIconFinder::FindIconForClass(Component->GetClass(), TEXT("SCS.Component")),
-				FUIAction(FExecuteAction::CreateLambda([CreateMenu, GetSocketChooserWidget, Component]()
-				{
-					CreateMenu(GetSocketChooserWidget(Component), FSlateApplication::Get().GetCursorPos());
-				})),
-				NAME_None,
-			   EUserInterfaceActionType::Button);
-		}
-	}
-	MenuBuilder.EndSection();
-	CreateMenu(MenuBuilder.MakeWidget(), SummonLocation);
 }
 
 /**
@@ -596,6 +666,8 @@ void SEditableConstraintItem::Construct(
 	const TSharedPtr<FEditableConstraintItem>& InItem,
 	TSharedPtr<SConstraintsEditionWidget> InConstraintsWidget)
 {
+	using namespace UE::Private;
+	
 	ConstraintItem = InItem;
 	ConstraintsWidget = InConstraintsWidget;
 
@@ -980,6 +1052,8 @@ FText FBaseConstraintListWidget::GetShowConstraintsTooltip(EShowConstraints Inde
 
 int32 FBaseConstraintListWidget::RefreshConstraintList()
 {
+	using namespace UE::Private;
+	
 	// get constraints
 	UWorld* World = GetCurrentWorld();
 	if (!IsValid(World))
@@ -1086,67 +1160,248 @@ void FBaseConstraintListWidget::OnActorSelectionChanged(const TArray<UObject*>&
 /// SConstraintsEditionWidget
 ///////////////////////////////////////////////////////////
 
+void SConstraintsEditionWidget::ResetParent()
+{
+	ConstrainableParent = FConstrainable();
+}
+
+const FConstrainable& SConstraintsEditionWidget::GetParent() const
+{
+	return ConstrainableParent;
+}
+
+bool SConstraintsEditionWidget::IsParentValid() const
+{
+	return ConstrainableParent.Actor.IsValid();
+}
+
+void SConstraintsEditionWidget::UpdateParent(AActor* InActor)
+{
+	UE::Private::GetConstrainable(InActor, ConstrainableParent);
+}
 
 void SConstraintsEditionWidget::Construct(const FArguments& InArgs)
 {
+	using namespace UE::Private;
+	
 	WeakSequencer = GetSequencerChecked();
 
-	ChildSlot
-	[
-		SNew(SBorder)
-		.Visibility(EVisibility::Visible)
-		.BorderImage(FAppStyle::GetBrush("Menu.Background"))
-		[
-
-			SNew( SVerticalBox )
+	auto GetParentWidget = [this]()
+	{
+		TSharedRef<SSplitter> Splitter = SNew(SSplitter)
+		.Style(FAppStyle::Get(), "DetailsView.Splitter")
+		.PhysicalSplitterHandleSize(1.0f)
+		.HitDetectionSplitterHandleSize(5.0f)
 
-			+ SVerticalBox::Slot()
-			.AutoHeight()
-			.Padding(FMargin(0.0f, 0.f))
+		// Parent label 
+		+ SSplitter::Slot()
+		.Value(0.4f)
+		[
+			SNew(SHorizontalBox)
+			.Clipping(EWidgetClipping::OnDemand)
+			+ SHorizontalBox::Slot()
+			.VAlign(VAlign_Center)
+			.HAlign(HAlign_Left)
+			.Padding(20.0f, 0.0f, 0.0f, 0.0f)
 			[
-				SNew(SOverlay)
+				SNew(STextBlock)
+				.Font(FAppStyle::Get().GetFontStyle("PropertyWindow.NormalFont"))
+				.Text(LOCTEXT("PreSelectedParent_Label", "Pre Selected Parent"))
+			]
+		]
 
-				+ SOverlay::Slot()
+		// Parent chooser 
+		+ SSplitter::Slot()
+		.Value(0.6f)
+		[
+			SNew(SAssetDropTarget)
+			.bSupportsMultiDrop(true)
+			.OnAreAssetsAcceptableForDrop_Lambda([](TArrayView<FAssetData> InAssets)
+			{
+				const bool bHasAnyNotConstrainable = InAssets.ContainsByPredicate([](const FAssetData& Data)
+				{
+					return !IsActorConstrainable(Cast<AActor>(Data.GetAsset()));	
+				});
+				return !bHasAnyNotConstrainable;
+			})
+			.OnAssetsDropped_Lambda( [this](const FDragDropEvent&, TArrayView<FAssetData> InAssets)
+			{
+				UpdateParent(InAssets.IsEmpty() ? nullptr : Cast<AActor>(InAssets[0].GetAsset()));
+			})
+			[
+				SNew(SHorizontalBox)
+											
+				// actor chooser widget
+				+ SHorizontalBox::Slot()
+				.Padding(FMargin(4.0f,2.0f))
+				.HAlign(HAlign_Right)
+				[
+					GetActorChooserWidget().ToSharedRef()
+				]
+			
+				// use selected widget
+				+ SHorizontalBox::Slot()
+				.AutoWidth()
 				[
-					SAssignNew(ListView, ConstraintItemListView)
-					.SelectionMode(ESelectionMode::Single)
-					.ListItemsSource( &ListItems )
-					.OnGenerateRow(this, &SConstraintsEditionWidget::OnGenerateWidgetForItem)
-					.OnContextMenuOpening(this, &SConstraintsEditionWidget::CreateContextMenu)
-					.OnMouseButtonDoubleClick(this, &SConstraintsEditionWidget::OnItemDoubleClicked)
+					SNew(SBox)
+					.Padding(FMargin(4.0f,0.0f))
+					[
+						GetUseSelectedActorWidget()
+					]
 				]
+			]
+		]
 
+		// Reset button
+		+ SSplitter::Slot()
+		.SizeRule(SSplitter::SizeToContent)
+		.Value(26.f)
+		.MinSize(26.f)
+		.Resizable(false)
+		[
+			SNew(SBorder)
+			.BorderImage(FAppStyle::Get().GetBrush("DetailsView.CategoryMiddle"))
+			.BorderBackgroundColor(FAppStyle::Get().GetSlateColor("Colors.Panel"))
+			.HAlign(HAlign_Center)
+			.VAlign(VAlign_Center)
+			[
+				SNew(SBox)
+				.HAlign(HAlign_Center)
+				.VAlign(VAlign_Center)
+				.WidthOverride(22.0f)
+				.HeightOverride(22.0f)
+				.ToolTipText(LOCTEXT("ResetParent_Tooltip", "Reset parent."))
+				[
+					SNew(SButton)
+					.ButtonStyle( FAppStyle::Get(), "SimpleButton" )
+					.OnClicked_Lambda([this]()
+					{
+						ResetParent();
+						return FReply::Handled();
+					})
+					.Visibility_Lambda([this]()
+					{
+						return IsParentValid() ? EVisibility::Visible : EVisibility::Hidden;
+					})
+					.ContentPadding(0.0f)
+					.IsFocusable(false)
+					[ 
+						SNew( SImage )
+						.Image(FAppStyle::GetBrush("PropertyWindow.DiffersFromDefault"))
+						.ColorAndOpacity( FSlateColor::UseForeground() )
+					]
+				]
+			]
+		];
+		
+		return SNew( SBorder )
+		.BorderImage(FAppStyle::Get().GetBrush("DetailsView.GridLine"))
+		.BorderBackgroundColor(FAppStyle::Get().GetSlateColor("Colors.Panel"))
+		.Padding(FMargin(0.0f,0.0f,0.0f,1.0f))
+		.Clipping(EWidgetClipping::ClipToBounds)
+		[
+			SNew(SBox)
+			.MinDesiredHeight(26.f)
+			[
+				SNew( SHorizontalBox )
+				+ SHorizontalBox::Slot()
+				.HAlign(HAlign_Fill)
+				.VAlign(VAlign_Fill)
+				[
+					SNew( SBorder )
+					.BorderImage(FAppStyle::Get().GetBrush("DetailsView.CategoryMiddle"))
+					.BorderBackgroundColor(FAppStyle::Get().GetSlateColor("Colors.Panel"))
+					.Padding(0.0f)
+					[
+						Splitter
+					]
+				]
 			]
+		];
+	};
+
+	auto GetConstraintListWidget = [this]()
+	{
+		return SNew(SOverlay)
 	
-			+ SVerticalBox::Slot()
+		+ SOverlay::Slot()
+		[
+			SAssignNew(ListView, ConstraintItemListView)
+			.SelectionMode(ESelectionMode::Single)
+			.ListItemsSource( &ListItems )
+			.OnGenerateRow(this, &SConstraintsEditionWidget::OnGenerateWidgetForItem)
+			.OnContextMenuOpening(this, &SConstraintsEditionWidget::CreateContextMenu)
+			.OnMouseButtonDoubleClick(this, &SConstraintsEditionWidget::OnItemDoubleClicked)
+		];
+	};
+
+	auto GetBakingWidget = [this]()
+	{
+		return SNew(SHorizontalBox)
+		
+		+ SHorizontalBox::Slot()
+		.FillWidth(1.f)
+		.HAlign(HAlign_Fill)
+		[
+			SNew(SSpacer)
+		]
+			
+		+ SHorizontalBox::Slot()
+		.AutoWidth()
+		.VAlign(VAlign_Center)
+		.HAlign(HAlign_Right)
+		.Padding(1.f)
+		[
+			SNew(SButton)
+			.ButtonStyle(FAppStyle::Get(), "FlatButton.Default")
+			.Text(LOCTEXT("BakeButton", "Bake..."))
+			.OnClicked(this, &SConstraintsEditionWidget::OnBakeClicked)
+			.IsEnabled_Lambda([this]()
+			{
+				return (ListItems.Num() > 0);
+			})
+			.ToolTipText(LOCTEXT("BakeConstraintButtonToolTip", "Bake selected constraints"))
+		];
+	};
+	
+	ChildSlot
+	[
+		SNew( SVerticalBox )
+	
+		// parent chooser widget
+		+ SVerticalBox::Slot()
+		.AutoHeight()
+		.Padding(FMargin(0.0f, 0.f))
+		[
+			GetParentWidget()
+		]
+
+		+ SVerticalBox::Slot()
+		.AutoHeight()
+		.Padding(FMargin(8.f))
+		[
+			SNew(SBorder)
+			.Visibility(EVisibility::Visible)
+			.BorderImage(FAppStyle::GetBrush("Menu.Background"))
+			[
+				SNew( SVerticalBox )
+		
+				// constraints list widget
+				+ SVerticalBox::Slot()
 				.AutoHeight()
 				.Padding(FMargin(0.0f, 0.f))
 				[
-					SNew(SHorizontalBox)
-					+ SHorizontalBox::Slot()
-					.FillWidth(1.f)
-					.HAlign(HAlign_Fill)
-					[
-						SNew(SSpacer)
-					]
+					GetConstraintListWidget()
+				]
 		
-					+ SHorizontalBox::Slot()
-					.AutoWidth()
-					.VAlign(VAlign_Center)
-					.HAlign(HAlign_Right)
-					.Padding(1.f)
-					[
-						SNew(SButton)
-						.ButtonStyle(FAppStyle::Get(), "FlatButton.Default")
-						.Text(LOCTEXT("BakeButton", "Bake..."))
-						.OnClicked(this, &SConstraintsEditionWidget::OnBakeClicked)
-						.IsEnabled_Lambda([this]()
-						{
-							return (ListItems.Num() > 0);
-						})
-						.ToolTipText(LOCTEXT("BakeConstraintButtonToolTip", "Bake selected contraints"))
-					]
+				// baking widget
+				+ SVerticalBox::Slot()
+				.AutoHeight()
+				.Padding(FMargin(0.0f, 0.f))
+				[
+					GetBakingWidget()
 				]
+			]
 		]
 	];
 
@@ -1178,6 +1433,124 @@ void SConstraintsEditionWidget::Tick(const FGeometry& AllottedGeometry, const do
 	}
 }
 
+TSharedRef<SWidget> SConstraintsEditionWidget::OnGetParentMenuContent()
+{
+	static constexpr bool bAllowClear = true, bAllowPicking = false, bUseSelected = true;
+
+	FOnShouldFilterActor FilterActorDelegate = FOnShouldFilterActor::CreateLambda([](const AActor* InActor)
+	{
+		return UE::Private::IsActorConstrainable(InActor);
+	});
+	
+	FOnActorSelected SelectedActorDelegate = FOnActorSelected::CreateLambda([this](AActor* InActor)
+	{
+		UpdateParent(InActor);
+	});
+
+	FSimpleDelegate OnCloseDelegate = FSimpleDelegate::CreateLambda([this]
+	{
+		GetActorChooserWidget()->SetIsOpen(false);
+	});
+
+	FSimpleDelegate UseSelectedDelegate = FSimpleDelegate::CreateLambda([]{ });
+	
+	return PropertyCustomizationHelpers::MakeActorPickerWithMenu(
+		ConstrainableParent.Actor.Get(),
+		bAllowClear,
+		bAllowPicking,
+		FilterActorDelegate,
+		SelectedActorDelegate,
+		OnCloseDelegate,
+		UseSelectedDelegate,
+		bUseSelected);
+}
+
+void SConstraintsEditionWidget::OnParentMenuOpenChanged(bool bOpened)
+{
+	if (!bOpened)
+	{
+		GetActorChooserWidget()->SetMenuContent(SNullWidget::NullWidget);
+	}
+}
+
+TSharedPtr<SComboButton> SConstraintsEditionWidget::GetActorChooserWidget()
+{
+	if (!ActorSelectionButtons.IsValid())
+	{
+		auto GetParentLabel = [this]()
+		{
+			FString Label;
+			if (ConstrainableParent.Actor.IsValid())
+			{
+				Label = ConstrainableParent.Actor->GetActorLabel();
+				if (ConstrainableParent.Socket != NAME_None)
+				{
+					Label = Label + ":" + ConstrainableParent.Socket.ToString();
+				}
+			}
+			return Label;
+		};
+		
+		SAssignNew(ActorSelectionButtons, SComboButton)
+		.ToolTipText_Lambda([this, GetParentLabel]()
+		{
+			if (ConstrainableParent.Actor.IsValid())
+			{
+				const FString Tooltip = FString::Printf(TEXT("'%s' will be used as the constraint's parent."), *GetParentLabel());
+				return FText::FromString(Tooltip);
+			}
+			
+			return LOCTEXT("SelectParent_Tooltip", "Select a parent actor.");
+		})
+		.OnGetMenuContent(this, &SConstraintsEditionWidget::OnGetParentMenuContent)
+		.OnMenuOpenChanged(this, &SConstraintsEditionWidget::OnParentMenuOpenChanged)
+		.ButtonContent()
+		[
+			SNew(SHorizontalBox)
+
+			+SHorizontalBox::Slot()
+			.FillWidth(1)
+			.VAlign(VAlign_Center)
+			[
+				// Show the name of the asset or actor
+				SNew(STextBlock)
+				.Font(FAppStyle::Get().GetFontStyle("PropertyWindow.NormalFont"))
+				.Text_Lambda([this, GetParentLabel]()
+				{
+					if (ConstrainableParent.Actor.IsValid())
+					{
+						return FText::FromString(GetParentLabel()); 
+					}
+					return LOCTEXT("SelectConstraintParent", "Select Parent");
+				})
+			]
+		];
+	}
+	
+	return ActorSelectionButtons;
+}
+
+TSharedRef<SWidget> SConstraintsEditionWidget::GetUseSelectedActorWidget()
+{
+	static constexpr bool bIsActor = true, bEnabled = true;
+	
+	return SNew( SHorizontalBox )
+
+	+ SHorizontalBox::Slot()
+	.VAlign(VAlign_Center)
+	.AutoWidth()
+	.Padding( 2.0f, 0.0f )
+	[
+		PropertyCustomizationHelpers::MakeUseSelectedButton(
+			FSimpleDelegate::CreateLambda([this]()
+			{
+				const TArray<AActor*> Actors = UE::Private::GetCurrentSelection();
+				UpdateParent(Actors.IsEmpty() ? nullptr : Actors[0]);
+			}),
+			LOCTEXT("SetFromSelection", "Set from current selection."),
+			bEnabled, bIsActor)
+	];
+}
 
 TSharedRef<ITableRow> SConstraintsEditionWidget::OnGenerateWidgetForItem(
 	ItemSharedPtr InItem,
@@ -1222,7 +1595,7 @@ void SConstraintsEditionWidget::MoveItemUp(const TSharedPtr<FEditableConstraintI
 	const int32 Index = ListItems.IndexOfByKey(Item);
 	if (Index > 0 && ListItems.IsValidIndex(Index))
 	{
-		UWorld* World = GetCurrentWorld();
+		UWorld* World = UE::Private::GetCurrentWorld();
 		if (!IsValid(World))
 		{
 			return;
@@ -1249,7 +1622,7 @@ void SConstraintsEditionWidget::MoveItemDown(const TSharedPtr<FEditableConstrain
 	const int32 Index = ListItems.IndexOfByKey(Item);
 	if (ListItems.IsValidIndex(Index))
 	{
-		UWorld* World = GetCurrentWorld();
+		UWorld* World = UE::Private::GetCurrentWorld();
 		if (!IsValid(World))
 		{
 			return;
@@ -1268,7 +1641,7 @@ void SConstraintsEditionWidget::MoveItemDown(const TSharedPtr<FEditableConstrain
 
 void SConstraintsEditionWidget::RemoveItem(const TSharedPtr<FEditableConstraintItem>& Item)
 {
-	UWorld* World = GetCurrentWorld();
+	UWorld* World = UE::Private::GetCurrentWorld();
 	if (!IsValid(World))
 	{
 		return;
@@ -1302,7 +1675,7 @@ void SConstraintsEditionWidget::UpdateSequencer()
 {
 	if (!WeakSequencer.IsValid())
 	{
-		WeakSequencer = GetSequencerChecked();
+		WeakSequencer = UE::Private::GetSequencerChecked();
 	}
 }
 
@@ -1335,7 +1708,7 @@ TSharedPtr<SWidget> SConstraintsEditionWidget::CreateContextMenu()
 		return SNullWidget::NullWidget;
 	}
 
-	UWorld* World = GetCurrentWorld();
+	UWorld* World = UE::Private::GetCurrentWorld();
 	if (!IsValid(World))
 	{
 		return SNullWidget::NullWidget;
@@ -1519,7 +1892,7 @@ void SConstraintsEditionWidget::OnItemDoubleClicked(ItemSharedPtr InItem)
 		return;
 	}
 
-	UWorld* World = GetCurrentWorld();
+	UWorld* World = UE::Private::GetCurrentWorld();
 	if (!IsValid(World))
 	{
 		return;
@@ -1705,7 +2078,7 @@ void SConstraintBakeWidget::BakeSelected()
 			Constraints.Add(TransformConstraint);
 		}
 	}
-	UWorld* World = GetCurrentWorld();
+	UWorld* World = UE::Private::GetCurrentWorld();
 
 	FBakingAnimationKeySettings* BakeSettings = Settings->Get();
 	
@@ -1776,6 +2149,8 @@ void SBakeConstraintItem::Construct(
 	const TSharedPtr<FEditableConstraintItem>& InItem,
 	TSharedPtr<SConstraintBakeWidget> InConstraintsWidget)
 {
+	using namespace UE::Private;
+	
 	ConstraintItem = InItem;
 	ConstraintsWidget = InConstraintsWidget;
 
diff --git a/Engine/Plugins/Animation/ControlRig/Source/ControlRigEditor/Private/Editor/Constraints/SConstraintsWidget.h b/Engine/Plugins/Animation/ControlRig/Source/ControlRigEditor/Private/Editor/Constraints/SConstraintsWidget.h
--- a/Engine/Plugins/Animation/ControlRig/Source/ControlRigEditor/Private/Editor/Constraints/SConstraintsWidget.h
+++ b/Engine/Plugins/Animation/ControlRig/Source/ControlRigEditor/Private/Editor/Constraints/SConstraintsWidget.h
@@ -16,9 +16,22 @@ class UTickableTransformConstraint;
 class AActor;
 class SConstraintsCreationWidget;
 class SConstraintsEditionWidget;
+class SComboButton;
 class ISequencer;
+
 DECLARE_DELEGATE(FOnConstraintCreated);
 
+/**
+ * FConstrainable is a UObject/Socket wrapper used to represent the future transformable handle 
+ */
+struct FConstrainable
+{
+	TWeakObjectPtr<AActor> Actor = nullptr;
+	TWeakObjectPtr<UObject> Object = nullptr;
+	FName Socket = NAME_None;
+};
+DECLARE_DELEGATE_RetVal(FConstrainable, FOnGetConstrainable);
+
 /**
  * FConstraintInfo
  */
@@ -77,6 +90,7 @@ public:
 	SLATE_BEGIN_ARGS(SConstraintMenuEntry){}
 
 		SLATE_EVENT(FOnConstraintCreated, OnConstraintCreated)
+		SLATE_EVENT(FOnGetConstrainable, OnGetParent)
 	
 	SLATE_END_ARGS()
 
@@ -97,19 +111,25 @@ public:
 	
 private:
 
+	/** Uses the current selected parent or pick a new one if none has been set. */
+	FReply CreateFromSelectionOrPicker(const bool bUseDefault) const;
+	
 	/** Activates the ActorPickerMode to pick a parent in the viewport. */
 	FReply CreateSelectionPicker(const bool bUseDefault) const;
 
 	/** Delegate triggered when a new constraint has been created. */
 	FOnConstraintCreated OnConstraintCreated;
 
+	/** Delegate to get constrainable parent data. */
+	FOnGetConstrainable OnGetParent;
+
 	/** Creates a widget to edit the constraint default properties. */
 	TSharedRef<SWidget> GenerateConstraintDefaultWidget();
 	
 	/** Creates the constraint between the current selection and the picked actor. */
 	static void CreateConstraint(
 		AActor* InParent,
-		FOnConstraintCreated InDelegate,
+		const FOnConstraintCreated& InDelegate,
 		const ETransformConstraintType InConstraintType,
 		const bool bUseDefault);
 
@@ -298,9 +318,28 @@ public:
 
 	/** Notify from sequencer changed. */
 	void SequencerChanged(const TWeakPtr<ISequencer>& InNewSequencer);
-	
+
+	/** Resets the selected parent data. */
+	void ResetParent();
+
+	/** Returns the parent that has been set using the chooser widget. */
+	const FConstrainable& GetParent() const;
+	bool IsParentValid() const;
+
 private:
 
+	/** Updates the current constrainable parent data. */
+	void UpdateParent(AActor* InActor);
+
+	/** Constrainable parent data to be sued to create a new constraint. */
+	FConstrainable ConstrainableParent;
+
+	/** Widgets to set the constrainable parent without using the picker. */
+	TSharedRef<SWidget> OnGetParentMenuContent();
+	TSharedPtr<SComboButton> GetActorChooserWidget();
+	TSharedRef<SWidget> GetUseSelectedActorWidget();
+	void OnParentMenuOpenChanged(bool bOpened);
+	
 	/** Generates a widget for the specified item */
 	TSharedRef<ITableRow> OnGenerateWidgetForItem(ItemSharedPtr InItem, const TSharedRef<STableViewBase>& OwnerTable);
 
@@ -321,6 +360,7 @@ private:
 	bool SequencerTimeChanged();
 	FQualifiedFrameTime SequencerLastTime;
 
+	TSharedPtr<SComboButton> ActorSelectionButtons;
 };
 
 /** Widget allowing baking of constraints */
