Control Rig: Introduce ControlRig Replay class as parent class of testdata

#rnx

diff --git a/Engine/Plugins/Animation/ControlRig/Source/ControlRig/Private/ControlRigReplay.cpp b/Engine/Plugins/Animation/ControlRig/Source/ControlRig/Private/ControlRigReplay.cpp
new file mode 100644
--- /dev/null
+++ b/Engine/Plugins/Animation/ControlRig/Source/ControlRig/Private/ControlRigReplay.cpp
@@ -0,0 +1,160 @@
+// Copyright Epic Games, Inc. All Rights Reserved.
+
+#include "ControlRigReplay.h"
+#include "ControlRig.h"
+#include "ControlRigObjectVersion.h"
+#include "Units/Execution/RigUnit_PrepareForExecution.h"
+#include "Units/Execution/RigUnit_BeginExecution.h"
+#include "HAL/PlatformTime.h"
+#if WITH_EDITOR
+#include "AssetToolsModule.h"
+#endif
+
+#include UE_INLINE_GENERATED_CPP_BY_NAME(ControlRigReplay)
+
+//////////////////////////////////////////////////////////////////////////////////////////////////////////////////////////////
+/////// FControlRigReplayTracks
+//////////////////////////////////////////////////////////////////////////////////////////////////////////////////////////////
+
+bool FControlRigReplayTracks::Serialize(FArchive& InArchive)
+{
+	if(!FSampleTrackHost::Serialize(InArchive))
+	{
+		return false;
+	}
+	return true;
+}
+
+//////////////////////////////////////////////////////////////////////////////////////////////////////////////////////////////
+/////// UControlRigReplay
+//////////////////////////////////////////////////////////////////////////////////////////////////////////////////////////////
+
+void UControlRigReplay::BeginDestroy()
+{
+	UObject::BeginDestroy();
+}
+
+void UControlRigReplay::Serialize(FArchive& Ar)
+{
+	Ar.UsingCustomVersion(FControlRigObjectVersion::GUID);
+	UObject::Serialize(Ar);
+
+	if (EventQueue.IsEmpty())
+	{
+		EventQueue = {FRigUnit_BeginExecution::EventName};
+	}
+}
+
+UControlRigReplay* UControlRigReplay::CreateNewAsset(FString InDesiredPackagePath, FString InBlueprintPathName, UClass* InAssetClass)
+{
+#if WITH_EDITOR
+	const FAssetToolsModule& AssetToolsModule = FModuleManager::LoadModuleChecked<FAssetToolsModule>("AssetTools");
+	FString UniquePackageName;
+	FString UniqueAssetName;
+	AssetToolsModule.Get().CreateUniqueAssetName(InDesiredPackagePath, TEXT(""), UniquePackageName, UniqueAssetName);
+
+	if (UniquePackageName.EndsWith(UniqueAssetName))
+	{
+		UniquePackageName = UniquePackageName.LeftChop(UniqueAssetName.Len() + 1);
+	}
+
+	UObject* NewAsset = AssetToolsModule.Get().CreateAsset(*UniqueAssetName, *UniquePackageName, InAssetClass, nullptr);
+	if(NewAsset)
+	{
+		// make sure the package is never cooked.
+		UPackage* Package = NewAsset->GetOutermost();
+		Package->SetPackageFlags(Package->GetPackageFlags() | PKG_EditorOnly);
+			
+		if(UControlRigReplay* TestData = Cast<UControlRigReplay>(NewAsset))
+		{
+			TestData->ControlRigObjectPath = InBlueprintPathName;
+			return TestData;
+		}
+	}
+#endif
+	return nullptr;
+}
+
+FVector2D UControlRigReplay::GetTimeRange() const
+{
+	return FVector2D(Tracks.GetContainer()->GetTimeRange());
+}
+
+bool UControlRigReplay::StartRecording(UControlRig* InControlRig, double InRecordingDuration)
+{
+	return false;
+}
+
+bool UControlRigReplay::StopRecording()
+{
+	return false;
+}
+
+EControlRigReplayPlaybackMode UControlRigReplay::GetPlaybackMode() const
+{
+	if(IsReplaying())
+	{
+		return EControlRigReplayPlaybackMode::ReplayInputs;
+	}
+	return EControlRigReplayPlaybackMode::Live;
+}
+
+bool UControlRigReplay::StartReplay(UControlRig* InControlRig)
+{
+	return false;
+}
+
+bool UControlRigReplay::StopReplay()
+{
+	if(UControlRig* ControlRig = ReplayControlRig.Get())
+	{
+		ClearDelegates(ControlRig);
+		ReplayControlRig.Reset();
+		return true;
+	}
+	return false;
+}
+
+bool UControlRigReplay::IsReplaying() const
+{
+	return ReplayControlRig.IsValid();
+}
+
+bool UControlRigReplay::IsRecording() const
+{
+	return RecordControlRig.IsValid();
+}
+
+bool UControlRigReplay::IsValidForTesting() const
+{
+	// todo
+	return false;
+}
+
+bool UControlRigReplay::PerformTest(UControlRig* InSubject, TFunction<void(EMessageSeverity::Type, const FString&)> InLogFunction) const
+{
+	// todo
+	return false;
+}
+
+void UControlRigReplay::ClearDelegates(UControlRig* InControlRig)
+{
+	if(InControlRig)
+	{
+		if(PreConstructionHandle.IsValid())
+		{
+			InControlRig->OnPreConstruction_AnyThread().Remove(PreConstructionHandle);
+			PreConstructionHandle.Reset();
+		}
+		if(PreEventHandle.IsValid())
+		{
+			InControlRig->OnPreExecuted_AnyThread().Remove(PreEventHandle);
+			PreEventHandle.Reset();
+		}
+		if(PostEventHandle.IsValid())
+		{
+			InControlRig->OnExecuted_AnyThread().Remove(PostEventHandle);
+			PostEventHandle.Reset();
+		}
+	}
+}
diff --git a/Engine/Plugins/Animation/ControlRig/Source/ControlRig/Private/ControlRigTestData.cpp b/Engine/Plugins/Animation/ControlRig/Source/ControlRig/Private/ControlRigTestData.cpp
--- a/Engine/Plugins/Animation/ControlRig/Source/ControlRig/Private/ControlRigTestData.cpp
+++ b/Engine/Plugins/Animation/ControlRig/Source/ControlRig/Private/ControlRigTestData.cpp
@@ -27,7 +27,7 @@ bool FControlRigTestDataFrame::Store(UControlRig* InControlRig, bool bInitial)
 	const TArray<FRigVMExternalVariable> ExternalVariables = InControlRig->GetExternalVariables(); 
 	for(const FRigVMExternalVariable& ExternalVariable : ExternalVariables)
 	{
-		FControlRigTestDataVariable VariableData;
+		FControlRigReplayVariable VariableData;
 		VariableData.Name = ExternalVariable.Name;
 		VariableData.CPPType = ExternalVariable.TypeName;
 
@@ -119,8 +119,8 @@ bool FControlRigTestDataFrame::RestoreVariables(UControlRig* InControlRig) const
 			return false;
 		}
 		
-		const FControlRigTestDataVariable* VariableData = Variables.FindByPredicate(
-			[ExternalVariable](const FControlRigTestDataVariable& InVariable) -> bool
+		const FControlRigReplayVariable* VariableData = Variables.FindByPredicate(
+			[ExternalVariable](const FControlRigReplayVariable& InVariable) -> bool
 			{
 				return InVariable.Name == ExternalVariable.Name &&
 					InVariable.CPPType == ExternalVariable.TypeName;
@@ -204,8 +204,7 @@ void UControlRigTestData::BeginDestroy()
 
 void UControlRigTestData::Serialize(FArchive& Ar)
 {
-	Ar.UsingCustomVersion(FControlRigObjectVersion::GUID);
-	UObject::Serialize(Ar);
+	UControlRigReplay::Serialize(Ar);
 	LastFrameIndex = INDEX_NONE;
 
 	Initial.ArchiveCustomVersions = Ar.GetCustomVersions();
@@ -220,7 +219,6 @@ void UControlRigTestData::Serialize(FArchive& Ar)
 		}
 	}
 	
-	
 	// If pose is older than RigPoseWithParentKey, set the active parent of all poses to invalid key
 	if (GetLinkerCustomVersion(FControlRigObjectVersion::GUID) < FControlRigObjectVersion::RigPoseWithParentKey)
 	{
@@ -243,51 +241,15 @@ void UControlRigTestData::Serialize(FArchive& Ar)
 			}
 		}
 	}
-
-	if (EventQueue.IsEmpty())
-	{
-		EventQueue = {FRigUnit_BeginExecution::EventName};
-	}
 }
 
-UControlRigTestData* UControlRigTestData::CreateNewAsset(FString InDesiredPackagePath, FString InBlueprintPathName)
+FVector2D UControlRigTestData::GetTimeRange() const
 {
-#if WITH_EDITOR
-	const FAssetToolsModule& AssetToolsModule = FModuleManager::LoadModuleChecked<FAssetToolsModule>("AssetTools");
-	FString UniquePackageName;
-	FString UniqueAssetName;
-	AssetToolsModule.Get().CreateUniqueAssetName(InDesiredPackagePath, TEXT(""), UniquePackageName, UniqueAssetName);
-
-	if (UniquePackageName.EndsWith(UniqueAssetName))
-	{
-		UniquePackageName = UniquePackageName.LeftChop(UniqueAssetName.Len() + 1);
-	}
-
-	UObject* NewAsset = AssetToolsModule.Get().CreateAsset(*UniqueAssetName, *UniquePackageName, UControlRigTestData::StaticClass(), nullptr);
-	if(NewAsset)
-	{
-		// make sure the package is never cooked.
-		UPackage* Package = NewAsset->GetOutermost();
-		Package->SetPackageFlags(Package->GetPackageFlags() | PKG_EditorOnly);
-			
-		if(UControlRigTestData* TestData = Cast<UControlRigTestData>(NewAsset))
-		{
-			TestData->ControlRigObjectPath = InBlueprintPathName;
-			return TestData;
-		}
-	}
-#endif
-	return nullptr;
-}
-
-FVector2D UControlRigTestData::GetTimeRange(bool bInput) const
-{
-	const TArray<FControlRigTestDataFrame>& Frames = bInput ? InputFrames : OutputFrames;
-	if(Frames.IsEmpty())
+	if(OutputFrames.IsEmpty())
 	{
 		return FVector2D::ZeroVector;
 	}
-	return FVector2D(Frames[0].AbsoluteTime, Frames.Last().AbsoluteTime);
+	return FVector2D(OutputFrames[0].AbsoluteTime, OutputFrames.Last().AbsoluteTime);
 }
 
 int32 UControlRigTestData::GetFrameIndexForTime(double InSeconds, bool bInput) const
@@ -314,14 +276,15 @@ int32 UControlRigTestData::GetFrameIndexForTime(double InSeconds, bool bInput) c
 	return LastFrameIndex;
 }
 
-bool UControlRigTestData::Record(UControlRig* InControlRig, double InRecordingDuration)
+bool UControlRigTestData::StartRecording(UControlRig* InControlRig, double InRecordingDuration)
 {
 	if(InControlRig == nullptr)
 	{
 		return false;
 	}
 
-	ReleaseReplay();
+	RecordControlRig = InControlRig;
+	StopReplay();
 	ClearDelegates(InControlRig);
 
 	EventQueue = InControlRig->EventQueue;
@@ -398,7 +361,7 @@ bool UControlRigTestData::Record(UControlRig* InControlRig, double InRecordingDu
 				DesiredRecordingDuration = 0.0;
 
 				// Once clear delegates is called, we no longer have access to this pointer
-				ClearDelegates(ControlRig);
+				StopRecording();
 			}
 		}
 	);
@@ -412,9 +375,21 @@ bool UControlRigTestData::Record(UControlRig* InControlRig, double InRecordingDu
 	return true;
 }
 
-bool UControlRigTestData::SetupReplay(UControlRig* InControlRig, bool bGroundTruth)
+bool UControlRigTestData::StopRecording()
 {
-	ReleaseReplay();
+	if(RecordControlRig.IsValid())
+	{
+		ClearDelegates(RecordControlRig.Get());
+		RecordControlRig.Reset();
+		return true;
+	}
+	return false;
+}
+
+bool UControlRigTestData::StartReplay(UControlRig* InControlRig)
+{
+	StopRecording();
+	StopReplay();
 	ClearDelegates(InControlRig);
 
 	if(InControlRig == nullptr)
@@ -423,7 +398,7 @@ bool UControlRigTestData::SetupReplay(UControlRig* InControlRig, bool bGroundTru
 	}
 
 	InControlRig->EventQueue = EventQueue;
-	bIsApplyingOutputs = bGroundTruth;
+	bIsApplyingOutputs = true;
 
 	if(InputFrames.IsEmpty() || OutputFrames.IsEmpty())
 	{
@@ -581,31 +556,20 @@ bool UControlRigTestData::SetupReplay(UControlRig* InControlRig, bool bGroundTru
 	return true;
 }
 
-void UControlRigTestData::ReleaseReplay()
+bool UControlRigTestData::StopReplay()
 {
 	if(UControlRig* ControlRig = ReplayControlRig.Get())
 	{
 		ClearDelegates(ControlRig);
 		ReplayControlRig.Reset();
+		return true;
 	}
+	return false;
 }
 
-EControlRigTestDataPlaybackMode UControlRigTestData::GetPlaybackMode() const
-{
-	if(IsReplaying())
-	{
-		if(bIsApplyingOutputs)
-		{
-			return EControlRigTestDataPlaybackMode::GroundTruth;
-		}
-		return EControlRigTestDataPlaybackMode::ReplayInputs;
-	}
-	return EControlRigTestDataPlaybackMode::Live;
-}
-
-bool UControlRigTestData::IsReplaying() const
+bool UControlRigTestData::IsValidForTesting() const
 {
-	return ReplayControlRig.IsValid();
+	return InputFrames.Num() == OutputFrames.Num();
 }
 
 bool UControlRigTestData::PerformTest(UControlRig* InSubject, TFunction<void(EMessageSeverity::Type, const FString&)> InLogFunction) const
@@ -724,8 +688,8 @@ bool UControlRigTestData::PerformTest(UControlRig* InSubject, TFunction<void(EMe
 		const TArray<FRigVMExternalVariable> ExternalVariables = InSubject->GetExternalVariables(); 
 		for(const FRigVMExternalVariable& ExternalVariable : ExternalVariables)
 		{
-			const FControlRigTestDataVariable* VariableData = OutputFrame.Variables.FindByPredicate(
-				[ExternalVariable](const FControlRigTestDataVariable& InVariable) -> bool
+			const FControlRigReplayVariable* VariableData = OutputFrame.Variables.FindByPredicate(
+				[ExternalVariable](const FControlRigReplayVariable& InVariable) -> bool
 				{
 					return InVariable.Name == ExternalVariable.Name &&
 						InVariable.CPPType == ExternalVariable.TypeName;
@@ -764,24 +728,3 @@ bool UControlRigTestData::PerformTest(UControlRig* InSubject, TFunction<void(EMe
 	return bSuccess;
 }
 
-void UControlRigTestData::ClearDelegates(UControlRig* InControlRig)
-{
-	if(InControlRig)
-	{
-		if(PreConstructionHandle.IsValid())
-		{
-			InControlRig->OnPreConstruction_AnyThread().Remove(PreConstructionHandle);
-			PreConstructionHandle.Reset();
-		}
-		if(PreEventHandle.IsValid())
-		{
-			InControlRig->OnPreExecuted_AnyThread().Remove(PreEventHandle);
-			PreEventHandle.Reset();
-		}
-		if(PostEventHandle.IsValid())
-		{
-			InControlRig->OnExecuted_AnyThread().Remove(PostEventHandle);
-			PostEventHandle.Reset();
-		}
-	}
-}
diff --git a/Engine/Plugins/Animation/ControlRig/Source/ControlRig/Private/Tracks/SampleTrackContainer.cpp b/Engine/Plugins/Animation/ControlRig/Source/ControlRig/Private/Tracks/SampleTrackContainer.cpp
--- a/Engine/Plugins/Animation/ControlRig/Source/ControlRig/Private/Tracks/SampleTrackContainer.cpp
+++ b/Engine/Plugins/Animation/ControlRig/Source/ControlRig/Private/Tracks/SampleTrackContainer.cpp
@@ -45,13 +45,18 @@ FSampleTrackHost& FSampleTrackHost::operator=(const FSampleTrackHost& InOther)
 	return *this;
 }
 
-void FSampleTrackHost::Serialize(FArchive& InArchive)
+bool FSampleTrackHost::Serialize(FArchive& InArchive)
 {
 	if(Container == nullptr)
 	{
 		Container = MakeUnique<FSampleTrackContainer>();
 	}
-	Container->Serialize(InArchive);
+
+	// Serialize normal tagged property data
+	UScriptStruct* const Struct = FSampleTrackHost::StaticStruct();
+	Struct->SerializeTaggedProperties(InArchive, reinterpret_cast<uint8*>(this), Struct, nullptr);
+
+	return Container->Serialize(InArchive);
 }
 
 const FSampleTrackContainer* FSampleTrackHost::GetContainer() const
@@ -126,7 +131,7 @@ void FSampleTrackContainer::Reserve(int32 InNum)
 	}
 }
 
-void FSampleTrackContainer::Serialize(FArchive& InArchive)
+bool FSampleTrackContainer::Serialize(FArchive& InArchive)
 {
 	InArchive.UsingCustomVersion(FControlRigObjectVersion::GUID);
 	
@@ -284,6 +289,8 @@ void FSampleTrackContainer::Serialize(FArchive& InArchive)
 			InArchive << bStoreCompressedBytes;
 		}
 	}
+
+	return true;
 }
 
 void FSampleTrackContainer::SetForceToUseCompression(bool InForce)
@@ -477,6 +484,15 @@ TSharedPtr<FSampleTrack<TArray<FInstancedStruct>>> FSampleTrackContainer::AddStr
 	return Track;
 }
 
+FVector2f FSampleTrackContainer::GetTimeRange() const
+{
+	if(TimeSamples.IsEmpty())
+	{
+		return FVector2f::ZeroVector;
+	}
+	return {TimeSamples[0], TimeSamples.Last()};
+}
+
 int32 FSampleTrackContainer::GetTrackIndex(const FName& InName) const
 {
 	if(const int32* Index = NameToIndex.Find(InName))
diff --git a/Engine/Plugins/Animation/ControlRig/Source/ControlRig/Public/ControlRigReplay.h b/Engine/Plugins/Animation/ControlRig/Source/ControlRig/Public/ControlRigReplay.h
new file mode 100644
--- /dev/null
+++ b/Engine/Plugins/Animation/ControlRig/Source/ControlRig/Public/ControlRigReplay.h
@@ -0,0 +1,145 @@
+// Copyright Epic Games, Inc. All Rights Reserved.
+
+#pragma once
+
+#include "CoreMinimal.h"
+#include "Rigs/RigHierarchy.h"
+#include "Rigs/RigHierarchyPose.h"
+#include "Logging/TokenizedMessage.h"
+#include "Tracks/SampleTrackContainer.h"
+
+#include "ControlRigReplay.generated.h"
+
+class UControlRigReplay;
+class UControlRig;
+
+USTRUCT(BlueprintType)
+struct CONTROLRIG_API FControlRigReplayVariable
+{
+	GENERATED_USTRUCT_BODY()
+
+	FControlRigReplayVariable()
+	{
+		Name = NAME_None;
+		CPPType = NAME_None;
+		Value = FString();
+	}
+
+	UPROPERTY(VisibleAnywhere, BlueprintReadOnly, Category = "ControlRigReplayVariable")
+	FName Name;
+
+	UPROPERTY(VisibleAnywhere, BlueprintReadOnly, Category = "ControlRigReplayVariable")
+	FName CPPType;
+
+	UPROPERTY(VisibleAnywhere, BlueprintReadOnly, Category = "ControlRigReplayVariable")
+	FString Value;
+};
+
+UENUM()
+enum class EControlRigReplayPlaybackMode : uint8
+{
+	Live,
+	ReplayInputs,
+	GroundTruth,
+	Max UMETA(Hidden),
+};
+
+USTRUCT()
+struct CONTROLRIG_API FControlRigReplayTracks : public FSampleTrackHost
+{
+	GENERATED_BODY()
+
+	virtual bool Serialize(FArchive& InArchive) override;
+};
+
+template<>
+struct TStructOpsTypeTraits<FControlRigReplayTracks> : public TStructOpsTypeTraitsBase2<FControlRigReplayTracks>
+{
+	enum 
+	{
+		WithSerializer = true, // struct has a Serialize function for serializing its state to an FArchive.
+		WithCopy = true, // struct can be copied via its copy assignment operator.
+		WithIdentical = true, // struct can be compared via an Identical(const T* Other, uint32 PortFlags) function.
+	};
+};
+
+UCLASS(BlueprintType)
+class CONTROLRIG_API UControlRigReplay : public UObject
+{
+	GENERATED_BODY()
+
+public:
+
+	UControlRigReplay()
+		: Tolerance(0.001)
+		, EnableTest(true)
+	{}
+
+	virtual void BeginDestroy() override;
+
+	virtual void Serialize(FArchive& Ar) override;
+
+	UFUNCTION(BlueprintCallable, Category = "ControlRigReplay")
+	static UControlRigReplay* CreateNewAsset(FString InDesiredPackagePath, FString InBlueprintPathName, UClass* InAssetClass);
+
+	UFUNCTION(BlueprintPure, Category = "ControlRigReplay")
+	virtual FVector2D GetTimeRange() const;
+
+	UPROPERTY(AssetRegistrySearchable, VisibleAnywhere, BlueprintReadOnly, Category = "ControlRigReplay")
+	FSoftObjectPath ControlRigObjectPath;
+
+	UPROPERTY(EditAnywhere, Category = "ControlRigReplay")
+	FControlRigReplayTracks Tracks;
+	
+	UPROPERTY(EditAnywhere, BlueprintReadOnly, Category = "ControlRigReplay")
+	double Tolerance;
+
+	UPROPERTY(EditAnywhere, BlueprintReadOnly, Category = "ControlRigReplay")
+	TArray<int32> FramesToSkip;
+
+	UPROPERTY(EditAnywhere, BlueprintReadOnly, Category = "ControlRigReplay")
+	TArray<FName> EventQueue;
+
+	UPROPERTY(EditAnywhere, BlueprintReadOnly, Category = "ControlRigReplay", AssetRegistrySearchable)
+	bool EnableTest;
+
+	UFUNCTION(BlueprintCallable, Category = "ControlRigReplay")
+	virtual bool StartRecording(UControlRig* InControlRig, double InRecordingDuration = 0.0);
+
+	UFUNCTION(BlueprintCallable, Category = "ControlRigReplay")
+	virtual bool StopRecording();
+
+	UFUNCTION(BlueprintCallable, Category = "ControlRigReplay")
+	virtual bool StartReplay(UControlRig* InControlRig);
+
+	UFUNCTION(BlueprintCallable, Category = "ControlRigReplay")
+	virtual bool StopReplay();
+
+	UFUNCTION(BlueprintPure, Category = "ControlRigReplay")
+	virtual bool IsReplaying() const;
+
+	UFUNCTION(BlueprintPure, Category = "ControlRigReplay")
+	virtual bool IsRecording() const;
+
+	UFUNCTION(BlueprintPure, Category = "ControlRigReplay")
+	EControlRigReplayPlaybackMode GetPlaybackMode() const;
+
+	bool IsTestEnabled() const
+	{
+		return EnableTest;
+	}
+
+	virtual bool IsValidForTesting() const;
+
+	virtual bool PerformTest(UControlRig* InSubject, TFunction<void(EMessageSeverity::Type, const FString&)> InLogFunction) const;
+
+protected:
+
+	void ClearDelegates(UControlRig* InControlRig);
+
+	TWeakObjectPtr<UControlRig> ReplayControlRig;
+	TWeakObjectPtr<UControlRig> RecordControlRig;
+	FDelegateHandle PreConstructionHandle;
+	FDelegateHandle PreEventHandle;
+	FDelegateHandle PostEventHandle;
+};
\ No newline at end of file
diff --git a/Engine/Plugins/Animation/ControlRig/Source/ControlRig/Public/ControlRigTestData.h b/Engine/Plugins/Animation/ControlRig/Source/ControlRig/Public/ControlRigTestData.h
--- a/Engine/Plugins/Animation/ControlRig/Source/ControlRig/Public/ControlRigTestData.h
+++ b/Engine/Plugins/Animation/ControlRig/Source/ControlRig/Public/ControlRigTestData.h
@@ -2,38 +2,13 @@
 
 #pragma once
 
-#include "CoreMinimal.h"
-#include "Rigs/RigHierarchy.h"
-#include "Rigs/RigHierarchyPose.h"
-#include "Logging/TokenizedMessage.h"
+#include "ControlRigReplay.h"
 
 #include "ControlRigTestData.generated.h"
 
 class UControlRigTestData;
 class UControlRig;
 
-USTRUCT(BlueprintType)
-struct CONTROLRIG_API FControlRigTestDataVariable
-{
-	GENERATED_USTRUCT_BODY()
-
-	FControlRigTestDataVariable()
-	{
-		Name = NAME_None;
-		CPPType = NAME_None;
-		Value = FString();
-	}
-
-	UPROPERTY(VisibleAnywhere, BlueprintReadOnly, Category = "ControlRigTestDataVariable")
-	FName Name;
-
-	UPROPERTY(VisibleAnywhere, BlueprintReadOnly, Category = "ControlRigTestDataVariable")
-	FName CPPType;
-
-	UPROPERTY(VisibleAnywhere, BlueprintReadOnly, Category = "ControlRigTestDataVariable")
-	FString Value;
-};
-
 USTRUCT(BlueprintType)
 struct CONTROLRIG_API FControlRigTestDataFrame
 {
@@ -60,7 +35,7 @@ struct CONTROLRIG_API FControlRigTestDataFrame
 	double DeltaTime;
 
 	UPROPERTY(VisibleAnywhere, BlueprintReadOnly, Category = "ControlRigTestDataFrame")
-	TArray<FControlRigTestDataVariable> Variables;
+	TArray<FControlRigReplayVariable> Variables;
 
 	UPROPERTY(VisibleAnywhere, BlueprintReadOnly, Category = "ControlRigTestDataFrame")
 	FRigPose Pose;
@@ -73,26 +48,15 @@ struct CONTROLRIG_API FControlRigTestDataFrame
 	mutable bool bTestMetadata;
 };
 
-UENUM()
-enum class EControlRigTestDataPlaybackMode : uint8
-{
-	Live,
-	ReplayInputs,
-	GroundTruth,
-	Max UMETA(Hidden),
-};
-
 UCLASS(BlueprintType)
-class CONTROLRIG_API UControlRigTestData : public UObject
+class CONTROLRIG_API UControlRigTestData : public UControlRigReplay
 {
 	GENERATED_BODY()
 
 public:
 
 	UControlRigTestData()
-		: Tolerance(0.001)
-		, EnableTest(true)
-		, LastFrameIndex(INDEX_NONE)
+		: LastFrameIndex(INDEX_NONE)
 		, DesiredRecordingDuration(0.0)
 		, TimeAtStartOfRecording(0.0)
 		, bIsApplyingOutputs(false)
@@ -102,35 +66,15 @@ public:
 
 	virtual void Serialize(FArchive& Ar) override;
 
-	UFUNCTION(BlueprintCallable, Category = "ControlRigTestData")
-	static UControlRigTestData* CreateNewAsset(FString InDesiredPackagePath, FString InBlueprintPathName);
-
-	UFUNCTION(BlueprintPure, Category = "ControlRigTestData")
-	FVector2D GetTimeRange(bool bInput = false) const;
+	virtual FVector2D GetTimeRange() const override;
 
 	UFUNCTION(BlueprintPure, Category = "ControlRigTestData")
 	int32 GetFrameIndexForTime(double InSeconds, bool bInput = false) const;
 
-	UFUNCTION(BlueprintCallable, Category = "ControlRigTestData")
-	bool Record(UControlRig* InControlRig, double InRecordingDuration = 0.0);
-
-	UFUNCTION(BlueprintCallable, Category = "ControlRigTestData")
-	bool SetupReplay(UControlRig* InControlRig, bool bGroundTruth = true);
-
-	UFUNCTION(BlueprintCallable, Category = "ControlRigTestData")
-	void ReleaseReplay();
-
-	UFUNCTION(BlueprintPure, Category = "ControlRigTestData")
-	EControlRigTestDataPlaybackMode GetPlaybackMode() const;
-
-	UFUNCTION(BlueprintPure, Category = "ControlRigTestData")
-	bool IsReplaying() const;
-
-	UFUNCTION(BlueprintPure, Category = "ControlRigTestData")
-	bool IsRecording() const { return DesiredRecordingDuration >= SMALL_NUMBER; }
-
-	UPROPERTY(AssetRegistrySearchable, VisibleAnywhere, BlueprintReadOnly, Category = "ControlRigTestData")
-	FSoftObjectPath ControlRigObjectPath;
+	virtual bool StartRecording(UControlRig* InControlRig, double InRecordingDuration = 0.0) override;
+	virtual bool StopRecording() override;
+	virtual bool StartReplay(UControlRig* InControlRig) override;
+	virtual bool StopReplay() override;
 
 	UPROPERTY(VisibleAnywhere, BlueprintReadOnly, Category = "ControlRigTestData")
 	FControlRigTestDataFrame Initial;
@@ -141,35 +85,14 @@ public:
 	UPROPERTY(VisibleAnywhere, BlueprintReadOnly, Category = "ControlRigTestData")
 	TArray<FControlRigTestDataFrame> OutputFrames;
 
-	UPROPERTY(EditAnywhere, BlueprintReadOnly, Category = "ControlRigTestData")
-	TArray<int32> FramesToSkip;
-
-	UPROPERTY(EditAnywhere, BlueprintReadOnly, Category = "ControlRigTestData")
-	double Tolerance;
-
-	UPROPERTY(EditAnywhere, BlueprintReadOnly, Category = "ControlRigTestData")
-	TArray<FName> EventQueue;
-
-	UPROPERTY(EditAnywhere, BlueprintReadOnly, Category = "ControlRigTestData", AssetRegistrySearchable)
-	bool EnableTest;
-
-	bool IsTestEnabled() const
-	{
-		return EnableTest;
-	}
-
-	bool PerformTest(UControlRig* InSubject, TFunction<void(EMessageSeverity::Type, const FString&)> InLogFunction) const;
-
-private:
+	virtual bool IsValidForTesting() const override;
+	
+	virtual bool PerformTest(UControlRig* InSubject, TFunction<void(EMessageSeverity::Type, const FString&)> InLogFunction) const override;
 
-	void ClearDelegates(UControlRig* InControlRig);
+protected:
 
 	mutable int32 LastFrameIndex;
 	double DesiredRecordingDuration;
 	double TimeAtStartOfRecording;
-	TWeakObjectPtr<UControlRig> ReplayControlRig;
 	bool bIsApplyingOutputs;
-	FDelegateHandle PreConstructionHandle;
-	FDelegateHandle PreEventHandle;
-	FDelegateHandle PostEventHandle;
 };
\ No newline at end of file
diff --git a/Engine/Plugins/Animation/ControlRig/Source/ControlRig/Public/Tracks/SampleTrackContainer.h b/Engine/Plugins/Animation/ControlRig/Source/ControlRig/Public/Tracks/SampleTrackContainer.h
--- a/Engine/Plugins/Animation/ControlRig/Source/ControlRig/Public/Tracks/SampleTrackContainer.h
+++ b/Engine/Plugins/Animation/ControlRig/Source/ControlRig/Public/Tracks/SampleTrackContainer.h
@@ -20,15 +20,30 @@ public:
 
 	FSampleTrackHost& operator =(const FSampleTrackHost& InOther);
 
-	virtual void Serialize(FArchive& InArchive);
+	/** By always returning false here, we can disable delta serialization */
+	virtual bool Identical(const FSampleTrackHost* Other, uint32 PortFlags) const { return false; }
+
+	virtual bool Serialize(FArchive& InArchive);
 
 	const FSampleTrackContainer* GetContainer() const;
 	FSampleTrackContainer* GetContainer();
 	
 private:
+
 	TUniquePtr<FSampleTrackContainer> Container;
 };
 
+template<>
+struct TStructOpsTypeTraits<FSampleTrackHost> : public TStructOpsTypeTraitsBase2<FSampleTrackHost>
+{
+	enum 
+	{
+		WithSerializer = true, // struct has a Serialize function for serializing its state to an FArchive.
+		WithCopy = true, // struct can be copied via its copy assignment operator.
+		WithIdentical = true, // struct can be compared via an Identical(const T* Other, uint32 PortFlags) function.
+	};
+};
+
 class CONTROLRIG_API FSampleTrackContainer : public TSharedFromThis<FSampleTrackContainer>
 {
 public:
@@ -39,7 +54,7 @@ public:
 	void Shrink();
 	void Compact(float InTolerance = KINDA_SMALL_NUMBER);
 	void Reserve(int32 InNum);
-	void Serialize(FArchive& InArchive);
+	bool Serialize(FArchive& InArchive);
 	void SetForceToUseCompression(bool InForce = true);
 
 	TSharedPtr<FSampleTrack<bool>> AddBoolTrack(const FName& InName);
@@ -74,6 +89,7 @@ public:
 		return Tracks.Num();
 	}
 
+	FVector2f GetTimeRange() const;
 	const TArray<float>& GetTimeSamples() const { return TimeSamples; }
 
 	int32 GetTrackIndex(const FName& InName) const;
diff --git a/Engine/Plugins/Animation/ControlRig/Source/ControlRigEditor/Private/ControlRigEditorStyle.h b/Engine/Plugins/Animation/ControlRig/Source/ControlRigEditor/Private/ControlRigEditorStyle.h
--- a/Engine/Plugins/Animation/ControlRig/Source/ControlRigEditor/Private/ControlRigEditorStyle.h
+++ b/Engine/Plugins/Animation/ControlRig/Source/ControlRigEditor/Private/ControlRigEditorStyle.h
@@ -150,7 +150,7 @@ public:
 
 		// Test Data
 		{
-			Set("ControlRig.TestData.Record", new IMAGE_BRUSH("Slate/RecordingIndicator", Icon32x32));
+			Set("ControlRig.Replay.Record", new IMAGE_BRUSH("Slate/RecordingIndicator", Icon32x32));
 		}
 
 		Set("ControlRig.ConnectorPrimary", new IMAGE_BRUSH_SVG("Slate/Connector_Primary", Icon128x128));
diff --git a/Engine/Plugins/Animation/ControlRig/Source/ControlRigEditor/Private/ControlRigElementDetails.cpp b/Engine/Plugins/Animation/ControlRig/Source/ControlRigEditor/Private/ControlRigElementDetails.cpp
--- a/Engine/Plugins/Animation/ControlRig/Source/ControlRigEditor/Private/ControlRigElementDetails.cpp
+++ b/Engine/Plugins/Animation/ControlRig/Source/ControlRigEditor/Private/ControlRigElementDetails.cpp
@@ -1717,7 +1717,7 @@ void FRigBaseElementDetails::CustomizeMetadata(IDetailLayoutBuilder& DetailBuild
 			{
 				if (FControlRigEditor* Editor = static_cast<FControlRigEditor*>(EditorInstance))
 				{
-					if (Editor->GetTestAssetPlaybackMode() != EControlRigTestDataPlaybackMode::Live)
+					if (Editor->GetReplayPlaybackMode() != EControlRigReplayPlaybackMode::Live)
 					{
 						Hierarchy->OnMetadataChanged().Remove(MetadataHandle);
 						MetadataHandle.Reset();
diff --git a/Engine/Plugins/Animation/ControlRig/Source/ControlRigEditor/Private/Editor/ControlRigEditor.cpp b/Engine/Plugins/Animation/ControlRig/Source/ControlRigEditor/Private/Editor/ControlRigEditor.cpp
--- a/Engine/Plugins/Animation/ControlRig/Source/ControlRigEditor/Private/Editor/ControlRigEditor.cpp
+++ b/Engine/Plugins/Animation/ControlRig/Source/ControlRigEditor/Private/Editor/ControlRigEditor.cpp
@@ -115,6 +115,7 @@
 #include "Editor/RigVMGraphDetailCustomization.h"
 #include "Widgets/SRigVMSwapAssetReferencesWidget.h"
 #include "Widgets/SRigVMBulkEditDialog.h"
+#include "ControlRigTestData.h"
 
 #define LOCTEXT_NAMESPACE "ControlRigEditor"
 
@@ -461,98 +462,98 @@ const FSlateBrush* FControlRigEditor::GetDefaultTabIcon() const
 	return TabIcon.GetIcon();
 }
 
-FText FControlRigEditor::GetTestAssetName() const
+FText FControlRigEditor::GetReplayAssetName() const
 {
-	if(TestDataStrongPtr.IsValid())
+	if(ReplayStrongPtr.IsValid())
 	{
-		return FText::FromString(TestDataStrongPtr->GetName());
+		return FText::FromString(ReplayStrongPtr->GetName());
 	}
 
-	static const FText NoTestAsset = LOCTEXT("NoTestAsset", "No Test Asset");
-	return NoTestAsset;
+	static const FText NoReplayAsset = LOCTEXT("NoReplayAsset", "No Replay Asset");
+	return NoReplayAsset;
 }
 
-FText FControlRigEditor::GetTestAssetTooltip() const
+FText FControlRigEditor::GetReplayAssetTooltip() const
 {
-	if(TestDataStrongPtr.IsValid())
+	if(ReplayStrongPtr.IsValid())
 	{
-		return FText::FromString(TestDataStrongPtr->GetPathName());
+		return FText::FromString(ReplayStrongPtr->GetPathName());
 	}
-	static const FText NoTestAssetTooltip = LOCTEXT("NoTestAssetTooltip", "Click the record button to the left to record a new frame");
-	return NoTestAssetTooltip;
+	static const FText NoReplayTooltip = LOCTEXT("NoReplayAssetTooltip", "Click the record button to the left to record a new replay");
+	return NoReplayTooltip;
 }
 
-bool FControlRigEditor::SetTestAssetPath(const FString& InAssetPath)
+bool FControlRigEditor::SetReplayAssetPath(const FString& InAssetPath)
 {
-	if(TestDataStrongPtr.IsValid())
+	if(ReplayStrongPtr.IsValid())
 	{
-		if(TestDataStrongPtr->GetPathName().Equals(InAssetPath, ESearchCase::CaseSensitive))
+		if(ReplayStrongPtr->GetPathName().Equals(InAssetPath, ESearchCase::CaseSensitive))
 		{
 			return false;
 		}
 	}
 
-	if(TestDataStrongPtr.IsValid())
+	if(ReplayStrongPtr.IsValid())
 	{
-		TestDataStrongPtr->ReleaseReplay();
-		TestDataStrongPtr.Reset();
+		ReplayStrongPtr->StopReplay();
+		ReplayStrongPtr.Reset();
 	}
 
 	if(!InAssetPath.IsEmpty())
 	{
-		if(UControlRigTestData* TestData = LoadObject<UControlRigTestData>(GetControlRigBlueprint(), *InAssetPath))
+		if(UControlRigReplay* Replay = LoadObject<UControlRigReplay>(GetControlRigBlueprint(), *InAssetPath))
 		{
-			TestDataStrongPtr = TStrongObjectPtr<UControlRigTestData>(TestData);
+			ReplayStrongPtr = TStrongObjectPtr<UControlRigReplay>(Replay);
 		}
 	}
 	return true;
 }
 
-TSharedRef<SWidget> FControlRigEditor::GenerateTestAssetModeMenuContent()
+TSharedRef<SWidget> FControlRigEditor::GenerateReplayAssetModeMenuContent()
 {
 	FMenuBuilder MenuBuilder(true, GetToolkitCommands());
 
 	MenuBuilder.BeginSection(TEXT("Default"));
 	MenuBuilder.AddMenuEntry(
-		LOCTEXT("ClearTestAssset", "Clear"),
-		LOCTEXT("ClearTestAsset_ToolTip", "Clears the test asset"),
+		LOCTEXT("ClearReplay", "Clear"),
+		LOCTEXT("ClearReplay_ToolTip", "Clears the test asset"),
 		FSlateIcon(),
 		FUIAction(
 			FExecuteAction::CreateLambda([this]()
 			{
-				SetTestAssetPath(FString());
+				SetReplayAssetPath(FString());
 			}
 		)	
 	));
 	MenuBuilder.EndSection();
 
 	const FAssetRegistryModule& AssetRegistryModule = FModuleManager::LoadModuleChecked<FAssetRegistryModule>("AssetRegistry");
-	TArray<FAssetData> TestDataAssets;
+	TArray<FAssetData> ReplayAssets;
 	FARFilter AssetFilter;
-	AssetFilter.ClassPaths.Add(UControlRigTestData::StaticClass()->GetClassPathName());
-	AssetRegistryModule.Get().GetAssets(AssetFilter, TestDataAssets);
+	AssetFilter.ClassPaths.Add(UControlRigReplay::StaticClass()->GetClassPathName());
+	AssetRegistryModule.Get().GetAssets(AssetFilter, ReplayAssets);
 
 	const FString CurrentObjectPath = GetControlRigBlueprint()->GetPathName();
-	TestDataAssets.RemoveAll([CurrentObjectPath](const FAssetData& InAssetData)
+	ReplayAssets.RemoveAll([CurrentObjectPath](const FAssetData& InAssetData)
 	{
-		const FString ControlRigObjectPath = InAssetData.GetTagValueRef<FString>(GET_MEMBER_NAME_CHECKED(UControlRigTestData, ControlRigObjectPath));
+		const FString ControlRigObjectPath = InAssetData.GetTagValueRef<FString>(GET_MEMBER_NAME_CHECKED(UControlRigReplay, ControlRigObjectPath));
 		return ControlRigObjectPath != CurrentObjectPath;
 	});
 
-	if(!TestDataAssets.IsEmpty())
+	if(!ReplayAssets.IsEmpty())
 	{
 		MenuBuilder.BeginSection(TEXT("Assets"));
-		for(const FAssetData& TestDataAsset : TestDataAssets)
+		for(const FAssetData& ReplayAsset : ReplayAssets)
 		{
-			const FString TestDataObjectPath = TestDataAsset.GetObjectPathString();
+			const FString ReplayObjectPath = ReplayAsset.GetObjectPathString();
 			FString Right;
-			if(TestDataObjectPath.Split(TEXT("."), nullptr, &Right))
+			if(ReplayObjectPath.Split(TEXT("."), nullptr, &Right))
 			{
-				MenuBuilder.AddMenuEntry(FText::FromString(Right), FText::FromString(TestDataObjectPath), FSlateIcon(),
+				MenuBuilder.AddMenuEntry(FText::FromString(Right), FText::FromString(ReplayObjectPath), FSlateIcon(),
 					FUIAction(
-						FExecuteAction::CreateLambda([this, TestDataObjectPath]()
+						FExecuteAction::CreateLambda([this, ReplayObjectPath]()
 						{
-							SetTestAssetPath(TestDataObjectPath);
+							SetReplayAssetPath(ReplayObjectPath);
 						}
 					)	
 				));
@@ -563,53 +564,53 @@ TSharedRef<SWidget> FControlRigEditor::GenerateTestAssetModeMenuContent()
 	return MenuBuilder.MakeWidget();
 }
 
-TSharedRef<SWidget> FControlRigEditor::GenerateTestAssetRecordMenuContent()
+TSharedRef<SWidget> FControlRigEditor::GenerateReplayAssetRecordMenuContent()
 {
 	FMenuBuilder MenuBuilder(true, GetToolkitCommands());
 
 	MenuBuilder.AddMenuEntry(
-		LOCTEXT("TestAssetRecordSingleFrame", "Single Frame"),
-		LOCTEXT("TestAssetRecordSingleFrame_ToolTip", "Records a single frame into the test data asset"),
+		LOCTEXT("ReplayReplayRecordSingleFrame", "Single Frame"),
+		LOCTEXT("ReplayReplayRecordSingleFrame_ToolTip", "Records a single frame into the replay asset"),
 		FSlateIcon(),
 		FUIAction(
 			FExecuteAction::CreateLambda([this]()
 			{
-				RecordTestData(0);
+				RecordReplay(0);
 			}
 		)	
 	));
 	MenuBuilder.AddMenuEntry(
-		LOCTEXT("TestAssetRecordOneSecond", "1 Second"),
-		LOCTEXT("TestAssetRecordOneSecond_ToolTip", "Records 1 second of animation into the test data asset"),
+		LOCTEXT("ReplayReplayRecordOneSecond", "1 Second"),
+		LOCTEXT("ReplayRecordOneSecond_ToolTip", "Records 1 second of animation into the replay asset"),
 		FSlateIcon(),
 		FUIAction(
 			FExecuteAction::CreateLambda([this]()
 			{
-				RecordTestData(1);
+				RecordReplay(1);
 			}
 		)	
 	));
 
 	MenuBuilder.AddMenuEntry(
-	LOCTEXT("TestAssetRecordFiveSeconds", "5 Seconds"),
-	LOCTEXT("TestAssetRecordFiveSeconds_ToolTip", "Records 5 seconds of animation into the test data asset"),
+	LOCTEXT("ReplayRecordFiveSeconds", "5 Seconds"),
+	LOCTEXT("ReplayRecordFiveSeconds_ToolTip", "Records 5 seconds of animation into the replay asset"),
 	FSlateIcon(),
 	FUIAction(
 			FExecuteAction::CreateLambda([this]()
 			{
-				RecordTestData(5);
+				RecordReplay(5);
 			}
 		)	
 	));
 
 	MenuBuilder.AddMenuEntry(
-	LOCTEXT("TestAssetRecordTenSeconds", "10 Seconds"),
-	LOCTEXT("TestAssetRecordTenSeconds_ToolTip", "Records 10 seconds of animation into the test data asset"),
+	LOCTEXT("ReplayRecordTenSeconds", "10 Seconds"),
+	LOCTEXT("ReplayRecordTenSeconds_ToolTip", "Records 10 seconds of animation into the replay asset"),
 	FSlateIcon(),
 	FUIAction(
 			FExecuteAction::CreateLambda([this]()
 			{
-				RecordTestData(10);
+				RecordReplay(10);
 			}
 		)	
 	));
@@ -617,7 +618,7 @@ TSharedRef<SWidget> FControlRigEditor::GenerateTestAssetRecordMenuContent()
 	return MenuBuilder.MakeWidget();
 }
 
-bool FControlRigEditor::RecordTestData(double InRecordingDuration)
+bool FControlRigEditor::RecordReplay(double InRecordingDuration)
 {
 	if(GetControlRig() == nullptr)
 	{
@@ -627,52 +628,47 @@ bool FControlRigEditor::RecordTestData(double InRecordingDuration)
 	
 	// create a new test asset
 	static const FString Folder = TEXT("/Game/Animation/ControlRig/NoCook/");
-	const FString DesiredPackagePath =  FString::Printf(TEXT("%s/%s_TestData"), *Folder, *GetControlRigBlueprint()->GetName());
+	const FString DesiredPackagePath =  FString::Printf(TEXT("%s/%s_Replay"), *Folder, *GetControlRigBlueprint()->GetName());
 
-	if(UControlRigTestData* TestData = UControlRigTestData::CreateNewAsset(DesiredPackagePath, GetControlRigBlueprint()->GetPathName()))
+	if(const UControlRigReplay* Replay = UControlRigReplay::CreateNewAsset(DesiredPackagePath, GetControlRigBlueprint()->GetPathName(), UControlRigTestData::StaticClass()))
 	{
-		SetTestAssetPath(TestData->GetPathName());
+		SetReplayAssetPath(Replay->GetPathName());
 	}
 	
-	if(UControlRigTestData* TestData = TestDataStrongPtr.Get())
+	if(UControlRigReplay* Replay = ReplayStrongPtr.Get())
 	{
-		TestData->Record(GetControlRig(), InRecordingDuration);
+		Replay->StartRecording(GetControlRig(), InRecordingDuration);
 	}
 	return true;
 }
 
-void FControlRigEditor::ToggleTestData()
+void FControlRigEditor::ToggleReplay()
 {
-	if(TestDataStrongPtr.IsValid())
+	if(ReplayStrongPtr.IsValid())
 	{
-		switch(TestDataStrongPtr->GetPlaybackMode())
+		switch(ReplayStrongPtr->GetPlaybackMode())
 		{
-			case EControlRigTestDataPlaybackMode::ReplayInputs:
+			case EControlRigReplayPlaybackMode::ReplayInputs:
 			{
-				TestDataStrongPtr->ReleaseReplay();
-				break;
-			}
-			case EControlRigTestDataPlaybackMode::GroundTruth:
-			{
-				TestDataStrongPtr->SetupReplay(GetControlRig(), false);
+				ReplayStrongPtr->StopReplay();
 				break;
 			}
 			default:
 			{
-				TestDataStrongPtr->SetupReplay(GetControlRig(), true);
+				ReplayStrongPtr->StartReplay(GetControlRig());
 				break;
 			}
 		}
 	}
 }
 
-EControlRigTestDataPlaybackMode FControlRigEditor::GetTestAssetPlaybackMode() const
+EControlRigReplayPlaybackMode FControlRigEditor::GetReplayPlaybackMode() const
 {
-	if (TestDataStrongPtr.IsValid())
+	if (ReplayStrongPtr.IsValid())
 	{
-		return TestDataStrongPtr->GetPlaybackMode();
+		return ReplayStrongPtr->GetPlaybackMode();
 	}
-	return EControlRigTestDataPlaybackMode::Live;
+	return EControlRigReplayPlaybackMode::Live;
 }
 
 void FControlRigEditor::FillToolbar(FToolBarBuilder& ToolbarBuilder, bool bEndSection)
@@ -778,59 +774,58 @@ void FControlRigEditor::FillToolbar(FToolBarBuilder& ToolbarBuilder, bool bEndSe
 		{
 			ToolbarBuilder.AddSeparator();
 
-			FUIAction TestAssetAction;
 			ToolbarBuilder.AddComboButton(
 				FUIAction(FExecuteAction(), FCanExecuteAction::CreateLambda([this]()
 				{
-					if(TestDataStrongPtr.IsValid())
+					if(ReplayStrongPtr.IsValid())
 					{
-						return !TestDataStrongPtr->IsReplaying() && !TestDataStrongPtr->IsRecording();
+						return !ReplayStrongPtr->IsReplaying() && !ReplayStrongPtr->IsRecording();
 					}
 					return true;
 				})),
-				FOnGetContent::CreateSP(this, &FControlRigEditor::GenerateTestAssetModeMenuContent),
-				TAttribute<FText>::CreateSP(this, &FControlRigEditor::GetTestAssetName),
-				TAttribute<FText>::CreateSP(this, &FControlRigEditor::GetTestAssetTooltip),
+				FOnGetContent::CreateSP(this, &FControlRigEditor::GenerateReplayAssetModeMenuContent),
+				TAttribute<FText>::CreateSP(this, &FControlRigEditor::GetReplayAssetName),
+				TAttribute<FText>::CreateSP(this, &FControlRigEditor::GetReplayAssetTooltip),
 				FSlateIcon(FAppStyle::GetAppStyleSetName(), "AutomationTools.TestAutomation"),
 				false);
 
 			ToolbarBuilder.AddToolBarButton(FUIAction(
 					FExecuteAction::CreateLambda([this]()
 					{
-						RecordTestData(0.0);
+						RecordReplay(0.0);
 					}),
 					FCanExecuteAction::CreateLambda([this]()
 					{
-						if(TestDataStrongPtr.IsValid())
+						if(ReplayStrongPtr.IsValid())
 						{
-							return !TestDataStrongPtr->IsReplaying() && !TestDataStrongPtr->IsRecording();
+							return !ReplayStrongPtr->IsReplaying() && !ReplayStrongPtr->IsRecording();
 						}
 						return true;
 					})
 				),
 				NAME_None,
-				LOCTEXT("TestDataRecordButton", "Record"),
-				LOCTEXT("TestDataRecordButton_Tooltip", "Records a new frame into the test data asset.\nA test data asset will be created if necessary."),
-				FSlateIcon(FControlRigEditorStyle::Get().GetStyleSetName(),"ControlRig.TestData.Record")
+				LOCTEXT("ReplayRecordButton", "Record"),
+				LOCTEXT("ReplayRecordButton_Tooltip", "Records a replay\nA replay asset will be created if necessary."),
+				FSlateIcon(FControlRigEditorStyle::Get().GetStyleSetName(),"ControlRig.Replay.Record")
 				);
 			ToolbarBuilder.AddComboButton(
 				FUIAction(),
-				FOnGetContent::CreateSP(this, &FControlRigEditor::GenerateTestAssetRecordMenuContent),
-				LOCTEXT("TestDataRecordMenu_Label", "Recording Modes"),
-				LOCTEXT("TestDataRecordMenu_ToolTip", "Pick between different modes for recording"),
+				FOnGetContent::CreateSP(this, &FControlRigEditor::GenerateReplayAssetRecordMenuContent),
+				LOCTEXT("ReplayRecordMenu_Label", "Recording Modes"),
+				LOCTEXT("ReplayRecordMenu_ToolTip", "Pick between different modes for recording"),
 				FSlateIcon(FAppStyle::GetAppStyleSetName(), "LevelEditor.Recompile"),
 				true);
 
 			ToolbarBuilder.AddToolBarButton(FUIAction(
 					FExecuteAction::CreateLambda([this]()
 					{
-						ToggleTestData();
+						ToggleReplay();
 					}),
 					FCanExecuteAction::CreateLambda([this]
 					{
-						if(TestDataStrongPtr.IsValid())
+						if(ReplayStrongPtr.IsValid())
 						{
-							return !TestDataStrongPtr->IsRecording();
+							return !ReplayStrongPtr->IsRecording();
 						}
 						return false;
 					})
@@ -840,19 +835,19 @@ void FControlRigEditor::FillToolbar(FToolBarBuilder& ToolbarBuilder, bool bEndSe
 				{
 					static const FText LiveStatus = LOCTEXT("LiveStatus", "Live"); 
 					static const FText ReplayInputsStatus = LOCTEXT("ReplayInputsStatus", "Replay Inputs");
-					static const FText GroundTruthStatus = LOCTEXT("GroundTruthStatus", "Ground Truth");
-					if(TestDataStrongPtr.IsValid())
+					//static const FText GroundTruthStatus = LOCTEXT("GroundTruthStatus", "Ground Truth");
+					if(ReplayStrongPtr.IsValid())
 					{
-						switch(TestDataStrongPtr->GetPlaybackMode())
+						switch(ReplayStrongPtr->GetPlaybackMode())
 						{
-							case EControlRigTestDataPlaybackMode::ReplayInputs:
+							case EControlRigReplayPlaybackMode::ReplayInputs:
 							{
 								return ReplayInputsStatus;
 							}
-							case EControlRigTestDataPlaybackMode::GroundTruth:
-							{
-								return GroundTruthStatus;
-							}
+							//case EControlRigReplayPlaybackMode::GroundTruth:
+							//{
+							//	return GroundTruthStatus;
+							//}
 							default:
 							{
 								break;
@@ -863,21 +858,21 @@ void FControlRigEditor::FillToolbar(FToolBarBuilder& ToolbarBuilder, bool bEndSe
 				}),
 				TAttribute<FText>::CreateLambda([this]()
 				{
-					static const FText LiveStatusTooltip = LOCTEXT("LiveStatusTooltip", "The test data is not affecting the rig."); 
-					static const FText ReplayInputsStatusTooltip = LOCTEXT("ReplayInputsStatusTooltip", "The test data inputs are being replayed onto the rig.");
-					static const FText GroundTruthStatusTooltip = LOCTEXT("GroundTruthStatusTooltip", "The complete result of the rig is being overwritten.");
-					if(TestDataStrongPtr.IsValid())
+					static const FText LiveStatusTooltip = LOCTEXT("LiveStatusTooltip", "The replay is not affecting the rig."); 
+					static const FText ReplayInputsStatusTooltip = LOCTEXT("ReplayInputsStatusTooltip", "The replay inputs are being replayed onto the rig.");
+					//static const FText GroundTruthStatusTooltip = LOCTEXT("GroundTruthStatusTooltip", "The complete result of the rig is being overwritten.");
+					if(ReplayStrongPtr.IsValid())
 					{
-						switch(TestDataStrongPtr->GetPlaybackMode())
+						switch(ReplayStrongPtr->GetPlaybackMode())
 						{
-							case EControlRigTestDataPlaybackMode::ReplayInputs:
+							case EControlRigReplayPlaybackMode::ReplayInputs:
 							{
 								return ReplayInputsStatusTooltip;
 							}
-							case EControlRigTestDataPlaybackMode::GroundTruth:
-							{
-								return GroundTruthStatusTooltip;
-							}
+							//case EControlRigReplayPlaybackMode::GroundTruth:
+							//{
+							//	return GroundTruthStatusTooltip;
+							//}
 							default:
 							{
 								break;
@@ -890,7 +885,7 @@ void FControlRigEditor::FillToolbar(FToolBarBuilder& ToolbarBuilder, bool bEndSe
 				{
 					static const FSlateIcon LiveIcon = FSlateIcon(FControlRigEditorStyle::Get().GetStyleSetName(),"ClassIcon.ControlRigBlueprint"); 
 					static const FSlateIcon ReplayIcon = FSlateIcon(FControlRigEditorStyle::Get().GetStyleSetName(),"ClassIcon.ControlRigSequence"); 
-					if(TestDataStrongPtr.IsValid() && TestDataStrongPtr->IsReplaying())
+					if(ReplayStrongPtr.IsValid() && ReplayStrongPtr->IsReplaying())
 					{
 						return ReplayIcon;
 					}
@@ -2805,9 +2800,9 @@ void FControlRigEditor::UpdateRigVMHost()
 
 void FControlRigEditor::UpdateRigVMHost_PreClearOldHost(URigVMHost* InPreviousHost)
 {
-	if(TestDataStrongPtr.IsValid())
+	if(ReplayStrongPtr.IsValid())
 	{
-		TestDataStrongPtr->ReleaseReplay();
+		ReplayStrongPtr->StopReplay();
 	}
 }
 
diff --git a/Engine/Plugins/Animation/ControlRig/Source/ControlRigEditor/Private/Editor/ControlRigEditor.h b/Engine/Plugins/Animation/ControlRig/Source/ControlRigEditor/Private/Editor/ControlRigEditor.h
--- a/Engine/Plugins/Animation/ControlRig/Source/ControlRigEditor/Private/Editor/ControlRigEditor.h
+++ b/Engine/Plugins/Animation/ControlRig/Source/ControlRigEditor/Private/Editor/ControlRigEditor.h
@@ -21,7 +21,7 @@
 #include "Graph/ControlRigGraphNode.h"
 #include "RigVMModel/RigVMController.h"
 #include "Editor/RigVMDetailsViewWrapperObject.h"
-#include "ControlRigTestData.h"
+#include "ControlRigReplay.h"
 #include "ModularRigController.h"
 #include "RigVMHost.h"
 #include "SchematicGraphPanel/SSchematicGraphPanel.h"
@@ -293,16 +293,16 @@ private:
 
 	virtual void HandleRefreshEditorFromBlueprint(URigVMBlueprint* InBlueprint) override;
 
-	FText GetTestAssetName() const;
-	FText GetTestAssetTooltip() const;
-	bool SetTestAssetPath(const FString& InAssetPath);
-	TSharedRef<SWidget> GenerateTestAssetModeMenuContent();
-	TSharedRef<SWidget> GenerateTestAssetRecordMenuContent();
-	bool RecordTestData(double InRecordingDuration);
-	void ToggleTestData();
+	FText GetReplayAssetName() const;
+	FText GetReplayAssetTooltip() const;
+	bool SetReplayAssetPath(const FString& InAssetPath);
+	TSharedRef<SWidget> GenerateReplayAssetModeMenuContent();
+	TSharedRef<SWidget> GenerateReplayAssetRecordMenuContent();
+	bool RecordReplay(double InRecordingDuration);
+	void ToggleReplay();
 
 public:
-	EControlRigTestDataPlaybackMode GetTestAssetPlaybackMode() const;
+	EControlRigReplayPlaybackMode GetReplayPlaybackMode() const;
 
 protected:
 
@@ -364,7 +364,7 @@ protected:
 	bool bIsConstructionEventRunning;
 	uint32 LastHierarchyHash;
 
-	TStrongObjectPtr<UControlRigTestData> TestDataStrongPtr;
+	TStrongObjectPtr<UControlRigReplay> ReplayStrongPtr;
 
 	TWeakObjectPtr<const URigVMUnitNode> DirectManipulationSubject;
 	mutable TArray<TSharedPtr<FString>> DirectManipulationTextList;
