Refactor tweening logic: Math has been extracted to into new Animation/TweeningUtils plugin

This is in preparation of introducing tweening to CurveEditor.
Follow up changes will address the following:
- Introduce new STweenSlider widget
- Add STweenSlider widget to CurveEditor toolbar and make it use the logic that was refactored by this CL
- Adjust ControlRig to use the new STweenSlider widget.

diff --git a/Engine/Plugins/Animation/ControlRig/ControlRig.uplugin b/Engine/Plugins/Animation/ControlRig/ControlRig.uplugin
--- a/Engine/Plugins/Animation/ControlRig/ControlRig.uplugin
+++ b/Engine/Plugins/Animation/ControlRig/ControlRig.uplugin
@@ -37,6 +37,10 @@
 		{
 			"Name": "SequencerScripting",
 			"Enabled": true
+		},
+		{
+			"Name": "TweeningUtils",
+			"Enabled": true
 		}
 	],
 	"Modules" :
diff --git a/Engine/Plugins/Animation/ControlRig/Source/ControlRigEditor/ControlRigEditor.Build.cs b/Engine/Plugins/Animation/ControlRig/Source/ControlRigEditor/ControlRigEditor.Build.cs
--- a/Engine/Plugins/Animation/ControlRig/Source/ControlRigEditor/ControlRigEditor.Build.cs
+++ b/Engine/Plugins/Animation/ControlRig/Source/ControlRigEditor/ControlRigEditor.Build.cs
@@ -11,7 +11,8 @@ namespace UnrealBuildTool.Rules
 				{
 					"MainFrame",
 					"AppFramework",
-                    "RigVMEditor",
+                    "RigVMEditor", 
+					"CurveEditor",
 				}
 			);
 
@@ -21,7 +22,6 @@ namespace UnrealBuildTool.Rules
 					"AssetDefinition",
                     "Core",
                     "CoreUObject",
-					"CurveEditor",
 					"DetailCustomizations",
 					"Slate",
                     "SlateCore",
@@ -83,7 +83,8 @@ namespace UnrealBuildTool.Rules
                     "ActorPickerMode",
                     "Constraints",
                     "AnimationEditMode",
-					"SequencerWidgets"
+					"SequencerWidgets",
+					"TweeningUtilsEditor"
 				}
             );
 
diff --git a/Engine/Plugins/Animation/ControlRig/Source/ControlRigEditor/Private/Tools/ControlRigTweener.cpp b/Engine/Plugins/Animation/ControlRig/Source/ControlRigEditor/Private/Tools/ControlRigTweener.cpp
--- a/Engine/Plugins/Animation/ControlRig/Source/ControlRigEditor/Private/Tools/ControlRigTweener.cpp
+++ b/Engine/Plugins/Animation/ControlRig/Source/ControlRigEditor/Private/Tools/ControlRigTweener.cpp
@@ -20,6 +20,8 @@
 #include "MovieSceneSection.h"
 #include "EditorModeManager.h"
 #include "EditMode/ControlRigEditMode.h"
+#include "Math/KeyBlendingFunctions.h"
+#include "Math/CurveBlending.h"
 
 #define LOCTEXT_NAMESPACE "ControlRigTweenWidget"
 
@@ -28,7 +30,7 @@
 * FAnimSliderKeySelection
 *
 */
-
+PRAGMA_DISABLE_DEPRECATION_WARNINGS
 FAnimSliderKeySelection::FContiguousKeys::FContiguousKeys(const TArray<FKeyHandle>& AllKeyHandles, const TArray<FKeyPosition>& AllKeyPositions, const TArray<int32>& InContiguousKeyIndices)
 {
 	Indices.Reset();
@@ -103,6 +105,7 @@ bool FAnimSliderKeySelection::Setup(const TWeakPtr<ISequencer>& InSequencer)
 	}
 	return (KeyMap.Num() > 0);
 }
+PRAGMA_ENABLE_DEPRECATION_WARNINGS
 
 bool FAnimSliderObjectSelection::Setup(TWeakPtr<ISequencer>& InSequencer, TWeakPtr<FControlRigEditMode>& InEditMode)
 {
@@ -111,6 +114,7 @@ bool FAnimSliderObjectSelection::Setup(TWeakPtr<ISequencer>& InSequencer, TWeakP
 	return Setup(ControlRigs, InSequencer);
 }
 
+PRAGMA_DISABLE_DEPRECATION_WARNINGS
 void FAnimSliderKeySelection::GetMapOfContiguousKeys()
 {
 	TArray<FKeyHandle> AllKeyHandles;
@@ -177,6 +181,7 @@ void FAnimSliderKeySelection::GetMapOfContiguousKeys()
 		}
 	}
 }
+PRAGMA_ENABLE_DEPRECATION_WARNINGS
 
 /**
 *
@@ -468,11 +473,30 @@ bool FAnimSliderObjectSelection::Setup(const TArray<UControlRig*>& SelectedContr
 * 
 */
 
+static TSharedPtr<FCurveEditor> GetCurveEditor(const TWeakPtr<ISequencer>& InSequencer)
+{
+	using namespace UE::Sequencer;
+	
+	if (const TSharedPtr<ISequencer> Sequencer = InSequencer.Pin())
+	{
+		const TSharedPtr<FSequencerEditorViewModel> SequencerViewModel = InSequencer.Pin()->GetViewModel();
+		const FCurveEditorExtension* CurveEditorExtension = SequencerViewModel->CastDynamic<FCurveEditorExtension>();
+		check(CurveEditorExtension);
+		return CurveEditorExtension->GetCurveEditor();
+	}
+	return nullptr;
+}
+
 bool FBaseAnimSlider::Setup(TWeakPtr<ISequencer>& InSequencer, TWeakPtr<FControlRigEditMode>& InEditMode)
 {
-	KeySelection.KeyMap.Reset();
+	using namespace UE::TweeningUtilsEditor;
+	
+	const TSharedPtr<FCurveEditor> CurveEditor = GetCurveEditor(InSequencer);
+	
+	ContiguousKeySelection = CurveEditor ? FContiguousKeyMapping(*CurveEditor) : FContiguousKeyMapping{};
 	ObjectSelection.ChannelsArray.Reset();
-	if (KeySelection.Setup(InSequencer) == false)
+	
+	if (ContiguousKeySelection.KeyMap.Num() == 0)
 	{
 		return ObjectSelection.Setup(InSequencer, InEditMode);
 	}
@@ -487,59 +511,24 @@ bool FBaseAnimSlider::Setup(TWeakPtr<ISequencer>& InSequencer, TWeakPtr<FControl
 
 bool FBasicBlendSlider::Blend(TWeakPtr<ISequencer>& InSequencer, const double BlendValue)
 {
-	if (InSequencer.IsValid() == false)
+	using namespace UE::CurveEditor;
+
+	if (!InSequencer.IsValid())
 	{
 		return false;
 	}
 	const ISequencer* Sequencer = InSequencer.Pin().Get();
-
-	bool bDidBlend = false;
-
-	if (KeySelection.KeyMap.Num() > 0)
+	const TSharedPtr<FCurveEditor> CurveEditor = GetCurveEditor(InSequencer);
+	
+	if (ContiguousKeySelection.KeyMap.Num() > 0 && ensure(CurveEditor))
 	{
-		TArray<FKeyHandle> KeyHandles;
-		TArray<FKeyPosition> KeyPositions;
-		for (TPair<FCurveModelID, FAnimSliderKeySelection::FContiguousKeysArray>& KeysArray : KeySelection.KeyMap)
+		return UE::TweeningUtilsEditor::BlendCurves(BlendValue, *CurveEditor, ContiguousKeySelection, [this](const UE::TweeningUtilsEditor::FKeyBlendInfo& InBlendInfo)
 		{
-			if (FCurveModel* Curve = KeySelection.CurveEditor->FindCurve(KeysArray.Key))
-			{
-				Curve->Modify();
-				for (FAnimSliderKeySelection::FContiguousKeys& Keys : KeysArray.Value.KeysArray)
-				{
-					const int32 PreviousIndex = Keys.PreviousIndex != INDEX_NONE ? Keys.PreviousIndex : Keys.Indices[0];
-					const double PreviousTime = KeysArray.Value.AllKeyPositions[PreviousIndex].InputValue;
-					const double PreviousValue = KeysArray.Value.AllKeyPositions[PreviousIndex].OutputValue;
-					const int32 NextIndex = Keys.NextIndex != INDEX_NONE ? Keys.NextIndex : Keys.Indices[Keys.Indices.Num() - 1];
-					const double NextTime = KeysArray.Value.AllKeyPositions[NextIndex].InputValue;
-					const double NextValue = KeysArray.Value.AllKeyPositions[NextIndex].OutputValue;
-
-					const int32 NumIndices = Keys.Indices.Num();
-					KeyHandles.SetNum(NumIndices);
-					KeyPositions.SetNum(NumIndices);
-					const double FirstValue = KeysArray.Value.AllKeyPositions[Keys.Indices[0]].OutputValue;
-					const double LastValue = KeysArray.Value.AllKeyPositions[Keys.Indices[Keys.Indices.Num() -1]].OutputValue;
-					FBlendStruct BlendStruct(KeysArray.Value.AllKeyPositions, Keys.Indices);
-					int32 CurrentIndex = 0;
-					for (int32 Index = 0; Index < NumIndices; ++Index)
-					{
-						const int32& KeyIndex = Keys.Indices[Index];
-						KeyHandles[Index] = KeysArray.Value.AllKeyHandles[KeyIndex];
-						const double CurrentValue = KeysArray.Value.AllKeyPositions[KeyIndex].OutputValue;
-						const double CurrentTime = KeysArray.Value.AllKeyPositions[KeyIndex].InputValue;
-						BlendStruct.SetValues(PreviousTime, PreviousValue, CurrentTime, CurrentValue,
-							NextTime, NextValue, BlendValue, FirstValue, LastValue, CurrentIndex);
-						++CurrentIndex;
-						double NewValue = DoBlend(BlendStruct);
-						KeyPositions[Index] = FKeyPosition(KeysArray.Value.AllKeyPositions[KeyIndex].InputValue, NewValue);
-					}
-					Curve->SetKeyPositions(KeyHandles, KeyPositions);
-					bDidBlend = true;
-				}
-			}
-		}
-		return bDidBlend;
+			return DoBlend(InBlendInfo);
+		});
 	}
 
+	bool bDidBlend = false;
 	const FFrameTime FrameTime = Sequencer->GetLocalTime().Time;
 	const FFrameRate TickResolution = Sequencer->GetFocusedTickResolution();
 
@@ -549,7 +538,7 @@ bool FBasicBlendSlider::Blend(TWeakPtr<ISequencer>& InSequencer, const double Bl
 		const TArray<FKeyPosition> AllKeyPositions;
 		const TArray <int32> Indices;
 		const int32 CurrentIndex = -1;
-		FBlendStruct BlendStruct(AllKeyPositions, Indices);
+		UE::TweeningUtilsEditor::FKeyBlendInfo BlendStruct(AllKeyPositions, Indices);
 
 		for (const FAnimSliderObjectSelection::FObjectChannels& ObjectChannels : ObjectSelection.ChannelsArray)
 		{
@@ -595,12 +584,9 @@ bool FBasicBlendSlider::Blend(TWeakPtr<ISequencer>& InSequencer, const double Bl
 *
 */
 
-double FControlsToTween::DoBlend(const FBlendStruct& BlendStruct)
+double FControlsToTween::DoBlend(const UE::TweeningUtilsEditor::FKeyBlendInfo& InBlendInfo)
 {
-	//clasic tween will move all to same location, not based on current time at all just blend and values
-	const double NormalizedBlendValue = (BlendStruct.BlendValue + 1.0f) * 0.5f;
-	const double Value = BlendStruct.PreviousValue + (BlendStruct.NextValue - BlendStruct.PreviousValue) * (NormalizedBlendValue);
-	return Value;
+	return Blend_ControlsToTween(InBlendInfo);
 }
 
 FText FControlsToTween::GetText() const
@@ -630,25 +616,9 @@ bool FControlsToTween::Setup(TWeakPtr<ISequencer>& InSequencer, TWeakPtr<FContro
 */
 
 
-double FPushPullSlider::DoBlend(const FBlendStruct& BlendStruct)
+double FPushPullSlider::DoBlend(const UE::TweeningUtilsEditor::FKeyBlendInfo& InBlendInfo)
 {
-	if (FMath::IsNearlyEqual(BlendStruct.NextTime, BlendStruct.PreviousTime))
-	{
-		return BlendStruct.CurrentValue;
-	}
-	const double T = (BlendStruct.CurrentTime - BlendStruct.PreviousTime) / (BlendStruct.NextTime - BlendStruct.PreviousTime);
-	const double ValueAtT = BlendStruct.PreviousValue + T * (BlendStruct.NextValue - BlendStruct.PreviousValue);
-	double NewValue;
-	if (BlendStruct.BlendValue < 0.0)
-	{
-		NewValue = BlendStruct.CurrentValue + (-1.0 * BlendStruct.BlendValue) * (ValueAtT - BlendStruct.CurrentValue);
-	}
-	else
-	{
-		const double AmplifyValueAtT = BlendStruct.CurrentValue + (BlendStruct.CurrentValue - ValueAtT);
-		NewValue = BlendStruct.CurrentValue + BlendStruct.BlendValue * (AmplifyValueAtT - BlendStruct.CurrentValue);
-	}
-	return NewValue;
+	return Blend_PushPull(InBlendInfo);
 }
 
 FText FPushPullSlider::GetText() const
@@ -668,18 +638,9 @@ FText FPushPullSlider::GetTooltipText() const
 *
 */
 
-double FBlendNeighborSlider::DoBlend(const FBlendStruct& BlendStruct)
+double FBlendNeighborSlider::DoBlend(const UE::TweeningUtilsEditor::FKeyBlendInfo& InBlendInfo)
 {
-	double NewValue = BlendStruct.CurrentValue;
-	if (BlendStruct.BlendValue < 0.0)
-	{
-		NewValue = BlendStruct.CurrentValue + (-1.0 * BlendStruct.BlendValue) * (BlendStruct.PreviousValue - BlendStruct.CurrentValue);
-	}
-	else
-	{
-		NewValue = BlendStruct.CurrentValue + BlendStruct.BlendValue * (BlendStruct.NextValue - BlendStruct.CurrentValue);
-	}
-	return NewValue;
+	return Blend_Neighbor(InBlendInfo);
 }
 
 FText FBlendNeighborSlider::GetText() const
@@ -698,18 +659,9 @@ FText FBlendNeighborSlider::GetTooltipText() const
 *
 */
 
-double FBlendRelativeSlider::DoBlend(const FBlendStruct& BlendStruct)
+double FBlendRelativeSlider::DoBlend(const UE::TweeningUtilsEditor::FKeyBlendInfo& InBlendInfo)
 {
-	double NewValue = BlendStruct.CurrentValue;
-	if (BlendStruct.BlendValue < 0.0)
-	{
-		NewValue = BlendStruct.CurrentValue + (-1.0 * BlendStruct.BlendValue) * (BlendStruct.PreviousValue - BlendStruct.FirstValue);
-	}
-	else
-	{
-		NewValue = BlendStruct.CurrentValue + BlendStruct.BlendValue * (BlendStruct.NextValue - BlendStruct.LastValue);
-	}
-	return NewValue;
+	return Blend_Relative(InBlendInfo);
 }
 
 FText FBlendRelativeSlider::GetText() const
@@ -722,67 +674,9 @@ FText FBlendRelativeSlider::GetTooltipText() const
 	return LOCTEXT("BlendRelativeSliderTooltip", "Blend relative to the next or previous value for selected keys or objects");
 }
 
-/*
-*
-*  Blend To Ease Slider
-*
-*/
-namespace BlendToEase
-{ 
-static float ExpIn(float InTime)
-{
-	return FMath::Pow(2, 10 * (InTime - 1.f));
-}
-static float ExpOut(float InTime)
-{
-	return 1.f - ExpIn(1.f - InTime);
-}
-
-static float SCurve(float X, float Slope, float Width, float Height, float XShift, float YShift)
-{
-	if (X > (XShift + Width))
-	{
-		return(Height + YShift);
-	}
-	else if (X < XShift)
-	{
-		return YShift;
-	}
-
-	float Val = Height * (FMath::Pow((X - XShift), Slope) / (FMath::Pow(X - XShift, Slope) + FMath::Pow((Width - (X - XShift)), Slope))) + YShift;
-	return Val;
-}
-}
-
-double FBlendToEaseSlider::DoBlend(const FBlendStruct& BlendStruct)
+double FBlendToEaseSlider::DoBlend(const UE::TweeningUtilsEditor::FKeyBlendInfo& InBlendInfo)
 {
-	if (FMath::IsNearlyEqual(BlendStruct.NextTime, BlendStruct.PreviousTime))
-	{
-		return BlendStruct.CurrentValue;
-	}
-	const double Source = BlendStruct.CurrentValue;
-	const double FullTimeDiff = BlendStruct.NextTime - BlendStruct.PreviousTime;
-	const double AbsValue = FMath::Abs(BlendStruct.BlendValue);
-	const double X = BlendStruct.CurrentTime - BlendStruct.PreviousTime;
-	const double Ratio = X / FullTimeDiff;
-	double Shift = 0.0, Delta = 0.0, Base = 0.0;
-	if (BlendStruct.BlendValue > 0)
-	{
-		Shift = -1.0;
-		Delta = BlendStruct.NextValue - Source;
-		Base = Source;
-	}
-	else
-	{
-		Shift = 0.0;
-		Delta = Source - BlendStruct.PreviousValue;
-		Base = BlendStruct.PreviousValue;
-	}
-	const double Slope = 5.0 * AbsValue;
-	const double EaseY = BlendToEase::SCurve(Ratio, Slope, 2.0, 2.0, Shift, Shift);
-	const double NewValue = Base + (Delta * EaseY);
-	
-	return NewValue;
+	return Blend_Ease(InBlendInfo);
 }
 
 FText FBlendToEaseSlider::GetText() const
@@ -801,21 +695,9 @@ FText FBlendToEaseSlider::GetTooltipText() const
 *
 */
 
-double FSmoothRoughSlider::DoBlend(const FBlendStruct& BlendStruct)
+double FSmoothRoughSlider::DoBlend(const UE::TweeningUtilsEditor::FKeyBlendInfo& InBlendInfo)
 {
-	double PrevVal = BlendStruct.CurrentIndex > 0 ? BlendStruct.AllKeyPositions[BlendStruct.Indices[BlendStruct.CurrentIndex - 1]].OutputValue : BlendStruct.PreviousValue;
-	double CurVal = BlendStruct.CurrentValue;
-	double NextVal = BlendStruct.CurrentIndex < BlendStruct.Indices.Num() - 1 ? BlendStruct.AllKeyPositions[BlendStruct.Indices[BlendStruct.CurrentIndex + 1]].OutputValue : BlendStruct.NextValue;
-	double NewValue = (PrevVal * 0.25 + CurVal * 0.5 + NextVal * 0.25);
-	if (BlendStruct.BlendValue < 0.0)
-	{
-		NewValue = BlendStruct.CurrentValue + (-1.0 * BlendStruct.BlendValue) * (NewValue - BlendStruct.CurrentValue);
-	}
-	else
-	{
-		NewValue = BlendStruct.CurrentValue + BlendStruct.BlendValue * (BlendStruct.CurrentValue - NewValue);
-	}
-	return NewValue;
+	return Blend_SmoothRough(InBlendInfo);
 }
 
 FText FSmoothRoughSlider::GetText() const
diff --git a/Engine/Plugins/Animation/ControlRig/Source/ControlRigEditor/Public/Tools/ControlRigTweener.h b/Engine/Plugins/Animation/ControlRig/Source/ControlRigEditor/Public/Tools/ControlRigTweener.h
--- a/Engine/Plugins/Animation/ControlRig/Source/ControlRigEditor/Public/Tools/ControlRigTweener.h
+++ b/Engine/Plugins/Animation/ControlRig/Source/ControlRigEditor/Public/Tools/ControlRigTweener.h
@@ -7,6 +7,7 @@
 #include "Curves/KeyHandle.h"
 #include "CurveDataAbstraction.h"
 #include "CurveModel.h"
+#include "Math/ContiguousKeyMapping.h"
 
 struct FMovieSceneFloatChannel;
 struct FMovieSceneDoubleChannel;
@@ -15,11 +16,12 @@ class ISequencer;
 struct FRigControl;
 class UMovieSceneSection;
 class FCurveEditor;
+namespace UE::TweeningUtilsEditor { struct FKeyBlendInfo; }
 
 /*
 * Contains the selection state for a set of keys to blend with the anim slider
 */
-struct FAnimSliderKeySelection
+struct UE_DEPRECATED(5.6, "Use FContiguousKeyMapping from CurveEditor instead.") FAnimSliderKeySelection
 {
 	//list of keys in a row
 	struct FContiguousKeys
@@ -103,8 +105,10 @@ private:
 */
 struct FBaseAnimSlider
 {
-	FBaseAnimSlider(){};
-	virtual ~FBaseAnimSlider() {};
+	PRAGMA_DISABLE_DEPRECATION_WARNINGS
+	FBaseAnimSlider() = default;
+	virtual ~FBaseAnimSlider() = default;
+	PRAGMA_DISABLE_DEPRECATION_WARNINGS
 	/**
 	 Setup up what's needed for the slider to blend
 	* @param InSequencer Sequencer to get keys, controls, etc. from
@@ -131,7 +135,12 @@ struct FBaseAnimSlider
 
 protected:
 
+	UE_DEPRECATED(5.6, "Use ContigousKeySelection instead.")
+	PRAGMA_DISABLE_DEPRECATION_WARNINGS
 	FAnimSliderKeySelection  KeySelection;
+	PRAGMA_ENABLE_DEPRECATION_WARNINGS
+
+	UE::TweeningUtilsEditor::FContiguousKeyMapping ContiguousKeySelection;
 	FAnimSliderObjectSelection ObjectSelection;
 };
 
@@ -164,7 +173,7 @@ private:
 *  Struct containing data that is used by Blending
 *
 **/
-struct FBlendStruct
+struct UE_DEPRECATED(5.6, "Use FKeyPositionBlendInfo instead.") FBlendStruct
 {
 	//Previous and Next are the keys that are before/after the contiguous set of keys, 
 	//to get the Previous key for example you use AllKeyPostions[Indices[CurrentIndex -1]]
@@ -203,6 +212,7 @@ private:
 	FBlendStruct() = delete;
 	
 };
+
 /**
 *
 *  Slider that used the next and previous values to do the blend
@@ -210,7 +220,13 @@ private:
 **/
 struct FBasicBlendSlider : public FBaseAnimSlider
 {
-	virtual double DoBlend(const FBlendStruct& BlendStruct) = 0;
+	UE_DEPRECATED(5.6, "Use the version that uses FKeyBlendInfo instead.")
+	PRAGMA_DISABLE_DEPRECATION_WARNINGS
+	virtual double DoBlend(const FBlendStruct& BlendStruct) { return BlendStruct.CurrentValue; }
+	PRAGMA_ENABLE_DEPRECATION_WARNINGS
+
+	virtual double DoBlend(const UE::TweeningUtilsEditor::FKeyBlendInfo& InBlendInfo) = 0;
+	
 	//BaseAnimSlider overrides
 	virtual bool Blend(TWeakPtr<ISequencer>& InSequencer, const double BlendValue) override;
 };
@@ -228,7 +244,7 @@ struct FControlsToTween :public FBasicBlendSlider
 	virtual bool Setup(TWeakPtr<ISequencer>& InSequencer, TWeakPtr<FControlRigEditMode>& InEditMode) override;
 
 	//BasicBlendSlider overrides
-	virtual double DoBlend(const FBlendStruct& BlendStruct) override;
+	virtual double DoBlend(const UE::TweeningUtilsEditor::FKeyBlendInfo& InBlendInfo) override;
 
 	bool Setup(const TArray<UControlRig*>& SelectedControlRigs, TWeakPtr<ISequencer>& InSequencer);
 
@@ -246,7 +262,7 @@ struct FPushPullSlider :public FBasicBlendSlider
 	virtual FText GetTooltipText() const;
 
 	//BasicBlendSlider overrides
-	virtual double DoBlend(const FBlendStruct& BlendStruct) override;
+	virtual double DoBlend(const UE::TweeningUtilsEditor::FKeyBlendInfo& InBlendInfo) override;
 };
 
 /**
@@ -261,7 +277,7 @@ struct FBlendNeighborSlider : public FBasicBlendSlider
 	virtual FText GetTooltipText() const;
 
 	//BasicBlendSlider overrides
-	virtual double DoBlend(const FBlendStruct& BlendStruct) override;
+	virtual double DoBlend(const UE::TweeningUtilsEditor::FKeyBlendInfo& InBlendInfo) override;
 };
 
 
@@ -277,7 +293,7 @@ struct FBlendRelativeSlider: public FBasicBlendSlider
 	virtual FText GetTooltipText() const;
 
 	//BasicBlendSlider overrides
-	virtual double DoBlend(const FBlendStruct& BlendStruct) override;
+	virtual double DoBlend(const UE::TweeningUtilsEditor::FKeyBlendInfo& InBlendInfo) override;
 };
 
 
@@ -293,7 +309,7 @@ struct FBlendToEaseSlider : public FBasicBlendSlider
 	virtual FText GetTooltipText() const;
 
 	//BasicBlendSlider overrides
-	virtual double DoBlend(const FBlendStruct& BlendStruct) override;
+	virtual double DoBlend(const UE::TweeningUtilsEditor::FKeyBlendInfo& InBlendInfo) override;
 };
 
 
@@ -309,6 +325,6 @@ struct FSmoothRoughSlider : public FBasicBlendSlider
 	virtual FText GetTooltipText() const;
 
 	//BasicBlendSlider overrides
-	virtual double DoBlend(const FBlendStruct& BlendStruct) override;
+	virtual double DoBlend(const UE::TweeningUtilsEditor::FKeyBlendInfo& InBlendInfo) override;
 };
 
diff --git a/Engine/Plugins/Animation/TweeningUtils/Source/TweeningUtilsEditor/Private/Math/ContiguousKeyMapping.cpp b/Engine/Plugins/Animation/TweeningUtils/Source/TweeningUtilsEditor/Private/Math/ContiguousKeyMapping.cpp
new file mode 100644
--- /dev/null
+++ b/Engine/Plugins/Animation/TweeningUtils/Source/TweeningUtilsEditor/Private/Math/ContiguousKeyMapping.cpp
@@ -0,0 +1,151 @@
+﻿// Copyright Epic Games, Inc. All Rights Reserved.
+
+#include "Math/ContiguousKeyMapping.h"
+
+#include "Algo/BinarySearch.h"
+#include "CurveDataAbstraction.h"
+#include "CurveEditorTypes.h"
+#include "CurveEditor.h"
+#include "CurveEditorSelection.h"
+#include "CurveModel.h"
+#include "Math/NumericLimits.h"
+
+namespace UE::TweeningUtilsEditor
+{
+static int32 GetIndex(const TArray<FKeyPosition>& Keys, double Time)
+{
+	const int32 Index = Algo::LowerBoundBy(Keys, Time, [](const FKeyPosition& Value) { return Value.InputValue; });
+
+	// don't trust precision issues so will double check to make sure the index is correct
+	if (Index != INDEX_NONE)
+	{
+		if (FMath::IsNearlyEqual(Keys[Index].InputValue, Time))
+		{
+			return Index;
+		}
+		if (((Index - 1) >= 0) && FMath::IsNearlyEqual(Keys[Index - 1].InputValue, Time))
+		{
+			return Index - 1;
+		}
+		if (((Index + 1) < Keys.Num()) && FMath::IsNearlyEqual(Keys[Index + 1].InputValue, Time))
+		{
+			return Index + 1;
+		}
+	}
+	return Index;
+}
+
+static TMap<FCurveModelID, FContiguousKeyMapping::FContiguousKeysArray> ComputeKeyMap(const FCurveEditor& InCurveEditor)
+{
+	using FContiguousKeysArray = FContiguousKeyMapping::FContiguousKeysArray;
+	
+	TArray<FKeyHandle> AllKeyHandles;
+	TArray<FKeyPosition> AllKeyPositions;
+
+	TArray<FKeyHandle> SelectedKeyHandles;
+	TArray<FKeyPosition> SelectedKeyPositions;
+
+	const TMap<FCurveModelID, FKeyHandleSet>& SelectionKeyMap = InCurveEditor.Selection.GetAll();
+
+	TArray<int32> ContiguousKeyIndices;
+
+	TMap<FCurveModelID, FContiguousKeyMapping::FContiguousKeysArray> KeyMap;
+	for (const TPair<FCurveModelID, FKeyHandleSet>& Pair : SelectionKeyMap)
+	{
+		if (const FCurveModel* Curve = InCurveEditor.FindCurve(Pair.Key))
+		{
+			AllKeyHandles.Reset();
+			Curve->GetKeys(TNumericLimits<double>::Lowest(), TNumericLimits<double>::Max(), TNumericLimits<double>::Lowest(), TNumericLimits<double>::Max(), AllKeyHandles);
+			AllKeyPositions.SetNum(AllKeyHandles.Num());
+			Curve->GetKeyPositions(AllKeyHandles, AllKeyPositions);
+			SelectedKeyHandles.Reset(Pair.Value.Num());
+			SelectedKeyHandles.Append(Pair.Value.AsArray().GetData(), Pair.Value.Num());
+
+			// Get all the selected keys
+			SelectedKeyPositions.SetNum(SelectedKeyHandles.Num());
+			Curve->GetKeyPositions(SelectedKeyHandles, SelectedKeyPositions);
+
+			ContiguousKeyIndices.Reset();
+			FContiguousKeysArray& KeyArray = KeyMap.Add(Pair.Key, FContiguousKeysArray(AllKeyHandles, AllKeyPositions));
+
+			TArray<int32> SelectedIndices;
+			for (int32 Index = 0; Index < SelectedKeyPositions.Num(); ++Index)
+			{
+				int32 AllKeysIndex = GetIndex(AllKeyPositions, SelectedKeyPositions[Index].InputValue);
+				SelectedIndices.Add(AllKeysIndex);
+			}
+			SelectedIndices.Sort();
+			for(int32 AllKeysIndex: SelectedIndices)
+			{
+				if (ContiguousKeyIndices.Num() > 0) //see if this key is next to the previous one
+				{
+					if (ContiguousKeyIndices[ContiguousKeyIndices.Num() - 1] + 1 == AllKeysIndex)
+					{
+						ContiguousKeyIndices.Add(AllKeysIndex);
+					}
+					else //not contiguous, so need to add to the OutKeyMap
+					{
+						//now start new set with the new index
+						KeyArray.Add(ContiguousKeyIndices);
+						ContiguousKeyIndices.Reset();
+						ContiguousKeyIndices.Add(AllKeysIndex);
+					}
+				}
+				else//first key in this set so just add
+				{
+					ContiguousKeyIndices.Add(AllKeysIndex);
+				}
+			}
+			if (ContiguousKeyIndices.Num() > 0)
+			{
+				KeyArray.Add(ContiguousKeyIndices);
+				ContiguousKeyIndices.Reset();
+			}
+		}
+	}
+	
+	return KeyMap;
+}
+
+FContiguousKeyMapping::FContiguousKeys::FContiguousKeys(
+	const TArray<FKeyHandle>& AllKeyHandles, const TArray<FKeyPosition>& AllKeyPositions, const TArray<int32>& InContiguousKeyIndices
+	)
+{
+	Indices.Reset();
+	if (AllKeyHandles.Num() > 0 && AllKeyPositions.Num() == AllKeyHandles.Num() && InContiguousKeyIndices.Num() > 0)
+	{
+		Indices = InContiguousKeyIndices;
+		//Previous
+		PreviousIndex = Indices[0] - 1;
+		if (PreviousIndex < 0)
+		{
+			PreviousIndex = INDEX_NONE;
+		}		
+		//Next
+		NextIndex = Indices[Indices.Num() -1] + 1;
+		if (NextIndex >= AllKeyHandles.Num())
+		{
+			NextIndex = INDEX_NONE;
+		}
+	}
+}
+
+FContiguousKeyMapping::FContiguousKeysArray::FContiguousKeysArray(
+	const TArray<FKeyHandle>& InAllKeyHandles, const TArray<FKeyPosition>& InAllKeyPositions
+	)
+	: AllKeyHandles(InAllKeyHandles)
+	, AllKeyPositions(InAllKeyPositions)
+{}
+
+void FContiguousKeyMapping::FContiguousKeysArray::Add(const TArray<int32>& InContiguousKeyIndices)
+{
+	if (AllKeyHandles.Num() > 0 && AllKeyPositions.Num() == AllKeyHandles.Num() && InContiguousKeyIndices.Num() > 0)
+	{
+		KeysArray.Add(FContiguousKeys(AllKeyHandles, AllKeyPositions, InContiguousKeyIndices));
+	}
+}
+
+FContiguousKeyMapping::FContiguousKeyMapping(const FCurveEditor& InCurveEditor)
+	: KeyMap(ComputeKeyMap(InCurveEditor))
+{}
+}
diff --git a/Engine/Plugins/Animation/TweeningUtils/Source/TweeningUtilsEditor/Private/Math/KeyBlendingAbstraction.cpp b/Engine/Plugins/Animation/TweeningUtils/Source/TweeningUtilsEditor/Private/Math/KeyBlendingAbstraction.cpp
new file mode 100644
--- /dev/null
+++ b/Engine/Plugins/Animation/TweeningUtils/Source/TweeningUtilsEditor/Private/Math/KeyBlendingAbstraction.cpp
@@ -0,0 +1,26 @@
+﻿// Copyright Epic Games, Inc. All Rights Reserved.
+
+#include "Math/KeyBlendingAbstraction.h"
+
+#include "Math/KeyBlendInfo.h"
+#include "Math/KeyBlendingFunctions.h"
+
+namespace UE::TweeningUtilsEditor
+{
+double EvaluateBlend(EBlendFunction InFunction, const FKeyBlendInfo& InBlendInfo)
+{
+	static_assert(static_cast<int32>(EBlendFunction::Num) == 6, "Extend this switch");
+	switch (InFunction)
+	{
+	case EBlendFunction::ControlsToTween: return Blend_ControlsToTween(InBlendInfo);
+	case EBlendFunction::PushPull: return Blend_PushPull(InBlendInfo);
+	case EBlendFunction::BlendNeighbor: return Blend_Neighbor(InBlendInfo);
+	case EBlendFunction::BlendRelative: return Blend_Relative(InBlendInfo);
+	case EBlendFunction::BlendEase: return Blend_Ease(InBlendInfo);
+	case EBlendFunction::SmoothRough: return Blend_SmoothRough(InBlendInfo);
+	default: checkNoEntry(); return InBlendInfo.CurrentValue;
+	}
+}
+}
+
+
diff --git a/Engine/Plugins/Animation/TweeningUtils/Source/TweeningUtilsEditor/Private/Math/KeyBlendingFunctions.cpp b/Engine/Plugins/Animation/TweeningUtils/Source/TweeningUtilsEditor/Private/Math/KeyBlendingFunctions.cpp
new file mode 100644
--- /dev/null
+++ b/Engine/Plugins/Animation/TweeningUtils/Source/TweeningUtilsEditor/Private/Math/KeyBlendingFunctions.cpp
@@ -0,0 +1,131 @@
+﻿// Copyright Epic Games, Inc. All Rights Reserved.
+
+#include "Math/KeyBlendingFunctions.h"
+
+#include "Math/KeyBlendInfo.h"
+
+namespace UE::TweeningUtilsEditor
+{
+double Blend_ControlsToTween(const FKeyBlendInfo& InBlendInfo)
+{
+	// Classic tween will move all to same location, not based on current time at all just blend and values
+	const double NormalizedBlendValue = (InBlendInfo.BlendValue + 1.0f) * 0.5f;
+	const double Value = InBlendInfo.PreviousValue + (InBlendInfo.NextValue - InBlendInfo.PreviousValue) * (NormalizedBlendValue);
+	return Value;
+}
+
+double Blend_PushPull(const FKeyBlendInfo& InBlendInfo)
+{
+	if (FMath::IsNearlyEqual(InBlendInfo.NextTime, InBlendInfo.PreviousTime))
+	{
+		return InBlendInfo.CurrentValue;
+	}
+	const double T = (InBlendInfo.CurrentTime - InBlendInfo.PreviousTime) / (InBlendInfo.NextTime - InBlendInfo.PreviousTime);
+	const double ValueAtT = InBlendInfo.PreviousValue + T * (InBlendInfo.NextValue - InBlendInfo.PreviousValue);
+	double NewValue;
+	if (InBlendInfo.BlendValue < 0.0)
+	{
+		NewValue = InBlendInfo.CurrentValue + (-1.0 * InBlendInfo.BlendValue) * (ValueAtT - InBlendInfo.CurrentValue);
+	}
+	else
+	{
+		const double AmplifyValueAtT = InBlendInfo.CurrentValue + (InBlendInfo.CurrentValue - ValueAtT);
+		NewValue = InBlendInfo.CurrentValue + InBlendInfo.BlendValue * (AmplifyValueAtT - InBlendInfo.CurrentValue);
+	}
+	return NewValue;
+}
+
+double Blend_Neighbor(const FKeyBlendInfo& InBlendInfo)
+{
+	double NewValue = InBlendInfo.CurrentValue;
+	if (InBlendInfo.BlendValue < 0.0)
+	{
+		NewValue = InBlendInfo.CurrentValue + (-1.0 * InBlendInfo.BlendValue) * (InBlendInfo.PreviousValue - InBlendInfo.CurrentValue);
+	}
+	else
+	{
+		NewValue = InBlendInfo.CurrentValue + InBlendInfo.BlendValue * (InBlendInfo.NextValue - InBlendInfo.CurrentValue);
+	}
+	return NewValue;
+}
+
+double Blend_Relative(const FKeyBlendInfo& InBlendInfo)
+{
+	double NewValue = InBlendInfo.CurrentValue;
+	if (InBlendInfo.BlendValue < 0.0)
+	{
+		NewValue = InBlendInfo.CurrentValue + (-1.0 * InBlendInfo.BlendValue) * (InBlendInfo.PreviousValue - InBlendInfo.FirstValue);
+	}
+	else
+	{
+		NewValue = InBlendInfo.CurrentValue + InBlendInfo.BlendValue * (InBlendInfo.NextValue - InBlendInfo.LastValue);
+	}
+	return NewValue;
+}
+
+namespace Private::BlendToEase
+{ 
+static float SCurve(float X, float Slope, float Width, float Height, float XShift, float YShift)
+{
+	if (X > (XShift + Width))
+	{
+		return(Height + YShift);
+	}
+	if (X < XShift)
+	{
+		return YShift;
+	}
+
+	const float Val = Height * (FMath::Pow((X - XShift), Slope) / (FMath::Pow(X - XShift, Slope) + FMath::Pow((Width - (X - XShift)), Slope))) + YShift;
+	return Val;
+}
+}
+
+double Blend_Ease(const FKeyBlendInfo& InBlendInfo)
+{
+	if (FMath::IsNearlyEqual(InBlendInfo.NextTime, InBlendInfo.PreviousTime))
+	{
+		return InBlendInfo.CurrentValue;
+	}
+	const double Source = InBlendInfo.CurrentValue;
+	const double FullTimeDiff = InBlendInfo.NextTime - InBlendInfo.PreviousTime;
+	const double AbsValue = FMath::Abs(InBlendInfo.BlendValue);
+	const double X = InBlendInfo.CurrentTime - InBlendInfo.PreviousTime;
+	const double Ratio = X / FullTimeDiff;
+	double Shift, Delta, Base;
+	if (InBlendInfo.BlendValue > 0)
+	{
+		Shift = -1.0;
+		Delta = InBlendInfo.NextValue - Source;
+		Base = Source;
+	}
+	else
+	{
+		Shift = 0.0;
+		Delta = Source - InBlendInfo.PreviousValue;
+		Base = InBlendInfo.PreviousValue;
+	}
+	const double Slope = 5.0 * AbsValue;
+	const double EaseY = Private::BlendToEase::SCurve(Ratio, Slope, 2.0, 2.0, Shift, Shift);
+	const double NewValue = Base + (Delta * EaseY);
+
+	return NewValue;
+}
+
+double Blend_SmoothRough(const FKeyBlendInfo& InBlendInfo)
+{
+	double PrevVal = InBlendInfo.CurrentIndex > 0 ? InBlendInfo.AllKeyPositions[InBlendInfo.Indices[InBlendInfo.CurrentIndex - 1]].OutputValue : InBlendInfo.PreviousValue;
+	double CurVal = InBlendInfo.CurrentValue;
+	double NextVal = InBlendInfo.CurrentIndex < InBlendInfo.Indices.Num() - 1 ? InBlendInfo.AllKeyPositions[InBlendInfo.Indices[InBlendInfo.CurrentIndex + 1]].OutputValue : InBlendInfo.NextValue;
+	double NewValue = (PrevVal * 0.25 + CurVal * 0.5 + NextVal * 0.25);
+	if (InBlendInfo.BlendValue < 0.0)
+	{
+		NewValue = InBlendInfo.CurrentValue + (-1.0 * InBlendInfo.BlendValue) * (NewValue - InBlendInfo.CurrentValue);
+	}
+	else
+	{
+		NewValue = InBlendInfo.CurrentValue + InBlendInfo.BlendValue * (InBlendInfo.CurrentValue - NewValue);
+	}
+	return NewValue;
+}
+}
diff --git a/Engine/Plugins/Animation/TweeningUtils/Source/TweeningUtilsEditor/Private/TweeningUtilsEditorModule.cpp b/Engine/Plugins/Animation/TweeningUtils/Source/TweeningUtilsEditor/Private/TweeningUtilsEditorModule.cpp
new file mode 100644
--- /dev/null
+++ b/Engine/Plugins/Animation/TweeningUtils/Source/TweeningUtilsEditor/Private/TweeningUtilsEditorModule.cpp
@@ -0,0 +1,20 @@
+﻿// Copyright Epic Games, Inc. All Rights Reserved.
+
+#include "TweeningUtilsEditorModule.h"
+
+namespace UE::TweeningUtilsEditor
+{
+void FTweeningUtilsEditorModule::StartupModule()
+{
+	
+}
+
+void FTweeningUtilsEditorModule::ShutdownModule()
+{
+	
+}
+	
+IMPLEMENT_MODULE(FTweeningUtilsEditorModule, TweeningUtilsEditor)
+}
+
+	
\ No newline at end of file
diff --git a/Engine/Plugins/Animation/TweeningUtils/Source/TweeningUtilsEditor/Private/TweeningUtilsEditorModule.h b/Engine/Plugins/Animation/TweeningUtils/Source/TweeningUtilsEditor/Private/TweeningUtilsEditorModule.h
new file mode 100644
--- /dev/null
+++ b/Engine/Plugins/Animation/TweeningUtils/Source/TweeningUtilsEditor/Private/TweeningUtilsEditorModule.h
@@ -0,0 +1,17 @@
+﻿// Copyright Epic Games, Inc. All Rights Reserved.
+
+#pragma once
+
+#include "Modules/ModuleManager.h"
+
+namespace UE::TweeningUtilsEditor
+{
+class FTweeningUtilsEditorModule : public IModuleInterface
+{
+public:
+
+	/** IModuleInterface implementation */
+	virtual void StartupModule() override;
+	virtual void ShutdownModule() override;
+};
+}
diff --git a/Engine/Plugins/Animation/TweeningUtils/Source/TweeningUtilsEditor/Public/Math/ContiguousKeyMapping.h b/Engine/Plugins/Animation/TweeningUtils/Source/TweeningUtilsEditor/Public/Math/ContiguousKeyMapping.h
new file mode 100644
--- /dev/null
+++ b/Engine/Plugins/Animation/TweeningUtils/Source/TweeningUtilsEditor/Public/Math/ContiguousKeyMapping.h
@@ -0,0 +1,46 @@
+﻿// Copyright Epic Games, Inc. All Rights Reserved.
+
+#pragma once
+
+#include "CurveDataAbstraction.h"
+#include "CurveEditorTypes.h"
+#include "Containers/Array.h"
+#include "Containers/Map.h"
+#include "HAL/Platform.h"
+
+class FCurveEditor;
+
+namespace UE::TweeningUtilsEditor
+{
+/** Data structure for use by BlendCurves. */
+struct TWEENINGUTILSEDITOR_API FContiguousKeyMapping
+{
+	// Set of keys in a row
+	struct FContiguousKeys
+	{
+		int32 PreviousIndex;
+		TArray<int32> Indices;
+		int32 NextIndex;
+		explicit FContiguousKeys(
+			const TArray<FKeyHandle>& InAllKeyHandles, const TArray<FKeyPosition>& InAllKeyPositions, const TArray<int32>& ContiguousKeyIndices
+			);
+	};
+
+	// All the selected keys in a row for a particular curve and pointers to all of the keys
+	struct FContiguousKeysArray
+	{
+		TArray<FKeyHandle> AllKeyHandles;
+		TArray<FKeyPosition> AllKeyPositions;
+
+		TArray<FContiguousKeys> KeysArray;
+
+		explicit FContiguousKeysArray(const TArray<FKeyHandle>& InAllKeyHandle, const TArray<FKeyPosition>& InAllKeyPositions);
+		void Add(const TArray<int32>& ContiguousKeyIndices);
+	};
+
+	FContiguousKeyMapping() = default;
+	explicit FContiguousKeyMapping(const FCurveEditor& InCurveEditor);
+
+	TMap<FCurveModelID, FContiguousKeysArray> KeyMap;
+};
+}
diff --git a/Engine/Plugins/Animation/TweeningUtils/Source/TweeningUtilsEditor/Public/Math/CurveBlending.h b/Engine/Plugins/Animation/TweeningUtils/Source/TweeningUtilsEditor/Public/Math/CurveBlending.h
new file mode 100644
--- /dev/null
+++ b/Engine/Plugins/Animation/TweeningUtils/Source/TweeningUtilsEditor/Public/Math/CurveBlending.h
@@ -0,0 +1,83 @@
+﻿// Copyright Epic Games, Inc. All Rights Reserved.
+
+#pragma once
+
+#include "ContiguousKeyMapping.h"
+#include "CurveDataAbstraction.h"
+#include "CurveEditor.h"
+#include "CurveModel.h"
+#include "Math/KeyBlendInfo.h"
+
+#include <type_traits>
+
+namespace UE::TweeningUtilsEditor
+{
+/**
+ * Invokes callback InDoBlend on all keys in InKeySelection.
+ *
+ * @note This function will call Modify on each curve.
+ * 
+ * @param InBlendValue The blend value to pass to InDoBlend
+ * @param InCurveEditor Used to look up curves.
+ * @param InKeySelection The selection to blend
+ * @param InDoBlend	Does the blending
+ * @return Whether InDoBlend was invokes at least once.
+ */
+template<typename TCallback> requires std::is_invocable_r_v<double, TCallback, const FKeyBlendInfo&>
+bool BlendCurves(float InBlendValue, const FCurveEditor& InCurveEditor, const FContiguousKeyMapping& InKeySelection, TCallback&& InDoBlend);
+}
+
+namespace UE::TweeningUtilsEditor
+{
+template<typename TCallback> requires std::is_invocable_r_v<double, TCallback, const FKeyBlendInfo&>
+bool BlendCurves(float InBlendValue, const FCurveEditor& InCurveEditor, const FContiguousKeyMapping& InKeySelection, TCallback&& InDoBlend)
+{
+	bool bDidBlend = false;
+	
+	TArray<FKeyHandle> KeyHandles;
+	TArray<FKeyPosition> KeyPositions;
+	for (const TPair<FCurveModelID, FContiguousKeyMapping::FContiguousKeysArray>& KeysArray : InKeySelection.KeyMap)
+	{
+		if (FCurveModel* Curve = InCurveEditor.FindCurve(KeysArray.Key))
+		{
+			Curve->Modify();
+			for (const FContiguousKeyMapping::FContiguousKeys& Keys : KeysArray.Value.KeysArray)
+			{
+				const int32 PreviousIndex = Keys.PreviousIndex != INDEX_NONE ? Keys.PreviousIndex : Keys.Indices[0];
+				const double PreviousTime = KeysArray.Value.AllKeyPositions[PreviousIndex].InputValue;
+				const double PreviousValue = KeysArray.Value.AllKeyPositions[PreviousIndex].OutputValue;
+				const int32 NextIndex = Keys.NextIndex != INDEX_NONE ? Keys.NextIndex : Keys.Indices[Keys.Indices.Num() - 1];
+				const double NextTime = KeysArray.Value.AllKeyPositions[NextIndex].InputValue;
+				const double NextValue = KeysArray.Value.AllKeyPositions[NextIndex].OutputValue;
+
+				const int32 NumIndices = Keys.Indices.Num();
+				KeyHandles.SetNum(NumIndices);
+				KeyPositions.SetNum(NumIndices);
+				const double FirstValue = KeysArray.Value.AllKeyPositions[Keys.Indices[0]].OutputValue;
+				const double LastValue = KeysArray.Value.AllKeyPositions[Keys.Indices[Keys.Indices.Num() -1]].OutputValue;
+				FKeyBlendInfo BlendStruct(KeysArray.Value.AllKeyPositions, Keys.Indices);
+				int32 CurrentIndex = 0;
+				for (int32 Index = 0; Index < NumIndices; ++Index)
+				{
+					const int32& KeyIndex = Keys.Indices[Index];
+					KeyHandles[Index] = KeysArray.Value.AllKeyHandles[KeyIndex];
+					const double CurrentValue = KeysArray.Value.AllKeyPositions[KeyIndex].OutputValue;
+					const double CurrentTime = KeysArray.Value.AllKeyPositions[KeyIndex].InputValue;
+					BlendStruct.SetValues(
+						PreviousTime, PreviousValue, CurrentTime, CurrentValue,
+						NextTime, NextValue, InBlendValue, FirstValue, LastValue, CurrentIndex
+						);
+					++CurrentIndex;
+
+					const double NewValue = InDoBlend(BlendStruct);
+					KeyPositions[Index] = FKeyPosition(KeysArray.Value.AllKeyPositions[KeyIndex].InputValue, NewValue);
+				}
+				Curve->SetKeyPositions(KeyHandles, KeyPositions);
+				bDidBlend = true;
+			}
+		}
+	}
+
+	return bDidBlend;
+}
+}
diff --git a/Engine/Plugins/Animation/TweeningUtils/Source/TweeningUtilsEditor/Public/Math/KeyBlendInfo.h b/Engine/Plugins/Animation/TweeningUtils/Source/TweeningUtilsEditor/Public/Math/KeyBlendInfo.h
new file mode 100644
--- /dev/null
+++ b/Engine/Plugins/Animation/TweeningUtils/Source/TweeningUtilsEditor/Public/Math/KeyBlendInfo.h
@@ -0,0 +1,56 @@
+﻿// Copyright Epic Games, Inc. All Rights Reserved.
+
+#pragma once
+
+#include "CurveDataAbstraction.h"
+#include "Containers/Array.h"
+#include "HAL/Platform.h"
+
+namespace UE::TweeningUtilsEditor
+{
+/** Holds information about blending consecutive range of keys; each key has a time (x) and value (y). */
+struct FKeyBlendInfo
+{
+	// Previous and Next are the keys that are before/after the contiguous set of keys, 
+	// to get the Previous key for example you use AllKeyPostions[Indices[CurrentIndex -1]]
+	double PreviousTime; 
+	double PreviousValue; 
+	double CurrentTime; 
+	double CurrentValue;
+	double NextTime;
+	double NextValue;
+	/** The user specified value based on which to blend. Range is in [-1.0, 1.0]. */
+	double BlendValue;
+	double FirstValue;
+	double LastValue;
+
+	/** All keys, including the ones not being blended. */
+	TArray<FKeyPosition> AllKeyPositions;
+	/** Indices to AllKeyPositions which are being blended. */
+	TArray<int32> Indices;
+	/** Index to Indices. Identifies the key that is being blended right now. The blend algorithm moves through Indices by incrementing CurrentIndex. */
+	int32 CurrentIndex; 
+
+	explicit FKeyBlendInfo(const TArray<FKeyPosition>& InAllKeyPositions, const TArray<int32>& InIndices) 
+		: AllKeyPositions(InAllKeyPositions), Indices(InIndices)
+	{}
+	
+	void SetValues(
+		double InPreviousTime, double InPreviousValue, double InCurrentTime, double InCurrentValue,
+		double InNextTime, double InNextValue, double InBlendValue, double InFirstValue, double InLastValue,
+		int32 InCurrentIndex
+		)
+	{
+		PreviousTime = InPreviousTime;
+		PreviousValue = InPreviousValue;
+		CurrentTime = InCurrentTime;
+		CurrentValue = InCurrentValue;
+		NextTime = InNextTime;
+		NextValue = InNextValue;
+		BlendValue = InBlendValue;
+		FirstValue = InFirstValue;
+		LastValue = InLastValue;
+		CurrentIndex = InCurrentIndex;
+	}
+};
+}
diff --git a/Engine/Plugins/Animation/TweeningUtils/Source/TweeningUtilsEditor/Public/Math/KeyBlendingAbstraction.h b/Engine/Plugins/Animation/TweeningUtils/Source/TweeningUtilsEditor/Public/Math/KeyBlendingAbstraction.h
new file mode 100644
--- /dev/null
+++ b/Engine/Plugins/Animation/TweeningUtils/Source/TweeningUtilsEditor/Public/Math/KeyBlendingAbstraction.h
@@ -0,0 +1,27 @@
+﻿// Copyright Epic Games, Inc. All Rights Reserved.
+
+#pragma once
+
+#include "HAL/Platform.h"
+
+/** The following just contains some utils that is common in UI code. Tools can opt for their own description but using this helps with consistency. */
+namespace UE::TweeningUtilsEditor
+{
+struct FKeyBlendInfo;
+	
+/** Util enum which is useful for UI code to abstract functions.*/
+enum class EBlendFunction : uint8
+{
+	ControlsToTween,
+	PushPull,
+	BlendNeighbor,
+	BlendRelative,
+	BlendEase,
+	SmoothRough,
+
+	// Add new entries above.
+	Num
+};
+
+TWEENINGUTILSEDITOR_API double EvaluateBlend(EBlendFunction InFunction, const FKeyBlendInfo& InBlendInfo);
+}
\ No newline at end of file
diff --git a/Engine/Plugins/Animation/TweeningUtils/Source/TweeningUtilsEditor/Public/Math/KeyBlendingFunctions.h b/Engine/Plugins/Animation/TweeningUtils/Source/TweeningUtilsEditor/Public/Math/KeyBlendingFunctions.h
new file mode 100644
--- /dev/null
+++ b/Engine/Plugins/Animation/TweeningUtils/Source/TweeningUtilsEditor/Public/Math/KeyBlendingFunctions.h
@@ -0,0 +1,20 @@
+﻿// Copyright Epic Games, Inc. All Rights Reserved.
+
+#pragma once
+
+#include "KeyBlendInfo.h"
+
+namespace UE::TweeningUtilsEditor
+{
+TWEENINGUTILSEDITOR_API double Blend_ControlsToTween(const FKeyBlendInfo& InBlendInfo);
+
+TWEENINGUTILSEDITOR_API double Blend_PushPull(const FKeyBlendInfo& InBlendInfo);
+
+TWEENINGUTILSEDITOR_API double Blend_Neighbor(const FKeyBlendInfo& InBlendInfo);
+
+TWEENINGUTILSEDITOR_API double Blend_Relative(const FKeyBlendInfo& InBlendInfo);
+
+TWEENINGUTILSEDITOR_API double Blend_Ease(const FKeyBlendInfo& InBlendInfo);
+
+TWEENINGUTILSEDITOR_API double Blend_SmoothRough(const FKeyBlendInfo& InBlendInfo);
+}
diff --git a/Engine/Plugins/Animation/TweeningUtils/Source/TweeningUtilsEditor/TweeningUtilsEditor.build.cs b/Engine/Plugins/Animation/TweeningUtils/Source/TweeningUtilsEditor/TweeningUtilsEditor.build.cs
new file mode 100644
--- /dev/null
+++ b/Engine/Plugins/Animation/TweeningUtils/Source/TweeningUtilsEditor/TweeningUtilsEditor.build.cs
@@ -0,0 +1,25 @@
+﻿// Copyright Epic Games, Inc. All Rights Reserved.
+
+using UnrealBuildTool;
+
+public class TweeningUtilsEditor : ModuleRules
+{
+	public TweeningUtilsEditor(ReadOnlyTargetRules Target) : base(Target)
+	{
+		PCHUsage = ModuleRules.PCHUsageMode.UseExplicitOrSharedPCHs;
+		
+		PublicDependencyModuleNames.AddRange(new string[]
+		{
+			"Core", 
+			"CurveEditor",
+			"InputCore",
+			"SlateCore",
+			"Slate"
+		});
+		
+		PrivateDependencyModuleNames.AddRange(new string[]
+		{
+			
+		});
+	}
+}
diff --git a/Engine/Plugins/Animation/TweeningUtils/TweeningUtils.uplugin b/Engine/Plugins/Animation/TweeningUtils/TweeningUtils.uplugin
new file mode 100644
--- /dev/null
+++ b/Engine/Plugins/Animation/TweeningUtils/TweeningUtils.uplugin
@@ -0,0 +1,23 @@
+{
+	"FileVersion": 3,
+	"Version": 1,
+	"VersionName": "1.0",
+	"FriendlyName": "Tweening Utils",
+	"Description": "Algorithms and widgets useful for inbetweening.",
+	"Category": "Animation",
+	"CreatedBy": "Epic Games, Inc.",
+	"CreatedByURL": "https://epicgames.com",
+	"DocsURL": "",
+	"MarketplaceURL": "",
+	"SupportURL": "",
+	"EnabledByDefault": true,
+	"CanContainContent": true,
+	"Installed": false,
+	"Modules": [
+		{
+			"Name": "TweeningUtilsEditor",
+			"Type": "Editor",
+			"LoadingPhase": "Default"
+		}
+	]
+}
\ No newline at end of file
diff --git a/Engine/Source/Editor/CurveEditor/CurveEditor.Build.cs b/Engine/Source/Editor/CurveEditor/CurveEditor.Build.cs
--- a/Engine/Source/Editor/CurveEditor/CurveEditor.Build.cs
+++ b/Engine/Source/Editor/CurveEditor/CurveEditor.Build.cs
@@ -25,6 +25,11 @@ public class CurveEditor : ModuleRules
 			}
 		);
 
-        PublicDependencyModuleNames.Add("SequencerWidgets");
+        PublicDependencyModuleNames.AddRange(
+	        new string[]
+	        {
+		        "SequencerWidgets"
+	        }
+	    );
 	}
 }
