Constraints: fixed regression introduced with 38662120

#rnx

diff --git a/Engine/Plugins/Animation/ControlRig/Source/ControlRigEditor/Private/Editor/Constraints/SConstraintsWidget.cpp b/Engine/Plugins/Animation/ControlRig/Source/ControlRigEditor/Private/Editor/Constraints/SConstraintsWidget.cpp
--- a/Engine/Plugins/Animation/ControlRig/Source/ControlRigEditor/Private/Editor/Constraints/SConstraintsWidget.cpp
+++ b/Engine/Plugins/Animation/ControlRig/Source/ControlRigEditor/Private/Editor/Constraints/SConstraintsWidget.cpp
@@ -612,20 +612,20 @@ FReply SConstraintMenuEntry::CreateSelectionPicker(const bool bUseDefault) const
 		FOnShouldFilterActor(), 
 		FOnActorSelected::CreateLambda([Args](AActor* InActor)
 		{
-			CreateConstraint(InActor, Args.CreationDelegate, Args.ConstraintType, Args.bUseDefault);
+			OnParentPicked(InActor, Args.CreationDelegate, Args.ConstraintType, Args.bUseDefault);
 		}) );
 	
 	return FReply::Handled();
 }
 
-void SConstraintMenuEntry::CreateConstraint(
+void SConstraintMenuEntry::OnParentPicked(
 	AActor* InParent,
 	const FOnConstraintCreated& InCreationDelegate,
 	const ETransformConstraintType InConstraintType,
 	const bool bUseDefault)
 {
 	using namespace UE::Private;
-
+	
 	if (!InParent)
 	{
 		return;	
@@ -641,20 +641,94 @@ void SConstraintMenuEntry::CreateConstraint(
 	const TArray<AActor*> Selection = GetCurrentSelection();
 	if (Selection.IsEmpty())
 	{
-		return;
+		return;	
 	}
 
-	FConstrainable ConstrainableParent;
-	GetConstrainable(InParent, ConstrainableParent);
-
-	if (ConstrainableParent.Object.IsValid())
+	// create constraints
+	auto CreateConstraint = [InCreationDelegate, InConstraintType, bUseDefault](
+		const TArray<AActor*>& Selection, UObject* InParent, const FName& InSocketName)
 	{
-		const bool bCreated = UE::Private::CreateConstraint(Selection, ConstrainableParent, InConstraintType, bUseDefault);
+		FConstrainable Parent;
+		Parent.Object = InParent;
+		Parent.Socket = InSocketName;
+
+		const bool bCreated = UE::Private::CreateConstraint(Selection, Parent, InConstraintType, bUseDefault);
+	
+		// update list
 		if (bCreated && InCreationDelegate.IsBound())
-        {
-        	InCreationDelegate.Execute();
-        }
+		{
+			InCreationDelegate.Execute();
+		}
+	};
+
+	// gather sub components with sockets
+	const TInlineComponentArray<USceneComponent*> Components(InParent);
+	const TArray<USceneComponent*> ComponentsWithSockets = Components.FilterByPredicate([](const USceneComponent* Component)
+	{
+		return Component->HasAnySockets();
+	});
+	const int32 NumComponentsWithSockets = ComponentsWithSockets.Num();
+
+	// if no component socket available then constrain the whole actor 
+	if (NumComponentsWithSockets == 0)
+	{
+		return CreateConstraint(Selection, InParent, NAME_None);
 	}
+
+	// creates a menu encapsulating InContent  
+	auto CreateMenu = [](const TSharedRef<SWidget>& InContent, const FVector2D& InLocation)
+	{
+		const FLevelEditorModule& LevelEditorModule = FModuleManager::GetModuleChecked<FLevelEditorModule>("LevelEditor");
+		const TSharedPtr< ILevelEditor > LevelEditor = LevelEditorModule.GetFirstLevelEditor();
+		FSlateApplication::Get().PushMenu(
+			LevelEditor.ToSharedRef(),
+			FWidgetPath(),
+			InContent,
+		InLocation,
+		FPopupTransitionEffect( FPopupTransitionEffect::ContextMenu ) );
+	};
+
+	// creates a new socket chooser popup widget
+	auto GetSocketChooserWidget = [CreateConstraint, Selection](USceneComponent* Component)
+	{
+		return SNew(SSocketChooserPopup)
+		.SceneComponent( Component )
+		.OnSocketChosen_Lambda([CreateConstraint, Selection, Component](FName InSocketName)
+		{
+			CreateConstraint(Selection, Component, InSocketName);
+		});
+	};
+
+	// decal the menu 
+	FVector2D SummonLocation = FSlateApplication::Get().GetCursorPos();
+	SummonLocation.Y += 4 * FSlateApplication::Get().GetCursorSize().Y;
+	
+	// if one component with sockets then constrain the component selecting the socket
+	if (NumComponentsWithSockets == 1)
+	{
+		return CreateMenu( GetSocketChooserWidget(ComponentsWithSockets[0]), SummonLocation);
+	}
+
+	// if there are several of them, then build a component chooser first then the socket chooser
+	static constexpr bool CloseAfterSelection = true;
+	FMenuBuilder MenuBuilder(CloseAfterSelection, nullptr);
+	MenuBuilder.BeginSection("ChooseComp", LOCTEXT("ChooseComponentSection", "Choose Component"));
+	{
+		for (USceneComponent* Component: ComponentsWithSockets)
+		{
+			MenuBuilder.AddMenuEntry(FText::FromName(Component->GetFName()),
+			FText(),
+				FSlateIconFinder::FindIconForClass(Component->GetClass(), TEXT("SCS.Component")),
+				FUIAction(FExecuteAction::CreateLambda([CreateMenu, GetSocketChooserWidget, Component]()
+				{
+					CreateMenu(GetSocketChooserWidget(Component), FSlateApplication::Get().GetCursorPos());
+				})),
+				NAME_None,
+			   EUserInterfaceActionType::Button);
+		}
+	}
+	MenuBuilder.EndSection();
+	CreateMenu(MenuBuilder.MakeWidget(), SummonLocation);
 }
 
 /**
diff --git a/Engine/Plugins/Animation/ControlRig/Source/ControlRigEditor/Private/Editor/Constraints/SConstraintsWidget.h b/Engine/Plugins/Animation/ControlRig/Source/ControlRigEditor/Private/Editor/Constraints/SConstraintsWidget.h
--- a/Engine/Plugins/Animation/ControlRig/Source/ControlRigEditor/Private/Editor/Constraints/SConstraintsWidget.h
+++ b/Engine/Plugins/Animation/ControlRig/Source/ControlRigEditor/Private/Editor/Constraints/SConstraintsWidget.h
@@ -127,7 +127,7 @@ private:
 	TSharedRef<SWidget> GenerateConstraintDefaultWidget();
 	
 	/** Creates the constraint between the current selection and the picked actor. */
-	static void CreateConstraint(
+	static void OnParentPicked(
 		AActor* InParent,
 		const FOnConstraintCreated& InDelegate,
 		const ETransformConstraintType InConstraintType,
