Fix mismatched deprecation macro

diff --git a/Engine/Plugins/Animation/ControlRig/Source/ControlRigEditor/Public/Tools/ControlRigTweener.h b/Engine/Plugins/Animation/ControlRig/Source/ControlRigEditor/Public/Tools/ControlRigTweener.h
--- a/Engine/Plugins/Animation/ControlRig/Source/ControlRigEditor/Public/Tools/ControlRigTweener.h
+++ b/Engine/Plugins/Animation/ControlRig/Source/ControlRigEditor/Public/Tools/ControlRigTweener.h
@@ -108,7 +108,7 @@ struct FBaseAnimSlider
 	PRAGMA_DISABLE_DEPRECATION_WARNINGS
 	FBaseAnimSlider() = default;
 	virtual ~FBaseAnimSlider() = default;
-	PRAGMA_DISABLE_DEPRECATION_WARNINGS
+	PRAGMA_ENABLE_DEPRECATION_WARNINGS
 	/**
 	 Setup up what's needed for the slider to blend
 	* @param InSequencer Sequencer to get keys, controls, etc. from
