Control Rig: Sample track container - introduce composed tracks

#rnx

diff --git a/Engine/Plugins/Animation/ControlRig/Source/ControlRig/Private/Tracks/ComposedSampleTrack.cpp b/Engine/Plugins/Animation/ControlRig/Source/ControlRig/Private/Tracks/ComposedSampleTrack.cpp
new file mode 100644
--- /dev/null
+++ b/Engine/Plugins/Animation/ControlRig/Source/ControlRig/Private/Tracks/ComposedSampleTrack.cpp
@@ -0,0 +1,51 @@
+// Copyright Epic Games, Inc. All Rights Reserved.
+
+#include "Tracks/ComposedSampleTrack.h"
+
+//////////////////////////////////////////////////////////////////////////////////////////////////////////////////////////////
+/////// FComposedSampleTrack
+//////////////////////////////////////////////////////////////////////////////////////////////////////////////////////////////
+
+template<>
+TArray<FSampleTrackBase::ETrackType> FComposedSampleTrack<FTransform3f>::GetChildTrackTypes() const
+{
+	return {
+		ETrackType_Vector3f,
+		ETrackType_Quatf,
+		ETrackType_Vector3f
+	};
+}
+
+template<>
+FTransform3f FComposedSampleTrack<FTransform3f>::GetValueAtTimeIndex(int32 InTimeIndex, int32& InOutSampleIndex) const
+{
+	check(ChildTracks.Num() == 3);
+
+	const FSampleTrack<FVector3f>* LocationTrack = static_cast<const FSampleTrack<FVector3f>*>(ChildTracks[0].Get());
+	const FSampleTrack<FQuat4f>* RotationTrack = static_cast<const FSampleTrack<FQuat4f>*>(ChildTracks[1].Get());
+	const FSampleTrack<FVector3f>* ScaleTrack = static_cast<const FSampleTrack<FVector3f>*>(ChildTracks[2].Get());
+
+	FTransform3f Result = FTransform3f::Identity;
+
+	// todo: we need to introduce a custom sample structure which allows to store sub sample indices... int32 isn't enough
+	// essentially the FSampleTrackIndex should just hold an array of sample indices - one for each track
+	Result.SetLocation(LocationTrack->GetValueAtTimeIndex(InTimeIndex, InOutSampleIndex));
+	Result.SetRotation(RotationTrack->GetValueAtTimeIndex(InTimeIndex, InOutSampleIndex));
+	Result.SetScale3D(ScaleTrack->GetValueAtTimeIndex(InTimeIndex, InOutSampleIndex));
+	return Result;
+}
+	
+template<>
+void FComposedSampleTrack<FTransform3f>::AddSample(const FTransform3f& InValue, float InTolerance)
+{
+	check(!FSampleTrackBase::IsReferenced());
+	check(ChildTracks.Num() == 3);
+
+	FSampleTrack<FVector3f>* LocationTrack = static_cast<FSampleTrack<FVector3f>*>(ChildTracks[0].Get());
+	FSampleTrack<FQuat4f>* RotationTrack = static_cast<FSampleTrack<FQuat4f>*>(ChildTracks[1].Get());
+	FSampleTrack<FVector3f>* ScaleTrack = static_cast<FSampleTrack<FVector3f>*>(ChildTracks[2].Get());
+
+	LocationTrack->AddSample(InValue.GetLocation(), InTolerance);
+	RotationTrack->AddSample(InValue.GetRotation(), InTolerance);
+	ScaleTrack->AddSample(InValue.GetScale3D(), InTolerance);
+}
\ No newline at end of file
diff --git a/Engine/Plugins/Animation/ControlRig/Source/ControlRig/Private/Tracks/SampleTrackBase.cpp b/Engine/Plugins/Animation/ControlRig/Source/ControlRig/Private/Tracks/SampleTrackBase.cpp
--- a/Engine/Plugins/Animation/ControlRig/Source/ControlRig/Private/Tracks/SampleTrackBase.cpp
+++ b/Engine/Plugins/Animation/ControlRig/Source/ControlRig/Private/Tracks/SampleTrackBase.cpp
@@ -422,6 +422,12 @@ void FSampleTrackBase::AddTimeIndex(bool bOnlyIncreaseUpperBound)
 	}
 }
 
+const TArray<int32>& FSampleTrackBase::GetChildTracks() const
+{
+	static const TArray<int32> EmptyIndices;
+	return EmptyIndices;
+}
+
 int32 FSampleTrackBase::RemoveObsoleteTimes(int32 InNumSamplesToRemove, const TArray<bool>& InRemoveIndexMap, const TArray<int32>& InNewIndexMap)
 {
 	if(!IsValid() || IsReferenced())
diff --git a/Engine/Plugins/Animation/ControlRig/Source/ControlRig/Private/Tracks/SampleTrackContainer.cpp b/Engine/Plugins/Animation/ControlRig/Source/ControlRig/Private/Tracks/SampleTrackContainer.cpp
--- a/Engine/Plugins/Animation/ControlRig/Source/ControlRig/Private/Tracks/SampleTrackContainer.cpp
+++ b/Engine/Plugins/Animation/ControlRig/Source/ControlRig/Private/Tracks/SampleTrackContainer.cpp
@@ -2,6 +2,7 @@
 
 #include "Tracks/SampleTrackContainer.h"
 #include "Tracks/SampleTrackArchive.h"
+#include "Tracks/ComposedSampleTrack.h"
 #include "Misc/Compression.h"
 
 #include UE_INLINE_GENERATED_CPP_BY_NAME(SampleTrackContainer)
@@ -208,7 +209,7 @@ bool FSampleTrackContainer::Serialize(FArchive& InArchive)
 		{
 			uint8 TrackType = 0;
 			(*InnerArchive) << TrackType;
-			AddTrack(MakeTrack((FSampleTrackBase::ETrackType)TrackType));
+			AddTrack(MakeTrack((FSampleTrackBase::ETrackType)TrackType), false);
 		}
 
 		// now load the track in the order they were saved (raw first, everything else after)
@@ -223,6 +224,7 @@ bool FSampleTrackContainer::Serialize(FArchive& InArchive)
 		for(int32 TrackIndex = 0; TrackIndex < NumTracks; TrackIndex++)
 		{
 			GetTrack(TrackIndex)->UpdateArrayViews();
+			GetTrack(TrackIndex)->UpdateChildTracks();
 		}
 
 		// the tracks were added with no names present,
@@ -356,7 +358,7 @@ TSharedPtr<FSampleTrack<FQuat4f>> FSampleTrackContainer::AddQuatTrack(const FNam
 
 TSharedPtr<FSampleTrack<FTransform3f>> FSampleTrackContainer::AddTransformTrack(const FName& InName)
 {
-	TSharedPtr<FSampleTrack<FTransform3f>> Track = MakeShareable(new FSampleTrack<FTransform3f>(InName));
+	TSharedPtr<FSampleTrack<FTransform3f>> Track = MakeShareable(new FComposedSampleTrack<FTransform3f>(InName));
 	AddTrack(Track);
 	return Track;
 }
@@ -532,15 +534,39 @@ TSharedPtr<FSampleTrackBase> FSampleTrackContainer::FindTrack(const FName& InNam
 	return nullptr;
 }
 
-void FSampleTrackContainer::AddTrack(TSharedPtr<FSampleTrackBase> InTrack)
+void FSampleTrackContainer::AddTrack(TSharedPtr<FSampleTrackBase> InTrack, bool bCreateChildTracks)
 {
+	TArray<int32> ChildTrackIndices;
+	if(InTrack->IsComposed() && bCreateChildTracks)
+	{
+		const TArray<FSampleTrackBase::ETrackType> ChildTrackTypes = InTrack->GetChildTrackTypes();
+		for(const FSampleTrackBase::ETrackType ChildTrackType : ChildTrackTypes)
+		{
+			FName ChildTrackName = TEXT("ChildTrack");
+			if(!InTrack->Names.IsEmpty())
+			{
+				ChildTrackName = *FString::Printf(TEXT("%s_Child_%d"), *InTrack->GetName().ToString(), ChildTrackIndices.Num());
+			}
+			TSharedPtr<FSampleTrackBase> ChildTrack = MakeTrack(ChildTrackType);
+			ChildTrack->Names.Add(ChildTrackName);
+			AddTrack(ChildTrack);
+			ChildTrackIndices.Add(ChildTrack->GetTrackIndex());
+		}
+	}
+	
 	InTrack->TrackIndex = Tracks.Num();
 	InTrack->Container = this;
+	InTrack->SetChildTracks(ChildTrackIndices);
 
 	// a track can have many names - so we need to make sure each one is unique
 	for(int32 Index = 0; Index < InTrack->Names.Num(); Index++)
 	{
 		FName& Name = InTrack->Names[Index];
+		if(Name.IsNone())
+		{
+			continue;
+		}
+		
 		int32 NameSuffix = 3; // causes the first duplicate name to be called "Foo_2"
 		while (NameToIndex.Contains(Name))
 		{
@@ -598,7 +624,7 @@ TSharedPtr<FSampleTrackBase> FSampleTrackContainer::MakeTrack(FSampleTrackBase::
 		}
 		case FSampleTrackBase::ETrackType_Transformf:
 		{
-			return MakeShareable(new FSampleTrack<FTransform3f>());
+			return MakeShareable(new FComposedSampleTrack<FTransform3f>());
 		}
 		case FSampleTrackBase::ETrackType_LinearColor:
 		{
@@ -680,14 +706,32 @@ TSharedPtr<FSampleTrackBase> FSampleTrackContainer::MakeTrack(FSampleTrackBase::
 
 void FSampleTrackContainer::RemoveInvalidTracks(bool bUpdateNameToIndexMap)
 {
+	TArray<int32> OldTrackIndexToNewTrackIndex;
+	OldTrackIndexToNewTrackIndex.AddUninitialized(Tracks.Num());
+	int32 MaxTrackIndex = 0;
+	for(int32 TrackIndex = 0; TrackIndex < Tracks.Num(); TrackIndex++)
+	{
+		if(Tracks[TrackIndex].IsValid() && Tracks[TrackIndex]->IsValid())
+		{
+			OldTrackIndexToNewTrackIndex[TrackIndex] = MaxTrackIndex++;
+		}
+		else
+		{
+			OldTrackIndexToNewTrackIndex[TrackIndex] = INDEX_NONE;
+		}
+	}
 	const int32 NumRemoved = Tracks.RemoveAll([](const TSharedPtr<FSampleTrackBase>& Track)
 	{
 		return !Track.IsValid() || !Track->IsValid();
 	});
 
-	if(NumRemoved > 0 && bUpdateNameToIndexMap)
+	if(NumRemoved > 0)
 	{
-		UpdateNameToIndexMap();
+		UpdateTrackIndices(OldTrackIndexToNewTrackIndex);
+		if(bUpdateNameToIndexMap)
+		{
+			UpdateNameToIndexMap();
+		}
 	}
 }
 
@@ -795,12 +839,21 @@ void FSampleTrackContainer::MergeDuplicateTracks(bool bUpdateNameToIndexMap, flo
 	int32 NumTracksMerged = 0;
 	TArray<bool> MergedTracks;
 	MergedTracks.AddZeroed(Tracks.Num());
+	TArray<int32> OldTrackIndexToNewTrackIndex;
+	OldTrackIndexToNewTrackIndex.Reserve(Tracks.Num());
+
+	for(int32 Index = 0; Index < Tracks.Num(); Index++)
+	{
+		OldTrackIndexToNewTrackIndex.Add(Index);
+	}
+
 	for(int32 IndexA = 0; IndexA < Tracks.Num() - 1; IndexA++)
 	{
 		if(MergedTracks[IndexA])
 		{
 			continue;
 		}
+
 		if(TrackInfos[IndexA].Mode == FSampleTrackBase::EMode_Invalid)
 		{
 			continue;
@@ -833,6 +886,7 @@ void FSampleTrackContainer::MergeDuplicateTracks(bool bUpdateNameToIndexMap, flo
 			// merge the track
 			TrackA->Names.Append(TrackB->Names);
 			MergedTracks[IndexB] = true;
+			OldTrackIndexToNewTrackIndex[IndexB] = IndexA;
 			NumTracksMerged++;
 		}
 	}
@@ -842,7 +896,7 @@ void FSampleTrackContainer::MergeDuplicateTracks(bool bUpdateNameToIndexMap, flo
 		return;
 	}
 
-	for(int32 Index = 0; Index < MergedTracks.Num(); Index++)
+	for(int32 Index = 0; Index < Tracks.Num(); Index++)
 	{
 		if(MergedTracks[Index])
 		{
@@ -850,6 +904,7 @@ void FSampleTrackContainer::MergeDuplicateTracks(bool bUpdateNameToIndexMap, flo
 		}
 	}
 
+	UpdateTrackIndices(OldTrackIndexToNewTrackIndex);
 	RemoveInvalidTracks(bUpdateNameToIndexMap);
 }
 
@@ -893,7 +948,7 @@ void FSampleTrackContainer::MergeTypedTracks(bool bUpdateNameToIndexMap, float I
 	for(int32 TrackIndex = 0; TrackIndex < Tracks.Num(); TrackIndex++)
 	{
 		TSharedPtr<FSampleTrackBase>& Track = Tracks[TrackIndex];
-		if(!Track.IsValid() || !Track->IsValid() || Track->IsReferenced())
+		if(!Track.IsValid() || !Track->IsValid() || Track->IsReferenced() || Track->IsComposed())
 		{
 			continue;
 		}
@@ -1014,6 +1069,31 @@ void FSampleTrackContainer::ConvertTracksToSampled(float InTolerance)
 	}
 }
 
+void FSampleTrackContainer::UpdateTrackIndices(const TArray<int32>& InNewTrackIndices)
+{
+	for(TSharedPtr<FSampleTrackBase>& Track : Tracks)
+	{
+		if(!Track.IsValid() || !Track->IsValid())
+		{
+			continue;
+		}
+		if(Track->ReferencedTrackIndex != INDEX_NONE)
+		{
+			Track->ReferencedTrackIndex = InNewTrackIndices[Track->ReferencedTrackIndex];
+			Track->UpdateArrayViews();
+		}
+		if(Track->IsComposed())
+		{
+			TArray<int32> ChildTrackIndices = Track->GetChildTracks();
+			for(int32& ChildTrackIndex : ChildTrackIndices)
+			{
+				ChildTrackIndex = InNewTrackIndices[ChildTrackIndex];
+			}
+			Track->SetChildTracks(ChildTrackIndices);
+		}
+	}
+}
+
 void FSampleTrackContainer::UpdateNameToIndexMap()
 {
 	NameToIndex.Reset();
@@ -1023,8 +1103,11 @@ void FSampleTrackContainer::UpdateNameToIndexMap()
 		Tracks[TrackIndex]->TrackIndex = TrackIndex;
 		for(const FName& Name : Tracks[TrackIndex]->GetAllNames())
 		{
-			check(!NameToIndex.Contains(Name));
-			NameToIndex.Add(Name, TrackIndex);
+			if(!Name.IsNone())
+			{
+				check(!NameToIndex.Contains(Name));
+				NameToIndex.Add(Name, TrackIndex);
+			}
 		}
 	}
 }
diff --git a/Engine/Plugins/Animation/ControlRig/Source/ControlRig/Public/Tracks/ComposedSampleTrack.h b/Engine/Plugins/Animation/ControlRig/Source/ControlRig/Public/Tracks/ComposedSampleTrack.h
new file mode 100644
--- /dev/null
+++ b/Engine/Plugins/Animation/ControlRig/Source/ControlRig/Public/Tracks/ComposedSampleTrack.h
@@ -0,0 +1,204 @@
+// Copyright Epic Games, Inc. All Rights Reserved.
+
+#pragma once
+
+#include "SampleTrack.h"
+#include "SampleTrackContainer.h"
+
+template<typename T>
+class FComposedSampleTrack : public FSampleTrack<T>
+{
+public:
+
+	FComposedSampleTrack()
+	{
+		FSampleTrackBase::Mode = FSampleTrackBase::EMode_Singleton;
+		FSampleTrack<T>::ValuesStorage.AddDefaulted(1);
+		FSampleTrack<T>::UpdateValueArrayView();
+	}
+	
+	FComposedSampleTrack(const FName& InName)
+	: FSampleTrack<T>(InName)
+	{
+		FSampleTrackBase::Mode = FSampleTrackBase::EMode_Singleton;
+		FSampleTrack<T>::ValuesStorage.AddDefaulted(1);
+		FSampleTrack<T>::UpdateValueArrayView();
+	}
+	
+	virtual bool IsComposed() const override { return true; } 
+
+	virtual TArray<FSampleTrackBase::ETrackType> GetChildTrackTypes() const override
+	{
+		return {};
+	}
+
+	virtual ~FComposedSampleTrack() override {}
+
+	// returns the value of this track at a given time index
+	virtual T GetValueAtTimeIndex(int32 InTimeIndex, int32& InOutSampleIndex) const override
+	{
+		// this needs to be implemented in template specializations
+		return FSampleTrack<T>::GetValueAtTimeIndex(InTimeIndex, InOutSampleIndex);
+	}
+	
+	// adds a new sample to this track
+	virtual void AddSample(const T& InValue, float InTolerance = KINDA_SMALL_NUMBER) override
+	{
+		// this needs to be implemented in template specializations
+	}
+
+	virtual void Reserve(int32 InSampleCount, int32 InValueCount) override
+	{
+		FSampleTrack<T>::Reserve(InSampleCount, InValueCount);
+		for(const TSharedPtr<FSampleTrackBase>& ChildTrack : ChildTracks)
+		{
+			ChildTrack->Reserve(InSampleCount, InValueCount);
+		}
+	}
+
+	virtual void Shrink() override
+	{
+		FSampleTrack<T>::Shrink();
+		for(const TSharedPtr<FSampleTrackBase>& ChildTrack : ChildTracks)
+		{
+			ChildTrack->Shrink();
+		}
+	}
+
+	virtual void Empty() override
+	{
+		FSampleTrack<T>::Empty();
+		for(const TSharedPtr<FSampleTrackBase>& ChildTrack : ChildTracks)
+		{
+			ChildTrack->Empty();
+		}
+	}
+	
+	// returns true if this track is identical to another one
+	virtual bool IsIdentical(const FSampleTrackBase* InOther, float InTolerance = KINDA_SMALL_NUMBER) override
+	{
+		if(!FSampleTrack<T>::IsIdentical(InOther))
+		{
+			return false;
+		}
+
+		const FComposedSampleTrack<T>* OtherComposedTrack = static_cast<const FComposedSampleTrack<T>*>(InOther);
+		
+		if(ChildTracks.Num() != OtherComposedTrack->ChildTracks.Num())
+		{
+			return false;
+		}
+
+		for(int32 ChildTrackIndex = 0; ChildTrackIndex < ChildTracks.Num(); ChildTrackIndex++)
+		{
+			if(!ChildTracks[ChildTrackIndex]->IsIdentical(OtherComposedTrack->ChildTracks[ChildTrackIndex].Get(), InTolerance))
+			{
+				return false;
+			}
+		}
+		
+		return true;
+	}
+
+	virtual void Serialize(FArchive& InArchive) override
+	{
+		FSampleTrack<T>::Serialize(InArchive);
+		InArchive << ChildTrackIndices;
+
+		if(InArchive.IsLoading())
+		{
+			SetChildTracks(ChildTrackIndices);
+		}
+	}
+
+protected:
+
+	virtual int32 RemoveObsoleteTimes(int32 InNumTimesToRemove, const TArray<bool>& InRemoveTimeIndexArray, const TArray<int32>& InOldTimeIndexToNewTimeIndex) override
+	{
+		return 0;
+	}
+	
+	virtual const TArray<int32>& GetChildTracks() const override
+	{
+		return ChildTrackIndices;
+	}
+
+	virtual void SetChildTracks(const TArray<int32>& InChildTrackIndices) override
+	{
+		ChildTrackIndices = InChildTrackIndices;
+
+		check(FSampleTrackBase::Container);
+		ChildTracks.Reset();
+		for(const int32& ChildTrackIndex : ChildTrackIndices)
+		{
+			ChildTracks.Add(FSampleTrackBase::Container->GetTrack(ChildTrackIndex));
+		}
+	}
+
+	virtual void UpdateChildTracks() override
+	{
+		SetChildTracks(ChildTrackIndices);
+	}
+
+public:
+	
+	/**
+	 * Analyses the values stored and introduces the atlas index array,
+	 * an indirection to save memory. After this Values is going to store unique elements,
+	 * and the Atlas array is used to pick the right unique value for each sample index.
+	 */
+	virtual bool AddAtlas(bool bForce = false, float InTolerance = KINDA_SMALL_NUMBER) override
+	{
+		return false;
+	}
+
+	/**
+	 * Removes an existing Atlas indirection by flattening out all unique values into a larger, non-unique Value array
+	 */
+	virtual bool RemoveAtlas() override
+	{
+		return false;
+	}
+
+	/**
+	 * Analyses the memory footprint of a sample track (just complete values vs time indices + values)
+	 * and converts the track back to a complete representation for memory efficiency
+	 */
+	virtual bool ConvertToComplete(bool bForce = false) override
+	{
+		return false;
+	}
+
+	/**
+	 * Adds all of the values back onto the track using a sampled representation
+	 */
+	virtual bool ConvertToSampled(bool bForce = false, float InTolerance = KINDA_SMALL_NUMBER) override
+	{
+		return false;
+	}
+
+	/**
+	 * Localizes values if this track is referenced
+	 */
+	virtual bool LocalizeValues() override
+	{
+		return false;
+	}
+
+protected:
+	
+	TArray<int32> ChildTrackIndices;
+	TArray<TSharedPtr<FSampleTrackBase>> ChildTracks;
+	
+	friend class FSampleTrackContainer;
+	friend class FControlRigBasicsSampleTrackTest;
+};
+
+template<>
+CONTROLRIG_API TArray<FSampleTrackBase::ETrackType> FComposedSampleTrack<FTransform3f>::GetChildTrackTypes() const;
+
+template<>
+CONTROLRIG_API FTransform3f FComposedSampleTrack<FTransform3f>::GetValueAtTimeIndex(int32 InTimeIndex, int32& InOutSampleIndex) const;
+	
+template<>
+CONTROLRIG_API void FComposedSampleTrack<FTransform3f>::AddSample(const FTransform3f& InValue, float InTolerance);
diff --git a/Engine/Plugins/Animation/ControlRig/Source/ControlRig/Public/Tracks/SampleTrack.h b/Engine/Plugins/Animation/ControlRig/Source/ControlRig/Public/Tracks/SampleTrack.h
--- a/Engine/Plugins/Animation/ControlRig/Source/ControlRig/Public/Tracks/SampleTrack.h
+++ b/Engine/Plugins/Animation/ControlRig/Source/ControlRig/Public/Tracks/SampleTrack.h
@@ -29,7 +29,7 @@ public:
 	}
 
 	// returns the value of this track at a given time index
-	const T& GetValueAtTimeIndex(int32 InTimeIndex, int32& InOutSampleIndex) const
+	virtual T GetValueAtTimeIndex(int32 InTimeIndex, int32& InOutSampleIndex) const
 	{
 		verify(GetSampleIndexForTimeIndex(InTimeIndex, InOutSampleIndex));
 		return GetValueAtSampleIndex(InOutSampleIndex);
@@ -80,7 +80,7 @@ public:
 	}
 
 	// returns the value of this track at a given sample index
-	const T& GetValueAtSampleIndex(int32 InSampleIndex) const
+	virtual const T& GetValueAtSampleIndex(int32 InSampleIndex) const
 	{
 		if(UsesAtlas(InSampleIndex))
 		{
@@ -134,7 +134,7 @@ public:
 	}
 
 	// adds a new sample to this track
-	void AddSample(const T& InValue, float InTolerance = KINDA_SMALL_NUMBER)
+	virtual void AddSample(const T& InValue, float InTolerance = KINDA_SMALL_NUMBER)
 	{
 		check(!IsReferenced());
 		
diff --git a/Engine/Plugins/Animation/ControlRig/Source/ControlRig/Public/Tracks/SampleTrackBase.h b/Engine/Plugins/Animation/ControlRig/Source/ControlRig/Public/Tracks/SampleTrackBase.h
--- a/Engine/Plugins/Animation/ControlRig/Source/ControlRig/Public/Tracks/SampleTrackBase.h
+++ b/Engine/Plugins/Animation/ControlRig/Source/ControlRig/Public/Tracks/SampleTrackBase.h
@@ -130,6 +130,12 @@ public:
 	// returns the index the track this track is referencing
 	int32 GetReferencedTrackIndex() const { return ReferencedTrackIndex; }
 
+	// returns true if this track is composed out of other tracks
+	virtual bool IsComposed() const { return false; }
+
+	// returns the required track types for the child tracks
+	virtual TArray<FSampleTrackBase::ETrackType> GetChildTrackTypes() const { return {}; }
+
 protected:
 
 	const FSampleTrackBase* GetReferencedTrack() const;
@@ -139,6 +145,9 @@ protected:
 	bool UsesAtlas(int32 InAtlasIndex = INDEX_NONE) const;
 	void AddTimeIndex(bool bOnlyIncreaseUpperBound);
 	virtual int32 AppendValuesFromTrack(const FSampleTrackBase* InOtherTrack) = 0;
+	virtual const TArray<int32>& GetChildTracks() const;
+	virtual void SetChildTracks(const TArray<int32>& InChildTrackIndices) {}
+	virtual void UpdateChildTracks() {}
 
 	/**
 	 * Removes samples from the track which don't present any value anywhere.
diff --git a/Engine/Plugins/Animation/ControlRig/Source/ControlRig/Public/Tracks/SampleTrackContainer.h b/Engine/Plugins/Animation/ControlRig/Source/ControlRig/Public/Tracks/SampleTrackContainer.h
--- a/Engine/Plugins/Animation/ControlRig/Source/ControlRig/Public/Tracks/SampleTrackContainer.h
+++ b/Engine/Plugins/Animation/ControlRig/Source/ControlRig/Public/Tracks/SampleTrackContainer.h
@@ -217,7 +217,7 @@ public:
 
 private:
 
-	void AddTrack(TSharedPtr<FSampleTrackBase> InTrack);
+	void AddTrack(TSharedPtr<FSampleTrackBase> InTrack, bool bCreateChildTracks = true);
 	static TSharedPtr<FSampleTrackBase> MakeTrack(FSampleTrackBase::ETrackType InTrackType);
 
 	// removes all tracks with no data in them or nullptr tracks
@@ -249,6 +249,11 @@ private:
 	 */
 	void ConvertTracksToSampled(float InTolerance = KINDA_SMALL_NUMBER);
 
+	/**
+	 * Updates referenced and child track indices given a list providing the new index for each track
+	 */
+	void UpdateTrackIndices(const TArray<int32>& InNewTrackIndices);
+
 	// Updates the name lookup after larger changes within the container
 	void UpdateNameToIndexMap();
 
