Anim In Engine: Add Copy/Paste Pose hot keys

diff --git a/Engine/Plugins/Animation/ControlRig/Source/ControlRig/Private/Tools/ControlRigPose.cpp b/Engine/Plugins/Animation/ControlRig/Source/ControlRig/Private/Tools/ControlRigPose.cpp
--- a/Engine/Plugins/Animation/ControlRig/Source/ControlRig/Private/Tools/ControlRigPose.cpp
+++ b/Engine/Plugins/Animation/ControlRig/Source/ControlRig/Private/Tools/ControlRigPose.cpp
@@ -374,13 +374,16 @@ void UControlRigPoseAsset::PastePose(UControlRig* InControlRig, bool bDoKey, boo
 	Pose.PastePose(InControlRig,bDoKey, bDoMirror);
 }
 
-void UControlRigPoseAsset::SelectControls(UControlRig* InControlRig, bool bDoMirror)
+void UControlRigPoseAsset::SelectControls(UControlRig* InControlRig, bool bDoMirror, bool bClearSelection)
 {
 #if WITH_EDITOR
 	FScopedTransaction ScopedTransaction(LOCTEXT("SelectControlTransaction", "Select Control"));
 	InControlRig->Modify();
 #endif
-	InControlRig->ClearControlSelection();
+	if (bClearSelection)
+	{
+		InControlRig->ClearControlSelection();
+	}
 	TArray<FName> Controls = Pose.GetControlNames();
 	FControlRigPoseMirrorTable MirrorTable;
 	FControlRigControlPose TempPose;
diff --git a/Engine/Plugins/Animation/ControlRig/Source/ControlRig/Public/Tools/ControlRigPose.h b/Engine/Plugins/Animation/ControlRig/Source/ControlRig/Public/Tools/ControlRigPose.h
--- a/Engine/Plugins/Animation/ControlRig/Source/ControlRig/Public/Tools/ControlRigPose.h
+++ b/Engine/Plugins/Animation/ControlRig/Source/ControlRig/Public/Tools/ControlRigPose.h
@@ -132,7 +132,7 @@ public:
 	void PastePose(UControlRig* InControlRig, bool bDoKey = false, bool bDoMirror = false);
 	
 	UFUNCTION(BlueprintCallable, Category = "Pose")
-	void SelectControls(UControlRig* InControlRig, bool bDoMirror = false);
+	void SelectControls(UControlRig* InControlRig, bool bDoMirror = false, bool bClearSelection = true);
 
 	TArray<FRigControlCopy> GetCurrentPose(UControlRig* InControlRig);
 
diff --git a/Engine/Plugins/Animation/ControlRig/Source/ControlRigEditor/Private/EditMode/ControlRigEditMode.cpp b/Engine/Plugins/Animation/ControlRig/Source/ControlRigEditor/Private/EditMode/ControlRigEditMode.cpp
--- a/Engine/Plugins/Animation/ControlRig/Source/ControlRigEditor/Private/EditMode/ControlRigEditMode.cpp
+++ b/Engine/Plugins/Animation/ControlRig/Source/ControlRigEditor/Private/EditMode/ControlRigEditMode.cpp
@@ -71,6 +71,7 @@
 #include "Widgets/Docking/SDockTab.h"
 #include "TransformConstraint.h"
 #include "Animation/DebugSkelMeshComponent.h"
+#include "Tools/ControlRigPose.h"
 #include "Materials/Material.h"
 #include "ControlRigEditorStyle.h"
 #include "DragTool_BoxSelect.h"
@@ -192,6 +193,7 @@ FControlRigEditMode::FControlRigEditMode()
 {
 	ControlProxy = NewObject<UControlRigDetailPanelControlProxies>(GetTransientPackage(), NAME_None);
 	ControlProxy->SetFlags(RF_Transactional);
+	StoredPose = NewObject<UControlRigPoseAsset>(GetTransientPackage(), NAME_None);
 	DetailKeyFrameCache = MakeShared<FDetailKeyFrameCacheAndHandler>();
 
 	UControlRigEditModeSettings* Settings = GetMutableDefault<UControlRigEditModeSettings>();
@@ -2333,6 +2335,10 @@ void FControlRigEditMode::AddReferencedObjects( FReferenceCollector& Collector )
 	{
 		Collector.AddReferencedObject(ControlProxy);
 	}
+	if (StoredPose)
+	{
+		Collector.AddReferencedObject(StoredPose);
+	}
 }
 
 void FControlRigEditMode::ClearRigElementSelection(uint32 InTypes)
@@ -2890,6 +2896,42 @@ void FControlRigEditMode::BindCommands()
 		Commands.ToggleControlShapeTransformEdit,
 		FExecuteAction::CreateRaw(this, &FControlRigEditMode::ToggleControlShapeTransformEdit));
 
+	CommandBindings->MapAction(
+		Commands.SelectMirroredControls,
+		FExecuteAction::CreateRaw(this, &FControlRigEditMode::SelectMirroredControls));
+
+	CommandBindings->MapAction(
+		Commands.AddMirroredControlsToSelection,
+		FExecuteAction::CreateRaw(this, &FControlRigEditMode::AddMirroredControlsToSelection));
+
+	CommandBindings->MapAction(
+		Commands.MirrorSelectedControls,
+		FExecuteAction::CreateRaw(this, &FControlRigEditMode::MirrorSelectedControls));
+
+	CommandBindings->MapAction(
+		Commands.SelectAllControls,
+		FExecuteAction::CreateRaw(this, &FControlRigEditMode::SelectAllControls));
+
+	CommandBindings->MapAction(
+		Commands.SavePose,
+		FExecuteAction::CreateRaw(this, &FControlRigEditMode::SavePose,0));
+
+	CommandBindings->MapAction(
+		Commands.SelectPose,
+		FExecuteAction::CreateRaw(this, &FControlRigEditMode::SelectPose,false, 0));
+
+	CommandBindings->MapAction(
+		Commands.SelectMirrorPose,
+		FExecuteAction::CreateRaw(this, &FControlRigEditMode::SelectPose, true, 0));
+
+	CommandBindings->MapAction(
+		Commands.PastePose,
+		FExecuteAction::CreateRaw(this, &FControlRigEditMode::PastePose, false, 0));
+
+	CommandBindings->MapAction(
+		Commands.PasteMirrorPose,
+		FExecuteAction::CreateRaw(this, &FControlRigEditMode::PastePose, true, 0));
+
 	CommandBindings->MapAction(
 		Commands.SetAnimLayerPassthroughKey,
 		FExecuteAction::CreateRaw(this, &FControlRigEditMode::SetAnimLayerPassthroughKey));
@@ -5696,6 +5738,189 @@ bool FControlRigEditMode::GetOnlySelectRigControls()const
 	return Settings->bOnlySelectRigControls;
 }
 
+static TArray<UControlRig*> GetControlRigsWithSelectedControls(const TArray<UControlRig*>& InControlRigs)
+{
+	TArray<UControlRig*> SelectedControlRigs;
+	for (UControlRig* ControlRig : InControlRigs)
+	{
+		if (ControlRig && ControlRig->CurrentControlSelection().Num() > 0)
+		{
+			SelectedControlRigs.Add(ControlRig);
+		}
+	}
+	return SelectedControlRigs;
+}
+
+/** Select Mirrored Controls on Current Selection*/
+void FControlRigEditMode::SelectMirroredControls()
+{
+	TArray<UControlRig*> ControlRigs = GetControlRigsArray(true);
+	ControlRigs = GetControlRigsWithSelectedControls(ControlRigs);
+	if (ControlRigs.Num() == 0)
+	{
+		return;
+	}
+
+	FScopedTransaction ScopedTransaction(LOCTEXT("SelectMirroredControls", "Select Mirrored Controls"), !GIsTransacting);
+	FGCScopeGuard Guard;
+	UControlRigPoseAsset* TempPose = NewObject<UControlRigPoseAsset>(GetTransientPackage(), NAME_None);
+
+	for (UControlRig* ControlRig : ControlRigs)
+	{
+		if (ControlRig)
+		{
+			TempPose->SavePose(ControlRig, false);
+			TempPose->SelectControls(ControlRig, true);
+		}
+	}
+}
+
+/** Select Mirrored Controls on Current Selection, keeping current selection*/
+void FControlRigEditMode::AddMirroredControlsToSelection()
+{
+	TArray<UControlRig*> ControlRigs = GetControlRigsArray(true);
+	ControlRigs = GetControlRigsWithSelectedControls(ControlRigs);
+	if (ControlRigs.Num() == 0)
+	{
+		return;
+	}
+
+	FScopedTransaction ScopedTransaction(LOCTEXT("AddMirroredControlsToSelection", "Add Mirrored Controls to Selection"), !GIsTransacting);
+	FGCScopeGuard Guard;
+	UControlRigPoseAsset* TempPose = NewObject<UControlRigPoseAsset>(GetTransientPackage(), NAME_None);
+	for (UControlRig* ControlRig : ControlRigs)
+	{
+		if (ControlRig)
+		{
+			TempPose->SavePose(ControlRig, false);
+			TempPose->SelectControls(ControlRig, true, false);
+		}
+	}
+}
+
+/** Put Selected Controls To Mirrored Pose*/
+void FControlRigEditMode::MirrorSelectedControls()
+{
+	TArray<UControlRig*> ControlRigs = GetControlRigsArray(true);
+	ControlRigs = GetControlRigsWithSelectedControls(ControlRigs);
+	if (ControlRigs.Num() == 0)
+	{
+		return;
+	}
+	FScopedTransaction ScopedTransaction(LOCTEXT("MirrorSelectedControls", "Mirror Selected Controls"), !GIsTransacting);
+	FGCScopeGuard Guard;
+	UControlRigPoseAsset* TempPose = NewObject<UControlRigPoseAsset>(GetTransientPackage(), NAME_None);
+
+	for (UControlRig* ControlRig : ControlRigs)
+	{
+		if (ControlRig)
+		{
+			TempPose->SavePose(ControlRig, true);
+			TempPose->PastePose(ControlRig, /*setkey*/false, true);
+		}
+	}
+}
+
+/** Select All Controls*/
+void FControlRigEditMode::SelectAllControls()
+{
+	TArray<UControlRig*> ControlRigs = GetControlRigsArray(true);
+	ControlRigs = GetControlRigsWithSelectedControls(ControlRigs);
+	if (ControlRigs.Num() == 0)
+	{
+		return;
+	}
+	FScopedTransaction ScopedTransaction(LOCTEXT("SelectAllControls", "Select All Controls"), !GIsTransacting);
+	FGCScopeGuard Guard;
+	UControlRigPoseAsset* TempPose = NewObject<UControlRigPoseAsset>(GetTransientPackage(), NAME_None);
+	for (UControlRig* ControlRig : ControlRigs)
+	{
+		if (ControlRig)
+		{
+			TempPose->SavePose(ControlRig, true); //this will save the whole pose
+			TempPose->SelectControls(ControlRig);
+		}
+	}
+}
+
+/** Save a pose of selected controls*/
+void FControlRigEditMode::SavePose(int32 PoseNum)
+{
+	if (StoredPose)
+	{
+		TArray<UControlRig*> ControlRigs = GetControlRigsArray(true);
+		ControlRigs = GetControlRigsWithSelectedControls(ControlRigs);
+		for (UControlRig* ControlRig : ControlRigs)
+		{
+			if (ControlRig)
+			{
+				StoredPose->SavePose(ControlRig, false);
+				return;
+			}
+		}
+	}
+}
+
+/** Select controls in saved pose*/
+void FControlRigEditMode::SelectPose(bool bMirror, int32 PoseNum)
+{
+	if (StoredPose)
+	{
+		TArray<UControlRig*> ControlRigs = GetControlRigsArray(true);
+		if (ControlRigs.Num() < 0)
+		{
+			return;
+		}
+		else if (ControlRigs.Num() == 1 && ControlRigs[0])
+		{
+			FScopedTransaction ScopedTransaction(LOCTEXT("SelectPoseControls", "Select Pose Controls"), !GIsTransacting);
+			StoredPose->SelectControls(ControlRigs[0], bMirror);
+			return;
+		}
+		ControlRigs = GetControlRigsWithSelectedControls(ControlRigs);
+		if (ControlRigs.Num() == 0)
+		{
+			return;
+		}
+		FScopedTransaction ScopedTransaction(LOCTEXT("SelectPoseControls", "Select Pose Controls"), !GIsTransacting);
+
+		for (UControlRig* ControlRig : ControlRigs)
+		{
+			if (ControlRig)
+			{
+				//if (ControlRig->CurrentControlSelection().Num() > 0)
+				{
+					StoredPose->SelectControls(ControlRigs[0], bMirror);
+					//return;
+				}
+			}
+		}
+	}
+}
+
+/** Paste saved pose */
+void FControlRigEditMode::PastePose(bool bMirror, int32 PoseNum)
+{
+	if (StoredPose)
+	{
+		TArray<UControlRig*> ControlRigs = GetControlRigsArray(true);
+		ControlRigs = GetControlRigsWithSelectedControls(ControlRigs);
+		if (ControlRigs.Num() == 0)
+		{
+			return;
+		}
+		FScopedTransaction ScopedTransaction(LOCTEXT("PastePose", "Paste Pose"), !GIsTransacting);
+		for (UControlRig* ControlRig : ControlRigs)
+		{
+			if (ControlRig)
+			{
+				StoredPose->PastePose(ControlRig,/*setkey*/false, bMirror);
+				return;
+			}
+		}
+	}
+}
+
 /**
 * FDetailKeyFrameCacheAndHandler
 */
diff --git a/Engine/Plugins/Animation/ControlRig/Source/ControlRigEditor/Private/EditMode/ControlRigEditModeCommands.cpp b/Engine/Plugins/Animation/ControlRig/Source/ControlRigEditor/Private/EditMode/ControlRigEditModeCommands.cpp
--- a/Engine/Plugins/Animation/ControlRig/Source/ControlRigEditor/Private/EditMode/ControlRigEditModeCommands.cpp
+++ b/Engine/Plugins/Animation/ControlRig/Source/ControlRigEditor/Private/EditMode/ControlRigEditModeCommands.cpp
@@ -22,7 +22,21 @@ void FControlRigEditModeCommands::RegisterCommands()
 	UI_COMMAND(DragAnimSliderTool, "Drag Anim Slider Tool", "Drag existing anim slider", EUserInterfaceActionType::Button, FInputChord(EKeys::U));
 	UI_COMMAND(ChangeAnimSliderTool, "Change Anim Slider Tool", "Go to the next anim slider", EUserInterfaceActionType::Button, FInputChord(EKeys::U, EModifierKey::Shift));
 
-	UI_COMMAND(SetAnimLayerPassthroughKey," Set Anim Layer Passthrough Key","Set previous layer value on current selected anim layer", EUserInterfaceActionType::Button, FInputChord());
+	UI_COMMAND(SetAnimLayerPassthroughKey,"Set Anim Layer Passthrough Key","Set previous layer value on current selected anim layer", EUserInterfaceActionType::Button, FInputChord());
+
+	UI_COMMAND(SelectMirroredControls, "Select Mirrored Controls", " Select mirrored Controls on current selection", EUserInterfaceActionType::Button, FInputChord(EKeys::M, EModifierKey::Control));
+
+	UI_COMMAND(AddMirroredControlsToSelection, "Add Mirrored Controls to Selection", "Select mirrored Controls on current selection, keeping current selection", EUserInterfaceActionType::Button, FInputChord(EKeys::M, EModifierKey::Shift));
+
+	UI_COMMAND(MirrorSelectedControls, "Mirror Selected Controls", "Put selected Controls to mirrored Pose", EUserInterfaceActionType::Button, FInputChord(EKeys::M));
+
+	UI_COMMAND(SelectAllControls, "Select All Controls", "Select All Controls", EUserInterfaceActionType::Button, FInputChord(EKeys::Escape, EModifierKey::Shift));
+
+	UI_COMMAND(SavePose, "Save Current Pose", "Save current Pose", EUserInterfaceActionType::Button, FInputChord(EKeys::Y, EModifierKey::Shift));
+	UI_COMMAND(SelectPose, "Select Saved Pose", "Select Controls on saved Pose", EUserInterfaceActionType::Button, FInputChord(EKeys::Y, EModifierKey::Control |EModifierKey::Shift));
+	UI_COMMAND(PastePose, "Paste Saved Pose", "Paste Controls from saved Pose", EUserInterfaceActionType::Button, FInputChord(EKeys::Y, EModifierKey::Control));
+	UI_COMMAND(SelectMirrorPose, "Select Mirrored Saved Pose", "Select mirrored Controls on saved Pose", EUserInterfaceActionType::Button, FInputChord(EKeys::Y, EModifierKey::Alt | EModifierKey::Control | EModifierKey::Shift));
+	UI_COMMAND(PasteMirrorPose, "Paste Mirrored Saved Pose", "Paste mirrored Controls from saved Pose", EUserInterfaceActionType::Button, FInputChord(EKeys::Y, EModifierKey::Alt | EModifierKey::Control));
 
 	UI_COMMAND(ToggleControlShapeTransformEdit, "Toggle Shape Transform Edit", "Toggle Editing Selected Control's Shape Transform", EUserInterfaceActionType::ToggleButton, FInputChord(EKeys::Period, EModifierKey::Control)); 
 	UI_COMMAND(OpenSpacePickerWidget, "Open the Space Picker", "Allows space switching on the control", EUserInterfaceActionType::Button, FInputChord(EKeys::Tab)); 
diff --git a/Engine/Plugins/Animation/ControlRig/Source/ControlRigEditor/Private/EditMode/ControlRigEditModeCommands.h b/Engine/Plugins/Animation/ControlRig/Source/ControlRigEditor/Private/EditMode/ControlRigEditModeCommands.h
--- a/Engine/Plugins/Animation/ControlRig/Source/ControlRigEditor/Private/EditMode/ControlRigEditModeCommands.h
+++ b/Engine/Plugins/Animation/ControlRig/Source/ControlRigEditor/Private/EditMode/ControlRigEditModeCommands.h
@@ -65,6 +65,35 @@ public:
 	/** Sets Passthrough Key on selected anim layers */
 	TSharedPtr< FUICommandInfo > SetAnimLayerPassthroughKey;
 
+	/** Select mirrored Controls on current selection*/
+	TSharedPtr< FUICommandInfo > SelectMirroredControls;
+
+	/** Select mirrored Controls on current selection, keeping current selection*/
+	TSharedPtr< FUICommandInfo > AddMirroredControlsToSelection;
+
+	/** Put selected Controls To mirrored Pose*/
+	TSharedPtr< FUICommandInfo > MirrorSelectedControls;
+
+	/** Select all Controls*/
+	TSharedPtr< FUICommandInfo > SelectAllControls;
+
+	/** Save a pose of selected Controls*/
+	TSharedPtr< FUICommandInfo > SavePose;
+
+	/** Select Controls in saved Pose*/
+	TSharedPtr< FUICommandInfo > SelectPose;
+
+	/** Paste saved Pose */
+	TSharedPtr< FUICommandInfo > PastePose;
+
+	/** Select Controls in saved Pose*/
+	TSharedPtr< FUICommandInfo > SelectMirrorPose;
+
+	/** Paste saved mirror Pose */
+	TSharedPtr< FUICommandInfo > PasteMirrorPose;
+
+
+
 	/** Opens up the space picker widget */
 	TSharedPtr< FUICommandInfo > OpenSpacePickerWidget;
 
diff --git a/Engine/Plugins/Animation/ControlRig/Source/ControlRigEditor/Public/EditMode/ControlRigEditMode.h b/Engine/Plugins/Animation/ControlRig/Source/ControlRigEditor/Public/EditMode/ControlRigEditMode.h
--- a/Engine/Plugins/Animation/ControlRig/Source/ControlRigEditor/Public/EditMode/ControlRigEditMode.h
+++ b/Engine/Plugins/Animation/ControlRig/Source/ControlRigEditor/Public/EditMode/ControlRigEditMode.h
@@ -38,7 +38,7 @@ class IMovieScenePlayer;
 class AControlRigShapeActor;
 class UDefaultControlRigManipulationLayer;
 class UControlRigDetailPanelControlProxies;
-class UControlRigControlsProxy;
+class UControlRigPoseAsset;
 struct FRigControl;
 class IControlRigManipulatable;
 class ISequencer;
@@ -301,6 +301,29 @@ public:
 	/** Sets Passthrough Key on selected anim layers */
 	void SetAnimLayerPassthroughKey();
 
+	/** Select Mirrored Controls on Current Selection*/
+	void SelectMirroredControls();
+
+	/** Select Mirrored Controls on Current Selection, keeping current selection*/
+	void AddMirroredControlsToSelection();
+
+	/** Put Selected Controls To Mirrored Pose*/
+	void MirrorSelectedControls();
+
+	/** Select All Controls*/
+	void SelectAllControls();
+
+	//for the following pose functions we only support one pose(PoseNum = 0)
+	//but may support more later
+	/** Save a pose of selected controls*/
+	void SavePose(int32 PoseNum = 0);
+
+	/** Select controls in saved pose*/
+	void SelectPose(bool bDoMirror, int32 PoseNum = 0 );
+
+	/** Paste saved pose */
+	void PastePose(bool bDoMirror, int32 PoseNum = 0);
+
 	/** Opens up the space picker widget */
 	void OpenSpacePickerWidget();
 
@@ -619,9 +642,12 @@ private:
 		~FPendingControlRigEvaluator();
 		
 	private:
-	
+
 		FControlRigEditMode* EditMode;
 	};
+private:
+	//pose used by the hotkeys
+	TObjectPtr<UControlRigPoseAsset> StoredPose;
 
 	friend class FControlRigEditorModule;
 	friend class FControlRigEditor;
