Control Rig: SampleTrack to use custom sample index to speed up value lookups

#rnx

diff --git a/Engine/Plugins/Animation/ControlRig/Source/ControlRig/Private/Tracks/ComposedSampleTrack.cpp b/Engine/Plugins/Animation/ControlRig/Source/ControlRig/Private/Tracks/ComposedSampleTrack.cpp
--- a/Engine/Plugins/Animation/ControlRig/Source/ControlRig/Private/Tracks/ComposedSampleTrack.cpp
+++ b/Engine/Plugins/Animation/ControlRig/Source/ControlRig/Private/Tracks/ComposedSampleTrack.cpp
@@ -17,7 +17,7 @@ TArray<FSampleTrackBase::ETrackType> FComposedSampleTrack<FTransform3f>::GetChil
 }
 
 template<>
-FTransform3f FComposedSampleTrack<FTransform3f>::GetValueAtTimeIndex(int32 InTimeIndex, int32& InOutSampleIndex) const
+FTransform3f FComposedSampleTrack<FTransform3f>::GetValueAtTimeIndex(int32 InTimeIndex, FSampleTrackIndex& InOutIndex) const
 {
 	check(ChildTracks.Num() == 3);
 
@@ -27,11 +27,9 @@ FTransform3f FComposedSampleTrack<FTransform3f>::GetValueAtTimeIndex(int32 InTim
 
 	FTransform3f Result = FTransform3f::Identity;
 
-	// todo: we need to introduce a custom sample structure which allows to store sub sample indices... int32 isn't enough
-	// essentially the FSampleTrackIndex should just hold an array of sample indices - one for each track
-	Result.SetLocation(LocationTrack->GetValueAtTimeIndex(InTimeIndex, InOutSampleIndex));
-	Result.SetRotation(RotationTrack->GetValueAtTimeIndex(InTimeIndex, InOutSampleIndex));
-	Result.SetScale3D(ScaleTrack->GetValueAtTimeIndex(InTimeIndex, InOutSampleIndex));
+	Result.SetLocation(LocationTrack->GetValueAtTimeIndex(InTimeIndex, InOutIndex));
+	Result.SetRotation(RotationTrack->GetValueAtTimeIndex(InTimeIndex, InOutIndex));
+	Result.SetScale3D(ScaleTrack->GetValueAtTimeIndex(InTimeIndex, InOutIndex));
 	return Result;
 }
 	
diff --git a/Engine/Plugins/Animation/ControlRig/Source/ControlRig/Private/Tracks/SampleTrackBase.cpp b/Engine/Plugins/Animation/ControlRig/Source/ControlRig/Private/Tracks/SampleTrackBase.cpp
--- a/Engine/Plugins/Animation/ControlRig/Source/ControlRig/Private/Tracks/SampleTrackBase.cpp
+++ b/Engine/Plugins/Animation/ControlRig/Source/ControlRig/Private/Tracks/SampleTrackBase.cpp
@@ -3,6 +3,51 @@
 #include "Tracks/SampleTrackBase.h"
 #include "Tracks/SampleTrackContainer.h"
 
+//////////////////////////////////////////////////////////////////////////////////////////////////////////////////////////////
+/////// FSampleTrackIndex
+//////////////////////////////////////////////////////////////////////////////////////////////////////////////////////////////
+
+FSampleTrackIndex::FSampleTrackIndex(int32 NumTracks)
+	: bIsSingleton(false)
+{
+	Allocate(NumTracks);
+}
+
+FSampleTrackIndex::FSampleTrackIndex(const FSampleTrackContainer& InContainer)
+	: bIsSingleton(false)
+{
+	Allocate(InContainer.NumTracks());
+}
+
+int32& FSampleTrackIndex::GetSample(int32 InTrackIndex)
+{
+	if(bIsSingleton)
+	{
+		check(Samples.Num() == 1);
+		return Samples[0];
+	}
+	
+	InTrackIndex = FMath::Max(InTrackIndex, 0);
+	Allocate(InTrackIndex+1);
+	return Samples[InTrackIndex];
+}
+
+FSampleTrackIndex FSampleTrackIndex::MakeSingleton()
+{
+	FSampleTrackIndex Index(1);
+	Index.bIsSingleton = true;
+	return Index;
+}
+
+void FSampleTrackIndex::Allocate(int32 NumTracks)
+{
+	if(Samples.Num() >= NumTracks)
+	{
+		return;
+	}
+	Samples.AddZeroed(NumTracks - Samples.Num());
+}
+
 //////////////////////////////////////////////////////////////////////////////////////////////////////////////////////////////
 /////// FSampleTrackBase
 //////////////////////////////////////////////////////////////////////////////////////////////////////////////////////////////
@@ -66,40 +111,42 @@ int32 FSampleTrackBase::NumSamples() const
 	return 0;
 }
 
-bool FSampleTrackBase::GetSampleIndexForTimeIndex(int32 InTimeIndex, int32& InOutSampleIndex) const
+bool FSampleTrackBase::GetSampleIndexForTimeIndex(int32 InTimeIndex, FSampleTrackIndex& InOutIndex) const
 {
+	int32& SampleIndex = InOutIndex.GetSample(TrackIndex);
+	
 	switch(Mode)
 	{
 		case EMode_Singleton:
 		{
-			InOutSampleIndex = 0;
+			SampleIndex = 0;
 			return true;
 		}
 		case EMode_Sampled:
 		{
-			InOutSampleIndex = FMath::Clamp(InOutSampleIndex, 0, TimeIndices.Num()-1);
+			SampleIndex = FMath::Clamp(SampleIndex, 0, TimeIndices.Num()-1);
 
 			if(TimeIndices.IsEmpty())
 			{
 				break;
 			}
-			if(!TimeIndices.IsValidIndex(InOutSampleIndex))
+			if(!TimeIndices.IsValidIndex(SampleIndex))
 			{
-				InOutSampleIndex = 0;
+				SampleIndex = 0;
 			}
-			while(TimeIndices.IsValidIndex(InOutSampleIndex - 1) && TimeIndices[InOutSampleIndex] > InTimeIndex)
+			while(TimeIndices.IsValidIndex(SampleIndex - 1) && TimeIndices[SampleIndex] > InTimeIndex)
 			{
-				InOutSampleIndex--;
+				SampleIndex--;
 			}
-			while(TimeIndices.IsValidIndex(InOutSampleIndex + 1) && TimeIndices[InOutSampleIndex + 1] <= InTimeIndex)
+			while(TimeIndices.IsValidIndex(SampleIndex + 1) && TimeIndices[SampleIndex + 1] <= InTimeIndex)
 			{
-				InOutSampleIndex++;
+				SampleIndex++;
 			}
 			return true;
 		}
 		case EMode_Complete:
 		{
-			InOutSampleIndex = InTimeIndex;
+			SampleIndex = InTimeIndex;
 			return true;
 		}
 		case EMode_Invalid:
diff --git a/Engine/Plugins/Animation/ControlRig/Source/ControlRig/Public/Tracks/ComposedSampleTrack.h b/Engine/Plugins/Animation/ControlRig/Source/ControlRig/Public/Tracks/ComposedSampleTrack.h
--- a/Engine/Plugins/Animation/ControlRig/Source/ControlRig/Public/Tracks/ComposedSampleTrack.h
+++ b/Engine/Plugins/Animation/ControlRig/Source/ControlRig/Public/Tracks/ComposedSampleTrack.h
@@ -35,10 +35,10 @@ public:
 	virtual ~FComposedSampleTrack() override {}
 
 	// returns the value of this track at a given time index
-	virtual T GetValueAtTimeIndex(int32 InTimeIndex, int32& InOutSampleIndex) const override
+	virtual T GetValueAtTimeIndex(int32 InTimeIndex, FSampleTrackIndex& InOutIndex) const override
 	{
 		// this needs to be implemented in template specializations
-		return FSampleTrack<T>::GetValueAtTimeIndex(InTimeIndex, InOutSampleIndex);
+		return FSampleTrack<T>::GetValueAtTimeIndex(InTimeIndex, InOutIndex);
 	}
 	
 	// adds a new sample to this track
@@ -198,7 +198,7 @@ template<>
 CONTROLRIG_API TArray<FSampleTrackBase::ETrackType> FComposedSampleTrack<FTransform3f>::GetChildTrackTypes() const;
 
 template<>
-CONTROLRIG_API FTransform3f FComposedSampleTrack<FTransform3f>::GetValueAtTimeIndex(int32 InTimeIndex, int32& InOutSampleIndex) const;
+CONTROLRIG_API FTransform3f FComposedSampleTrack<FTransform3f>::GetValueAtTimeIndex(int32 InTimeIndex, FSampleTrackIndex& InOutIndex) const;
 	
 template<>
 CONTROLRIG_API void FComposedSampleTrack<FTransform3f>::AddSample(const FTransform3f& InValue, float InTolerance);
diff --git a/Engine/Plugins/Animation/ControlRig/Source/ControlRig/Public/Tracks/SampleTrack.h b/Engine/Plugins/Animation/ControlRig/Source/ControlRig/Public/Tracks/SampleTrack.h
--- a/Engine/Plugins/Animation/ControlRig/Source/ControlRig/Public/Tracks/SampleTrack.h
+++ b/Engine/Plugins/Animation/ControlRig/Source/ControlRig/Public/Tracks/SampleTrack.h
@@ -29,10 +29,10 @@ public:
 	}
 
 	// returns the value of this track at a given time index
-	virtual T GetValueAtTimeIndex(int32 InTimeIndex, int32& InOutSampleIndex) const
+	virtual T GetValueAtTimeIndex(int32 InTimeIndex, FSampleTrackIndex& InOutIndex) const
 	{
-		verify(GetSampleIndexForTimeIndex(InTimeIndex, InOutSampleIndex));
-		return GetValueAtSampleIndex(InOutSampleIndex);
+		verify(GetSampleIndexForTimeIndex(InTimeIndex, InOutIndex));
+		return GetValueAtSampleIndex(InOutIndex.GetSample(TrackIndex));
 	}
 	
 	// returns the value of this track at a given time index (overload for instanced structs only)
@@ -40,10 +40,10 @@ public:
 		typename StructType,
 		typename TEnableIf<TIsInstancedStructCompliant<T, StructType>::Value, StructType>::Type* = nullptr
 	>
-	StructType GetValueAtTimeIndex(int32 InTimeIndex, int32& InOutSampleIndex) const
+	StructType GetValueAtTimeIndex(int32 InTimeIndex, FSampleTrackIndex& InOutIndex) const
 	{
 		check(StructType::StaticStruct() == ScriptStruct);
-		const FInstancedStruct& InstancedStruct = GetValueAtTimeIndex(InTimeIndex, InOutSampleIndex);
+		const FInstancedStruct& InstancedStruct = GetValueAtTimeIndex(InTimeIndex, InOutIndex);
 		StructType Result;
 		check(InstancedStruct.GetScriptStruct() == ScriptStruct);
 		ScriptStruct->CopyScriptStruct(&Result, InstancedStruct.GetMemory(), 1);
@@ -55,10 +55,10 @@ public:
 		typename StructArrayType,
 		typename TEnableIf<TIsInstancedStructArrayCompliant<T, StructArrayType>::Value, StructArrayType>::Type* = nullptr
 	>
-	StructArrayType GetValueAtTimeIndex(int32 InTimeIndex, int32& InOutSampleIndex) const
+	StructArrayType GetValueAtTimeIndex(int32 InTimeIndex, FSampleTrackIndex& InOutIndex) const
 	{
 		check(StructArrayType::ElementType::StaticStruct() == ScriptStruct);
-		const TArray<FInstancedStruct>& InstancedStructArray = GetValueAtTimeIndex(InTimeIndex, InOutSampleIndex);
+		const TArray<FInstancedStruct>& InstancedStructArray = GetValueAtTimeIndex(InTimeIndex, InOutIndex);
 		StructArrayType Result;
 		Result.SetNum(InstancedStructArray.Num());
 		for(int32 ElementIndex = 0; ElementIndex < InstancedStructArray.Num(); ElementIndex++)
@@ -74,9 +74,9 @@ public:
 		typename StructType,
 		typename TEnableIf<TIsInstancedStructArrayCompliant<T, TArray<StructType>>::Value, StructType>::Type* = nullptr
 	>
-	TArray<StructType> GetArrayValueAtTimeIndex(int32 InTimeIndex, int32& InOutSampleIndex) const
+	TArray<StructType> GetArrayValueAtTimeIndex(int32 InTimeIndex, FSampleTrackIndex& InOutIndex) const
 	{
-		return GetValueAtTimeIndex<TArray<StructType>>(InTimeIndex, InOutSampleIndex);
+		return GetValueAtTimeIndex<TArray<StructType>>(InTimeIndex, InOutIndex);
 	}
 
 	// returns the value of this track at a given sample index
@@ -478,10 +478,10 @@ public:
 		{
 			TArray<int32> NewAtlas;
 			NewAtlas.Reserve(NumTimesInContainer);
-			int32 SampleIndex = 0;
+			FSampleTrackIndex SampleTrackIndex = FSampleTrackIndex::MakeSingleton();
 			for(int32 TimeIndex = 0; TimeIndex < NumTimesInContainer; TimeIndex++)
 			{
-				NewAtlas.Add(Atlas[GetSampleIndexForTimeIndex(TimeIndex, SampleIndex)]);
+				NewAtlas.Add(Atlas[GetSampleIndexForTimeIndex(TimeIndex, SampleTrackIndex)]);
 			}
 			AtlasStorage = MoveTemp(NewAtlas);
 		}
@@ -489,7 +489,7 @@ public:
 		{
 			TArray<T> NewValues;
 			NewValues.Reserve(NumTimesInContainer);
-			int32 SampleIndex = 0;
+			FSampleTrackIndex SampleIndex = FSampleTrackIndex::MakeSingleton();
 			for(int32 TimeIndex = 0; TimeIndex < NumTimesInContainer; TimeIndex++)
 			{
 				NewValues.Add(GetValueAtTimeIndex(TimeIndex, SampleIndex));
diff --git a/Engine/Plugins/Animation/ControlRig/Source/ControlRig/Public/Tracks/SampleTrackBase.h b/Engine/Plugins/Animation/ControlRig/Source/ControlRig/Public/Tracks/SampleTrackBase.h
--- a/Engine/Plugins/Animation/ControlRig/Source/ControlRig/Public/Tracks/SampleTrackBase.h
+++ b/Engine/Plugins/Animation/ControlRig/Source/ControlRig/Public/Tracks/SampleTrackBase.h
@@ -9,6 +9,21 @@
 
 class FSampleTrackContainer;
 
+struct CONTROLRIG_API FSampleTrackIndex
+{
+public:
+	explicit FSampleTrackIndex(int32 NumTracks);
+	explicit FSampleTrackIndex(const FSampleTrackContainer& InContainer);
+	int32& GetSample(int32 InTrackIndex);
+	static FSampleTrackIndex MakeSingleton();
+
+private:
+	void Allocate(int32 NumTracks);
+
+	TArray<int32, TInlineAllocator<1>> Samples;
+	bool bIsSingleton;
+};
+
 class CONTROLRIG_API FSampleTrackBase : public TSharedFromThis<FSampleTrackBase>
 {
 public:
@@ -87,7 +102,7 @@ public:
 	virtual int32 NumStoredValues() const = 0;
 
 	// returns the sample index within the track given a time index
-	bool GetSampleIndexForTimeIndex(int32 InTimeIndex, int32& InOutSampleIndex) const;
+	bool GetSampleIndexForTimeIndex(int32 InTimeIndex, FSampleTrackIndex& InOutIndex) const;
 
 	// returns true if the track stores a value specifically at a given time index
 	bool StoresValueForTimeIndex(int32 InTimeIndex) const;
