Control Rig: Fix CIS build break for sample track container

#rnx

diff --git a/Engine/Plugins/Animation/ControlRig/Source/ControlRig/Private/Tracks/SampleTrackContainer.cpp b/Engine/Plugins/Animation/ControlRig/Source/ControlRig/Private/Tracks/SampleTrackContainer.cpp
--- a/Engine/Plugins/Animation/ControlRig/Source/ControlRig/Private/Tracks/SampleTrackContainer.cpp
+++ b/Engine/Plugins/Animation/ControlRig/Source/ControlRig/Private/Tracks/SampleTrackContainer.cpp
@@ -109,7 +109,7 @@ void FSampleTrackContainer::Compact(float InTolerance)
 {
 	RemoveInvalidTracks(false);
 	RemoveObsoleteTimes();
-	MergeDuplicateTracks(false, InTolerance);
+	RemoveRedundantTracks(false, InTolerance);
 	MergeTypedTracks(false, InTolerance);
 	ConvertTracksToSampled();
 	ConvertTracksToComplete();
@@ -802,7 +802,7 @@ void FSampleTrackContainer::RemoveObsoleteTimes()
 	}
 }
 
-void FSampleTrackContainer::MergeDuplicateTracks(bool bUpdateNameToIndexMap, float InTolerance)
+void FSampleTrackContainer::RemoveRedundantTracks(bool bUpdateNameToIndexMap, float InTolerance)
 {
 	struct FTrackInfo
 	{
diff --git a/Engine/Plugins/Animation/ControlRig/Source/ControlRig/Public/Tracks/SampleTrackContainer.h b/Engine/Plugins/Animation/ControlRig/Source/ControlRig/Public/Tracks/SampleTrackContainer.h
--- a/Engine/Plugins/Animation/ControlRig/Source/ControlRig/Public/Tracks/SampleTrackContainer.h
+++ b/Engine/Plugins/Animation/ControlRig/Source/ControlRig/Public/Tracks/SampleTrackContainer.h
@@ -220,6 +220,8 @@ private:
 	void AddTrack(TSharedPtr<FSampleTrackBase> InTrack, bool bCreateChildTracks = true);
 	static TSharedPtr<FSampleTrackBase> MakeTrack(FSampleTrackBase::ETrackType InTrackType);
 
+public:
+	
 	// removes all tracks with no data in them or nullptr tracks
 	void RemoveInvalidTracks(bool bUpdateNameToIndexMap = true);
 
@@ -227,7 +229,7 @@ private:
 	void RemoveObsoleteTimes();
 
 	// combines tracks that have the exact same data in them
-	void MergeDuplicateTracks(bool bUpdateNameToIndexMap = true, float InTolerance = KINDA_SMALL_NUMBER);
+	void RemoveRedundantTracks(bool bUpdateNameToIndexMap = true, float InTolerance = KINDA_SMALL_NUMBER);
 
 	// combines tracks of the same type into a single, longer track and
 	// spawns tracks referencing the larger one's sections
@@ -249,6 +251,8 @@ private:
 	 */
 	void ConvertTracksToSampled(float InTolerance = KINDA_SMALL_NUMBER);
 
+private:
+	
 	/**
 	 * Updates referenced and child track indices given a list providing the new index for each track
 	 */
