Control Rig: Fix crash when debugging rig during PIE
#rnx

diff --git a/Engine/Plugins/Animation/ControlRig/Source/ControlRigEditor/Private/Editor/ControlRigEditor.cpp b/Engine/Plugins/Animation/ControlRig/Source/ControlRigEditor/Private/Editor/ControlRigEditor.cpp
--- a/Engine/Plugins/Animation/ControlRig/Source/ControlRigEditor/Private/Editor/ControlRigEditor.cpp
+++ b/Engine/Plugins/Animation/ControlRig/Source/ControlRigEditor/Private/Editor/ControlRigEditor.cpp
@@ -3688,7 +3688,10 @@ void FControlRigEditor::RefreshDirectManipulationTextList()
 	(void)GetDirectManipulationTargetTextList();
 	if(DirectManipulationCombo.IsValid())
 	{
-		DirectManipulationCombo->RefreshOptions();
+		FFunctionGraphTask::CreateAndDispatchWhenReady([this]()
+		{
+			DirectManipulationCombo->RefreshOptions();
+		}, TStatId(), NULL, ENamedThreads::GameThread);
 	}
 }
 
diff --git a/Engine/Plugins/Runtime/RigVM/Source/RigVMEditor/Private/Editor/RigVMEditor.cpp b/Engine/Plugins/Runtime/RigVM/Source/RigVMEditor/Private/Editor/RigVMEditor.cpp
--- a/Engine/Plugins/Runtime/RigVM/Source/RigVMEditor/Private/Editor/RigVMEditor.cpp
+++ b/Engine/Plugins/Runtime/RigVM/Source/RigVMEditor/Private/Editor/RigVMEditor.cpp
@@ -2010,7 +2010,10 @@ void FRigVMEditor::HandleVMExecutionHalted(const int32 InstructionIndex, UObject
 				{
 					if(UEdGraphNode* EdNode = EdGraph->FindNodeForModelNodeName(InNode->GetFName()))
 					{
-						JumpToHyperlink(EdNode, false);
+						FFunctionGraphTask::CreateAndDispatchWhenReady([this, EdNode]()
+						{
+							JumpToHyperlink(EdNode, false);
+						}, TStatId(), nullptr, ENamedThreads::GameThread);
 					}
 				}
 			}
