RigHierarchy / Storage improvements

#rnx

diff --git a/Engine/Plugins/Animation/ControlRig/Source/ControlRig/Private/Rigs/RigHierarchy.cpp b/Engine/Plugins/Animation/ControlRig/Source/ControlRig/Private/Rigs/RigHierarchy.cpp
--- a/Engine/Plugins/Animation/ControlRig/Source/ControlRig/Private/Rigs/RigHierarchy.cpp
+++ b/Engine/Plugins/Animation/ControlRig/Source/ControlRig/Private/Rigs/RigHierarchy.cpp
@@ -6774,9 +6774,10 @@ bool URigHierarchy::SortElementStorage()
 	}
 
 	FRigReusableElementStorage<FTransform> NewElementTransforms;
+	NewElementTransforms.Reserve(ElementTransforms.Num() - ElementTransforms.GetFreeList().Num());
+	
 	FRigReusableElementStorage<bool> NewElementDirtyStates;
 	NewElementDirtyStates.Reserve(ElementDirtyStates.Num() - ElementDirtyStates.GetFreeList().Num());
-	NewElementDirtyStates.Reserve(ElementDirtyStates.Num() - ElementDirtyStates.GetFreeList().Num());
 
 	bHasMatchingRanges = true;
 	NewRanges.Reset();
diff --git a/Engine/Plugins/Animation/ControlRig/Source/ControlRig/Public/Rigs/RigReusableElementStorage.h b/Engine/Plugins/Animation/ControlRig/Source/ControlRig/Public/Rigs/RigReusableElementStorage.h
--- a/Engine/Plugins/Animation/ControlRig/Source/ControlRig/Public/Rigs/RigReusableElementStorage.h
+++ b/Engine/Plugins/Animation/ControlRig/Source/ControlRig/Public/Rigs/RigReusableElementStorage.h
@@ -140,7 +140,8 @@ public:
 		{
 			Storage.Reserve(Storage.Num() + NumToAllocate);
 		}
-	
+
+		Indices.Reserve(InCount);
 		for(int32 Index = 0; Index < InCount; Index++)
 		{
 			if(FreeList.IsEmpty())
