Control Rig: Sequencer ECS fixes and improvements

- Implement ISequencerTrackEditor::ProcessKeyOperation for control rig tracks which is used by key navigation buttons and middle-mouse clicks for adding keys
- Fixed weight channel not being added to the base entity
- Fixed controls not aligning to skel animations correctly due to missing system dependency between skel anims and control rigs (control rigs must run after skel anims)

diff --git a/Engine/Plugins/Animation/ControlRig/Source/ControlRig/Private/Sequencer/MovieSceneControlRigParameterBuffer.cpp b/Engine/Plugins/Animation/ControlRig/Source/ControlRig/Private/Sequencer/MovieSceneControlRigParameterBuffer.cpp
--- a/Engine/Plugins/Animation/ControlRig/Source/ControlRig/Private/Sequencer/MovieSceneControlRigParameterBuffer.cpp
+++ b/Engine/Plugins/Animation/ControlRig/Source/ControlRig/Private/Sequencer/MovieSceneControlRigParameterBuffer.cpp
@@ -717,6 +717,34 @@ void FControlRigParameterValues::PopulateFrom(UControlRig* Rig)
 	}
 }
 
+FControlRigValueView FControlRigParameterValues::Find(FName Name) const
+{
+	for (const FControlRigParameterValueHeader& Header : Headers)
+	{
+		const int32 Index = Header.Find(Name);
+		if (Index != INDEX_NONE)
+		{
+			return FControlRigValueView(Header.GetParameter(Index), Header.GetType());
+		}
+	}
+	return FControlRigValueView();
+}
+
+FControlRigValueView FControlRigParameterValues::FindParameter(FName Name) const
+{
+	TArrayView<const FControlRigParameterValueHeader> ParameterHeaders(Headers);
+	ParameterHeaders.RightChopInline(1);
+
+	for (const FControlRigParameterValueHeader& Header : ParameterHeaders)
+	{
+		const int32 Index = Header.Find(Name);
+		if (Index != INDEX_NONE)
+		{
+			return FControlRigValueView(Header.GetParameter(Index), Header.GetType());
+		}
+	}
+	return FControlRigValueView();
+}
 
 FControlRigParameterBuffer::FControlRigParameterBuffer(UControlRig* InRig, EControlRigParameterBufferIndexStability IndexStability)
 	: WeakControlRig(InRig)
diff --git a/Engine/Plugins/Animation/ControlRig/Source/ControlRig/Private/Sequencer/MovieSceneControlRigParameterSection.cpp b/Engine/Plugins/Animation/ControlRig/Source/ControlRig/Private/Sequencer/MovieSceneControlRigParameterSection.cpp
--- a/Engine/Plugins/Animation/ControlRig/Source/ControlRig/Private/Sequencer/MovieSceneControlRigParameterSection.cpp
+++ b/Engine/Plugins/Animation/ControlRig/Source/ControlRig/Private/Sequencer/MovieSceneControlRigParameterSection.cpp
@@ -32,6 +32,74 @@
 
 #include UE_INLINE_GENERATED_CPP_BY_NAME(MovieSceneControlRigParameterSection)
 
+namespace UE::MovieScene
+{
+
+	/**
+	 * Finds an entry within the specified array based on a pointer that might exist within the
+	 *     memory address of the entry. Used for finding entries that relate to channels.
+	 * 
+	 * This function is O(1) in all cases and works using pointer arithmetic.
+	 */
+	template<typename T>
+	const T* FindEntryWithinArrayByPtr(const TArray<T>& Array, const void* Ptr)
+	{
+		UPTRINT ByteOffset = ((uint8*)Ptr) - ((uint8*)Array.GetData());
+
+		if (ByteOffset >= 0 && ByteOffset < Array.Num()*sizeof(T))
+		{
+			const int32 Index = ByteOffset / sizeof(T);
+			return &Array[Index];
+		}
+		return nullptr;
+	}
+
+	FControlRigChannelMetaData::FControlRigChannelMetaData()
+		: Type(EControlRigControlType::Num)
+		, IndexWithinControl(INDEX_NONE)
+		, EntitySystemID(MAX_uint32)
+	{}
+
+	FControlRigChannelMetaData::FControlRigChannelMetaData(EControlRigControlType InType, FName InControlName, int32 InIndexWithinControl, uint32 InEntitySystemID)
+		: Type(InType)
+		, ControlName(InControlName)
+		, IndexWithinControl(InIndexWithinControl)
+		, EntitySystemID(InEntitySystemID)
+	{
+	}
+
+	FControlRigChannelMetaData::operator bool() const
+	{
+		return Type != EControlRigControlType::Num && IndexWithinControl != INDEX_NONE;
+	}
+
+	enum class EControlRigEntityType : uint8
+	{
+		Base,
+		Space,
+		BoolParameter,
+		EnumParameter,
+		IntegerParameter,
+		ScalarParameter,
+		VectorParameter,
+		TransformParameter,
+	};
+
+	/* Entity IDs are an encoded type and index, with the upper 8 bits being the type, and the lower 24 bits as the index */
+	uint32 EncodeControlRigEntityID(int32 InIndex, EControlRigEntityType InType)
+	{
+		check(InIndex >= 0 && InIndex < int32(0x00FFFFFF));
+		return static_cast<uint32>(InIndex) | (uint32(InType) << 24);
+	}
+	void DecodeControlRigEntityID(uint32 InEntityID, int32& OutIndex, EControlRigEntityType& OutType)
+	{
+		// Mask out the type to get the index
+		OutIndex = static_cast<int32>(InEntityID & 0x00FFFFFF);
+		OutType = (EControlRigEntityType)(InEntityID >> 24);
+	}
+
+} // namespace UE::MovieScene
+
 #if WITH_EDITOR
 #include "AnimPose.h"
 #endif
@@ -1233,12 +1301,10 @@ FSpaceControlNameAndChannel* UMovieSceneControlRigParameterSection::GetSpaceChan
 
 FName UMovieSceneControlRigParameterSection::FindControlNameFromSpaceChannel(const FMovieSceneControlRigSpaceChannel* InSpaceChannel) const
 {
-	for (const FSpaceControlNameAndChannel& SpaceChannel : SpaceChannels)
+	if (const FSpaceControlNameAndChannel* Space = UE::MovieScene::FindEntryWithinArrayByPtr(SpaceChannels, InSpaceChannel))
 	{
-		if (&(SpaceChannel.SpaceCurve) == InSpaceChannel)
-		{
-			return SpaceChannel.ControlName;
-		}
+		check(InSpaceChannel== &Space->SpaceCurve);
+		return Space->ControlName;
 	}
 	return NAME_None;
 }
@@ -1594,13 +1660,11 @@ const TArray<FConstraintAndActiveChannel>& UMovieSceneControlRigParameterSection
 const FName& UMovieSceneControlRigParameterSection::FindControlNameFromConstraintChannel(
 	const FMovieSceneConstraintChannel* InConstraintChannel) const
 {
-	const int32 Index = ConstraintsChannels.IndexOfByPredicate([InConstraintChannel](const FConstraintAndActiveChannel& InChannel)
-	{
-		return &(InChannel.ActiveChannel) == InConstraintChannel;
-	});
-	
-	if (Index != INDEX_NONE)
+	const FConstraintAndActiveChannel* Entry = UE::MovieScene::FindEntryWithinArrayByPtr(ConstraintsChannels, InConstraintChannel);
+	if (Entry)
 	{
+		const int32 Index = Entry - ConstraintsChannels.GetData();
+
 		// look for info referencing that constraint index
 		using NameInfoIterator = TMap<FName, FChannelMapInfo>::TRangedForConstIterator;
 		for (NameInfoIterator It = ControlChannelMap.begin(); It; ++It)
@@ -2364,6 +2428,114 @@ EMovieSceneChannelProxyType UMovieSceneControlRigParameterSection::CacheChannelP
 	return EMovieSceneChannelProxyType::Dynamic;
 }
 
+UE::MovieScene::FControlRigChannelMetaData UMovieSceneControlRigParameterSection::GetChannelMetaData(const FMovieSceneChannel* Channel) const
+{
+	using namespace UE::MovieScene;
+
+	if (const FTransformParameterNameAndCurves* Transform = FindEntryWithinArrayByPtr(TransformParameterNamesAndCurves, Channel))
+	{
+		const FMovieSceneFloatChannel* ChannelStart = &Transform->Translation[0];
+
+		return FControlRigChannelMetaData(
+			EControlRigControlType::Parameter_Transform,
+			Transform->ParameterName,
+			static_cast<const FMovieSceneFloatChannel*>(Channel) - ChannelStart,
+			EncodeControlRigEntityID(Transform - TransformParameterNamesAndCurves.GetData(), EControlRigEntityType::TransformParameter)
+		);
+	}
+
+	if (const FVectorParameterNameAndCurves* Vector = FindEntryWithinArrayByPtr(VectorParameterNamesAndCurves, Channel))
+	{
+		const FMovieSceneFloatChannel* ChannelStart = &Vector->XCurve;
+
+		return FControlRigChannelMetaData(
+			EControlRigControlType::Parameter_Vector,
+			Vector->ParameterName,
+			static_cast<const FMovieSceneFloatChannel*>(Channel) - ChannelStart,
+			EncodeControlRigEntityID(Vector - VectorParameterNamesAndCurves.GetData(), EControlRigEntityType::VectorParameter)
+		);
+	}
+
+	if (const FEnumParameterNameAndCurve* Enum = FindEntryWithinArrayByPtr(EnumParameterNamesAndCurves, Channel))
+	{
+		check(Channel == &Enum->ParameterCurve);
+
+		return FControlRigChannelMetaData(
+			EControlRigControlType::Parameter_Enum,
+			Enum->ParameterName,
+			0,
+			EncodeControlRigEntityID(Enum - EnumParameterNamesAndCurves.GetData(), EControlRigEntityType::EnumParameter)
+			);
+	}
+
+	if (const FIntegerParameterNameAndCurve* Integer = FindEntryWithinArrayByPtr(IntegerParameterNamesAndCurves, Channel))
+	{
+		check(Channel == &Integer->ParameterCurve);
+
+		return FControlRigChannelMetaData(
+			EControlRigControlType::Parameter_Integer,
+			Integer->ParameterName,
+			0,
+			EncodeControlRigEntityID(Integer - IntegerParameterNamesAndCurves.GetData(), EControlRigEntityType::IntegerParameter)
+		);
+	}
+
+	if (const FSpaceControlNameAndChannel* Space = FindEntryWithinArrayByPtr(SpaceChannels, Channel))
+	{
+		check(Channel == &Space->SpaceCurve);
+
+		return FControlRigChannelMetaData(
+			EControlRigControlType::Space,
+			Space->ControlName,
+			0,
+			EncodeControlRigEntityID(Space - SpaceChannels.GetData(), EControlRigEntityType::Space)
+		);
+	}
+
+	if (const FBoolParameterNameAndCurve* Bool = FindEntryWithinArrayByPtr(BoolParameterNamesAndCurves, Channel))
+	{
+		check(Channel == &Bool->ParameterCurve);
+
+		return FControlRigChannelMetaData(
+			EControlRigControlType::Parameter_Bool,
+			Bool->ParameterName,
+			0,
+			EncodeControlRigEntityID(Bool - BoolParameterNamesAndCurves.GetData(), EControlRigEntityType::BoolParameter)
+		);
+	}
+
+	if (const FScalarParameterNameAndCurve* Scalar = FindEntryWithinArrayByPtr(ScalarParameterNamesAndCurves, Channel))
+	{
+		check(Channel == &Scalar->ParameterCurve);
+
+		return FControlRigChannelMetaData(
+			EControlRigControlType::Parameter_Scalar,
+			Scalar->ParameterName,
+			0,
+			EncodeControlRigEntityID(Scalar - ScalarParameterNamesAndCurves.GetData(), EControlRigEntityType::ScalarParameter)
+		);
+	}
+
+	if (const FVector2DParameterNameAndCurves* Vector2 = FindEntryWithinArrayByPtr(Vector2DParameterNamesAndCurves, Channel))
+	{
+		const FMovieSceneFloatChannel* ChannelStart = &Vector2->XCurve;
+
+		return FControlRigChannelMetaData(
+			EControlRigControlType::Parameter_Vector,
+			Vector2->ParameterName,
+			static_cast<const FMovieSceneFloatChannel*>(Channel) - ChannelStart,
+			EncodeControlRigEntityID(Vector2 - Vector2DParameterNamesAndCurves.GetData(), EControlRigEntityType::VectorParameter)
+		);
+	}
+
+	// @todo: Are colors even supported??
+	// if (const FColorParameterNameAndCurves* Color = FindEntryWithinArrayByPtr(ColorParameterNamesAndCurves, Channel))
+	// {
+	// }
+
+	return FControlRigChannelMetaData();
+}
+
 FMovieSceneInterrogationKey UMovieSceneControlRigParameterSection::GetFloatInterrogationKey()
 {
 	static FMovieSceneAnimTypeID TypeID = FMovieSceneAnimTypeID::Unique();
@@ -3888,37 +4060,6 @@ FIntegerParameterNameAndCurve::FIntegerParameterNameAndCurve(FName InParameterNa
 	ParameterName = InParameterName;
 }
 
-namespace UE::MovieScene
-{
-
-	enum class EControlRigEntityType : uint8
-	{
-		Base,
-		Space,
-		BoolParameter,
-		EnumParameter,
-		IntegerParameter,
-		ScalarParameter,
-		Vector2Parameter,
-		Vector3Parameter,
-		TransformParameter,
-	};
-	/* Entity IDs are an encoded type and index, with the upper 8 bits being the type, and the lower 24 bits as the index */
-	uint32 EncodeControlRigEntityID(int32 InIndex, EControlRigEntityType InType)
-	{
-		check(InIndex >= 0 && InIndex < int32(0x00FFFFFF));
-		return static_cast<uint32>(InIndex) | (uint32(InType) << 24);
-	}
-	void DecodeControlRigEntityID(uint32 InEntityID, int32& OutIndex, EControlRigEntityType& OutType)
-	{
-		// Mask out the type to get the index
-		OutIndex = static_cast<int32>(InEntityID & 0x00FFFFFF);
-		OutType = (EControlRigEntityType)(InEntityID >> 24);
-	}
-
-
-}// namespace UE::MovieScene
-
 void UMovieSceneControlRigParameterSection::ImportEntityImpl(UMovieSceneEntitySystemLinker* EntityLinker, const FEntityImportParams& Params, FImportedEntity* OutImportedEntity)
 {
 	using namespace UE::MovieScene;
@@ -3961,6 +4102,7 @@ void UMovieSceneControlRigParameterSection::ImportEntityImpl(UMovieSceneEntitySy
 				.Add(BuiltInComponentTypes->BoundObjectResolver, NeverResolve)
 				.Add(ControlRigComponents->ControlRigSource, ControlRigSource)
 				.Add(ControlRigComponents->BaseControlRigEvalData, FBaseControlRigEvalData{ this })
+				.AddConditional(BuiltInComponentTypes->WeightChannel, &Weight, EnumHasAnyFlags(ChannelMask, EMovieSceneTransformChannel::Weight) && Weight.HasAnyData())
 				.AddDefaulted(BuiltInComponentTypes->EvalTime)
 				.AddDefaulted(BuiltInComponentTypes->EvalSeconds)
 				.AddTag(ControlRigComponents->Tags.BaseControlRig)
@@ -4047,42 +4189,41 @@ void UMovieSceneControlRigParameterSection::ImportEntityImpl(UMovieSceneEntitySy
 			break;
 		}
 
-		case EControlRigEntityType::Vector2Parameter:
-		{
-			const FVector2DParameterNameAndCurves& Vector2D = Vector2DParameterNamesAndCurves[EntityIndex];
-			OutImportedEntity->AddBuilder(
-				FEntityBuilder()
-				.Add(ControlRigComponents->ControlRigSource, ControlRigSource)
-				.Add(TracksComponentTypes->GenericParameterName, Vector2D.ParameterName)
-				.AddTag(ControlRigComponents->Tags.ControlRigParameter)
-				.AddTag(TracksComponentTypes->Parameters.Vector3.PropertyTag)
-				.AddConditional(BuiltInComponentTypes->FloatChannel[0], &Vector2D.XCurve, Vector2D.XCurve.HasAnyData())
-				.AddConditional(BuiltInComponentTypes->FloatChannel[1], &Vector2D.YCurve, Vector2D.YCurve.HasAnyData())
-				.AddConditional(BuiltInComponentTypes->WeightChannel, &Weight, EnumHasAnyFlags(ChannelMask, EMovieSceneTransformChannel::Weight) && Weight.HasAnyData())
-				.AddTagConditional(BuiltInComponentTypes->Tags.ExternalBlending, bExternalBlending)
-				.AddMutualComponents()
-			);
-			break;
-		}
-		
-		case EControlRigEntityType::Vector3Parameter:
-		{
-			const FVectorParameterNameAndCurves& Vector = VectorParameterNamesAndCurves[EntityIndex];
-			OutImportedEntity->AddBuilder(
-				FEntityBuilder()
-				.Add(ControlRigComponents->ControlRigSource, ControlRigSource)
-				.Add(TracksComponentTypes->GenericParameterName, Vector.ParameterName)
-				.AddTag(ControlRigComponents->Tags.ControlRigParameter)
-				.AddTag(TracksComponentTypes->Parameters.Vector3.PropertyTag)
-				.AddConditional(BuiltInComponentTypes->FloatChannel[0], &Vector.XCurve, Vector.XCurve.HasAnyData())
-				.AddConditional(BuiltInComponentTypes->FloatChannel[1], &Vector.YCurve, Vector.YCurve.HasAnyData())
-				.AddConditional(BuiltInComponentTypes->FloatChannel[2], &Vector.ZCurve, Vector.ZCurve.HasAnyData())
-				.AddConditional(BuiltInComponentTypes->WeightChannel, &Weight, EnumHasAnyFlags(ChannelMask, EMovieSceneTransformChannel::Weight) && Weight.HasAnyData())
-				.AddTagConditional(BuiltInComponentTypes->Tags.ExternalBlending, bExternalBlending)
-				.AddMutualComponents()
-			);
+		case EControlRigEntityType::VectorParameter:
+			if (EntityIndex < Vector2DParameterNamesAndCurves.Num())
+			{
+				const FVector2DParameterNameAndCurves& Vector2D = Vector2DParameterNamesAndCurves[EntityIndex];
+				OutImportedEntity->AddBuilder(
+					FEntityBuilder()
+					.Add(ControlRigComponents->ControlRigSource, ControlRigSource)
+					.Add(TracksComponentTypes->GenericParameterName, Vector2D.ParameterName)
+					.AddTag(ControlRigComponents->Tags.ControlRigParameter)
+					.AddTag(TracksComponentTypes->Parameters.Vector3.PropertyTag)
+					.AddConditional(BuiltInComponentTypes->FloatChannel[0], &Vector2D.XCurve, Vector2D.XCurve.HasAnyData())
+					.AddConditional(BuiltInComponentTypes->FloatChannel[1], &Vector2D.YCurve, Vector2D.YCurve.HasAnyData())
+					.AddConditional(BuiltInComponentTypes->WeightChannel, &Weight, EnumHasAnyFlags(ChannelMask, EMovieSceneTransformChannel::Weight) && Weight.HasAnyData())
+					.AddTagConditional(BuiltInComponentTypes->Tags.ExternalBlending, bExternalBlending)
+					.AddMutualComponents()
+				);
+			}
+			else
+			{
+				const FVectorParameterNameAndCurves& Vector = VectorParameterNamesAndCurves[EntityIndex - Vector2DParameterNamesAndCurves.Num()];
+				OutImportedEntity->AddBuilder(
+					FEntityBuilder()
+					.Add(ControlRigComponents->ControlRigSource, ControlRigSource)
+					.Add(TracksComponentTypes->GenericParameterName, Vector.ParameterName)
+					.AddTag(ControlRigComponents->Tags.ControlRigParameter)
+					.AddTag(TracksComponentTypes->Parameters.Vector3.PropertyTag)
+					.AddConditional(BuiltInComponentTypes->FloatChannel[0], &Vector.XCurve, Vector.XCurve.HasAnyData())
+					.AddConditional(BuiltInComponentTypes->FloatChannel[1], &Vector.YCurve, Vector.YCurve.HasAnyData())
+					.AddConditional(BuiltInComponentTypes->FloatChannel[2], &Vector.ZCurve, Vector.ZCurve.HasAnyData())
+					.AddConditional(BuiltInComponentTypes->WeightChannel, &Weight, EnumHasAnyFlags(ChannelMask, EMovieSceneTransformChannel::Weight) && Weight.HasAnyData())
+					.AddTagConditional(BuiltInComponentTypes->Tags.ExternalBlending, bExternalBlending)
+					.AddMutualComponents()
+				);
+			}
 			break;
-		}
 
 		case EControlRigEntityType::TransformParameter:
 		{
@@ -4147,9 +4288,7 @@ bool UMovieSceneControlRigParameterSection::PopulateEvaluationFieldImpl(const TR
 	for (int32 Index = 0; Index < NumBoolID; ++Index)
 	{
 		const FBoolParameterNameAndCurve& Bool = BoolParameterNamesAndCurves[Index];
-
-		const FChannelMapInfo* ChannelInfo = ControlChannelMap.Find(Bool.ParameterName);
-		if (!ChannelInfo || (ControlsMask.IsValidIndex(ChannelInfo->MaskIndex) && !ControlsMask[ChannelInfo->MaskIndex]))
+		if (ControlNameMask.Contains(Bool.ParameterName))
 		{
 			continue;
 		}
@@ -4164,9 +4303,7 @@ bool UMovieSceneControlRigParameterSection::PopulateEvaluationFieldImpl(const TR
 	for (int32 Index = 0; Index < NumEnumID; ++Index)
 	{
 		const FEnumParameterNameAndCurve& Enum = EnumParameterNamesAndCurves[Index];
-
-		const FChannelMapInfo* ChannelInfo = ControlChannelMap.Find(Enum.ParameterName);
-		if (!ChannelInfo || (ControlsMask.IsValidIndex(ChannelInfo->MaskIndex) && !ControlsMask[ChannelInfo->MaskIndex]))
+		if (ControlNameMask.Contains(Enum.ParameterName))
 		{
 			continue;
 		}
@@ -4181,9 +4318,7 @@ bool UMovieSceneControlRigParameterSection::PopulateEvaluationFieldImpl(const TR
 	for (int32 Index = 0; Index < NumIntegerID; ++Index)
 	{
 		const FIntegerParameterNameAndCurve& Integer = IntegerParameterNamesAndCurves[Index];
-
-		const FChannelMapInfo* ChannelInfo = ControlChannelMap.Find(Integer.ParameterName);
-		if (!ChannelInfo || (ControlsMask.IsValidIndex(ChannelInfo->MaskIndex) && !ControlsMask[ChannelInfo->MaskIndex]))
+		if (ControlNameMask.Contains(Integer.ParameterName))
 		{
 			continue;
 		}
@@ -4198,9 +4333,7 @@ bool UMovieSceneControlRigParameterSection::PopulateEvaluationFieldImpl(const TR
 	for (int32 Index = 0; Index < NumScalarID; ++Index)
 	{
 		const FScalarParameterNameAndCurve& Scalar = ScalarParameterNamesAndCurves[Index];
-
-		const FChannelMapInfo* ChannelInfo = ControlChannelMap.Find(Scalar.ParameterName);
-		if (!ChannelInfo || (ControlsMask.IsValidIndex(ChannelInfo->MaskIndex) && !ControlsMask[ChannelInfo->MaskIndex]))
+		if (ControlNameMask.Contains(Scalar.ParameterName))
 		{
 			continue;
 		}
@@ -4215,16 +4348,14 @@ bool UMovieSceneControlRigParameterSection::PopulateEvaluationFieldImpl(const TR
 	for (int32 Index = 0; Index < NumVector2ID; ++Index)
 	{
 		const FVector2DParameterNameAndCurves& Vector2D = Vector2DParameterNamesAndCurves[Index];
-
-		const FChannelMapInfo* ChannelInfo = ControlChannelMap.Find(Vector2D.ParameterName);
-		if (!ChannelInfo || (ControlsMask.IsValidIndex(ChannelInfo->MaskIndex) && !ControlsMask[ChannelInfo->MaskIndex]))
+		if (ControlNameMask.Contains(Vector2D.ParameterName))
 		{
 			continue;
 		}
 
 		if (Vector2D.XCurve.HasAnyData() || Vector2D.YCurve.HasAnyData())
 		{
-			const int32 EntityIndex = OutFieldBuilder->FindOrAddEntity(this, EncodeControlRigEntityID(Index, EControlRigEntityType::Vector2Parameter));
+			const int32 EntityIndex = OutFieldBuilder->FindOrAddEntity(this, EncodeControlRigEntityID(Index, EControlRigEntityType::VectorParameter));
 			OutFieldBuilder->AddPersistentEntity(EffectiveRange, EntityIndex, MetaDataIndex);
 		}
 	}
@@ -4232,16 +4363,14 @@ bool UMovieSceneControlRigParameterSection::PopulateEvaluationFieldImpl(const TR
 	for (int32 Index = 0; Index < NumVector3ID; ++Index)
 	{
 		const FVectorParameterNameAndCurves& Vector = VectorParameterNamesAndCurves[Index];
-
-		const FChannelMapInfo* ChannelInfo = ControlChannelMap.Find(Vector.ParameterName);
-		if (!ChannelInfo || (ControlsMask.IsValidIndex(ChannelInfo->MaskIndex) && !ControlsMask[ChannelInfo->MaskIndex]))
+		if (ControlNameMask.Contains(Vector.ParameterName))
 		{
 			continue;
 		}
 
 		if (Vector.XCurve.HasAnyData() || Vector.YCurve.HasAnyData() || Vector.ZCurve.HasAnyData())
 		{
-			const int32 EntityIndex = OutFieldBuilder->FindOrAddEntity(this, EncodeControlRigEntityID(Index, EControlRigEntityType::Vector3Parameter));
+			const int32 EntityIndex = OutFieldBuilder->FindOrAddEntity(this, EncodeControlRigEntityID(Index + Vector2DParameterNamesAndCurves.Num(), EControlRigEntityType::VectorParameter));
 			OutFieldBuilder->AddPersistentEntity(EffectiveRange, EntityIndex, MetaDataIndex);
 		}
 	}
@@ -4249,9 +4378,7 @@ bool UMovieSceneControlRigParameterSection::PopulateEvaluationFieldImpl(const TR
 	for (int32 Index = 0; Index < NumTransformID; ++Index)
 	{
 		const FTransformParameterNameAndCurves& Transform = TransformParameterNamesAndCurves[Index];
-
-		const FChannelMapInfo* ChannelInfo = ControlChannelMap.Find(Transform.ParameterName);
-		if (!ChannelInfo || (ControlsMask.IsValidIndex(ChannelInfo->MaskIndex) && !ControlsMask[ChannelInfo->MaskIndex]))
+		if (ControlNameMask.Contains(Transform.ParameterName))
 		{
 			continue;
 		}
diff --git a/Engine/Plugins/Animation/ControlRig/Source/ControlRig/Private/Sequencer/MovieSceneControlRigSystem.cpp b/Engine/Plugins/Animation/ControlRig/Source/ControlRig/Private/Sequencer/MovieSceneControlRigSystem.cpp
--- a/Engine/Plugins/Animation/ControlRig/Source/ControlRig/Private/Sequencer/MovieSceneControlRigSystem.cpp
+++ b/Engine/Plugins/Animation/ControlRig/Source/ControlRig/Private/Sequencer/MovieSceneControlRigSystem.cpp
@@ -22,6 +22,7 @@
 #include "Systems/MovieScenePiecewiseDoubleBlenderSystem.h"
 #include "Systems/MovieSceneInitialValueSystem.h"
 #include "Systems/WeightAndEasingEvaluatorSystem.h"
+#include "Systems/MovieSceneSkeletalAnimationSystem.h"
 #include "EntitySystem/MovieSceneEntityGroupingSystem.h"
 #include "ControlRig.h"
 #include "Components/SkeletalMeshComponent.h"
@@ -1105,9 +1106,15 @@ struct FEvaluateBaseControlRigs
 		InOutBaseEvalData.bWasDoNotKey = InOutBaseEvalData.Section->GetDoNotKey();
 		InOutBaseEvalData.Section->SetDoNotKey(false);
 
+		UObject* BoundObject = ControlRig->GetObjectBinding()->GetBoundObject();
+		if (!BoundObject)
+		{
+			return;
+		}
+
 		if (InOutBaseEvalData.bIsActive)
 		{
-			if (USkeletalMeshComponent* SkeletalMeshComponent = Cast<USkeletalMeshComponent>(ControlRig->GetObjectBinding()->GetBoundObject()))
+			if (USkeletalMeshComponent* SkeletalMeshComponent = Cast<USkeletalMeshComponent>(BoundObject))
 			{
 				if (UControlRigLayerInstance* AnimInstance = Cast<UControlRigLayerInstance>(SkeletalMeshComponent->GetAnimInstance()))
 				{
@@ -1125,8 +1132,6 @@ struct FEvaluateBaseControlRigs
 			ControlRig->SetAbsoluteTime(EvalTimeSeconds);
 		}
 
-		UObject* BoundObject = ControlRig->GetObjectBinding()->GetBoundObject();
-
 		// when playing animation, instead of scrubbing/stepping thru frames, InTime might have a subframe of 0.999928
 		// leading to a decimal value of 24399.999928 (for example). This results in evaluating one frame less than expected
 		// (24399 instead of 24400) and leads to spaces and constraints switching parents/state after the control changes
@@ -1480,6 +1485,8 @@ UMovieSceneControlRigParameterEvaluatorSystem::UMovieSceneControlRigParameterEva
 		DefineImplicitPrerequisite(UMovieSceneHierarchicalEasingInstantiatorSystem::StaticClass(), GetClass());
 		DefineImplicitPrerequisite(GetClass(), UMovieSceneInitialValueSystem::StaticClass());
 		DefineImplicitPrerequisite(GetClass(), UMovieSceneRestorePreAnimatedStateSystem::StaticClass());
+
+		DefineImplicitPrerequisite(UMovieSceneSkeletalAnimationSystem::StaticClass(), GetClass());
 	}
 }
 
diff --git a/Engine/Plugins/Animation/ControlRig/Source/ControlRig/Public/Sequencer/MovieSceneControlRigParameterBuffer.h b/Engine/Plugins/Animation/ControlRig/Source/ControlRig/Public/Sequencer/MovieSceneControlRigParameterBuffer.h
--- a/Engine/Plugins/Animation/ControlRig/Source/ControlRig/Public/Sequencer/MovieSceneControlRigParameterBuffer.h
+++ b/Engine/Plugins/Animation/ControlRig/Source/ControlRig/Public/Sequencer/MovieSceneControlRigParameterBuffer.h
@@ -168,6 +168,58 @@ private:
 	EControlRigControlType Type;
 };
 
+/**
+ * Reflects a single value within a FControlRigParameterValues structure, allowing for type-safe casting and checking.
+ * A FControlRigValueView is only valid as long its owning header remains unmodified. Any insertion, removal or sorting
+ * will invalidate all views.
+ */
+struct FControlRigValueView
+{
+	FControlRigValueView()
+		: Value(nullptr)
+		, Type(EControlRigControlType::Num)
+	{}
+	FControlRigValueView(void* InValue, EControlRigControlType InType)
+		: Value(InValue)
+		, Type(InType)
+	{}
+
+	/**
+	 * Test to see if this value is valid
+	 */
+	explicit operator bool() const
+	{
+		return Value && Type != EControlRigControlType::Num;
+	}
+
+	/**
+	 * Cast this view to a specific value, returning nullptr if the cast is invalid.
+	 */
+	template<typename T>
+	T* Cast() const
+	{
+		if (Type == Private::ControlRigTypeFromNativeType<T>())
+		{
+			return static_cast<T*>(Value);
+		}
+		return nullptr;
+	}
+
+	/**
+	 * Get the type of this value. Must be checked for valididity first.
+	 */
+	EControlRigControlType GetType() const
+	{
+		check(*this);
+		return Type;
+	}
+
+private:
+
+	void* Value;
+	EControlRigControlType Type;
+};
+
 
 /**
  * Parameter buffer structure containing any number of control values arranged by type and name
@@ -218,6 +270,9 @@ struct CONTROLRIG_API FControlRigParameterValues
 		return false;
 	}
 
+	FControlRigValueView Find(FName Name) const;
+	FControlRigValueView FindParameter(FName Name) const;
+
 	~FControlRigParameterValues();
 
 	void CopyFrom(const FControlRigParameterValues& Other, FName ControlName);
diff --git a/Engine/Plugins/Animation/ControlRig/Source/ControlRig/Public/Sequencer/MovieSceneControlRigParameterSection.h b/Engine/Plugins/Animation/ControlRig/Source/ControlRig/Public/Sequencer/MovieSceneControlRigParameterSection.h
--- a/Engine/Plugins/Animation/ControlRig/Source/ControlRig/Public/Sequencer/MovieSceneControlRigParameterSection.h
+++ b/Engine/Plugins/Animation/ControlRig/Source/ControlRig/Public/Sequencer/MovieSceneControlRigParameterSection.h
@@ -27,6 +27,48 @@ class UAnimSequence;
 class USkeletalMeshComponent;
 struct FKeyDataOptimizationParams;
 
+namespace UE::MovieScene
+{
+
+	enum class EControlRigControlType : uint8;
+
+	struct FControlRigChannelMetaData
+	{
+		CONTROLRIG_API FControlRigChannelMetaData();
+		CONTROLRIG_API FControlRigChannelMetaData(EControlRigControlType InType, FName InControlName, int32 InIndexWithinControl, uint32 InEntitySystemID);
+
+		CONTROLRIG_API explicit operator bool() const;
+
+		EControlRigControlType GetType() const
+		{
+			check((bool)*this);
+			return Type;
+		}
+
+		FName GetControlName() const
+		{
+			return ControlName;
+		}
+
+		int32 GetChannelIndex() const
+		{
+			return IndexWithinControl;
+		}
+
+		uint32 GetEntitySystemID() const
+		{
+			return EntitySystemID;
+		}
+
+	private:
+		EControlRigControlType Type;
+		FName ControlName;
+		int32 IndexWithinControl;
+		uint32 EntitySystemID;
+	};
+
+} // namespace UE::MovieScene
+
 struct CONTROLRIG_API FControlRigBindingHelper
 {
 	static void BindToSequencerInstance(UControlRig* ControlRig);
@@ -541,7 +583,16 @@ public:
 	 static FMovieSceneInterrogationKey GetVectorInterrogationKey();
 	 static FMovieSceneInterrogationKey GetTransformInterrogationKey();
 
+	/**
+	 * Retrieve meta-data pertaining to a given channel ptr including the control it animates and its index within the control.
+	 * 
+	 * @param Channel A pointer to a channel
+	 * @return The channel's meta-data. Must be checked for validity before use
+	 */
+	UE::MovieScene::FControlRigChannelMetaData GetChannelMetaData(const FMovieSceneChannel* Channel) const;
+
 	virtual void ReconstructChannelProxy() override;
+	virtual float GetTotalWeightValue(FFrameTime InTime) const override;
 
 protected:
 
@@ -551,7 +602,6 @@ protected:
 	virtual void Serialize(FArchive& Ar) override;
 	virtual void PostEditImport() override;
 	virtual void PostLoad() override;
-	virtual float GetTotalWeightValue(FFrameTime InTime) const override;
 	virtual void OnBindingIDsUpdated(const TMap<UE::MovieScene::FFixedObjectBindingID, UE::MovieScene::FFixedObjectBindingID>& OldFixedToNewFixedMap, FMovieSceneSequenceID LocalSequenceID, TSharedRef<UE::MovieScene::FSharedPlaybackState> SharedPlaybackState) override;
 	virtual void GetReferencedBindings(TArray<FGuid>& OutBindings) override;
 	virtual void PreSave(FObjectPreSaveContext SaveContext) override;
diff --git a/Engine/Plugins/Animation/ControlRig/Source/ControlRigEditor/Private/Sequencer/ControlRigParameterTrackEditor.cpp b/Engine/Plugins/Animation/ControlRig/Source/ControlRigEditor/Private/Sequencer/ControlRigParameterTrackEditor.cpp
--- a/Engine/Plugins/Animation/ControlRig/Source/ControlRigEditor/Private/Sequencer/ControlRigParameterTrackEditor.cpp
+++ b/Engine/Plugins/Animation/ControlRig/Source/ControlRigEditor/Private/Sequencer/ControlRigParameterTrackEditor.cpp
@@ -12,6 +12,8 @@
 #include "MVVM/ViewModels/OutlinerViewModel.h"
 #include "MVVM/Selection/Selection.h"
 #include "EntitySystem/Interrogation/MovieSceneInterrogationLinker.h"
+#include "EntitySystem/MovieSceneDecompositionQuery.h"
+#include "Systems/MovieScenePiecewiseDoubleBlenderSystem.h"
 #include "Animation/AnimMontage.h"
 #include "Sequencer/MovieSceneControlRigParameterTrack.h"
 #include "Sequencer/MovieSceneControlRigParameterSection.h"
@@ -4518,6 +4520,336 @@ bool FControlRigParameterTrackEditor::ModifyOurGeneratedKeysByCurrentAndWeight(U
 	return true;
 }
 
+void FControlRigParameterTrackEditor::ProcessKeyOperation(FFrameNumber InKeyTime, const UE::Sequencer::FKeyOperation& Operation, ISequencer& InSequencer, TArray<UE::Sequencer::FAddKeyResult>* OutResults)
+{
+	using namespace UE::Sequencer;
+
+	if (UMovieSceneControlRigParameterTrack::ShouldUseLegacyTemplate())
+	{
+		// Legacy behavior uses all the legacy mechanisms
+		return FKeyframeTrackEditor<UMovieSceneControlRigParameterTrack>::ProcessKeyOperation(InKeyTime, Operation, InSequencer, OutResults);
+	}
+
+	// ECS system performs a full recomposition on the ECS data
+	auto Iterator = [this, InKeyTime, &Operation, &InSequencer, OutResults](UMovieSceneTrack* Track, TArrayView<const UE::Sequencer::FKeySectionOperation> Operations)
+	{
+		UMovieSceneControlRigParameterTrack* ControlRigTrack = CastChecked<UMovieSceneControlRigParameterTrack>(Track);
+
+		FGuid ObjectBinding = Track->FindObjectBindingGuid();
+		if (ObjectBinding.IsValid())
+		{
+			for (TWeakObjectPtr<> WeakObject : InSequencer.FindBoundObjects(ObjectBinding, InSequencer.GetFocusedTemplateID()))
+			{
+				if (UObject* Object = WeakObject.Get())
+				{
+					this->ProcessKeyOperation(Object, ControlRigTrack, Operations, InSequencer, InKeyTime, OutResults);
+					return;
+				}
+			}
+		}
+
+		// Default behavior
+		FKeyOperation::ApplyOperations(InKeyTime, Operations, ObjectBinding, InSequencer, OutResults);
+	};
+
+	Operation.IterateOperations(Iterator);
+}
+
+void FControlRigParameterTrackEditor::ProcessKeyOperation(UObject* ObjectToKey, UMovieSceneControlRigParameterTrack* Track, TArrayView<const UE::Sequencer::FKeySectionOperation> SectionsToKey, ISequencer& InSequencer, FFrameNumber KeyTime, TArray<UE::Sequencer::FAddKeyResult>* OutResults)
+{
+	using namespace UE::MovieScene;
+	using namespace UE::Sequencer;
+
+	// ~~~~~~~~~~~~~~~~~~~~~~~~~~~~~~~~~~~~~~~~~~~~~~~~~~~~~~~~~~~~~~~~~~~~~~~~~~~~~~~~~~~~~~~~~~~
+	// @todo: This should really be unified with AddControlKeys and ModifyOurGeneratedKeysByCurrentAndWeight
+	//        so that everything goes through the common FControlRigParameterValues container, but to do so 
+	//        we need to port some additional logic around constraint space and a few other pieces.
+	//
+	//        From there constructing an FGeneratedKeys structure should be routine, and should allow us
+	//        to remove FChannelMapInfo::GeneratedKeyIndex intogether
+	// ~~~~~~~~~~~~~~~~~~~~~~~~~~~~~~~~~~~~~~~~~~~~~~~~~~~~~~~~~~~~~~~~~~~~~~~~~~~~~~~~~~~~~~~~~~~
+
+	if (!ensure(Track))
+	{
+		return;
+	}
+
+	const EMovieSceneKeyInterpolation InterpMode = InSequencer.GetSequencerSettings()->GetKeyInterpolation();
+
+	// Iterate each section and perform the key operation on recomposed values
+	for (const FKeySectionOperation& Operation : SectionsToKey)
+	{
+		UMovieSceneControlRigParameterSection* ControlRigSection = Cast<UMovieSceneControlRigParameterSection>(Operation.Section->GetSectionObject());
+		ControlRigSection->Modify();
+
+		FControlRigParameterValues RecomposedValues = GetRecomposedControlValues(Track, ControlRigSection, KeyTime);
+
+		for (TSharedPtr<IKeyArea> KeyArea : Operation.KeyAreas)
+		{
+			const FMovieSceneChannelHandle& ChannelHandle = KeyArea->GetChannel();
+			FMovieSceneChannel*             Channel       = KeyArea->ResolveChannel();
+
+			if (Channel == &ControlRigSection->Weight)
+			{
+				FKeyHandle KeyHandle = KeyArea->AddOrUpdateKey(KeyTime, FGuid(), InSequencer);
+				if (OutResults)
+				{
+					FAddKeyResult Result;
+					Result.KeyArea = KeyArea;
+					Result.KeyHandle = KeyHandle;
+					OutResults->Add(Result);
+				}
+				continue;
+			}
+
+			FControlRigChannelMetaData ChannelMetaData = ControlRigSection->GetChannelMetaData(Channel);
+			if (!ChannelMetaData || !ControlRigSection->GetControlNameMask(ChannelMetaData.GetControlName()))
+			{
+				continue;
+			}
+
+			FControlRigValueView DesiredValue = RecomposedValues.Find(ChannelMetaData.GetControlName());
+			if (DesiredValue.GetType() == EControlRigControlType::Space)
+			{
+
+			}
+			else if (DesiredValue.GetType() == EControlRigControlType::Parameter_Bool)
+			{
+				FMovieSceneBoolChannel* Bool = static_cast<FMovieSceneBoolChannel*>(Channel);
+				Bool->GetData().UpdateOrAddKey(KeyTime, *DesiredValue.Cast<bool>());
+			}
+			else if (DesiredValue.GetType() == EControlRigControlType::Parameter_Enum)
+			{
+				FMovieSceneByteChannel* Byte = static_cast<FMovieSceneByteChannel*>(Channel);
+				Byte->GetData().UpdateOrAddKey(KeyTime, *DesiredValue.Cast<uint8>());
+			}
+			else if (DesiredValue.GetType() == EControlRigControlType::Parameter_Integer)
+			{
+				FMovieSceneIntegerChannel* Int = static_cast<FMovieSceneIntegerChannel*>(Channel);
+				Int->GetData().UpdateOrAddKey(KeyTime, *DesiredValue.Cast<int32>());
+			}
+			else 
+			{
+				const int32 ControlChannelIndex = ChannelMetaData.GetChannelIndex();
+
+				float NewValue = 0.f;
+				if (DesiredValue.GetType() == EControlRigControlType::Parameter_Scalar)
+				{
+					NewValue = *DesiredValue.Cast<float>();
+				}
+				else if (DesiredValue.GetType() == EControlRigControlType::Parameter_Vector)
+				{
+					NewValue = (*DesiredValue.Cast<FVector3f>())[ControlChannelIndex];
+				}
+				else if (DesiredValue.GetType() == EControlRigControlType::Parameter_Transform)
+				{
+					FEulerTransform Transform = (*DesiredValue.Cast<FEulerTransform>());
+					if (ControlChannelIndex < 3)
+					{
+						NewValue = Transform.Location[ControlChannelIndex];
+					}
+					else if (ControlChannelIndex < 6)
+					{
+						NewValue = Transform.Rotation.Euler()[ControlChannelIndex - 3];
+					}
+					else
+					{
+						NewValue = Transform.Scale[ControlChannelIndex - 6];
+					}
+				}
+
+				FMovieSceneFloatChannel* FloatChannel = static_cast<FMovieSceneFloatChannel*>(Channel);
+
+				int32 KeyIndex = FloatChannel->GetData().FindKey(KeyTime);
+				if (KeyIndex == INDEX_NONE)
+				{
+					switch (InterpMode)
+					{
+					case EMovieSceneKeyInterpolation::Linear:
+						KeyIndex = FloatChannel->AddLinearKey(KeyTime, NewValue);
+						break;
+					case EMovieSceneKeyInterpolation::Constant:
+						KeyIndex = FloatChannel->AddConstantKey(KeyTime, NewValue);
+						break;
+					case EMovieSceneKeyInterpolation::Auto:
+						KeyIndex = FloatChannel->AddCubicKey(KeyTime, NewValue, ERichCurveTangentMode::RCTM_Auto);
+						break;
+					case EMovieSceneKeyInterpolation::SmartAuto:
+					default:
+						KeyIndex = FloatChannel->AddCubicKey(KeyTime, NewValue, ERichCurveTangentMode::RCTM_SmartAuto);
+						break;
+					}
+				}
+				else
+				{
+					FloatChannel->GetData().GetValues()[KeyIndex].Value = NewValue;
+				}
+
+				if (OutResults && KeyIndex != INDEX_NONE)
+				{
+					FAddKeyResult Result;
+					Result.KeyArea = KeyArea;
+					Result.KeyHandle = FloatChannel->GetData().GetHandle(KeyIndex);
+					OutResults->Add(Result);
+				}
+			}
+		}
+	}
+}
+
+UE::MovieScene::FControlRigParameterValues FControlRigParameterTrackEditor::GetRecomposedControlValues(UMovieSceneControlRigParameterTrack* Track, UMovieSceneControlRigParameterSection* Section, FFrameNumber KeyTime)
+{
+	using namespace UE::MovieScene;
+	using namespace UE::Sequencer;
+
+	check(Track && Section);
+
+	FControlRigParameterValues CurrentValues(EControlRigParameterBufferIndexStability::Unstable);
+	CurrentValues.PopulateFrom(Track->GetControlRig());
+
+	// Create the interrogator.
+	FSystemInterrogator Interrogator;
+	Interrogator.TrackImportedEntities(true);
+
+	TGuardValue<FEntityManager*> DebugVizGuard(GEntityManagerForDebuggingVisualizers, &Interrogator.GetLinker()->EntityManager);
+
+	// Run an interrogation on the track at the specified time.
+	FInterrogationKey InterrogationKey(FInterrogationKey::Default());
+	FInterrogationChannel InterrogationChannel = Interrogator.AllocateChannel(Track->GetControlRig(), FMovieScenePropertyBinding());
+	InterrogationKey.Channel = InterrogationChannel;
+	Interrogator.ImportTrack(Track, InterrogationChannel);
+
+	Interrogator.AddInterrogation(KeyTime);
+	Interrogator.Update();
+
+	UMovieScenePiecewiseDoubleBlenderSystem* BlenderSystem = Interrogator.GetLinker()->FindSystem<UMovieScenePiecewiseDoubleBlenderSystem>();
+	if (!BlenderSystem)
+	{
+		return CurrentValues;
+	}
+
+	Interrogator.GetLinker()->EntityManager.LockDown();
+	ON_SCOPE_EXIT
+	{
+		Interrogator.GetLinker()->EntityManager.ReleaseLockDown();
+	};
+
+	FBuiltInComponentTypes*          BuiltInComponents = FBuiltInComponentTypes::Get();
+	FMovieSceneTracksComponentTypes* TracksComponents  = FMovieSceneTracksComponentTypes::Get();
+
+	const float SectionWeight = Section->GetTotalWeightValue(KeyTime);
+
+	auto RecomposeValue = [&Interrogator, InterrogationKey, BlenderSystem, BuiltInComponents](FMovieSceneEntityID EntityID, double Value, const double* InitialValue, FComponentTypeID ParameterTag, TComponentTypeID<double> ResultComponent, const FControlRigChannelMetaData& ChannelMetaData)
+	{
+		check(BlenderSystem);
+		TOptionalComponentReader<FMovieSceneBlendChannelID> BlendChannelInput = Interrogator.GetLinker()->EntityManager.ReadComponent(EntityID, BuiltInComponents->BlendChannelInput);
+		if (BlendChannelInput)
+		{
+			FAlignedDecomposedValue AlignedOutput;
+
+			FValueDecompositionParams Params;
+			Params.Query.Entities = MakeArrayView(&EntityID, 1);
+			Params.Query.bConvertFromSourceEntityIDs = false;
+			Params.DecomposeBlendChannel = BlendChannelInput->ChannelID;
+			Params.ResultComponentType = ResultComponent;
+			Params.PropertyTag = ParameterTag;
+
+			FGraphEventRef Task = BlenderSystem->DispatchDecomposeTask(Params, &AlignedOutput);
+			if (Task)
+			{
+				FTaskGraphInterface::Get().WaitUntilTaskCompletes(Task, ENamedThreads::GameThread);
+			}
+
+			return AlignedOutput.Value.Recompose(EntityID, Value, InitialValue);
+		}
+		return Value;
+	};
+
+
+	TPropertyValueStorage<FFloatParameterTraits>*     InitialScalarStorage    = FInitialValueCache::GetGlobalInitialValues()->FindStorage(TracksComponents->Parameters.Scalar);
+	TPropertyValueStorage<FVector3ParameterTraits>*   InitialVectorStorage    = FInitialValueCache::GetGlobalInitialValues()->FindStorage(TracksComponents->Parameters.Vector3);
+	TPropertyValueStorage<FTransformParameterTraits>* InitialTransformStorage = FInitialValueCache::GetGlobalInitialValues()->FindStorage(TracksComponents->Parameters.Transform);
+
+	// ------------------------------------------------------------------------------------------------------
+	// Recompose Scalars
+	for (const FScalarParameterNameAndCurve& Scalar : Section->GetScalarParameterNamesAndCurves())
+	{
+		FControlRigValueView Value = CurrentValues.Find(Scalar.ParameterName);
+		if (Value && Value.GetType() == EControlRigControlType::Parameter_Scalar)
+		{
+			FControlRigChannelMetaData ChannelMetaData = Section->GetChannelMetaData(&Scalar.ParameterCurve);
+			check(ChannelMetaData);
+
+			FMovieSceneEntityID Entity = Interrogator.FindEntityFromOwner(InterrogationKey, Section, ChannelMetaData.GetEntitySystemID());
+			if (Entity)
+			{
+				const double* InitialValue = InitialScalarStorage ? InitialScalarStorage->FindCachedValue(Track->GetControlRig(), ChannelMetaData.GetControlName()) : nullptr;
+				*Value.Cast<float>() = (float)RecomposeValue(Entity, *Value.Cast<float>(), InitialValue, TracksComponents->Parameters.Scalar.PropertyTag, BuiltInComponents->DoubleResult[0], ChannelMetaData);
+			}
+		}
+	}
+
+	// ------------------------------------------------------------------------------------------------------
+	// Recompose Vectors
+	for (const FVectorParameterNameAndCurves& Vector : Section->GetVectorParameterNamesAndCurves())
+	{
+		FControlRigValueView Value = CurrentValues.Find(Vector.ParameterName);
+		if (Value && Value.GetType() == EControlRigControlType::Parameter_Vector)
+		{
+			FControlRigChannelMetaData ChannelMetaData = Section->GetChannelMetaData(&Vector.XCurve);
+			check(ChannelMetaData);
+
+			FMovieSceneEntityID Entity = Interrogator.FindEntityFromOwner(InterrogationKey, Section, ChannelMetaData.GetEntitySystemID());
+			if (Entity)
+			{
+				const FFloatIntermediateVector* InitialValue = InitialVectorStorage ? InitialVectorStorage->FindCachedValue(Track->GetControlRig(), ChannelMetaData.GetControlName()) : nullptr;
+				FVector3f* ValueAsVector = Value.Cast<FVector3f>();
+
+				FComponentTypeID PropertyTag = TracksComponents->Parameters.Vector3.PropertyTag;
+				ValueAsVector->X = (float)RecomposeValue(Entity, ValueAsVector->X, InitialValue ? &InitialValue->X : nullptr, PropertyTag, BuiltInComponents->DoubleResult[0], ChannelMetaData);
+				ValueAsVector->Y = (float)RecomposeValue(Entity, ValueAsVector->Y, InitialValue ? &InitialValue->Y : nullptr, PropertyTag, BuiltInComponents->DoubleResult[1], ChannelMetaData);
+				ValueAsVector->Z = (float)RecomposeValue(Entity, ValueAsVector->Z, InitialValue ? &InitialValue->Z : nullptr, PropertyTag, BuiltInComponents->DoubleResult[2], ChannelMetaData);
+			}
+		}
+	}
+
+	// ------------------------------------------------------------------------------------------------------
+	// Recompose Vectors
+	for (const FTransformParameterNameAndCurves& Transform : Section->GetTransformParameterNamesAndCurves())
+	{
+		FControlRigValueView Value = CurrentValues.Find(Transform.ParameterName);
+		if (Value && Value.GetType() == EControlRigControlType::Parameter_Transform)
+		{
+			FControlRigChannelMetaData ChannelMetaData = Section->GetChannelMetaData(&Transform.Translation[0]);
+			check(ChannelMetaData);
+
+			FMovieSceneEntityID Entity = Interrogator.FindEntityFromOwner(InterrogationKey, Section, ChannelMetaData.GetEntitySystemID());
+			if (Entity)
+			{
+				const FIntermediate3DTransform* InitialValue = InitialTransformStorage ? InitialTransformStorage->FindCachedValue(Track->GetControlRig(), ChannelMetaData.GetControlName()) : nullptr;
+
+				FEulerTransform* ValueAsTransform = Value.Cast<FEulerTransform>();
+
+				FComponentTypeID PropertyTag = TracksComponents->Parameters.Transform.PropertyTag;
+
+				ValueAsTransform->Location.X     = RecomposeValue(Entity, ValueAsTransform->Location.X,     InitialValue ? &InitialValue->T_X : nullptr, PropertyTag, BuiltInComponents->DoubleResult[0], ChannelMetaData);
+				ValueAsTransform->Location.Y     = RecomposeValue(Entity, ValueAsTransform->Location.Y,     InitialValue ? &InitialValue->T_Y : nullptr, PropertyTag, BuiltInComponents->DoubleResult[1], ChannelMetaData);
+				ValueAsTransform->Location.Z     = RecomposeValue(Entity, ValueAsTransform->Location.Z,     InitialValue ? &InitialValue->T_Z : nullptr, PropertyTag, BuiltInComponents->DoubleResult[2], ChannelMetaData);
+
+				ValueAsTransform->Rotation.Roll  = RecomposeValue(Entity, ValueAsTransform->Rotation.Roll,  InitialValue ? &InitialValue->R_X : nullptr, PropertyTag, BuiltInComponents->DoubleResult[3], ChannelMetaData);
+				ValueAsTransform->Rotation.Pitch = RecomposeValue(Entity, ValueAsTransform->Rotation.Pitch, InitialValue ? &InitialValue->R_Y : nullptr, PropertyTag, BuiltInComponents->DoubleResult[4], ChannelMetaData);
+				ValueAsTransform->Rotation.Yaw   = RecomposeValue(Entity, ValueAsTransform->Rotation.Yaw,   InitialValue ? &InitialValue->R_Z : nullptr, PropertyTag, BuiltInComponents->DoubleResult[5], ChannelMetaData);
+
+				ValueAsTransform->Scale.X        = RecomposeValue(Entity, ValueAsTransform->Scale.X,        InitialValue ? &InitialValue->S_X : nullptr, PropertyTag, BuiltInComponents->DoubleResult[6], ChannelMetaData);
+				ValueAsTransform->Scale.Y        = RecomposeValue(Entity, ValueAsTransform->Scale.Y,        InitialValue ? &InitialValue->S_Y : nullptr, PropertyTag, BuiltInComponents->DoubleResult[7], ChannelMetaData);
+				ValueAsTransform->Scale.Z        = RecomposeValue(Entity, ValueAsTransform->Scale.Z,        InitialValue ? &InitialValue->S_Z : nullptr, PropertyTag, BuiltInComponents->DoubleResult[8], ChannelMetaData);
+			}
+		}
+	}
+
+	return CurrentValues;
+}
+
 void FControlRigParameterTrackEditor::BuildTrackContextMenu(FMenuBuilder& MenuBuilder, UMovieSceneTrack* InTrack)
 {
 	bool bSectionAdded;
diff --git a/Engine/Plugins/Animation/ControlRig/Source/ControlRigEditor/Private/Sequencer/ControlRigParameterTrackEditor.h b/Engine/Plugins/Animation/ControlRig/Source/ControlRigEditor/Private/Sequencer/ControlRigParameterTrackEditor.h
--- a/Engine/Plugins/Animation/ControlRig/Source/ControlRigEditor/Private/Sequencer/ControlRigParameterTrackEditor.h
+++ b/Engine/Plugins/Animation/ControlRig/Source/ControlRigEditor/Private/Sequencer/ControlRigParameterTrackEditor.h
@@ -40,6 +40,11 @@ class ISequencer;
 class UMovieSceneTrack;
 class IStructureDetailsView;
 
+namespace UE::MovieScene
+{
+	struct FControlRigParameterValues;
+}
+
 //////////////////////////////////////////////////////////////
 /// SCollapseControlsWidget
 ///////////////////////////////////////////////////////////
@@ -119,6 +124,7 @@ public:
 	virtual void ObjectImplicitlyRemoved(UObject* InObject)  override;
 	virtual void BuildTrackContextMenu(FMenuBuilder& MenuBuilder, UMovieSceneTrack* InTrack) override;
 	virtual bool HandleAssetAdded(UObject* Asset, const FGuid& TargetObjectGuid) override;
+	virtual void ProcessKeyOperation(FFrameNumber InKeyTime, const UE::Sequencer::FKeyOperation& Operation, ISequencer& InSequencer, TArray<UE::Sequencer::FAddKeyResult>* OutResults) override;
 
 	//IMovieSceneToolsAnimationBakeHelper
 	virtual void PostEvaluation(UMovieScene* MovieScene, FFrameNumber Frame);
@@ -130,6 +136,9 @@ public:
 
 private:
 
+	void ProcessKeyOperation(UObject* ObjectToKey, UMovieSceneControlRigParameterTrack* Track, TArrayView<const UE::Sequencer::FKeySectionOperation> SectionsToKey, ISequencer& InSequencer, FFrameNumber KeyTime, TArray<UE::Sequencer::FAddKeyResult>* OutResults);
+	UE::MovieScene::FControlRigParameterValues GetRecomposedControlValues(UMovieSceneControlRigParameterTrack* Track, UMovieSceneControlRigParameterSection* Section, FFrameNumber KeyTime);
+
 	void HandleAddTrackSubMenu(FMenuBuilder& MenuBuilder, TArray<FGuid> ObjectBindings, UMovieSceneTrack* Track);
 
 	void ToggleIsAdditiveControlRig();
diff --git a/Engine/Source/Runtime/MovieScene/Public/EntitySystem/MovieSceneInitialValueCache.h b/Engine/Source/Runtime/MovieScene/Public/EntitySystem/MovieSceneInitialValueCache.h
--- a/Engine/Source/Runtime/MovieScene/Public/EntitySystem/MovieSceneInitialValueCache.h
+++ b/Engine/Source/Runtime/MovieScene/Public/EntitySystem/MovieSceneInitialValueCache.h
@@ -104,6 +104,18 @@ public:
 		return NewStorage;
 	}
 
+	/**
+	 * Retrieve the initial value storage for a given property type, creating it if necessary
+	 *
+	 * @param PropertyComponents    The property component definitions for the storage including its traits
+	 * @return Property storage for the initial values or nullptr if none exists
+	 */
+	template<typename PropertyTraits>
+	TPropertyValueStorage<PropertyTraits>* GetStorage(const TPropertyComponents<PropertyTraits>& PropertyComponents)
+	{
+		return GetStorage<PropertyTraits>(PropertyComponents.InitialValueType);
+	}
+
 
 	/**
 	 * Retrieve the initial value storage for a given initial value type.
@@ -126,6 +138,19 @@ public:
 		return nullptr;
 	}
 
+
+	/**
+	 * Retrieve the initial value storage for a given property type.
+	 *
+	 * @param PropertyComponents    The property component definitions for the storage including its traits
+	 * @return Property storage for the initial values or nullptr if none exists
+	 */
+	template<typename PropertyTraits>
+	TPropertyValueStorage<PropertyTraits>* FindStorage(const TPropertyComponents<PropertyTraits>& PropertyComponents)
+	{
+		return FindStorage<PropertyTraits>(PropertyComponents.InitialValue);
+	}
+
 private:
 
 	TMap<FComponentTypeID, TUniquePtr<IPropertyValueStorage>> StorageByComponent;
