Control Rig: Replay to support, pose, metadata, variables etc

#rnx

diff --git a/Engine/Plugins/Animation/ControlRig/Source/ControlRig/Private/ControlRigReplay.cpp b/Engine/Plugins/Animation/ControlRig/Source/ControlRig/Private/ControlRigReplay.cpp
--- a/Engine/Plugins/Animation/ControlRig/Source/ControlRig/Private/ControlRigReplay.cpp
+++ b/Engine/Plugins/Animation/ControlRig/Source/ControlRig/Private/ControlRigReplay.cpp
@@ -22,9 +22,1004 @@ bool FControlRigReplayTracks::Serialize(FArchive& InArchive)
 	{
 		return false;
 	}
+
+	if(InArchive.IsLoading())
+	{
+		SampleTrackIndex = FSampleTrackIndex(*GetContainer());
+	}
 	return true;
 }
 
+void FControlRigReplayTracks::Reset()
+{
+	FSampleTrackHost::Reset();
+	SampleTrackIndex = FSampleTrackIndex();
+	ParentIndices.Reset();
+}
+
+int32 FControlRigReplayTracks::AddTimeSample(float InTime)
+{
+	return GetContainer()->AddTimeSample(InTime);
+}
+
+int32 FControlRigReplayTracks::AddTimeSampleFromDeltaTime(float InDeltaTime)
+{
+	if(GetContainer()->GetTimeSamples().IsEmpty())
+	{
+		return GetContainer()->AddTimeSample(0.f);
+	}
+	const float LastTime = GetContainer()->GetTimeSamples().Last();
+	return GetContainer()->AddTimeSample(LastTime + InDeltaTime);
+}
+
+void FControlRigReplayTracks::GetTimeIndex(float InTime, int32& InOutTimeIndex) const
+{
+	GetContainer()->GetTimeIndex(InTime, InOutTimeIndex);
+}
+
+const TArray<float>& FControlRigReplayTracks::GetTimeSamples() const
+{
+	return GetContainer()->GetTimeSamples();
+}
+
+void FControlRigReplayTracks::StoreHierarchy(URigHierarchy* InHierarchy, const TArrayView<FRigElementKey>& InKeys, bool bStorePose,
+	bool bStoreComponents, bool bStoreMetadata)
+{
+	static const FName TopologyHashName = TEXT("TopologyHash");
+	static const FName MetadataVersionName = TEXT("MetadataVersion");
+	static const FName ParentIndicesName = TEXT("ParentIndices");
+
+	const TSharedPtr<FSampleTrack<uint32>> TopologyHashTrack = StaticCastSharedPtr<FSampleTrack<uint32>>(GetContainer()->FindOrAddTrack(TopologyHashName, FSampleTrackBase::ETrackType_Uint32));
+	const TSharedPtr<FSampleTrack<uint32>> MetadataVersionTrack = StaticCastSharedPtr<FSampleTrack<uint32>>(GetContainer()->FindOrAddTrack(MetadataVersionName, FSampleTrackBase::ETrackType_Uint32));
+	const TSharedPtr<FSampleTrack<TArray<int32>>> ParentIndicesTrack = StaticCastSharedPtr<FSampleTrack<TArray<int32>>>(GetContainer()->FindOrAddTrack(ParentIndicesName, FSampleTrackBase::ETrackType_Int32Array));
+
+	ParentIndices.SetNum(InHierarchy->Num(), EAllowShrinking::No);
+	InHierarchy->ForEach([this, InHierarchy](const FRigBaseElement* InElement) -> bool
+	{
+		ParentIndices[InElement->GetIndex()] = InHierarchy->GetIndex(InHierarchy->GetDefaultParent(InElement->GetKey()));
+		return true;
+	});
+
+	TopologyHashTrack->AddSample(InHierarchy->GetTopologyHash());
+	MetadataVersionTrack->AddSample(InHierarchy->GetMetadataTagVersion());
+	ParentIndicesTrack->AddSample(ParentIndices);
+
+	ForEachElement(InHierarchy, InKeys, [this, InHierarchy, bStorePose, bStoreComponents, bStoreMetadata](FRigBaseElement* InElement, bool& bSuccess)
+	{
+		if(bStorePose)
+		{
+			StorePose(InHierarchy, InElement);
+		}
+		if(bStoreComponents)
+		{
+			StoreComponents(InHierarchy, InElement);
+		}
+		if(bStoreMetadata)
+		{
+			StoreMetaData(InHierarchy, InElement);
+		}
+	});
+}
+
+bool FControlRigReplayTracks::RestoreHierarchy(int32 InTimeIndex, URigHierarchy* InHierarchy, const TArrayView<FRigElementKey>& InKeys,
+	TReportFunction InReportFunction, bool bRestorePose, bool bRestoreComponents, bool bRestoreMetadata) const
+{
+	return ForEachElement(InHierarchy, InKeys,
+	[this, InTimeIndex, InHierarchy, InReportFunction, bRestorePose, bRestoreComponents, bRestoreMetadata](FRigBaseElement* InElement, bool& bSuccess)
+	{
+		if(bRestorePose)
+		{
+			if(!RestorePose(InTimeIndex, InHierarchy, InElement, InReportFunction))
+			{
+				bSuccess = false;
+			}
+		}
+		if(bRestoreComponents)
+		{
+			if(!RestoreComponents(InTimeIndex, InHierarchy, InElement, InReportFunction))
+			{
+				bSuccess = false;
+			}
+		}
+		if(bRestoreMetadata)
+		{
+			if(!RestoreMetaData(InTimeIndex, InHierarchy, InElement, InReportFunction))
+			{
+				bSuccess = false;
+			}
+		}
+	});
+}
+
+void FControlRigReplayTracks::StorePose(URigHierarchy* InHierarchy, const TArrayView<FRigElementKey>& InKeys)
+{
+	ForEachElement(InHierarchy, InKeys, [this, InHierarchy](FRigBaseElement* InElement, bool& bSuccess)
+	{
+		StorePose(InHierarchy, InElement);
+	});
+}
+
+void FControlRigReplayTracks::StorePose(URigHierarchy* InHierarchy, FRigBaseElement* InElement)
+{
+	if(!InElement->IsA<FRigTransformElement>() && !InElement->IsA<FRigCurveElement>())
+	{
+		return;
+	}
+
+	FSampleTrackContainer* Storage = GetContainer();
+	const FName& TrackName = GetTrackName(InElement->GetKey());
+	if(FRigCurveElement* Curve = Cast<FRigCurveElement>(InElement))
+	{
+		const TSharedPtr<FSampleTrack<float>> Track = StaticCastSharedPtr<FSampleTrack<float>>(Storage->FindOrAddTrack(TrackName, FSampleTrackBase::ETrackType_Float));
+		Track->AddSample(InHierarchy->GetCurveValue(Curve));
+		return;
+	}
+	
+	TSharedPtr<FSampleTrack<FTransform3f>> Track = Storage->FindTrack<FTransform3f>(TrackName);
+	if(!Track.IsValid())
+	{
+		Track = Storage->AddTransformTrack(TrackName);
+		SampleTrackIndex.Update(*Storage);
+	}
+	const FTransform LocalTransform = InHierarchy->GetLocalTransform(InElement->GetIndex());
+	Track->AddSample(FTransform3f(LocalTransform));
+}
+
+bool FControlRigReplayTracks::RestorePose(int32 InTimeIndex, URigHierarchy* InHierarchy, const TArrayView<FRigElementKey>& InKeys, TReportFunction InReportFunction) const
+{
+	return ForEachElement(InHierarchy, InKeys, [this, InTimeIndex, InHierarchy, InReportFunction](FRigBaseElement* InElement, bool& bSuccess)
+	{
+		if(!RestorePose(InTimeIndex, InHierarchy, InElement, InReportFunction))
+		{
+			bSuccess = false;
+		}
+	});
+}
+
+bool FControlRigReplayTracks::RestorePose(int32 InTimeIndex, URigHierarchy* InHierarchy, FRigBaseElement* InElement, TReportFunction InReportFunction) const
+{
+	if(!InElement->IsA<FRigTransformElement>() && !InElement->IsA<FRigCurveElement>())
+	{
+		return true;
+	}
+
+	const FSampleTrackContainer* Storage = GetContainer();
+	const FName& TrackName = GetTrackName(InElement->GetKey());
+
+	if(FRigCurveElement* Curve = Cast<FRigCurveElement>(InElement))
+	{
+		const TSharedPtr<const FSampleTrack<float>> Track = Storage->FindTrack<float>(TrackName);
+		if(!Track)
+		{
+			if(InReportFunction)
+			{
+				InReportFunction(EMessageSeverity::Warning, TrackName, TEXT("Track not found."));
+			}
+			return false;
+		}
+		InHierarchy->SetCurveValue(Curve, Track->GetValueAtTimeIndex(InTimeIndex, SampleTrackIndex));
+		return true;
+	}
+	
+	if(const TSharedPtr<const FSampleTrack<FTransform3f>> Track = Storage->FindTrack<FTransform3f>(TrackName))
+	{
+		const FTransform3f& Transform = Track->GetValueAtTimeIndex(InTimeIndex, SampleTrackIndex);
+		InHierarchy->SetLocalTransform(InElement->GetIndex(), FTransform(Transform));
+	}
+	else
+	{
+		if(InReportFunction)
+		{
+			InReportFunction(EMessageSeverity::Warning, TrackName, TEXT("Track not found."));
+		}
+		return false;
+	}
+	return true;
+}
+
+void FControlRigReplayTracks::StoreComponents(URigHierarchy* InHierarchy, const TArrayView<FRigComponentKey>& InKeys)
+{
+	ForEachComponent(InHierarchy, InKeys, [this, InHierarchy](FRigBaseComponent* InComponent, bool& bSuccess)
+	{
+		StoreComponent(InHierarchy, InComponent);
+	});
+}
+
+void FControlRigReplayTracks::StoreComponents(URigHierarchy* InHierarchy, FRigBaseElement* InElement)
+{
+	TArray<FRigComponentKey> ComponentKeys = InElement->GetComponentKeys();
+	if(!ComponentKeys.IsEmpty())
+	{
+		ForEachComponent(InHierarchy, ComponentKeys, [this, InHierarchy](FRigBaseComponent* InComponent, bool& bSuccess)
+		{
+			StoreComponent(InHierarchy, InComponent);
+		});
+	}
+}
+
+void FControlRigReplayTracks::StoreComponent(URigHierarchy* InHierarchy, FRigBaseComponent* InComponent)
+{
+	FSampleTrackContainer* Storage = GetContainer();
+	const FName& TrackName = GetTrackName(InComponent->GetKey());
+	TSharedPtr<FSampleTrack<FInstancedStruct>> Track = Storage->FindTrack<FInstancedStruct>(TrackName);
+	if(Track)
+	{
+		check(Track->GetTrackType() == FSampleTrackBase::ETrackType_Struct);
+		check(Track->GetScriptStruct() == InComponent->GetScriptStruct());
+	}
+	else
+	{
+		Track = Storage->AddStructTrack(TrackName, InComponent->GetScriptStruct());
+		SampleTrackIndex.Update(*Storage);
+	}
+	FInstancedStruct Struct(InComponent->GetScriptStruct());
+	Struct.GetScriptStruct()->CopyScriptStruct(Struct.GetMutableMemory(), InComponent);
+	Track->AddSample(Struct);
+}
+
+bool FControlRigReplayTracks::RestoreComponents(int32 InTimeIndex, URigHierarchy* InHierarchy, const TArrayView<FRigComponentKey>& InKeys, TReportFunction InReportFunction) const
+{
+	return ForEachComponent(InHierarchy, InKeys, [this, InTimeIndex, InHierarchy, InReportFunction](FRigBaseComponent* InComponent, bool& bSuccess)
+	{
+		if(!RestoreComponent(InTimeIndex, InHierarchy, InComponent, InReportFunction))
+		{
+			bSuccess = false;
+		}
+	});
+}
+
+bool FControlRigReplayTracks::RestoreComponents(int32 InTimeIndex, URigHierarchy* InHierarchy, FRigBaseElement* InElement, TReportFunction InReportFunction) const
+{
+	TArray<FRigComponentKey> ComponentKeys = InElement->GetComponentKeys();
+	if(!ComponentKeys.IsEmpty())
+	{
+		return ForEachComponent(InHierarchy, ComponentKeys, [this, InTimeIndex, InHierarchy, InReportFunction](FRigBaseComponent* InComponent, bool& bSuccess)
+		{
+			if(!RestoreComponent(InTimeIndex, InHierarchy, InComponent, InReportFunction))
+			{
+				bSuccess = false;
+			}
+		});
+	}
+	return true;
+}
+
+bool FControlRigReplayTracks::RestoreComponent(int32 InTimeIndex, URigHierarchy* InHierarchy, FRigBaseComponent* InComponent, TReportFunction InReportFunction) const
+{
+	const FSampleTrackContainer* Storage = GetContainer();
+	const FName& TrackName = GetTrackName(InComponent->GetKey());
+	if(const TSharedPtr<const FSampleTrack<FInstancedStruct>> Track = Storage->FindTrack<FInstancedStruct>(TrackName))
+	{
+		if(Track->GetScriptStruct() != InComponent->GetScriptStruct())
+		{
+			if(InReportFunction)
+			{
+				InReportFunction(EMessageSeverity::Error, TrackName, TEXT("Component doesn't match track scriptstruct"));
+			}
+			return false;
+		}
+		TGuardValue<int32> GuardIndexInHierarchy(InComponent->IndexInHierarchy, InComponent->IndexInHierarchy);
+		TGuardValue<int32> GuardIndexInElement(InComponent->IndexInElement, InComponent->IndexInElement);
+		TGuardValue<bool> GuardSelected(InComponent->bSelected, InComponent->bSelected);
+
+		const FInstancedStruct& Struct = Track->GetValueAtTimeIndex(InTimeIndex, SampleTrackIndex);
+		Struct.GetScriptStruct()->CopyScriptStruct(InComponent, Struct.GetMemory());
+	}
+	else
+	{
+		if(InReportFunction)
+		{
+			InReportFunction(EMessageSeverity::Warning, TrackName, TEXT("Track not found."));
+		}
+		return false;
+	}
+	return true;
+}
+
+void FControlRigReplayTracks::StoreMetaData(URigHierarchy* InHierarchy, const TArrayView<FRigElementKey>& InKeys)
+{
+	ForEachElement(InHierarchy, InKeys, [this, InHierarchy](FRigBaseElement* InElement, bool& bSuccess)
+	{
+		StoreMetaData(InHierarchy, InElement);
+	});
+}
+
+void FControlRigReplayTracks::StoreMetaData(URigHierarchy* InHierarchy, FRigBaseElement* InElement)
+{
+	const TArray<FName> MetadataNames = InHierarchy->GetMetadataNames(InElement->GetKey());
+	for(const FName& MetadataName : MetadataNames)
+	{
+		if(FRigBaseMetadata* Metadata = InElement->GetMetadata(MetadataName))
+		{
+			StoreMetaData(InHierarchy, InElement, Metadata);
+		}
+	}
+}
+
+void FControlRigReplayTracks::StoreMetaData(URigHierarchy* InHierarchy, FRigBaseElement* InElement, FRigBaseMetadata* InMetadata)
+{
+	FSampleTrackContainer* Storage = GetContainer();
+	const FName TrackName = GetTrackName(InElement->GetKey(), InMetadata->GetName());
+	const FSampleTrackBase::ETrackType TrackType = GetTrackTypeFromMetadata(InMetadata->GetType());
+	TSharedPtr<FSampleTrackBase> Track = Storage->FindTrack(TrackName);
+	if(Track)
+	{
+		check(Track->GetTrackType() == TrackType);
+	}
+	else
+	{
+		Track = Storage->AddTrack(TrackName, TrackType);
+		SampleTrackIndex.Update(*Storage);
+	}
+	
+	switch(InMetadata->GetType())
+	{
+		case ERigMetadataType::Bool:
+		{
+			const bool& Value = static_cast<FRigBoolMetadata*>(InMetadata)->GetValue(); 
+			StaticCastSharedPtr<FSampleTrack<bool>>(Track)->AddSample(Value);
+			break;
+		}
+		case ERigMetadataType::BoolArray:
+		{
+			const TArray<bool>& Value = static_cast<FRigBoolArrayMetadata*>(InMetadata)->GetValue(); 
+			StaticCastSharedPtr<FSampleTrack<TArray<bool>>>(Track)->AddSample(Value);
+			break;
+		}
+		case ERigMetadataType::Float:
+		{
+			const float& Value = static_cast<FRigFloatMetadata*>(InMetadata)->GetValue(); 
+			StaticCastSharedPtr<FSampleTrack<float>>(Track)->AddSample(Value);
+			break;
+		}
+		case ERigMetadataType::FloatArray:
+		{
+			const TArray<float>& Value = static_cast<FRigFloatArrayMetadata*>(InMetadata)->GetValue(); 
+			StaticCastSharedPtr<FSampleTrack<TArray<float>>>(Track)->AddSample(Value);
+			break;
+		}
+		case ERigMetadataType::Int32:
+		{
+			const int32& Value = static_cast<FRigInt32Metadata*>(InMetadata)->GetValue(); 
+			StaticCastSharedPtr<FSampleTrack<int32>>(Track)->AddSample(Value);
+			break;
+		}
+		case ERigMetadataType::Int32Array:
+		{
+			const TArray<int32>& Value = static_cast<FRigInt32ArrayMetadata*>(InMetadata)->GetValue(); 
+			StaticCastSharedPtr<FSampleTrack<TArray<int32>>>(Track)->AddSample(Value);
+			break;
+		}
+		case ERigMetadataType::Name:
+		{
+			const FName& Value = static_cast<FRigNameMetadata*>(InMetadata)->GetValue(); 
+			StaticCastSharedPtr<FSampleTrack<FName>>(Track)->AddSample(Value);
+			break;
+		}
+		case ERigMetadataType::NameArray:
+		{
+			const TArray<FName>& Value = static_cast<FRigNameArrayMetadata*>(InMetadata)->GetValue(); 
+			StaticCastSharedPtr<FSampleTrack<TArray<FName>>>(Track)->AddSample(Value);
+			break;
+		}
+		case ERigMetadataType::Vector:
+		{
+			const FVector& Value = static_cast<FRigVectorMetadata*>(InMetadata)->GetValue(); 
+			StaticCastSharedPtr<FSampleTrack<FVector3f>>(Track)->AddSample(FVector3f(Value));
+			break;
+		}
+		case ERigMetadataType::VectorArray:
+		{
+			const TArray<FVector>& Value = static_cast<FRigVectorArrayMetadata*>(InMetadata)->GetValue();
+			TArray<FVector3f> ConvertedValue;
+			ConvertedValue.Reserve(Value.Num());
+			for(const FVector& SingleValue : Value)
+			{
+				ConvertedValue.Emplace(SingleValue);
+			}
+			StaticCastSharedPtr<FSampleTrack<TArray<FVector3f>>>(Track)->AddSample(ConvertedValue);
+			break;
+		}
+		case ERigMetadataType::Rotator:
+		{
+			const FRotator& Value = static_cast<FRigRotatorMetadata*>(InMetadata)->GetValue();
+			StaticCastSharedPtr<FSampleTrack<FVector3f>>(Track)->AddSample(FVector3f(Value.Euler()));
+			break;
+		}
+		case ERigMetadataType::RotatorArray:
+		{
+			const TArray<FRotator>& Value = static_cast<FRigRotatorArrayMetadata*>(InMetadata)->GetValue();
+			TArray<FVector3f> ConvertedValue;
+			ConvertedValue.Reserve(Value.Num());
+			for(const FRotator& SingleValue : Value)
+			{
+				ConvertedValue.Emplace(SingleValue.Euler());
+			}
+			StaticCastSharedPtr<FSampleTrack<TArray<FVector3f>>>(Track)->AddSample(ConvertedValue);
+			break;
+		}
+		case ERigMetadataType::Quat:
+		{
+			const FQuat& Value = static_cast<FRigQuatMetadata*>(InMetadata)->GetValue(); 
+			StaticCastSharedPtr<FSampleTrack<FQuat4f>>(Track)->AddSample(FQuat4f(Value));
+			break;
+		}
+		case ERigMetadataType::QuatArray:
+		{
+			const TArray<FQuat>& Value = static_cast<FRigQuatArrayMetadata*>(InMetadata)->GetValue(); 
+			TArray<FQuat4f> ConvertedValue;
+			ConvertedValue.Reserve(Value.Num());
+			for(const FQuat& SingleValue : Value)
+			{
+				ConvertedValue.Emplace(SingleValue);
+			}
+			StaticCastSharedPtr<FSampleTrack<TArray<FQuat4f>>>(Track)->AddSample(ConvertedValue);
+			break;
+		}
+		case ERigMetadataType::Transform:
+		{
+			const FTransform& Value = static_cast<FRigTransformMetadata*>(InMetadata)->GetValue(); 
+			StaticCastSharedPtr<FSampleTrack<FTransform3f>>(Track)->AddSample(FTransform3f(Value));
+			break;
+		}
+		case ERigMetadataType::TransformArray:
+		{
+			const TArray<FTransform>& Value = static_cast<FRigTransformArrayMetadata*>(InMetadata)->GetValue(); 
+			TArray<FTransform3f> ConvertedValue;
+			ConvertedValue.Reserve(Value.Num());
+			for(const FTransform& SingleValue : Value)
+			{
+				ConvertedValue.Emplace(SingleValue);
+			}
+			StaticCastSharedPtr<FSampleTrack<TArray<FTransform3f>>>(Track)->AddSample(ConvertedValue);
+			break;
+		}
+		case ERigMetadataType::LinearColor:
+		{
+			const FLinearColor& Value = static_cast<FRigLinearColorMetadata*>(InMetadata)->GetValue(); 
+			StaticCastSharedPtr<FSampleTrack<FLinearColor>>(Track)->AddSample(Value);
+			break;
+		}
+		case ERigMetadataType::LinearColorArray:
+		{
+			const TArray<FLinearColor>& Value = static_cast<FRigLinearColorArrayMetadata*>(InMetadata)->GetValue(); 
+			StaticCastSharedPtr<FSampleTrack<TArray<FLinearColor>>>(Track)->AddSample(Value);
+			break;
+		}
+		case ERigMetadataType::RigElementKey:
+		{
+			const FRigElementKey& Value = static_cast<FRigElementKeyMetadata*>(InMetadata)->GetValue(); 
+			StaticCastSharedPtr<FSampleTrack<FRigElementKey>>(Track)->AddSample(Value);
+			break;
+		}
+		case ERigMetadataType::RigElementKeyArray:
+		{
+			const TArray<FRigElementKey>& Value = static_cast<FRigElementKeyArrayMetadata*>(InMetadata)->GetValue(); 
+			StaticCastSharedPtr<FSampleTrack<TArray<FRigElementKey>>>(Track)->AddSample(Value);
+			break;
+		}
+		case ERigMetadataType::Invalid:
+		{
+			break;
+		}
+	}
+}
+
+bool FControlRigReplayTracks::RestoreMetaData(int32 InTimeIndex, URigHierarchy* InHierarchy, const TArrayView<FRigElementKey>& InKeys, TReportFunction InReportFunction) const
+{
+	return ForEachElement(InHierarchy, InKeys, [this, InTimeIndex, InHierarchy, InReportFunction](FRigBaseElement* InElement, bool& bSuccess)
+	{
+		if(!RestoreMetaData(InTimeIndex, InHierarchy, InElement, InReportFunction))
+		{
+			bSuccess = false;
+		}
+	});
+}
+
+bool FControlRigReplayTracks::RestoreMetaData(int32 InTimeIndex, URigHierarchy* InHierarchy, FRigBaseElement* InElement, TReportFunction InReportFunction) const
+{
+	bool bSuccess = true;
+	const TArray<FName> MetadataNames = InHierarchy->GetMetadataNames(InElement->GetKey());
+	for(const FName& MetadataName : MetadataNames)
+	{
+		if(FRigBaseMetadata* Metadata = InElement->GetMetadata(MetadataName))
+		{
+			if(!RestoreMetaData(InTimeIndex, InHierarchy, InElement, Metadata, InReportFunction))
+			{
+				bSuccess = false;
+			}
+		}
+	}
+	return bSuccess;
+}
+
+bool FControlRigReplayTracks::RestoreMetaData(int32 InTimeIndex, URigHierarchy* InHierarchy, FRigBaseElement* InElement, FRigBaseMetadata* InMetadata, TReportFunction InReportFunction) const
+{
+	const FSampleTrackContainer* Storage = GetContainer();
+	const FName TrackName = GetTrackName(InElement->GetKey(), InMetadata->GetName());
+	const TSharedPtr<const FSampleTrackBase> Track = Storage->FindTrack(TrackName);
+	if(!Track.IsValid())
+	{
+		if(InReportFunction)
+		{
+			InReportFunction(EMessageSeverity::Warning, TrackName, TEXT("Track not found."));
+		}
+		return false;
+	}
+	if(Track->GetTrackType() != GetTrackTypeFromMetadata(InMetadata->GetType()))
+	{
+		if(InReportFunction)
+		{
+			InReportFunction(EMessageSeverity::Error, TrackName, TEXT("Track has incorrect type."));
+		}
+		return false;
+	}
+	
+	switch(InMetadata->GetType())
+	{
+		case ERigMetadataType::Bool:
+		{
+			const bool& Value = StaticCastSharedPtr<const FSampleTrack<bool>>(Track)->GetValueAtTimeIndex(InTimeIndex, SampleTrackIndex);
+			static_cast<FRigBoolMetadata*>(InMetadata)->SetValue(Value);
+			break;
+		}
+		case ERigMetadataType::BoolArray:
+		{
+			const TArray<bool>& Value = StaticCastSharedPtr<const FSampleTrack<TArray<bool>>>(Track)->GetValueAtTimeIndex(InTimeIndex, SampleTrackIndex);
+			static_cast<FRigBoolArrayMetadata*>(InMetadata)->SetValue(Value);
+			break;
+		}
+		case ERigMetadataType::Float:
+		{
+			const float& Value = StaticCastSharedPtr<const FSampleTrack<float>>(Track)->GetValueAtTimeIndex(InTimeIndex, SampleTrackIndex);
+			static_cast<FRigFloatMetadata*>(InMetadata)->SetValue(Value);
+			break;
+		}
+		case ERigMetadataType::FloatArray:
+		{
+			const TArray<float>& Value = StaticCastSharedPtr<const FSampleTrack<TArray<float>>>(Track)->GetValueAtTimeIndex(InTimeIndex, SampleTrackIndex);
+			static_cast<FRigFloatArrayMetadata*>(InMetadata)->SetValue(Value);
+			break;
+		}
+		case ERigMetadataType::Int32:
+		{
+			const int32& Value = StaticCastSharedPtr<const FSampleTrack<int32>>(Track)->GetValueAtTimeIndex(InTimeIndex, SampleTrackIndex);
+			static_cast<FRigInt32Metadata*>(InMetadata)->SetValue(Value);
+			break;
+		}
+		case ERigMetadataType::Int32Array:
+		{
+			const TArray<int32>& Value = StaticCastSharedPtr<const FSampleTrack<TArray<int32>>>(Track)->GetValueAtTimeIndex(InTimeIndex, SampleTrackIndex);
+			static_cast<FRigInt32ArrayMetadata*>(InMetadata)->SetValue(Value);
+			break;
+		}
+		case ERigMetadataType::Name:
+		{
+			const FName& Value = StaticCastSharedPtr<const FSampleTrack<FName>>(Track)->GetValueAtTimeIndex(InTimeIndex, SampleTrackIndex);
+			static_cast<FRigNameMetadata*>(InMetadata)->SetValue(Value);
+			break;
+		}
+		case ERigMetadataType::NameArray:
+		{
+			const TArray<FName>& Value = StaticCastSharedPtr<const FSampleTrack<TArray<FName>>>(Track)->GetValueAtTimeIndex(InTimeIndex, SampleTrackIndex);
+			static_cast<FRigNameArrayMetadata*>(InMetadata)->SetValue(Value);
+			break;
+		}
+		case ERigMetadataType::Vector:
+		{
+			const FVector3f& Value = StaticCastSharedPtr<const FSampleTrack<FVector3f>>(Track)->GetValueAtTimeIndex(InTimeIndex, SampleTrackIndex);
+			static_cast<FRigVectorMetadata*>(InMetadata)->SetValue(FVector(Value));
+			break;
+		}
+		case ERigMetadataType::VectorArray:
+		{
+			const TArray<FVector3f>& Value = StaticCastSharedPtr<const FSampleTrack<TArray<FVector3f>>>(Track)->GetValueAtTimeIndex(InTimeIndex, SampleTrackIndex);
+			TArray<FVector> ConvertedValue;
+			ConvertedValue.Reserve(Value.Num());
+			for(const FVector3f& SingleValue : Value)
+			{
+				ConvertedValue.Emplace(SingleValue);
+			}
+			static_cast<FRigVectorArrayMetadata*>(InMetadata)->SetValue(ConvertedValue);
+			break;
+		}
+		case ERigMetadataType::Rotator:
+		{
+			const FVector3f& Value = StaticCastSharedPtr<const FSampleTrack<FVector3f>>(Track)->GetValueAtTimeIndex(InTimeIndex, SampleTrackIndex);
+			static_cast<FRigRotatorMetadata*>(InMetadata)->SetValue(FRotator::MakeFromEuler(FVector(Value)));
+			break;
+		}
+		case ERigMetadataType::RotatorArray:
+		{
+			const TArray<FVector3f>& Value = StaticCastSharedPtr<const FSampleTrack<TArray<FVector3f>>>(Track)->GetValueAtTimeIndex(InTimeIndex, SampleTrackIndex);
+			TArray<FRotator> ConvertedValue;
+			ConvertedValue.Reserve(Value.Num());
+			for(const FVector3f& SingleValue : Value)
+			{
+				ConvertedValue.Add(FRotator::MakeFromEuler(FVector(SingleValue)));
+			}
+			static_cast<FRigRotatorArrayMetadata*>(InMetadata)->SetValue(ConvertedValue);
+			break;
+		}
+		case ERigMetadataType::Quat:
+		{
+			const FQuat4f& Value = StaticCastSharedPtr<const FSampleTrack<FQuat4f>>(Track)->GetValueAtTimeIndex(InTimeIndex, SampleTrackIndex);
+			static_cast<FRigQuatMetadata*>(InMetadata)->SetValue(FQuat(Value)); 
+			break;
+		}
+		case ERigMetadataType::QuatArray:
+		{
+			const TArray<FQuat4f>& Value = StaticCastSharedPtr<const FSampleTrack<TArray<FQuat4f>>>(Track)->GetValueAtSampleIndex(InTimeIndex);
+			TArray<FQuat> ConvertedValue;
+			ConvertedValue.Reserve(Value.Num());
+			for(const FQuat4f& SingleValue : Value)
+			{
+				ConvertedValue.Emplace(SingleValue);
+			}
+			static_cast<FRigQuatArrayMetadata*>(InMetadata)->SetValue(ConvertedValue); 
+			break;
+		}
+		case ERigMetadataType::Transform:
+		{
+			const FTransform3f& Value = StaticCastSharedPtr<const FSampleTrack<FTransform3f>>(Track)->GetValueAtSampleIndex(InTimeIndex);
+			static_cast<FRigTransformMetadata*>(InMetadata)->SetValue(FTransform(Value)); 
+			break;
+		}
+		case ERigMetadataType::TransformArray:
+		{
+			const TArray<FTransform3f>& Value = StaticCastSharedPtr<const FSampleTrack<TArray<FTransform3f>>>(Track)->GetValueAtSampleIndex(InTimeIndex);
+			TArray<FTransform> ConvertedValue;
+			ConvertedValue.Reserve(Value.Num());
+			for(const FTransform3f& SingleValue : Value)
+			{
+				ConvertedValue.Emplace(SingleValue);
+			}
+			static_cast<FRigTransformArrayMetadata*>(InMetadata)->SetValue(ConvertedValue); 
+			break;
+		}
+		case ERigMetadataType::LinearColor:
+		{
+			const FLinearColor& Value = StaticCastSharedPtr<const FSampleTrack<FLinearColor>>(Track)->GetValueAtSampleIndex(InTimeIndex);
+			static_cast<FRigLinearColorMetadata*>(InMetadata)->SetValue(Value); 
+			break;
+		}
+		case ERigMetadataType::LinearColorArray:
+		{
+			const TArray<FLinearColor>& Value = StaticCastSharedPtr<const FSampleTrack<TArray<FLinearColor>>>(Track)->GetValueAtSampleIndex(InTimeIndex);
+			static_cast<FRigLinearColorArrayMetadata*>(InMetadata)->SetValue(Value); 
+			break;
+		}
+		case ERigMetadataType::RigElementKey:
+		{
+			const FRigElementKey& Value = StaticCastSharedPtr<const FSampleTrack<FRigElementKey>>(Track)->GetValueAtSampleIndex(InTimeIndex);
+			static_cast<FRigElementKeyMetadata*>(InMetadata)->SetValue(Value); 
+			break;
+		}
+		case ERigMetadataType::RigElementKeyArray:
+		{
+			const TArray<FRigElementKey>& Value = StaticCastSharedPtr<const FSampleTrack<TArray<FRigElementKey>>>(Track)->GetValueAtSampleIndex(InTimeIndex);
+			static_cast<FRigElementKeyArrayMetadata*>(InMetadata)->SetValue(Value); 
+			break;
+		}
+		case ERigMetadataType::Invalid:
+		{
+			if(InReportFunction)
+			{
+				InReportFunction(EMessageSeverity::Error, TrackName, TEXT("Unsupported Metadata Type."));
+			}
+			return false;
+		}
+	}
+	return true;
+}
+
+void FControlRigReplayTracks::StoreVariables(URigVMHost* InHost)
+{
+	FSampleTrackContainer* Storage = GetContainer();
+	for (TFieldIterator<FProperty> PropertyIt(InHost->GetClass()); PropertyIt; ++PropertyIt)
+	{
+		const FProperty* Property = *PropertyIt;
+		if(Property->IsNative())
+		{
+			continue;
+		}
+
+		const UScriptStruct* ScriptStruct = nullptr;
+		if(const FStructProperty* StructProperty = CastField<FStructProperty>(Property))
+		{
+			ScriptStruct = StructProperty->Struct;
+		}
+		else if(const FArrayProperty* ArrayProperty = CastField<FArrayProperty>(Property))
+		{
+			if(const FStructProperty* InnerStructProperty = CastField<FStructProperty>(ArrayProperty->Inner))
+			{
+				ScriptStruct = InnerStructProperty->Struct;
+			}
+		}
+
+		const FName& TrackName = GetTrackName(Property);
+		const FSampleTrackBase::ETrackType TrackType = GetTrackTypeFromProperty(Property);
+		TSharedPtr<FSampleTrackBase> Track = Storage->FindTrack(TrackName);
+		if(Track)
+		{
+			check(Track->GetTrackType() == TrackType);
+			check(Track->GetScriptStruct() == ScriptStruct);
+		}
+		else
+		{
+			Track = Storage->AddTrack(TrackName, TrackType, ScriptStruct);
+		}
+		const uint8* Memory = Property->ContainerPtrToValuePtr<uint8>(InHost);
+		Track->AddSampleFromProperty(Property, Memory);
+	}
+}
+
+bool FControlRigReplayTracks::RestoreVariables(int32 InTimeIndex, URigVMHost* InHost, TReportFunction InReportFunction) const
+{
+	const FSampleTrackContainer* Storage = GetContainer();
+	bool bSuccess = true;
+	for (TFieldIterator<FProperty> PropertyIt(InHost->GetClass()); PropertyIt; ++PropertyIt)
+	{
+		const FProperty* Property = *PropertyIt;
+		if(Property->IsNative())
+		{
+			continue;
+		}
+
+		const FName& TrackName = GetTrackName(Property);
+		const FSampleTrackBase::ETrackType TrackType = GetTrackTypeFromProperty(Property);
+		const TSharedPtr<const FSampleTrackBase> Track = Storage->FindTrack(TrackName);
+		if(Track)
+		{
+			check(Track->GetTrackType() == TrackType);
+			uint8* Memory = Property->ContainerPtrToValuePtr<uint8>(InHost);
+			Track->GetSampleForProperty(InTimeIndex, SampleTrackIndex, Property, Memory);
+		}
+		else
+		{
+			if(InReportFunction)
+			{
+				InReportFunction(EMessageSeverity::Warning, TrackName, TEXT("Track not found."));
+			}
+			bSuccess = false;
+		}
+	}
+	return bSuccess;
+}
+
+bool FControlRigReplayTracks::ForEachElement(URigHierarchy* InHierarchy, const TArrayView<FRigElementKey>& InKeys, TFunction<void(FRigBaseElement*, bool&)> InFunction)
+{
+	bool bSuccess = true;
+	if(InKeys.IsEmpty())
+	{
+		InHierarchy->Traverse([InFunction, &bSuccess](FRigBaseElement* InElement, bool& bContinue)
+		{
+			InFunction(InElement, bSuccess);
+			bContinue = true;
+		}, true);
+	}
+	else
+	{
+		for(const FRigElementKey& Key : InKeys)
+		{
+			if(FRigBaseElement* Element = InHierarchy->Find(Key))
+			{
+				InFunction(Element, bSuccess);
+			}
+		}
+	}
+	return bSuccess;
+}
+
+bool FControlRigReplayTracks::ForEachComponent(URigHierarchy* InHierarchy, const TArrayView<FRigComponentKey>& InKeys,
+	TFunction<void(FRigBaseComponent*, bool&)> InFunction)
+{
+	TArrayView<FRigComponentKey> View = InKeys;
+	TArray<FRigComponentKey> AllComponentKeys;
+	if(View.IsEmpty())
+	{
+		AllComponentKeys = InHierarchy->GetAllComponentKeys();
+		View = TArrayView<FRigComponentKey>(AllComponentKeys);
+	}
+
+	bool bSuccess = true;
+	for(const FRigComponentKey& Key : View)
+	{
+		if(FRigBaseComponent* Component = InHierarchy->FindComponent(Key))
+		{
+			InFunction(Component, bSuccess);
+		}
+	}
+	return bSuccess;
+}
+
+const FName& FControlRigReplayTracks::GetTrackName(const FRigElementKey& InElementKey) const
+{
+	if(const FName* ExistingTrackName = ElementKeyToTrackName.Find(InElementKey))
+	{
+		return *ExistingTrackName;
+	}
+	const FName TrackName = *InElementKey.ToString();
+	return ElementKeyToTrackName.Add(InElementKey, TrackName);
+}
+
+const FName& FControlRigReplayTracks::GetTrackName(const FRigComponentKey& InComponentKey) const
+{
+	if(const FName* ExistingTrackName = ComponentKeyToTrackName.Find(InComponentKey))
+	{
+		return *ExistingTrackName;
+	}
+	const FName TrackName = *InComponentKey.ToString();
+	return ComponentKeyToTrackName.Add(InComponentKey, TrackName);
+}
+
+const FName& FControlRigReplayTracks::GetTrackName(const FRigElementKey& InElementKey, const FName& InMetadataName) const
+{
+	const TTuple<FRigElementKey,FName> MetadataKey(InElementKey, InMetadataName);
+	if(const FName* ExistingTrackName = MetadataToTrackName.Find(MetadataKey))
+	{
+		return *ExistingTrackName;
+	}
+	FStringBuilderBase Builder;
+	Builder.Append(InElementKey.ToString());
+	Builder.AppendChar(TEXT(':'));
+	Builder.Append(InMetadataName.ToString());
+	const FName TrackName = Builder.ToString();
+	return MetadataToTrackName.Add(MetadataKey, TrackName);
+}
+
+const FName& FControlRigReplayTracks::GetTrackName(const FProperty* InProperty) const
+{
+	check(InProperty);
+	if(const FName* ExistingTrackName = PropertyNameToTrackName.Find(InProperty->GetFName()))
+	{
+		return *ExistingTrackName;
+	}
+	FStringBuilderBase Builder;
+	Builder.Append(TEXT("Variable:"));
+	Builder.Append(InProperty->GetName());
+	const FName TrackName = Builder.ToString();
+	return PropertyNameToTrackName.Add(InProperty->GetFName(), TrackName);
+}
+
+FSampleTrackBase::ETrackType FControlRigReplayTracks::GetTrackTypeFromMetadata(ERigMetadataType InMetadataType)
+{
+	switch(InMetadataType)
+	{
+		case ERigMetadataType::Bool:
+		{
+			return FSampleTrackBase::ETrackType_Bool;
+		}
+		case ERigMetadataType::BoolArray:
+		{
+			return FSampleTrackBase::ETrackType_BoolArray;
+		}
+		case ERigMetadataType::Float:
+		{
+			return FSampleTrackBase::ETrackType_Float;
+		}
+		case ERigMetadataType::FloatArray:
+		{
+			return FSampleTrackBase::ETrackType_FloatArray;
+		}
+		case ERigMetadataType::Int32:
+		{
+			return FSampleTrackBase::ETrackType_Int32;
+		}
+		case ERigMetadataType::Int32Array:
+		{
+			return FSampleTrackBase::ETrackType_Int32Array;
+		}
+		case ERigMetadataType::Name:
+		{
+			return FSampleTrackBase::ETrackType_Name;
+		}
+		case ERigMetadataType::NameArray:
+		{
+			return FSampleTrackBase::ETrackType_NameArray;
+		}
+		case ERigMetadataType::Vector:
+		case ERigMetadataType::Rotator:
+		{
+			return FSampleTrackBase::ETrackType_Vector3f;
+		}
+		case ERigMetadataType::VectorArray:
+		case ERigMetadataType::RotatorArray:
+		{
+			return FSampleTrackBase::ETrackType_Vector3fArray;
+		}
+		case ERigMetadataType::Quat:
+		{
+			return FSampleTrackBase::ETrackType_Quatf;
+		}
+		case ERigMetadataType::QuatArray:
+		{
+			return FSampleTrackBase::ETrackType_QuatfArray;
+		}
+		case ERigMetadataType::Transform:
+		{
+			return FSampleTrackBase::ETrackType_Transformf;
+		}
+		case ERigMetadataType::TransformArray:
+		{
+			return FSampleTrackBase::ETrackType_TransformfArray;
+		}
+		case ERigMetadataType::LinearColor:
+		{
+			return FSampleTrackBase::ETrackType_LinearColor;
+		}
+		case ERigMetadataType::LinearColorArray:
+		{
+			return FSampleTrackBase::ETrackType_LinearColorArray;
+		}
+		case ERigMetadataType::RigElementKey:
+		{
+			return FSampleTrackBase::ETrackType_ElementKey;
+		}
+		case ERigMetadataType::RigElementKeyArray:
+		{
+			return FSampleTrackBase::ETrackType_ElementKeyArray;
+		}
+		case ERigMetadataType::Invalid:
+		{
+			break;
+		}
+	}
+	return FSampleTrackBase::ETrackType_Unknown;
+}
+
+FSampleTrackBase::ETrackType FControlRigReplayTracks::GetTrackTypeFromProperty(const FProperty* InProperty)
+{
+	check(InProperty);
+
+	static constexpr uint8 SingleToArrayOffset = uint8(FSampleTrackBase::ETrackType_BoolArray)-uint8(FSampleTrackBase::ETrackType_Bool);
+
+	bool bIsArray = false;
+	if(const FArrayProperty* ArrayProperty = CastField<FArrayProperty>(InProperty))
+	{
+		InProperty = ArrayProperty->Inner;
+		bIsArray = true;
+	}
+
+	FSampleTrackBase::ETrackType TrackType = FSampleTrackBase::ETrackType_Unknown; 
+	
+	if(InProperty->IsA<FBoolProperty>())
+	{
+		TrackType = FSampleTrackBase::ETrackType_Bool;
+	}
+	else if(InProperty->IsA<FFloatProperty>() || InProperty->IsA<FDoubleProperty>())
+	{
+		TrackType = FSampleTrackBase::ETrackType_Float;
+	}
+	else if(InProperty->IsA<FIntProperty>() || InProperty->IsA<FInt16Property>())
+	{
+		TrackType = FSampleTrackBase::ETrackType_Int32;
+	}
+	else if(InProperty->IsA<FUInt32Property>() || InProperty->IsA<FUInt16Property>() ||  InProperty->IsA<FByteProperty>() || InProperty->IsA<FEnumProperty>())
+	{
+		TrackType = FSampleTrackBase::ETrackType_Uint32;
+	}
+	else if(InProperty->IsA<FNameProperty>())
+	{
+		TrackType = FSampleTrackBase::ETrackType_Name;
+	}
+	else if(InProperty->IsA<FStrProperty>())
+	{
+		TrackType = FSampleTrackBase::ETrackType_String;
+	}
+	else if(InProperty->IsA<FStructProperty>())
+	{
+		TrackType = FSampleTrackBase::ETrackType_Struct;
+	}
+
+	if(bIsArray)
+	{
+		TrackType = (FSampleTrackBase::ETrackType)(uint8(TrackType) + SingleToArrayOffset);
+	}
+
+	return TrackType;
+}
+
 //////////////////////////////////////////////////////////////////////////////////////////////////////////////////////////////
 /////// UControlRigReplay
 //////////////////////////////////////////////////////////////////////////////////////////////////////////////////////////////
@@ -80,13 +1075,99 @@ FVector2D UControlRigReplay::GetTimeRange() const
 	return FVector2D(Tracks.GetContainer()->GetTimeRange());
 }
 
-bool UControlRigReplay::StartRecording(UControlRig* InControlRig, double InRecordingDuration)
+bool UControlRigReplay::StartRecording(UControlRig* InControlRig)
 {
-	return false;
+	if(InControlRig == nullptr)
+	{
+		return false;
+	}
+
+	(void)StopRecording();
+	(void)StopReplay();
+
+	Tracks.Reset();
+	RecordControlRig = InControlRig;
+	ClearDelegates(InControlRig);
+
+	EventQueue = InControlRig->EventQueue;
+	TimeAtStartOfRecording = TimeOfLastFrame = FPlatformTime::Seconds();
+
+	PreConstructionHandle = InControlRig->OnPreConstruction_AnyThread().AddLambda(
+		[this](UControlRig* InControlRig, const FName& InEventName)
+		{
+			// nothing to do here
+		}
+	);
+
+	PreEventHandle = InControlRig->OnPreExecuted_AnyThread().AddLambda(
+		[this](URigVMHost* InRigVMHost, const FName& InEventName)
+		{
+			UControlRig* ControlRig = Cast<UControlRig>(InRigVMHost);
+			if (!ControlRig)
+			{
+				return;
+			}
+
+			if (InEventName == FRigUnit_PrepareForExecution::EventName)
+			{
+				return;
+			}
+			
+			if (InEventName != EventQueue[0])
+			{
+				return;
+			}
+
+			// todo: record a frame if we are recording backwards?
+		}
+	);
+
+	PostEventHandle = InControlRig->OnExecuted_AnyThread().AddLambda(
+		[this](URigVMHost* InRigVMHost, const FName& InEventName)
+		{
+			UControlRig* ControlRig = Cast<UControlRig>(InRigVMHost);
+			if (!ControlRig)
+			{
+				return;
+			}
+
+			if (InEventName == FRigUnit_PrepareForExecution::EventName)
+			{
+				return;
+			}
+
+			if (InEventName != EventQueue.Last())
+			{
+				return;
+			}
+
+			const double CurrentTime = FPlatformTime::Seconds(); 
+			const double DeltaTime = CurrentTime - TimeOfLastFrame;
+			TimeOfLastFrame = CurrentTime;
+
+			Tracks.AddTimeSampleFromDeltaTime(static_cast<float>(DeltaTime));
+			Tracks.StoreHierarchy(ControlRig->GetHierarchy());
+			Tracks.StoreVariables(ControlRig);
+		}
+	);
+
+	if(Tracks.GetContainer()->NumTracks() == 0)
+	{
+		InControlRig->RequestInit();
+	}
+	
+	return true;
 }
 
 bool UControlRigReplay::StopRecording()
 {
+	if(UControlRig* ControlRig = RecordControlRig.Get())
+	{
+		ClearDelegates(ControlRig);
+		RecordControlRig.Reset();
+		Tracks.Compact();
+		return true;
+	}
 	return false;
 }
 
diff --git a/Engine/Plugins/Animation/ControlRig/Source/ControlRig/Private/ControlRigTestData.cpp b/Engine/Plugins/Animation/ControlRig/Source/ControlRig/Private/ControlRigTestData.cpp
--- a/Engine/Plugins/Animation/ControlRig/Source/ControlRig/Private/ControlRigTestData.cpp
+++ b/Engine/Plugins/Animation/ControlRig/Source/ControlRig/Private/ControlRigTestData.cpp
@@ -276,7 +276,7 @@ int32 UControlRigTestData::GetFrameIndexForTime(double InSeconds, bool bInput) c
 	return LastFrameIndex;
 }
 
-bool UControlRigTestData::StartRecording(UControlRig* InControlRig, double InRecordingDuration)
+bool UControlRigTestData::StartRecording(UControlRig* InControlRig)
 {
 	if(InControlRig == nullptr)
 	{
@@ -288,7 +288,6 @@ bool UControlRigTestData::StartRecording(UControlRig* InControlRig, double InRec
 	ClearDelegates(InControlRig);
 
 	EventQueue = InControlRig->EventQueue;
-	DesiredRecordingDuration = InRecordingDuration;
 	TimeAtStartOfRecording = FPlatformTime::Seconds();
 
 	PreConstructionHandle = InControlRig->OnPreConstruction_AnyThread().AddLambda(
diff --git a/Engine/Plugins/Animation/ControlRig/Source/ControlRig/Private/Rigs/RigHierarchy.cpp b/Engine/Plugins/Animation/ControlRig/Source/ControlRig/Private/Rigs/RigHierarchy.cpp
--- a/Engine/Plugins/Animation/ControlRig/Source/ControlRig/Private/Rigs/RigHierarchy.cpp
+++ b/Engine/Plugins/Animation/ControlRig/Source/ControlRig/Private/Rigs/RigHierarchy.cpp
@@ -993,7 +993,7 @@ void URigHierarchy::CopyHierarchy(URigHierarchy* InHierarchy)
 		ComponentIndexLookup.Reset();
 		ComponentIndexLookup.Reserve(InHierarchy->ElementComponents.Num());
 
-		for(int32 SourceComponentIndex = 0; SourceComponentIndex < ElementComponents.Num(); SourceComponentIndex++)
+		for(int32 SourceComponentIndex = 0; SourceComponentIndex < InHierarchy->ElementComponents.Num(); SourceComponentIndex++)
 		{
 			if(!InHierarchy->ElementComponents[SourceComponentIndex].IsValid())
 			{
@@ -1544,6 +1544,13 @@ int32 URigHierarchy::NumComponents(FRigElementKey InElement) const
 	return 0;
 }
 
+TArray<FRigComponentKey> URigHierarchy::GetAllComponentKeys() const
+{
+	TArray<FRigComponentKey> Keys;
+	ComponentIndexLookup.GenerateKeyArray(Keys);
+	return Keys;
+}
+
 TArray<FRigComponentKey> URigHierarchy::GetComponentKeys(FRigElementKey InElement) const
 {
 	if(const FRigBaseElement* Element = Find(InElement))
diff --git a/Engine/Plugins/Animation/ControlRig/Source/ControlRig/Private/Tracks/ComposedSampleTrack.cpp b/Engine/Plugins/Animation/ControlRig/Source/ControlRig/Private/Tracks/ComposedSampleTrack.cpp
--- a/Engine/Plugins/Animation/ControlRig/Source/ControlRig/Private/Tracks/ComposedSampleTrack.cpp
+++ b/Engine/Plugins/Animation/ControlRig/Source/ControlRig/Private/Tracks/ComposedSampleTrack.cpp
@@ -16,6 +16,31 @@ TArray<FSampleTrackBase::ETrackType> FComposedSampleTrack<FTransform3f>::GetChil
 	};
 }
 
+template<>
+FString FComposedSampleTrack<FTransform3f>::GetChildTrackNameSuffix(int32 InChildTrackIndex) const
+{
+	switch(InChildTrackIndex)
+	{
+		case 0:
+		{
+			return TEXT("Location");
+		}
+		case 1:
+		{
+			return TEXT("Rotation");
+		}
+		case 2:
+		{
+			return TEXT("Scale");
+		}
+		default:
+		{
+			break;
+		}
+	}
+	return FSampleTrack<FTransform3f>::GetChildTrackNameSuffix(InChildTrackIndex);
+}
+
 template<>
 FTransform3f FComposedSampleTrack<FTransform3f>::GetValueAtTimeIndex(int32 InTimeIndex, FSampleTrackIndex& InOutIndex) const
 {
diff --git a/Engine/Plugins/Animation/ControlRig/Source/ControlRig/Private/Tracks/SampleTrack.cpp b/Engine/Plugins/Animation/ControlRig/Source/ControlRig/Private/Tracks/SampleTrack.cpp
--- a/Engine/Plugins/Animation/ControlRig/Source/ControlRig/Private/Tracks/SampleTrack.cpp
+++ b/Engine/Plugins/Animation/ControlRig/Source/ControlRig/Private/Tracks/SampleTrack.cpp
@@ -77,3 +77,62 @@ int32 FSampleTrack<FInstancedStruct>::GetSizePerValue() const
 	}
 	return sizeof(FInstancedStruct);
 }
+
+template<>
+void FSampleTrack<FInstancedStruct>::AddSampleFromProperty(const FProperty* InProperty, const uint8* InMemory)
+{
+	const FStructProperty* StructProperty = CastFieldChecked<FStructProperty>(InProperty);
+	check(StructProperty->Struct == GetScriptStruct());
+
+	FInstancedStruct Value(GetScriptStruct());
+	StructProperty->CopyCompleteValue(Value.GetMutableMemory(), InMemory);
+	AddSample(Value);
+}
+
+template<>
+void FSampleTrack<FInstancedStruct>::GetSampleForProperty(int32 InTimeIndex, FSampleTrackIndex& InOutSampleTrackIndex, const FProperty* InProperty, uint8* OutMemory) const
+{
+	const FStructProperty* StructProperty = CastFieldChecked<FStructProperty>(InProperty);
+	check(StructProperty->Struct == GetScriptStruct());
+
+	const FInstancedStruct& Value = GetValueAtTimeIndex(InTimeIndex, InOutSampleTrackIndex);
+	StructProperty->CopyCompleteValue(OutMemory, Value.GetMemory());
+}
+
+template<>
+void FSampleTrack<TArray<FInstancedStruct>>::AddSampleFromProperty(const FProperty* InProperty, const uint8* InMemory)
+{
+	const FArrayProperty* ArrayProperty = CastFieldChecked<FArrayProperty>(InProperty);
+	const FStructProperty* StructProperty = CastFieldChecked<FStructProperty>(ArrayProperty->Inner);
+	check(StructProperty->Struct == GetScriptStruct());
+	FScriptArrayHelper ArrayHelper(ArrayProperty, InMemory);
+
+	TArray<FInstancedStruct> NewValues;
+	NewValues.Reserve(ArrayHelper.Num());
+	for(int32 Index = 0; Index < ArrayHelper.Num(); Index++)
+	{
+		FInstancedStruct Value(GetScriptStruct());
+		StructProperty->CopyCompleteValue(Value.GetMutableMemory(), ArrayHelper.GetElementPtr(Index));
+		NewValues.Add(Value);
+	}
+	AddSample(NewValues);
+}
+
+template<>
+void FSampleTrack<TArray<FInstancedStruct>>::GetSampleForProperty(int32 InTimeIndex, FSampleTrackIndex& InOutSampleTrackIndex, const FProperty* InProperty, uint8* OutMemory) const
+{
+	const FArrayProperty* ArrayProperty = CastFieldChecked<FArrayProperty>(InProperty);
+	const FStructProperty* StructProperty = CastFieldChecked<FStructProperty>(ArrayProperty->Inner);
+	check(StructProperty->Struct == GetScriptStruct());
+
+	const TArray<FInstancedStruct>& ValuesAtTime = GetValueAtTimeIndex(InTimeIndex, InOutSampleTrackIndex);
+
+	FScriptArrayHelper ArrayHelper(ArrayProperty, OutMemory);
+	ArrayHelper.Resize(ValuesAtTime.Num());
+
+	for(int32 Index = 0; Index < ValuesAtTime.Num(); Index++)
+	{
+		const FInstancedStruct& Value = ValuesAtTime[Index];
+		StructProperty->CopyCompleteValue(ArrayHelper.GetElementPtr(Index), Value.GetMemory());
+	}
+}
diff --git a/Engine/Plugins/Animation/ControlRig/Source/ControlRig/Private/Tracks/SampleTrackBase.cpp b/Engine/Plugins/Animation/ControlRig/Source/ControlRig/Private/Tracks/SampleTrackBase.cpp
--- a/Engine/Plugins/Animation/ControlRig/Source/ControlRig/Private/Tracks/SampleTrackBase.cpp
+++ b/Engine/Plugins/Animation/ControlRig/Source/ControlRig/Private/Tracks/SampleTrackBase.cpp
@@ -7,16 +7,21 @@
 /////// FSampleTrackIndex
 //////////////////////////////////////////////////////////////////////////////////////////////////////////////////////////////
 
+FSampleTrackIndex::FSampleTrackIndex()
+: bIsSingleton(false)
+{
+}
+
 FSampleTrackIndex::FSampleTrackIndex(int32 NumTracks)
-	: bIsSingleton(false)
+: bIsSingleton(false)
 {
 	Allocate(NumTracks);
 }
 
 FSampleTrackIndex::FSampleTrackIndex(const FSampleTrackContainer& InContainer)
-	: bIsSingleton(false)
+: bIsSingleton(false)
 {
-	Allocate(InContainer.NumTracks());
+	Update(InContainer);
 }
 
 int32& FSampleTrackIndex::GetSample(int32 InTrackIndex)
@@ -39,6 +44,11 @@ FSampleTrackIndex FSampleTrackIndex::MakeSingleton()
 	return Index;
 }
 
+void FSampleTrackIndex::Update(const FSampleTrackContainer& InContainer)
+{
+	Allocate(InContainer.NumTracks());
+}
+
 void FSampleTrackIndex::Allocate(int32 NumTracks)
 {
 	if(Samples.Num() >= NumTracks)
@@ -53,7 +63,7 @@ void FSampleTrackIndex::Allocate(int32 NumTracks)
 //////////////////////////////////////////////////////////////////////////////////////////////////////////////////////////////
 
 FSampleTrackBase::FSampleTrackBase()
-: Names({NAME_None})
+: Names()
 , TrackIndex(INDEX_NONE)
 , ReferencedTrackIndex(INDEX_NONE)
 , ReferencedTimeIndicesRange({INDEX_NONE,INDEX_NONE})
@@ -80,6 +90,16 @@ FSampleTrackBase::FSampleTrackBase(const FName& InName)
 {
 }
 
+const FName& FSampleTrackBase::GetName() const
+{
+	if(Names.IsEmpty())
+	{
+		static const FName NoneName(NAME_None);
+		return NoneName;
+	}
+	return Names[0];
+}
+
 int32 FSampleTrackBase::NumTimes() const
 {
 	return NumTimesInContainer;
diff --git a/Engine/Plugins/Animation/ControlRig/Source/ControlRig/Private/Tracks/SampleTrackContainer.cpp b/Engine/Plugins/Animation/ControlRig/Source/ControlRig/Private/Tracks/SampleTrackContainer.cpp
--- a/Engine/Plugins/Animation/ControlRig/Source/ControlRig/Private/Tracks/SampleTrackContainer.cpp
+++ b/Engine/Plugins/Animation/ControlRig/Source/ControlRig/Private/Tracks/SampleTrackContainer.cpp
@@ -46,6 +46,16 @@ FSampleTrackHost& FSampleTrackHost::operator=(const FSampleTrackHost& InOther)
 	return *this;
 }
 
+void FSampleTrackHost::Reset()
+{
+	Container->Reset();
+}
+
+void FSampleTrackHost::Compact()
+{
+	Container->Compact();
+}
+
 bool FSampleTrackHost::Serialize(FArchive& InArchive)
 {
 	if(Container == nullptr)
@@ -300,6 +310,27 @@ void FSampleTrackContainer::SetForceToUseCompression(bool InForce)
 	bForceToUseCompression = InForce;
 }
 
+TSharedPtr<FSampleTrackBase> FSampleTrackContainer::AddTrack(const FName& InName, FSampleTrackBase::ETrackType InTrackType, const UScriptStruct* InScriptStruct)
+{
+	TSharedPtr<FSampleTrackBase> Track = MakeTrack(InTrackType);
+	Track->Names.Reset();
+	Track->Names.Add(InName);
+	Track->ScriptStruct = InScriptStruct;
+	AddTrack(Track);
+	return Track;
+}
+
+TSharedPtr<FSampleTrackBase> FSampleTrackContainer::FindOrAddTrack(const FName& InName, FSampleTrackBase::ETrackType InTrackType,
+	const UScriptStruct* InScriptStruct)
+{
+	TSharedPtr<FSampleTrackBase> Track = FindTrack(InName);
+	if(!Track)
+	{
+		Track = AddTrack(InName, InTrackType, InScriptStruct);
+	}
+	return Track;
+}
+
 TSharedPtr<FSampleTrack<bool>> FSampleTrackContainer::AddBoolTrack(const FName& InName)
 {
 	TSharedPtr<FSampleTrack<bool>> Track = MakeShareable(new FSampleTrack<bool>(InName));
@@ -495,6 +526,30 @@ FVector2f FSampleTrackContainer::GetTimeRange() const
 	return {TimeSamples[0], TimeSamples.Last()};
 }
 
+void FSampleTrackContainer::GetTimeIndex(float InTime, int32& InOutTimeIndex) const
+{
+	if(TimeSamples.IsEmpty())
+	{
+		InOutTimeIndex = INDEX_NONE;
+		return;
+	}
+
+	InOutTimeIndex = FMath::Clamp(InOutTimeIndex, 0, TimeSamples.Num() - 1);
+	if(TimeSamples.Num() == 1)
+	{
+		return;
+	}
+
+	while(TimeSamples.IsValidIndex(InOutTimeIndex-1) && TimeSamples[InOutTimeIndex] > InTime)
+	{
+		InOutTimeIndex--;
+	}
+	while(TimeSamples.IsValidIndex(InOutTimeIndex+1) && TimeSamples[InOutTimeIndex+1] < InTime)
+	{
+		InOutTimeIndex++;
+	}
+}
+
 int32 FSampleTrackContainer::GetTrackIndex(const FName& InName) const
 {
 	if(const int32* Index = NameToIndex.Find(InName))
@@ -540,14 +595,17 @@ void FSampleTrackContainer::AddTrack(TSharedPtr<FSampleTrackBase> InTrack, bool
 	if(InTrack->IsComposed() && bCreateChildTracks)
 	{
 		const TArray<FSampleTrackBase::ETrackType> ChildTrackTypes = InTrack->GetChildTrackTypes();
-		for(const FSampleTrackBase::ETrackType ChildTrackType : ChildTrackTypes)
+		for(int32 ChildTrackIndex = 0; ChildTrackIndex < ChildTrackTypes.Num(); ChildTrackIndex++)
 		{
-			FName ChildTrackName = TEXT("ChildTrack");
+			const FSampleTrackBase::ETrackType ChildTrackType = ChildTrackTypes[ChildTrackIndex];
+			
+			FName ChildTrackName = *(TEXT("ChildTrack_") + InTrack->GetChildTrackNameSuffix(ChildTrackIndex));
 			if(!InTrack->Names.IsEmpty())
 			{
-				ChildTrackName = *FString::Printf(TEXT("%s_Child_%d"), *InTrack->GetName().ToString(), ChildTrackIndices.Num());
+				ChildTrackName = *(InTrack->GetName().ToString() + TEXT("_") + InTrack->GetChildTrackNameSuffix(ChildTrackIndex));
 			}
 			TSharedPtr<FSampleTrackBase> ChildTrack = MakeTrack(ChildTrackType);
+			ChildTrack->Names.Reset();
 			ChildTrack->Names.Add(ChildTrackName);
 			AddTrack(ChildTrack);
 			ChildTrackIndices.Add(ChildTrack->GetTrackIndex());
@@ -885,6 +943,7 @@ void FSampleTrackContainer::RemoveRedundantTracks(bool bUpdateNameToIndexMap, fl
 
 			// merge the track
 			TrackA->Names.Append(TrackB->Names);
+			TrackA->Names.Remove(NAME_None);
 			MergedTracks[IndexB] = true;
 			OldTrackIndexToNewTrackIndex[IndexB] = IndexA;
 			NumTracksMerged++;
diff --git a/Engine/Plugins/Animation/ControlRig/Source/ControlRig/Public/ControlRigReplay.h b/Engine/Plugins/Animation/ControlRig/Source/ControlRig/Public/ControlRigReplay.h
--- a/Engine/Plugins/Animation/ControlRig/Source/ControlRig/Public/ControlRigReplay.h
+++ b/Engine/Plugins/Animation/ControlRig/Source/ControlRig/Public/ControlRigReplay.h
@@ -47,9 +47,101 @@ enum class EControlRigReplayPlaybackMode : uint8
 USTRUCT()
 struct CONTROLRIG_API FControlRigReplayTracks : public FSampleTrackHost
 {
+public:
 	GENERATED_BODY()
 
+	typedef TFunction<void(EMessageSeverity::Type,const FName&,const FString&)> TReportFunction;
+	
+	virtual ~FControlRigReplayTracks() override {}
+
 	virtual bool Serialize(FArchive& InArchive) override;
+
+	virtual void Reset() override;
+
+	int32 AddTimeSample(float InTime);
+	int32 AddTimeSampleFromDeltaTime(float InDeltaTime);
+	void GetTimeIndex(float InTime, int32& InOutTimeIndex) const;
+	const TArray<float>& GetTimeSamples() const;
+
+	void StoreHierarchy(
+		URigHierarchy* InHierarchy,
+		const TArrayView<FRigElementKey>& InKeys = TArrayView<FRigElementKey>(),
+		bool bStorePose = true,
+		bool bStoreComponents = true,
+		bool bStoreMetadata = true);
+	bool RestoreHierarchy(
+		int32 InTimeIndex,
+		URigHierarchy* InHierarchy,
+		const TArrayView<FRigElementKey>& InKeys = TArrayView<FRigElementKey>(),
+		TReportFunction InReportFunction = nullptr,
+		bool bRestorePose = true,
+		bool bRestoreComponents = true,
+		bool bRestoreMetadata = true) const;
+
+	void StorePose(
+		URigHierarchy* InHierarchy,
+		const TArrayView<FRigElementKey>& InKeys = TArrayView<FRigElementKey>());
+	bool RestorePose(
+		int32 InTimeIndex,
+		URigHierarchy* InHierarchy,
+		const TArrayView<FRigElementKey>& InKeys = TArrayView<FRigElementKey>(),
+		TReportFunction InReportFunction = nullptr) const;
+
+	void StoreComponents(
+		URigHierarchy* InHierarchy,
+		const TArrayView<FRigComponentKey>& InKeys = TArrayView<FRigComponentKey>());
+	bool RestoreComponents(
+		int32 InTimeIndex,
+		URigHierarchy* InHierarchy,
+		const TArrayView<FRigComponentKey>& InKeys = TArrayView<FRigComponentKey>(),
+		TReportFunction InReportFunction = nullptr) const;
+
+	void StoreMetaData(
+		URigHierarchy* InHierarchy,
+		const TArrayView<FRigElementKey>& InKeys = TArrayView<FRigElementKey>());
+	bool RestoreMetaData(
+		int32 InTimeIndex,
+		URigHierarchy* InHierarchy,
+		const TArrayView<FRigElementKey>& InKeys = TArrayView<FRigElementKey>(),
+		TReportFunction InReportFunction = nullptr) const;
+
+	void StoreVariables(URigVMHost* InHost);
+	bool RestoreVariables(int32 InTimeIndex, URigVMHost* InHost, TReportFunction InReportFunction = nullptr) const;
+
+private:
+	const FName& GetTrackName(const FRigElementKey& InElementKey) const;
+	const FName& GetTrackName(const FRigComponentKey& InComponentKey) const;
+	const FName& GetTrackName(const FRigElementKey& InElementKey, const FName& InMetadataName) const;
+	const FName& GetTrackName(const FProperty* InProperty) const;
+	static FSampleTrackBase::ETrackType GetTrackTypeFromMetadata(ERigMetadataType InMetadataType);
+	static FSampleTrackBase::ETrackType GetTrackTypeFromProperty(const FProperty* InProperty) ;
+
+	void StorePose(URigHierarchy* InHierarchy, FRigBaseElement* InElement);
+	bool RestorePose(int32 InTimeIndex, URigHierarchy* InHierarchy, FRigBaseElement* InElement, TReportFunction InReportFunction) const;
+
+	void StoreComponents(URigHierarchy* InHierarchy, FRigBaseElement* InElement);
+	bool RestoreComponents(int32 InTimeIndex, URigHierarchy* InHierarchy, FRigBaseElement* InElement, TReportFunction InReportFunction) const;
+
+	void StoreComponent(URigHierarchy* InHierarchy, FRigBaseComponent* InComponent);
+	bool RestoreComponent(int32 InTimeIndex, URigHierarchy* InHierarchy, FRigBaseComponent* InComponent, TReportFunction InReportFunction) const;
+
+	void StoreMetaData(URigHierarchy* InHierarchy, FRigBaseElement* InElement);
+	bool RestoreMetaData(int32 InTimeIndex, URigHierarchy* InHierarchy, FRigBaseElement* InElement, TReportFunction InReportFunction) const;
+
+	void StoreMetaData(URigHierarchy* InHierarchy, FRigBaseElement* InElement, FRigBaseMetadata* InMetadata);
+	bool RestoreMetaData(int32 InTimeIndex, URigHierarchy* InHierarchy, FRigBaseElement* InElement, FRigBaseMetadata* InMetadata, TReportFunction InReportFunction) const;
+
+	static bool ForEachElement(URigHierarchy* InHierarchy, const TArrayView<FRigElementKey>& InKeys, TFunction<void(FRigBaseElement*, bool&)> InFunction);
+	static bool ForEachComponent(URigHierarchy* InHierarchy, const TArrayView<FRigComponentKey>& InKeys, TFunction<void(FRigBaseComponent*, bool&)> InFunction);
+
+	static inline FLazyName HierarchyTopologyHashName = FLazyName(TEXT("HierarchyTopologyHash"));
+	static inline FLazyName HierarchyTopologyParentIndicesName = FLazyName(TEXT("HierarchyTopologyParentIndices"));
+	mutable TMap<FRigElementKey, FName> ElementKeyToTrackName;
+	mutable TMap<FRigComponentKey, FName> ComponentKeyToTrackName;
+	mutable TMap<TTuple<FRigElementKey, FName>, FName> MetadataToTrackName;
+	mutable TMap<FName, FName> PropertyNameToTrackName;
+	mutable FSampleTrackIndex SampleTrackIndex;
+	TArray<int32> ParentIndices;
 };
 
 template<>
@@ -73,6 +165,9 @@ public:
 	UControlRigReplay()
 		: Tolerance(0.001)
 		, EnableTest(true)
+		, TimeAtStartOfRecording(0.0)
+		, TimeOfLastFrame(0.0)
+		, DesiredRecordingDuration(0.0)
 	{}
 
 	virtual void BeginDestroy() override;
@@ -104,7 +199,7 @@ public:
 	bool EnableTest;
 
 	UFUNCTION(BlueprintCallable, Category = "ControlRigReplay")
-	virtual bool StartRecording(UControlRig* InControlRig, double InRecordingDuration = 0.0);
+	virtual bool StartRecording(UControlRig* InControlRig);
 
 	UFUNCTION(BlueprintCallable, Category = "ControlRigReplay")
 	virtual bool StopRecording();
@@ -142,4 +237,9 @@ protected:
 	FDelegateHandle PreConstructionHandle;
 	FDelegateHandle PreEventHandle;
 	FDelegateHandle PostEventHandle;
+	double TimeAtStartOfRecording;
+	double TimeOfLastFrame;
+	double DesiredRecordingDuration;
+
+	friend class FControlRigEditor;
 };
\ No newline at end of file
diff --git a/Engine/Plugins/Animation/ControlRig/Source/ControlRig/Public/ControlRigTestData.h b/Engine/Plugins/Animation/ControlRig/Source/ControlRig/Public/ControlRigTestData.h
--- a/Engine/Plugins/Animation/ControlRig/Source/ControlRig/Public/ControlRigTestData.h
+++ b/Engine/Plugins/Animation/ControlRig/Source/ControlRig/Public/ControlRigTestData.h
@@ -57,8 +57,6 @@ public:
 
 	UControlRigTestData()
 		: LastFrameIndex(INDEX_NONE)
-		, DesiredRecordingDuration(0.0)
-		, TimeAtStartOfRecording(0.0)
 		, bIsApplyingOutputs(false)
 	{}
 
@@ -71,7 +69,7 @@ public:
 	UFUNCTION(BlueprintPure, Category = "ControlRigTestData")
 	int32 GetFrameIndexForTime(double InSeconds, bool bInput = false) const;
 
-	virtual bool StartRecording(UControlRig* InControlRig, double InRecordingDuration = 0.0) override;
+	virtual bool StartRecording(UControlRig* InControlRig) override;
 	virtual bool StopRecording() override;
 	virtual bool StartReplay(UControlRig* InControlRig) override;
 	virtual bool StopReplay() override;
@@ -92,7 +90,5 @@ public:
 protected:
 
 	mutable int32 LastFrameIndex;
-	double DesiredRecordingDuration;
-	double TimeAtStartOfRecording;
 	bool bIsApplyingOutputs;
 };
\ No newline at end of file
diff --git a/Engine/Plugins/Animation/ControlRig/Source/ControlRig/Public/Rigs/RigHierarchy.h b/Engine/Plugins/Animation/ControlRig/Source/ControlRig/Public/Rigs/RigHierarchy.h
--- a/Engine/Plugins/Animation/ControlRig/Source/ControlRig/Public/Rigs/RigHierarchy.h
+++ b/Engine/Plugins/Animation/ControlRig/Source/ControlRig/Public/Rigs/RigHierarchy.h
@@ -798,6 +798,9 @@ public:
 	UFUNCTION(BlueprintPure, Category = URigHierarchy)
 	int32 NumComponents(FRigElementKey InElement) const;
 
+	UFUNCTION(BlueprintPure, Category = URigHierarchy)
+	TArray<FRigComponentKey> GetAllComponentKeys() const;
+
 	UFUNCTION(BlueprintPure, Category = URigHierarchy)
 	TArray<FRigComponentKey> GetComponentKeys(FRigElementKey InElement) const;
 
diff --git a/Engine/Plugins/Animation/ControlRig/Source/ControlRig/Public/Rigs/RigHierarchyComponents.h b/Engine/Plugins/Animation/ControlRig/Source/ControlRig/Public/Rigs/RigHierarchyComponents.h
--- a/Engine/Plugins/Animation/ControlRig/Source/ControlRig/Public/Rigs/RigHierarchyComponents.h
+++ b/Engine/Plugins/Animation/ControlRig/Source/ControlRig/Public/Rigs/RigHierarchyComponents.h
@@ -162,6 +162,7 @@ protected:
 	friend class FControlRigEditor;
 	friend class URigHierarchy;
 	friend class URigHierarchyController;
+	friend struct FControlRigReplayTracks;
 };
 
 template<>
diff --git a/Engine/Plugins/Animation/ControlRig/Source/ControlRig/Public/Tracks/ComposedSampleTrack.h b/Engine/Plugins/Animation/ControlRig/Source/ControlRig/Public/Tracks/ComposedSampleTrack.h
--- a/Engine/Plugins/Animation/ControlRig/Source/ControlRig/Public/Tracks/ComposedSampleTrack.h
+++ b/Engine/Plugins/Animation/ControlRig/Source/ControlRig/Public/Tracks/ComposedSampleTrack.h
@@ -32,6 +32,11 @@ public:
 		return {};
 	}
 
+	virtual FString GetChildTrackNameSuffix(int32 InChildTrackIndex) const override
+	{
+		return FSampleTrack<T>::GetChildTrackNameSuffix(InChildTrackIndex);
+	}
+
 	virtual ~FComposedSampleTrack() override {}
 
 	// returns the value of this track at a given time index
@@ -197,6 +202,9 @@ protected:
 template<>
 CONTROLRIG_API TArray<FSampleTrackBase::ETrackType> FComposedSampleTrack<FTransform3f>::GetChildTrackTypes() const;
 
+template<>
+CONTROLRIG_API FString FComposedSampleTrack<FTransform3f>::GetChildTrackNameSuffix(int32 InChildTrackIndex) const;
+
 template<>
 CONTROLRIG_API FTransform3f FComposedSampleTrack<FTransform3f>::GetValueAtTimeIndex(int32 InTimeIndex, FSampleTrackIndex& InOutIndex) const;
 	
diff --git a/Engine/Plugins/Animation/ControlRig/Source/ControlRig/Public/Tracks/SampleTrack.h b/Engine/Plugins/Animation/ControlRig/Source/ControlRig/Public/Tracks/SampleTrack.h
--- a/Engine/Plugins/Animation/ControlRig/Source/ControlRig/Public/Tracks/SampleTrack.h
+++ b/Engine/Plugins/Animation/ControlRig/Source/ControlRig/Public/Tracks/SampleTrack.h
@@ -186,6 +186,21 @@ public:
 		AddSample(InstancedStructArray, InTolerance);
 	}
 
+	// helper to store a sample based on a given property
+	virtual void AddSampleFromProperty(const FProperty* InProperty, const uint8* InMemory) override
+	{
+		T Value;
+		InProperty->CopyCompleteValue(&Value, InMemory);
+		AddSample(Value);
+	}
+
+	// helper to retrieve a sample based on a given property
+	virtual void GetSampleForProperty(int32 InTimeIndex, FSampleTrackIndex& InOutSampleTrackIndex, const FProperty* InProperty, uint8* OutMemory) const override
+	{
+		const T& Value = GetValueAtTimeIndex(InTimeIndex, InOutSampleTrackIndex);
+		InProperty->CopyCompleteValue(OutMemory, &Value);
+	}
+
 	// returns the type of this track
 	virtual ETrackType GetTrackType() const override
 	{
@@ -795,3 +810,15 @@ CONTROLRIG_API void FSampleTrack<TArray<FInstancedStruct>>::Serialize(FArchive&
 
 template<>
 CONTROLRIG_API int32 FSampleTrack<FInstancedStruct>::GetSizePerValue() const;
+
+template<>
+CONTROLRIG_API void FSampleTrack<FInstancedStruct>::AddSampleFromProperty(const FProperty* InProperty, const uint8* InMemory);
+
+template<>
+CONTROLRIG_API void FSampleTrack<FInstancedStruct>::GetSampleForProperty(int32 InTimeIndex, FSampleTrackIndex& InOutSampleTrackIndex, const FProperty* InProperty, uint8* OutMemory) const;
+
+template<>
+CONTROLRIG_API void FSampleTrack<TArray<FInstancedStruct>>::AddSampleFromProperty(const FProperty* InProperty, const uint8* InMemory);
+
+template<>
+CONTROLRIG_API void FSampleTrack<TArray<FInstancedStruct>>::GetSampleForProperty(int32 InTimeIndex, FSampleTrackIndex& InOutSampleTrackIndex, const FProperty* InProperty, uint8* OutMemory) const;
diff --git a/Engine/Plugins/Animation/ControlRig/Source/ControlRig/Public/Tracks/SampleTrackBase.h b/Engine/Plugins/Animation/ControlRig/Source/ControlRig/Public/Tracks/SampleTrackBase.h
--- a/Engine/Plugins/Animation/ControlRig/Source/ControlRig/Public/Tracks/SampleTrackBase.h
+++ b/Engine/Plugins/Animation/ControlRig/Source/ControlRig/Public/Tracks/SampleTrackBase.h
@@ -12,10 +12,12 @@ class FSampleTrackContainer;
 struct CONTROLRIG_API FSampleTrackIndex
 {
 public:
+	FSampleTrackIndex();
 	explicit FSampleTrackIndex(int32 NumTracks);
 	explicit FSampleTrackIndex(const FSampleTrackContainer& InContainer);
 	int32& GetSample(int32 InTrackIndex);
 	static FSampleTrackIndex MakeSingleton();
+	void Update(const FSampleTrackContainer& InContainer);
 
 private:
 	void Allocate(int32 NumTracks);
@@ -75,7 +77,7 @@ public:
 	bool IsValid() const { return Mode != EMode_Invalid; }
 
 	// the first name of this track
-	const FName& GetName() const { return Names[0]; }
+	const FName& GetName() const;
 
 	// all names this track is using / this track represents
 	const TArray<FName>& GetAllNames() const { return Names; }
@@ -151,6 +153,15 @@ public:
 	// returns the required track types for the child tracks
 	virtual TArray<FSampleTrackBase::ETrackType> GetChildTrackTypes() const { return {}; }
 
+	// returns the name for a given child track
+	virtual FString GetChildTrackNameSuffix(int32 InChildTrackIndex) const { return TEXT("Child"); }
+
+	// helper to store a sample based on a given property
+	virtual void AddSampleFromProperty(const FProperty* InProperty, const uint8* InMemory) = 0;
+
+	// helper to retrieve a sample based on a given property
+	virtual void GetSampleForProperty(int32 InTimeIndex, FSampleTrackIndex& InOutSampleTrackIndex, const FProperty* InProperty, uint8* OutMemory) const = 0;
+
 protected:
 
 	const FSampleTrackBase* GetReferencedTrack() const;
diff --git a/Engine/Plugins/Animation/ControlRig/Source/ControlRig/Public/Tracks/SampleTrackContainer.h b/Engine/Plugins/Animation/ControlRig/Source/ControlRig/Public/Tracks/SampleTrackContainer.h
--- a/Engine/Plugins/Animation/ControlRig/Source/ControlRig/Public/Tracks/SampleTrackContainer.h
+++ b/Engine/Plugins/Animation/ControlRig/Source/ControlRig/Public/Tracks/SampleTrackContainer.h
@@ -20,6 +20,9 @@ public:
 
 	FSampleTrackHost& operator =(const FSampleTrackHost& InOther);
 
+	virtual void Reset();
+	virtual void Compact();
+
 	/** By always returning false here, we can disable delta serialization */
 	virtual bool Identical(const FSampleTrackHost* Other, uint32 PortFlags) const { return false; }
 
@@ -57,6 +60,8 @@ public:
 	bool Serialize(FArchive& InArchive);
 	void SetForceToUseCompression(bool InForce = true);
 
+	TSharedPtr<FSampleTrackBase> AddTrack(const FName& InName, FSampleTrackBase::ETrackType InTrackType, const UScriptStruct* InScriptStruct = nullptr);
+	TSharedPtr<FSampleTrackBase> FindOrAddTrack(const FName& InName, FSampleTrackBase::ETrackType InTrackType, const UScriptStruct* InScriptStruct = nullptr);
 	TSharedPtr<FSampleTrack<bool>> AddBoolTrack(const FName& InName);
 	TSharedPtr<FSampleTrack<int32>> AddInt32Track(const FName& InName);
 	TSharedPtr<FSampleTrack<uint32>> AddUint32Track(const FName& InName);
@@ -91,6 +96,7 @@ public:
 
 	FVector2f GetTimeRange() const;
 	const TArray<float>& GetTimeSamples() const { return TimeSamples; }
+	void GetTimeIndex(float InTime, int32& InOutTimeIndex) const;
 
 	int32 GetTrackIndex(const FName& InName) const;
 	TSharedPtr<const FSampleTrackBase> GetTrack(int32 InIndex) const;
diff --git a/Engine/Plugins/Animation/ControlRig/Source/ControlRigEditor/Private/Editor/ControlRigEditor.cpp b/Engine/Plugins/Animation/ControlRig/Source/ControlRigEditor/Private/Editor/ControlRigEditor.cpp
--- a/Engine/Plugins/Animation/ControlRig/Source/ControlRigEditor/Private/Editor/ControlRigEditor.cpp
+++ b/Engine/Plugins/Animation/ControlRig/Source/ControlRigEditor/Private/Editor/ControlRigEditor.cpp
@@ -637,7 +637,8 @@ bool FControlRigEditor::RecordReplay(double InRecordingDuration)
 	
 	if(UControlRigReplay* Replay = ReplayStrongPtr.Get())
 	{
-		Replay->StartRecording(GetControlRig(), InRecordingDuration);
+		Replay->StartRecording(GetControlRig());
+		Replay->DesiredRecordingDuration = InRecordingDuration;
 	}
 	return true;
 }
