Constraints: added the ability to select spawnable objects as parents

#rnx

diff --git a/Engine/Plugins/Animation/ControlRig/Source/ControlRigEditor/Private/Editor/Constraints/SConstraintsWidget.cpp b/Engine/Plugins/Animation/ControlRig/Source/ControlRigEditor/Private/Editor/Constraints/SConstraintsWidget.cpp
--- a/Engine/Plugins/Animation/ControlRig/Source/ControlRigEditor/Private/Editor/Constraints/SConstraintsWidget.cpp
+++ b/Engine/Plugins/Animation/ControlRig/Source/ControlRigEditor/Private/Editor/Constraints/SConstraintsWidget.cpp
@@ -1509,7 +1509,7 @@ void SConstraintsEditionWidget::Tick(const FGeometry& AllottedGeometry, const do
 
 TSharedRef<SWidget> SConstraintsEditionWidget::OnGetParentMenuContent()
 {
-	static constexpr bool bAllowClear = true, bAllowPicking = false, bUseSelected = true;
+	static constexpr bool bAllowClear = true, bAllowPicking = false, bUseSelected = true, bShowTransient = true;
 
 	FOnShouldFilterActor FilterActorDelegate = FOnShouldFilterActor::CreateLambda([](const AActor* InActor)
 	{
@@ -1536,7 +1536,8 @@ TSharedRef<SWidget> SConstraintsEditionWidget::OnGetParentMenuContent()
 		SelectedActorDelegate,
 		OnCloseDelegate,
 		UseSelectedDelegate,
-		bUseSelected);
+		bUseSelected,
+		bShowTransient);
 }
 
 void SConstraintsEditionWidget::OnParentMenuOpenChanged(bool bOpened)
