[Viewport Toolbar] Control Rig category options are missing in the Viewport Toolbar of Control Rig editor

Creates shared code to be used by both old and new toolbar.
Not using FExtender for this specific set of entries, since in this case they are added in a way which is specific to AnimViewport + AnimViewportToolbar.
Also, they are added from the ControlRig module using a delegate, after the new toolbar context menu is created, another reason why adding the existing FExtender to the FToolMenuContext is not straightforward.

diff --git a/Engine/Plugins/Animation/ControlRig/Source/ControlRigEditor/Private/Editor/ControlRigEditor.cpp b/Engine/Plugins/Animation/ControlRig/Source/ControlRigEditor/Private/Editor/ControlRigEditor.cpp
--- a/Engine/Plugins/Animation/ControlRig/Source/ControlRigEditor/Private/Editor/ControlRigEditor.cpp
+++ b/Engine/Plugins/Animation/ControlRig/Source/ControlRigEditor/Private/Editor/ControlRigEditor.cpp
@@ -116,9 +116,141 @@
 #include "Widgets/SRigVMSwapAssetReferencesWidget.h"
 #include "Widgets/SRigVMBulkEditDialog.h"
 #include "ControlRigTestData.h"
+#include "SEditorViewport.h"
+#include "Editor/Persona/Private/SAnimationEditorViewport.h"
+#include "ViewportToolbar/UnrealEdViewportToolbar.h"
+#include "ViewportToolbar/UnrealEdViewportToolbarContext.h"
 
 #define LOCTEXT_NAMESPACE "ControlRigEditor"
 
+namespace UE::ControlRigEditor::Private
+{
+TOptional<float> GetToolbarAxesScale()
+{
+	if (const UControlRigEditModeSettings* Settings = GetDefault<UControlRigEditModeSettings>())
+	{
+		return Settings->AxisScale;
+	}
+	return 0.0f;
+}
+
+void OnToolbarAxesScaleChanged(float InValue)
+{
+	if (UControlRigEditModeSettings* Settings = GetMutableDefault<UControlRigEditModeSettings>())
+	{
+		Settings->AxisScale = InValue;
+	}
+}
+
+UControlRigBlueprint* GetControlRigBlueprintFromContext(const UUnrealEdViewportToolbarContext* InContext)
+{
+	if (InContext)
+	{
+		if (TSharedPtr<SEditorViewport> Viewport = InContext->Viewport.Pin())
+		{
+			if (const TSharedPtr<FAnimationViewportClient>& AnimViewportClient =
+					StaticCastSharedPtr<FAnimationViewportClient>(Viewport->GetViewportClient()))
+			{
+				if (AnimViewportClient->GetAssetEditorToolkit()->IsBlueprintEditor())
+				{
+					const TSharedRef<FBlueprintEditor>& BlueprintEditor =
+						StaticCastSharedRef<FBlueprintEditor>(AnimViewportClient->GetAssetEditorToolkit());
+
+					return Cast<UControlRigBlueprint>(BlueprintEditor->GetBlueprintObj());
+				}
+			}
+		}
+	}
+
+	return nullptr;
+}
+
+void FillControlRigMenu(UToolMenu* InMenu)
+{
+	if (const UUnrealEdViewportToolbarContext* const Context = InMenu->FindContext<UUnrealEdViewportToolbarContext>())
+	{
+		if (const UControlRigBlueprint* const ControlRigBlueprint = GetControlRigBlueprintFromContext(Context))
+		{
+			FToolMenuSection& ControlRigSection = InMenu->FindOrAddSection("ControlRig", LOCTEXT("ControlRig_Label", "Control Rig"));
+			ControlRigSection.AddMenuEntry(FControlRigEditorCommands::Get().ToggleControlVisibility);
+			ControlRigSection.AddMenuEntry(FControlRigEditorCommands::Get().ToggleControlsAsOverlay);
+			ControlRigSection.AddMenuEntry(FControlRigEditorCommands::Get().ToggleDrawNulls);
+			ControlRigSection.AddMenuEntry(FControlRigEditorCommands::Get().ToggleDrawSockets);
+			ControlRigSection.AddMenuEntry(FControlRigEditorCommands::Get().ToggleDrawAxesOnSelection);
+
+			// Asset Axes Custom Widget
+			// clang-format off
+			TSharedRef<SWidget> AxesScaleEntryWidget = SNew(SBox)
+				.HAlign(HAlign_Right)
+				[
+					SNew(SBox)
+					.Padding(FMargin(4.0f, 0.0f, 0.0f, 0.0f))
+					.WidthOverride(100.0f)
+					[
+						SNew(SNumericEntryBox<float>)
+						.Font(FAppStyle::GetFontStyle(TEXT("MenuItem.Font")))
+						.AllowSpin(true)
+						.MinSliderValue(0.0f)
+						.MaxSliderValue(100.0f)
+						.Value_Static(&UE::ControlRigEditor::Private::GetToolbarAxesScale)
+						.OnValueChanged_Static(&UE::ControlRigEditor::Private::OnToolbarAxesScaleChanged)
+						.ToolTipText(LOCTEXT("ControlRigAxesScaleToolTip", "Scale of axes drawn for selected rig elements"))
+					]
+				];
+			// clang-format on
+
+			ControlRigSection.AddEntry(
+				FToolMenuEntry::InitWidget("AxesScale", AxesScaleEntryWidget, LOCTEXT("ControlRigAxesScale", "Axes Scale"))
+			);
+
+			// This entry only for Modular Rig
+			if (ControlRigBlueprint->IsModularRig())
+			{
+				FToolMenuSection& ModularRigSection =
+					InMenu->FindOrAddSection("ModularRig", LOCTEXT("ModularRig_Label", "Modular Rig"));
+				ModularRigSection.AddMenuEntry(FControlRigEditorCommands::Get().ToggleSchematicViewportVisibility);
+			}
+		}
+	}
+}
+
+TSharedRef<SWidget> GenerateControlRigMenuWidget(const TWeakPtr<SEditorViewport>& InViewport)
+{
+	if (TSharedPtr<SEditorViewport> Viewport = InViewport.Pin())
+	{
+		// We generate a menu via UToolMenus, so we can use FillControlRigMenu call from both old and new toolbar
+		FName OldFollowModeMenuName = "AnimationEditor.OldViewportToolbar.ControlRig";
+		if (!UToolMenus::Get()->IsMenuRegistered(OldFollowModeMenuName))
+		{
+			if (UToolMenu* const Menu = UToolMenus::Get()->RegisterMenu(OldFollowModeMenuName, NAME_None, EMultiBoxType::Menu, false))
+			{
+				Menu->AddDynamicSection(
+					"BaseSection",
+					FNewToolMenuDelegate::CreateLambda(
+						[](UToolMenu* InMenu)
+						{
+							UE::ControlRigEditor::Private::FillControlRigMenu(InMenu);
+						}
+					)
+				);
+			}
+		}
+
+		// Providing the UToolMenu with a context containing the data needed by its current entries
+		FToolMenuContext MenuContext;
+		{
+			UUnrealEdViewportToolbarContext* const ContextObject = UE::UnrealEd::CreateViewportToolbarDefaultContext(Viewport);
+			MenuContext.AddObject(ContextObject);
+
+			MenuContext.AppendCommandList(Viewport->GetCommandList());
+		}
+		return UToolMenus::Get()->GenerateWidget(OldFollowModeMenuName, MenuContext);
+	}
+
+	return SNullWidget::NullWidget;
+}
+} // namespace UE::ControlRigEditor::Private
+
 TAutoConsoleVariable<bool> CVarControlRigShowTestingToolbar(TEXT("ControlRig.Test.EnableTestingToolbar"), false, TEXT("When true we'll show the testing toolbar in Control Rig Editor."));
 TAutoConsoleVariable<bool> CVarShowSchematicPanelOverlay(TEXT("ControlRig.Preview.ShowSchematicPanelOverlay"), true, TEXT("When true we'll add an overlay to the persona viewport to show modular rig information."));
 
@@ -2334,42 +2466,60 @@ void FControlRigEditor::HandleViewportCreated(const TSharedRef<class IPersonaVie
 		ChangeShapeTransformNotificationOptions
 	);
 
-	InViewport->AddToolbarExtender(TEXT("AnimViewportDefaultCamera"), FMenuExtensionDelegate::CreateLambda(
-		[&](FMenuBuilder& InMenuBuilder)
-		{
-			InMenuBuilder.AddMenuSeparator(TEXT("Control Rig"));
-			InMenuBuilder.BeginSection("ControlRig", LOCTEXT("ControlRig_Label", "Control Rig"));
-			{
-				InMenuBuilder.AddMenuEntry(FControlRigEditorCommands::Get().ToggleControlVisibility);
-				InMenuBuilder.AddMenuEntry(FControlRigEditorCommands::Get().ToggleControlsAsOverlay);
-				InMenuBuilder.AddMenuEntry(FControlRigEditorCommands::Get().ToggleDrawNulls);
-				InMenuBuilder.AddMenuEntry(FControlRigEditorCommands::Get().ToggleDrawSockets);
-				InMenuBuilder.AddMenuEntry(FControlRigEditorCommands::Get().ToggleDrawAxesOnSelection);
+	if (const TSharedPtr<SEditorViewport>& EditorViewportWidget = InViewport->GetViewportClient().GetEditorViewportWidget())
+	{
+		TSharedPtr<FUICommandList> CommandList = EditorViewportWidget->GetCommandList();
+		InViewport->AddToolbarExtender(
+			TEXT("AnimViewportDefaultCamera"),
+			FMenuExtensionDelegate::CreateLambda(
+				[EditorViewportWidgetWeak = EditorViewportWidget.ToWeakPtr()](FMenuBuilder& InMenuBuilder)
+				{
+					if (TSharedPtr<SEditorViewport> ViewportPinned = EditorViewportWidgetWeak.Pin())
+					{
+						InMenuBuilder.AddWidget(
+							UE::ControlRigEditor::Private::GenerateControlRigMenuWidget(EditorViewportWidgetWeak), FText()
+						);
+					}
+				}
+			)
+		);
+	}
 
-				InMenuBuilder.AddWidget(
-					SNew(SBox)
-					.HAlign(HAlign_Right)
-					[
-						SNew(SBox)
-						.Padding(FMargin(4.0f, 0.0f, 0.0f, 0.0f))
-						.WidthOverride(100.0f)
-						[
-							SNew(SNumericEntryBox<float>)
-							.Font(FAppStyle::GetFontStyle(TEXT("MenuItem.Font")))
-							.AllowSpin(true)
-							.MinSliderValue(0.0f)
-							.MaxSliderValue(100.0f)
-							.Value(this, &FControlRigEditor::GetToolbarAxesScale)
-							.OnValueChanged(this, &FControlRigEditor::OnToolbarAxesScaleChanged)
-							.ToolTipText(LOCTEXT("ControlRigAxesScaleToolTip", "Scale of axes drawn for selected rig elements"))
-						]
-					], 
-					LOCTEXT("ControlRigAxesScale", "Axes Scale")
-				);
-			}
-			InMenuBuilder.EndSection();
-		}
-	));
+	// Extending for new toolbar as well. Using Extenders setup for the old toolbar was convoluted, since they are currently
+	// stored in the old toolbar, making them difficult to access and handle from the new one. Also, they are being added
+	// from this very same callback (HandleViewportCreated) after the creation of the FToolMenuContext in the new toolbar, making things slightly more complicated.
+	if (UToolMenu* const Submenu = UToolMenus::Get()->ExtendMenu("AnimationEditor.ViewportToolbar.Show"))
+	{
+		FToolMenuSection& ControlRigSection = Submenu->FindOrAddSection(
+			"AnimViewportSceneElements", LOCTEXT("CharacterMenu_SceneElements", "Scene Elements")
+		);
+
+		ControlRigSection.AddDynamicEntry(
+			"ControlRigSubmenu",
+			FNewToolMenuSectionDelegate::CreateLambda(
+				[](FToolMenuSection& InSection)
+				{
+					if (UUnrealEdViewportToolbarContext* Context = InSection.FindContext<UUnrealEdViewportToolbarContext>())
+					{
+						const bool bIsControlRigEditor = !!UE::ControlRigEditor::Private::GetControlRigBlueprintFromContext(Context);
+
+						// This Submenu needs to be visible only for Control Rig Viewport Toolbars
+						if (bIsControlRigEditor)
+						{
+							InSection.AddSubMenu(
+								"ControlRigSubmenu",
+								LOCTEXT("ControlRigSubmenuLabel", "Control Rig"),
+								LOCTEXT("ControlRigSubmenuLabelToolTip", "Control Rig Options"),
+								FNewToolMenuChoice(
+									FNewToolMenuDelegate::CreateStatic(&UE::ControlRigEditor::Private::FillControlRigMenu)
+								)
+							);
+						}
+					}
+				}
+			)
+		);
+	}
 
 	auto GetBorderColorAndOpacity = [this]()
 	{
@@ -2426,18 +2576,6 @@ void FControlRigEditor::HandleViewportCreated(const TSharedRef<class IPersonaVie
 																.Visibility(this, &FControlRigEditor::GetSchematicOverlayVisibility)
 				;
 				InViewport->AddOverlayWidget(SchematicViewport.ToSharedRef());
-
-				
-				InViewport->AddToolbarExtender(TEXT("ControlRig"), FMenuExtensionDelegate::CreateLambda([&](FMenuBuilder& InMenuBuilder)
-					{
-						InMenuBuilder.AddMenuSeparator(TEXT("Modular Rig"));
-						InMenuBuilder.BeginSection("ModularRig", LOCTEXT("ModularRig_Label", "Modular Rig"));
-						{
-							InMenuBuilder.AddMenuEntry(FControlRigEditorCommands::Get().ToggleSchematicViewportVisibility);
-						}
-						InMenuBuilder.EndSection();
-					}
-				));
 			}
 		}
 	}
@@ -2483,23 +2621,6 @@ void FControlRigEditor::HandleViewportCreated(const TSharedRef<class IPersonaVie
 	}
 }
 
-TOptional<float> FControlRigEditor::GetToolbarAxesScale() const
-{
-	if (const UControlRigEditModeSettings* Settings = GetDefault<UControlRigEditModeSettings>())
-	{
-		return Settings->AxisScale;
-	}
-	return 0.f;
-}
-
-void FControlRigEditor::OnToolbarAxesScaleChanged(float InValue)
-{
-	if (UControlRigEditModeSettings* Settings = GetMutableDefault<UControlRigEditModeSettings>())
-	{
-		Settings->AxisScale = InValue;
-	}
-}
-
 void FControlRigEditor::HandleToggleControlVisibility()
 {
 	if (FControlRigEditMode* EditMode = GetEditMode())
diff --git a/Engine/Plugins/Animation/ControlRig/Source/ControlRigEditor/Private/Editor/ControlRigEditor.h b/Engine/Plugins/Animation/ControlRig/Source/ControlRigEditor/Private/Editor/ControlRigEditor.h
--- a/Engine/Plugins/Animation/ControlRig/Source/ControlRigEditor/Private/Editor/ControlRigEditor.h
+++ b/Engine/Plugins/Animation/ControlRig/Source/ControlRigEditor/Private/Editor/ControlRigEditor.h
@@ -254,8 +254,6 @@ private:
 	void HandleToggleToolbarDrawSockets();
 	bool GetToolbarDrawAxesOnSelection() const;
 	void HandleToggleToolbarDrawAxesOnSelection();
-	TOptional<float> GetToolbarAxesScale() const;
-	void OnToolbarAxesScaleChanged(float InValue);
 	void HandleToggleSchematicViewport();
 	bool IsSchematicViewportActive() const;
 	EVisibility GetSchematicOverlayVisibility() const;
