Control Rig: Follow up Replay to support, pose, metadata, variables etc

#rnx

diff --git a/Engine/Plugins/Animation/ControlRig/Source/ControlRig/Private/Tracks/SampleTrackBase.cpp b/Engine/Plugins/Animation/ControlRig/Source/ControlRig/Private/Tracks/SampleTrackBase.cpp
--- a/Engine/Plugins/Animation/ControlRig/Source/ControlRig/Private/Tracks/SampleTrackBase.cpp
+++ b/Engine/Plugins/Animation/ControlRig/Source/ControlRig/Private/Tracks/SampleTrackBase.cpp
@@ -94,8 +94,8 @@ const FName& FSampleTrackBase::GetName() const
 {
 	if(Names.IsEmpty())
 	{
-		static const FName NoneName(NAME_None);
-		return NoneName;
+		static const FName InvalidName = FName(NAME_None);
+		return InvalidName;
 	}
 	return Names[0];
 }
