Control Rig: Add redirector for moving tests

#rnx

diff --git a/Engine/Plugins/Animation/ControlRig/Config/BaseControlRig.ini b/Engine/Plugins/Animation/ControlRig/Config/BaseControlRig.ini
--- a/Engine/Plugins/Animation/ControlRig/Config/BaseControlRig.ini
+++ b/Engine/Plugins/Animation/ControlRig/Config/BaseControlRig.ini
@@ -440,5 +440,9 @@
 +ClassRedirects=(OldName="/Script/ControlRig.NameSpacedUserData", NewName="/Script/RigVM.NameSpacedUserData")
 +ClassRedirects=(OldName="/Script/ControlRig.DataAssetLink", NewName="/Script/RigVM.DataAssetLink")
 
+; moving control rig test data to control rig replay
++StructRedirects=(OldName="/Script/ControlRig.ControlRigTestDataVariable",NewName="/Script/ControlRig.ControlRigReplayVariable")
++EnumRedirects=(OldName="EControlRigTestDataPlaybackMode",NewName="/Script/ControlRig.EControlRigReplayPlaybackMode")
+
 [/Script/Engine.Engine]
 +K2ParamRedirects=(NodeName="/Script/BlueprintGraph.K2Node_CallFunction", OldParamName="ControlRigSequencerEditorLibrary.FindOrCreateControlRigTrack.bIsAdditiveControlRig", NewParamName="bIsLayeredControlRig")
