- Disabled usage of bRecordTransformsAtRuntime for Native ControlRig classes
- Made context ScopeLock conditional on bRecordTransformsAtRuntime

diff --git a/Engine/Plugins/Animation/ControlRig/Source/ControlRig/Private/ControlRig.cpp b/Engine/Plugins/Animation/ControlRig/Source/ControlRig/Private/ControlRig.cpp
--- a/Engine/Plugins/Animation/ControlRig/Source/ControlRig/Private/ControlRig.cpp
+++ b/Engine/Plugins/Animation/ControlRig/Source/ControlRig/Private/ControlRig.cpp
@@ -948,6 +948,12 @@ bool UControlRig::Execute(const FName& InEventName)
 				bRecordTransformsAtRuntime = false;
 			}
 		}
+		
+		if (IsNativeRig())
+		{
+			bRecordTransformsAtRuntime = false;
+		}
+		
 		RecordTransformsPerInstructionGuard = MakeShared<TGuardValue<bool>>(Hierarchy->bRecordTransformsAtRuntime, bRecordTransformsAtRuntime);
 		ResetRecordedTransforms(InEventName);
 #endif
diff --git a/Engine/Plugins/Animation/ControlRig/Source/ControlRig/Private/Rigs/FKControlRig.cpp b/Engine/Plugins/Animation/ControlRig/Source/ControlRig/Private/Rigs/FKControlRig.cpp
--- a/Engine/Plugins/Animation/ControlRig/Source/ControlRig/Private/Rigs/FKControlRig.cpp
+++ b/Engine/Plugins/Animation/ControlRig/Source/ControlRig/Private/Rigs/FKControlRig.cpp
@@ -118,6 +118,7 @@ bool UFKControlRig::Execute_Internal(const FName& InEventName)
 		{
 			Hierarchy->StorePoseForTrace(TEXT("UFKControlRig::BeginExecuteUnits"));
 		}
+		Hierarchy->bRecordTransformsAtRuntime = false;
 	}
 #endif	
 
@@ -289,6 +290,13 @@ void UFKControlRig::Initialize(bool bInitRigUnits /*= true*/)
 
 	Super::Initialize(bInitRigUnits);
 
+#if WITH_EDITOR
+	if(URigHierarchy* Hierarchy = GetHierarchy())
+	{
+		Hierarchy->bRecordTransformsAtRuntime = false;
+	}
+#endif // WITH_EDITOR
+
 	if(const TSharedPtr<IControlRigObjectBinding> Binding = GetObjectBinding())
 	{
 		// we do this after Initialize because Initialize will copy from CDO. 
diff --git a/Engine/Plugins/Animation/ControlRig/Source/ControlRig/Private/Rigs/RigHierarchy.cpp b/Engine/Plugins/Animation/ControlRig/Source/ControlRig/Private/Rigs/RigHierarchy.cpp
--- a/Engine/Plugins/Animation/ControlRig/Source/ControlRig/Private/Rigs/RigHierarchy.cpp
+++ b/Engine/Plugins/Animation/ControlRig/Source/ControlRig/Private/Rigs/RigHierarchy.cpp
@@ -4154,15 +4154,18 @@ FTransform URigHierarchy::GetTransform(FRigTransformElement* InTransformElement,
 
 #if WITH_EDITOR
 	{
-		FScopeLock Lock(&ExecuteContextLock);
-		if(bRecordTransformsAtRuntime && ExecuteContext != nullptr)
+		if(bRecordTransformsAtRuntime)
 		{
-			ReadTransformsAtRuntime.Emplace(
-				ExecuteContext->GetPublicData<>().GetInstructionIndex(),
-				ExecuteContext->GetSlice().GetIndex(),
-				InTransformElement->GetIndex(),
-				InTransformType
-			);
+			FScopeLock Lock(&ExecuteContextLock);
+			if (ExecuteContext != nullptr)
+			{
+				ReadTransformsAtRuntime.Emplace(
+					ExecuteContext->GetPublicData<>().GetInstructionIndex(),
+					ExecuteContext->GetSlice().GetIndex(),
+					InTransformElement->GetIndex(),
+					InTransformType
+				);
+			}			
 		}
 	}
 	
@@ -4311,49 +4314,51 @@ void URigHierarchy::SetTransform(FRigTransformElement* InTransformElement, const
 
 	// lock execute context scope
 	{
-		FScopeLock Lock(&ExecuteContextLock);
-	
-		if(bRecordTransformsAtRuntime && ExecuteContext)
+		if(bRecordTransformsAtRuntime)
 		{
-			const FRigVMExecuteContext& PublicData = ExecuteContext->GetPublicData<>();
-			const FRigVMSlice& Slice = ExecuteContext->GetSlice();
-			WrittenTransformsAtRuntime.Emplace(
-				PublicData.GetInstructionIndex(),
-				Slice.GetIndex(),
-				InTransformElement->GetIndex(),
-				InTransformType
-			);
-
-			// if we are setting a control / null parent after a child here - let's let the user know
-			if(InTransformElement->IsA<FRigControlElement>() || InTransformElement->IsA<FRigNullElement>())
+			FScopeLock Lock(&ExecuteContextLock);
+			if (ExecuteContext)
 			{
-				if(const UWorld* World = GetWorld())
+				const FRigVMExecuteContext& PublicData = ExecuteContext->GetPublicData<>();
+				const FRigVMSlice& Slice = ExecuteContext->GetSlice();
+				WrittenTransformsAtRuntime.Emplace(
+					PublicData.GetInstructionIndex(),
+					Slice.GetIndex(),
+					InTransformElement->GetIndex(),
+					InTransformType
+				);
+
+				// if we are setting a control / null parent after a child here - let's let the user know
+				if(InTransformElement->IsA<FRigControlElement>() || InTransformElement->IsA<FRigNullElement>())
 				{
-					// only fire these notes if we are inside the asset editor
-					if(World->WorldType == EWorldType::EditorPreview)
+					if(const UWorld* World = GetWorld())
 					{
-						for (FRigBaseElement* Child : GetChildren(InTransformElement))
+						// only fire these notes if we are inside the asset editor
+						if(World->WorldType == EWorldType::EditorPreview)
 						{
-							const bool bChildFound = WrittenTransformsAtRuntime.ContainsByPredicate([Child](const TInstructionSliceElement& Entry) -> bool
+							for (FRigBaseElement* Child : GetChildren(InTransformElement))
 							{
-								return Entry.Get<2>() == Child->GetIndex();
-							});
+								const bool bChildFound = WrittenTransformsAtRuntime.ContainsByPredicate([Child](const TInstructionSliceElement& Entry) -> bool
+								{
+									return Entry.Get<2>() == Child->GetIndex();
+								});
 
-							if(bChildFound)
-							{
-								const FControlRigExecuteContext& CRContext = ExecuteContext->GetPublicData<FControlRigExecuteContext>();
-								if(CRContext.GetLog())
+								if(bChildFound)
 								{
-									static constexpr TCHAR MessageFormat[] = TEXT("Setting transform of parent (%s) after setting child (%s).\nThis may lead to unexpected results.");
-									const FString& Message = FString::Printf(
-										MessageFormat,
-										*InTransformElement->GetName(),
-										*Child->GetName());
-									CRContext.Report(
-										EMessageSeverity::Info,
-										ExecuteContext->GetPublicData<>().GetFunctionName(),
-										ExecuteContext->GetPublicData<>().GetInstructionIndex(),
-										Message);
+									const FControlRigExecuteContext& CRContext = ExecuteContext->GetPublicData<FControlRigExecuteContext>();
+									if(CRContext.GetLog())
+									{
+										static constexpr TCHAR MessageFormat[] = TEXT("Setting transform of parent (%s) after setting child (%s).\nThis may lead to unexpected results.");
+										const FString& Message = FString::Printf(
+											MessageFormat,
+											*InTransformElement->GetName(),
+											*Child->GetName());
+										CRContext.Report(
+											EMessageSeverity::Info,
+											ExecuteContext->GetPublicData<>().GetFunctionName(),
+											ExecuteContext->GetPublicData<>().GetInstructionIndex(),
+											Message);
+									}
 								}
 							}
 						}
@@ -4492,15 +4497,18 @@ FTransform URigHierarchy::GetControlOffsetTransform(FRigControlElement* InContro
 
 	// lock execute context scope
 	{
-		FScopeLock Lock(&ExecuteContextLock);
-		if(bRecordTransformsAtRuntime && ExecuteContext)
+		if(bRecordTransformsAtRuntime) 
 		{
-			ReadTransformsAtRuntime.Emplace(
-				ExecuteContext->GetPublicData<>().GetInstructionIndex(),
-				ExecuteContext->GetSlice().GetIndex(),
-				InControlElement->GetIndex(),
-				InTransformType
-			);
+			FScopeLock Lock(&ExecuteContextLock);
+			if (ExecuteContext)
+			{
+				ReadTransformsAtRuntime.Emplace(
+					ExecuteContext->GetPublicData<>().GetInstructionIndex(),
+					ExecuteContext->GetSlice().GetIndex(),
+					InControlElement->GetIndex(),
+					InTransformType
+				);
+			}
 		}
 	}
 	
@@ -4582,15 +4590,18 @@ void URigHierarchy::SetControlOffsetTransform(FRigControlElement* InControlEleme
 
 	// lock execute context scope
 	{
-		FScopeLock Lock(&ExecuteContextLock);
-		if(bRecordTransformsAtRuntime && ExecuteContext)
-		{
-			WrittenTransformsAtRuntime.Emplace(
-				ExecuteContext->GetPublicData<>().GetInstructionIndex(),
-				ExecuteContext->GetSlice().GetIndex(),
-				InControlElement->GetIndex(),
-				InTransformType
-			);
+		if(bRecordTransformsAtRuntime)
+		{
+			FScopeLock Lock(&ExecuteContextLock);
+			if (ExecuteContext)
+			{
+				WrittenTransformsAtRuntime.Emplace(
+					ExecuteContext->GetPublicData<>().GetInstructionIndex(),
+					ExecuteContext->GetSlice().GetIndex(),
+					InControlElement->GetIndex(),
+					InTransformType
+				);
+			}
 		}
 	}
 	
diff --git a/Engine/Plugins/Animation/ControlRig/Source/ControlRig/Public/Rigs/RigHierarchy.h b/Engine/Plugins/Animation/ControlRig/Source/ControlRig/Public/Rigs/RigHierarchy.h
--- a/Engine/Plugins/Animation/ControlRig/Source/ControlRig/Public/Rigs/RigHierarchy.h
+++ b/Engine/Plugins/Animation/ControlRig/Source/ControlRig/Public/Rigs/RigHierarchy.h
@@ -5184,6 +5184,7 @@ private:
 	
 	friend class URigHierarchyController;
 	friend class UControlRig;
+	friend class UFKControlRig;
 	friend class UModularRig;
 	friend class FControlRigEditor;
 	friend struct FRigBaseElement;
