Control Rig Replay: Store absolute time and delta time with sample track

#rnx

diff --git a/Engine/Plugins/Animation/ControlRig/Source/ControlRig/Private/ControlRigReplay.cpp b/Engine/Plugins/Animation/ControlRig/Source/ControlRig/Private/ControlRigReplay.cpp
--- a/Engine/Plugins/Animation/ControlRig/Source/ControlRig/Private/ControlRigReplay.cpp
+++ b/Engine/Plugins/Animation/ControlRig/Source/ControlRig/Private/ControlRigReplay.cpp
@@ -37,31 +37,6 @@ void FControlRigReplayTracks::Reset()
 	ParentIndices.Reset();
 }
 
-int32 FControlRigReplayTracks::AddTimeSample(float InTime)
-{
-	return GetContainer()->AddTimeSample(InTime);
-}
-
-int32 FControlRigReplayTracks::AddTimeSampleFromDeltaTime(float InDeltaTime)
-{
-	if(GetContainer()->GetTimeSamples().IsEmpty())
-	{
-		return GetContainer()->AddTimeSample(0.f);
-	}
-	const float LastTime = GetContainer()->GetTimeSamples().Last();
-	return GetContainer()->AddTimeSample(LastTime + InDeltaTime);
-}
-
-void FControlRigReplayTracks::GetTimeIndex(float InTime, int32& InOutTimeIndex) const
-{
-	GetContainer()->GetTimeIndex(InTime, InOutTimeIndex);
-}
-
-const TArray<float>& FControlRigReplayTracks::GetTimeSamples() const
-{
-	return GetContainer()->GetTimeSamples();
-}
-
 void FControlRigReplayTracks::StoreHierarchy(URigHierarchy* InHierarchy, const TArrayView<FRigElementKey>& InKeys, bool bStorePose,
 	bool bStoreComponents, bool bStoreMetadata)
 {
@@ -1072,7 +1047,7 @@ UControlRigReplay* UControlRigReplay::CreateNewAsset(FString InDesiredPackagePat
 
 FVector2D UControlRigReplay::GetTimeRange() const
 {
-	return FVector2D(Tracks.GetContainer()->GetTimeRange());
+	return FVector2D(Tracks.GetTimeRange());
 }
 
 bool UControlRigReplay::StartRecording(UControlRig* InControlRig)
diff --git a/Engine/Plugins/Animation/ControlRig/Source/ControlRig/Private/Tracks/SampleTrackContainer.cpp b/Engine/Plugins/Animation/ControlRig/Source/ControlRig/Private/Tracks/SampleTrackContainer.cpp
--- a/Engine/Plugins/Animation/ControlRig/Source/ControlRig/Private/Tracks/SampleTrackContainer.cpp
+++ b/Engine/Plugins/Animation/ControlRig/Source/ControlRig/Private/Tracks/SampleTrackContainer.cpp
@@ -80,12 +80,71 @@ FSampleTrackContainer* FSampleTrackHost::GetContainer()
 	return Container.Get();
 }
 
+int32 FSampleTrackHost::AddTimeSample(float InAbsoluteTime, float InDeltaTime)
+{
+	return GetContainer()->AddTimeSample(InAbsoluteTime, InDeltaTime);
+}
+
+int32 FSampleTrackHost::AddTimeSampleFromDeltaTime(float InDeltaTime)
+{
+	return GetContainer()->AddTimeSampleFromDeltaTime(InDeltaTime);
+}
+
+int32 FSampleTrackHost::GetNumTimes() const
+{
+	return GetContainer()->GetNumTimes();
+}
+
+FVector2f FSampleTrackHost::GetTimeRange() const
+{
+	return GetContainer()->GetTimeRange();
+}
+
+void FSampleTrackHost::GetTimeIndex(float InAbsoluteTime, FSampleTrackIndex& InOutTrackIndex) const
+{
+	return GetContainer()->GetTimeIndex(InAbsoluteTime, InOutTrackIndex);
+}
+
+float FSampleTrackHost::GetAbsoluteTime(int32 InTimeIndex, FSampleTrackIndex& InOutTrackIndex) const
+{
+	return GetContainer()->GetAbsoluteTime(InTimeIndex, InOutTrackIndex);
+}
+
+float FSampleTrackHost::GetDeltaTime(int32 InTimeIndex, FSampleTrackIndex& InOutTrackIndex) const
+{
+	return GetContainer()->GetDeltaTime(InTimeIndex, InOutTrackIndex);
+}
+
+float FSampleTrackHost::GetAbsoluteTime(int32 InTimeIndex) const
+{
+	return GetContainer()->GetAbsoluteTime(InTimeIndex);
+}
+
+float FSampleTrackHost::GetDeltaTime(int32 InTimeIndex) const
+{
+	return GetContainer()->GetDeltaTime(InTimeIndex);
+}
+
+float FSampleTrackHost::GetLastAbsoluteTime() const
+{
+	return GetContainer()->GetLastAbsoluteTime();
+}
+
+float FSampleTrackHost::GetLastDeltaTime() const
+{
+	return GetContainer()->GetLastDeltaTime();
+}
+
 //////////////////////////////////////////////////////////////////////////////////////////////////////////////////////////////
 /////// FSampleTrackContainer
 //////////////////////////////////////////////////////////////////////////////////////////////////////////////////////////////
 
+const FLazyName FSampleTrackContainer::AbsoluteTimeName = FLazyName(TEXT("AbsoluteTime"));
+const FLazyName FSampleTrackContainer::DeltaTimeName = FLazyName(TEXT("DeltaTime"));
+
 FSampleTrackContainer::FSampleTrackContainer()
 : bForceToUseCompression(false)
+, TimeSampleTrackIndex(FSampleTrackIndex::MakeSingleton())
 {
 }
 
@@ -95,14 +154,13 @@ FSampleTrackContainer::~FSampleTrackContainer()
 
 void FSampleTrackContainer::Reset()
 {
-	TimeSamples.Reset();
 	NameToIndex.Reset();
 	Tracks.Reset();
+	TimeSampleTrackIndex = FSampleTrackIndex::MakeSingleton();
 }
 
 void FSampleTrackContainer::Shrink()
 {
-	TimeSamples.Shrink();
 	NameToIndex.Shrink();
 
 	for(const TSharedPtr<FSampleTrackBase>& Track : Tracks)
@@ -118,7 +176,6 @@ void FSampleTrackContainer::Shrink()
 void FSampleTrackContainer::Compact(float InTolerance)
 {
 	RemoveInvalidTracks(false);
-	RemoveObsoleteTimes();
 	RemoveRedundantTracks(false, InTolerance);
 	MergeTypedTracks(false, InTolerance);
 	ConvertTracksToSampled();
@@ -130,8 +187,6 @@ void FSampleTrackContainer::Compact(float InTolerance)
 
 void FSampleTrackContainer::Reserve(int32 InNum)
 {
-	TimeSamples.Reserve(InNum);
-
 	for(const TSharedPtr<FSampleTrackBase>& Track : Tracks)
 	{
 		Track->Reserve(InNum, InNum);
@@ -207,8 +262,6 @@ bool FSampleTrackContainer::Serialize(FArchive& InArchive)
 		}
 	}
 	
-	(*InnerArchive) << TimeSamples;
-
 	int32 NumTracks = Tracks.Num();
 	(*InnerArchive) << NumTracks;
 	
@@ -519,35 +572,106 @@ TSharedPtr<FSampleTrack<TArray<FInstancedStruct>>> FSampleTrackContainer::AddStr
 
 FVector2f FSampleTrackContainer::GetTimeRange() const
 {
-	if(TimeSamples.IsEmpty())
+	const TSharedPtr<const FSampleTrack<float>> AbsoluteTimeTrack = FindTrack<float>(AbsoluteTimeName);
+	if(!AbsoluteTimeTrack.IsValid() || AbsoluteTimeTrack->NumSamples() == 0)
 	{
 		return FVector2f::ZeroVector;
 	}
-	return {TimeSamples[0], TimeSamples.Last()};
+	return {
+		AbsoluteTimeTrack->GetValueAtSampleIndex(0),
+		AbsoluteTimeTrack->GetValueAtSampleIndex(AbsoluteTimeTrack->NumSamples()-1)
+	};
 }
 
-void FSampleTrackContainer::GetTimeIndex(float InTime, int32& InOutTimeIndex) const
+int32 FSampleTrackContainer::GetNumTimes() const
 {
-	if(TimeSamples.IsEmpty())
+	const TSharedPtr<const FSampleTrack<float>> AbsoluteTimeTrack = FindTrack<float>(AbsoluteTimeName);
+	if(!AbsoluteTimeTrack.IsValid() || AbsoluteTimeTrack->NumTimes() == 0)
 	{
-		InOutTimeIndex = INDEX_NONE;
-		return;
+		return 0;
+	}
+	return AbsoluteTimeTrack->NumTimes(); 
+}
+
+float FSampleTrackContainer::GetAbsoluteTime(int32 InTimeIndex, FSampleTrackIndex& InOutTrackIndex) const
+{
+	const TSharedPtr<const FSampleTrack<float>> AbsoluteTimeTrack = FindTrack<float>(AbsoluteTimeName);
+	if(!AbsoluteTimeTrack.IsValid() || AbsoluteTimeTrack->NumTimes() == 0)
+	{
+		return 0.f;
 	}
+	return AbsoluteTimeTrack->GetValueAtTimeIndex(InTimeIndex, InOutTrackIndex); 
+}
+
+float FSampleTrackContainer::GetDeltaTime(int32 InTimeIndex, FSampleTrackIndex& InOutTrackIndex) const
+{
+	const TSharedPtr<const FSampleTrack<float>> DeltaTimeTrack = FindTrack<float>(DeltaTimeName);
+	if(!DeltaTimeTrack.IsValid() || DeltaTimeTrack->NumTimes() == 0)
+	{
+		return 0.f;
+	}
+	return DeltaTimeTrack->GetValueAtTimeIndex(InTimeIndex, InOutTrackIndex); 
+}
+
+float FSampleTrackContainer::GetAbsoluteTime(int32 InTimeIndex) const
+{
+	return GetAbsoluteTime(InTimeIndex, TimeSampleTrackIndex);
+}
 
-	InOutTimeIndex = FMath::Clamp(InOutTimeIndex, 0, TimeSamples.Num() - 1);
-	if(TimeSamples.Num() == 1)
+float FSampleTrackContainer::GetDeltaTime(int32 InTimeIndex) const
+{
+	return GetDeltaTime(InTimeIndex, TimeSampleTrackIndex);
+}
+
+float FSampleTrackContainer::GetLastAbsoluteTime() const
+{
+	const TSharedPtr<const FSampleTrack<float>> AbsoluteTimeTrack = FindTrack<float>(AbsoluteTimeName);
+	if(!AbsoluteTimeTrack.IsValid() || AbsoluteTimeTrack->NumSamples() == 0)
 	{
+		return 0.f;
+	}
+	return AbsoluteTimeTrack->GetValueAtSampleIndex(AbsoluteTimeTrack->NumSamples()-1);
+}
+
+float FSampleTrackContainer::GetLastDeltaTime() const
+{
+	const TSharedPtr<const FSampleTrack<float>> DeltaTimeTrack = FindTrack<float>(DeltaTimeName);
+	if(!DeltaTimeTrack.IsValid() || DeltaTimeTrack->NumSamples() == 0)
+	{
+		return 0.f;
+	}
+	return DeltaTimeTrack->GetValueAtSampleIndex(DeltaTimeTrack->NumSamples()-1);
+}
+
+void FSampleTrackContainer::GetTimeIndex(float InAbsoluteTime, FSampleTrackIndex& InOutTrackIndex) const
+{
+	const TSharedPtr<const FSampleTrack<float>> AbsoluteTimeTrack = FindTrack<float>(AbsoluteTimeName);
+	check(AbsoluteTimeTrack.IsValid() && AbsoluteTimeTrack->NumTimes() > 0);
+
+	int32& InOutTimeIndex = InOutTrackIndex.GetSample(AbsoluteTimeTrack->GetTrackIndex());
+
+	if(AbsoluteTimeTrack->GetMode() == FSampleTrackBase::EMode_Singleton)
+	{
+		InOutTimeIndex = 0;
 		return;
 	}
 
-	while(TimeSamples.IsValidIndex(InOutTimeIndex-1) && TimeSamples[InOutTimeIndex] > InTime)
+	const int32 NumSamples = AbsoluteTimeTrack->NumSamples();
+	InOutTimeIndex = FMath::Clamp(InOutTimeIndex, 0, NumSamples - 1);
+
+	while((InOutTimeIndex > 1) && AbsoluteTimeTrack->GetValueAtSampleIndex(InOutTimeIndex) > InAbsoluteTime - SMALL_NUMBER)
 	{
 		InOutTimeIndex--;
 	}
-	while(TimeSamples.IsValidIndex(InOutTimeIndex+1) && TimeSamples[InOutTimeIndex+1] < InTime)
+	while((InOutTimeIndex < NumSamples - 1) && AbsoluteTimeTrack->GetValueAtSampleIndex(InOutTimeIndex+1) <= InAbsoluteTime + SMALL_NUMBER)
 	{
 		InOutTimeIndex++;
 	}
+
+	if(AbsoluteTimeTrack->GetMode() == FSampleTrackBase::EMode_Sampled)
+	{
+		InOutTimeIndex = AbsoluteTimeTrack->TimeIndices[InOutTimeIndex];
+	}
 }
 
 int32 FSampleTrackContainer::GetTrackIndex(const FName& InName) const
@@ -793,73 +917,6 @@ void FSampleTrackContainer::RemoveInvalidTracks(bool bUpdateNameToIndexMap)
 	}
 }
 
-void FSampleTrackContainer::RemoveObsoleteTimes()
-{
-	TArray<bool> RemoveTimeIndexArray;
-	RemoveTimeIndexArray.SetNumUninitialized(TimeSamples.Num());
-	int32 NumTimesToRemove = TimeSamples.Num();
-
-	for(int32 Index = 0; Index < RemoveTimeIndexArray.Num(); Index++)
-	{
-		RemoveTimeIndexArray[Index] = true;
-		for(const TSharedPtr<FSampleTrackBase>& Track : Tracks)
-		{
-			if(Track->StoresValueForTimeIndex(Index))
-			{
-				if(RemoveTimeIndexArray[Index])
-				{
-					NumTimesToRemove--;
-					RemoveTimeIndexArray[Index] = false;
-				}
-			}
-		}
-	}
-
-	if(NumTimesToRemove == 0)
-	{
-		return;
-	}
-
-	TArray<float> NewTimeSamples;
-	NewTimeSamples.Reserve(TimeSamples.Num() - NumTimesToRemove);
-	TArray<int32> OldIndexToNewIndex;
-	OldIndexToNewIndex.Reserve(TimeSamples.Num());
-	for(int32 Index = 0; Index < RemoveTimeIndexArray.Num(); Index++)
-	{
-		// don't allow to remove the first time sample
-		if(Index == 0)
-		{
-			verify(!RemoveTimeIndexArray[Index]);
-		}
-		
-		if(RemoveTimeIndexArray[Index])
-		{
-			const int32 LastIndex = OldIndexToNewIndex.Last();
-			OldIndexToNewIndex.Add(LastIndex);
-		}
-		else
-		{
-			NewTimeSamples.Add(TimeSamples[Index]);
-			if(OldIndexToNewIndex.IsEmpty())
-			{
-				OldIndexToNewIndex.Add(0);
-			}
-			else
-			{
-				const int32 LastIndex = OldIndexToNewIndex.Last();
-				OldIndexToNewIndex.Add(LastIndex + 1);
-			}
-		}
-	}
-	
-	TimeSamples = MoveTemp(NewTimeSamples);
-
-	for(const TSharedPtr<FSampleTrackBase>& Track : Tracks)
-	{
-		Track->RemoveObsoleteTimes(NumTimesToRemove, RemoveTimeIndexArray, OldIndexToNewIndex);
-	}
-}
-
 void FSampleTrackContainer::RemoveRedundantTracks(bool bUpdateNameToIndexMap, float InTolerance)
 {
 	struct FTrackInfo
@@ -1171,9 +1228,20 @@ void FSampleTrackContainer::UpdateNameToIndexMap()
 	}
 }
 
-int32 FSampleTrackContainer::AddTimeSample(float InTimeSample)
+int32 FSampleTrackContainer::AddTimeSample(float InAbsoluteTime, float InDeltaTime)
+{
+	const TSharedPtr<FSampleTrack<float>> AbsoluteTimeTrack = FindOrAddTrack<float>(AbsoluteTimeName, FSampleTrackBase::ETrackType_Float);
+	const TSharedPtr<FSampleTrack<float>> DeltaTimeTrack = FindOrAddTrack<float>(DeltaTimeName, FSampleTrackBase::ETrackType_Float);
+	const int32 AddedIndex = AbsoluteTimeTrack->NumTimes();
+	AbsoluteTimeTrack->AddSample(InAbsoluteTime);
+	DeltaTimeTrack->AddSample(InDeltaTime);
+	return AddedIndex;
+}
+
+int32 FSampleTrackContainer::AddTimeSampleFromDeltaTime(float InDeltaTime)
 {
-	return TimeSamples.Add(InTimeSample);
+	const float LastTime = GetLastAbsoluteTime() + GetLastDeltaTime();
+	return AddTimeSample(LastTime, InDeltaTime);
 }
 
 bool FSampleTrackContainer::IsEditable() const
diff --git a/Engine/Plugins/Animation/ControlRig/Source/ControlRig/Public/ControlRigReplay.h b/Engine/Plugins/Animation/ControlRig/Source/ControlRig/Public/ControlRigReplay.h
--- a/Engine/Plugins/Animation/ControlRig/Source/ControlRig/Public/ControlRigReplay.h
+++ b/Engine/Plugins/Animation/ControlRig/Source/ControlRig/Public/ControlRigReplay.h
@@ -58,11 +58,6 @@ public:
 
 	virtual void Reset() override;
 
-	int32 AddTimeSample(float InTime);
-	int32 AddTimeSampleFromDeltaTime(float InDeltaTime);
-	void GetTimeIndex(float InTime, int32& InOutTimeIndex) const;
-	const TArray<float>& GetTimeSamples() const;
-
 	void StoreHierarchy(
 		URigHierarchy* InHierarchy,
 		const TArrayView<FRigElementKey>& InKeys = TArrayView<FRigElementKey>(),
diff --git a/Engine/Plugins/Animation/ControlRig/Source/ControlRig/Public/Tracks/SampleTrackContainer.h b/Engine/Plugins/Animation/ControlRig/Source/ControlRig/Public/Tracks/SampleTrackContainer.h
--- a/Engine/Plugins/Animation/ControlRig/Source/ControlRig/Public/Tracks/SampleTrackContainer.h
+++ b/Engine/Plugins/Animation/ControlRig/Source/ControlRig/Public/Tracks/SampleTrackContainer.h
@@ -30,7 +30,19 @@ public:
 
 	const FSampleTrackContainer* GetContainer() const;
 	FSampleTrackContainer* GetContainer();
-	
+
+	int32 AddTimeSample(float InAbsoluteTime, float InDeltaTime);
+	int32 AddTimeSampleFromDeltaTime(float InDeltaTime);
+	int32 GetNumTimes() const;
+	FVector2f GetTimeRange() const;
+	void GetTimeIndex(float InAbsoluteTime, FSampleTrackIndex& InOutTrackIndex) const;
+	float GetAbsoluteTime(int32 InTimeIndex, FSampleTrackIndex& InOutTrackIndex) const;
+	float GetDeltaTime(int32 InTimeIndex, FSampleTrackIndex& InOutTrackIndex) const;
+	float GetAbsoluteTime(int32 InTimeIndex) const;
+	float GetDeltaTime(int32 InTimeIndex) const;
+	float GetLastAbsoluteTime() const;
+	float GetLastDeltaTime() const;
+
 private:
 
 	TUniquePtr<FSampleTrackContainer> Container;
@@ -62,6 +74,21 @@ public:
 
 	TSharedPtr<FSampleTrackBase> AddTrack(const FName& InName, FSampleTrackBase::ETrackType InTrackType, const UScriptStruct* InScriptStruct = nullptr);
 	TSharedPtr<FSampleTrackBase> FindOrAddTrack(const FName& InName, FSampleTrackBase::ETrackType InTrackType, const UScriptStruct* InScriptStruct = nullptr);
+
+	template<typename T>
+	TSharedPtr<FSampleTrack<T>> AddTrack(const FName& InName, FSampleTrackBase::ETrackType InTrackType, const UScriptStruct* InScriptStruct = nullptr)
+	{
+		check(FSampleTrack<T>().GetTrackType() == InTrackType);
+		return StaticCastSharedPtr<FSampleTrack<T>>(AddTrack(InName, InTrackType, InScriptStruct));
+	}
+	
+	template<typename T>
+	TSharedPtr<FSampleTrack<T>> FindOrAddTrack(const FName& InName, FSampleTrackBase::ETrackType InTrackType, const UScriptStruct* InScriptStruct = nullptr)
+	{
+		check(FSampleTrack<T>().GetTrackType() == InTrackType);
+		return StaticCastSharedPtr<FSampleTrack<T>>(FindOrAddTrack(InName, InTrackType, InScriptStruct));
+	}
+	
 	TSharedPtr<FSampleTrack<bool>> AddBoolTrack(const FName& InName);
 	TSharedPtr<FSampleTrack<int32>> AddInt32Track(const FName& InName);
 	TSharedPtr<FSampleTrack<uint32>> AddUint32Track(const FName& InName);
@@ -94,9 +121,15 @@ public:
 		return Tracks.Num();
 	}
 
+	int32 GetNumTimes() const;
 	FVector2f GetTimeRange() const;
-	const TArray<float>& GetTimeSamples() const { return TimeSamples; }
-	void GetTimeIndex(float InTime, int32& InOutTimeIndex) const;
+	void GetTimeIndex(float InAbsoluteTime, FSampleTrackIndex& InOutTrackIndex) const;
+	float GetAbsoluteTime(int32 InTimeIndex, FSampleTrackIndex& InOutTrackIndex) const;
+	float GetDeltaTime(int32 InTimeIndex, FSampleTrackIndex& InOutTrackIndex) const;
+	float GetAbsoluteTime(int32 InTimeIndex) const;
+	float GetDeltaTime(int32 InTimeIndex) const;
+	float GetLastAbsoluteTime() const;
+	float GetLastDeltaTime() const;
 
 	int32 GetTrackIndex(const FName& InName) const;
 	TSharedPtr<const FSampleTrackBase> GetTrack(int32 InIndex) const;
@@ -212,7 +245,8 @@ public:
 		return FindArrayTrack<FInstancedStruct>(InName);
 	}
 	
-	int32 AddTimeSample(float InTimeSample);
+	int32 AddTimeSample(float InAbsoluteTime, float InDeltaTime);
+	int32 AddTimeSampleFromDeltaTime(float InDeltaTime);
 
 	// returns true if all tracks are complete / singleton tracks
 	// without references and without atlases
@@ -231,9 +265,6 @@ public:
 	// removes all tracks with no data in them or nullptr tracks
 	void RemoveInvalidTracks(bool bUpdateNameToIndexMap = true);
 
-	// removes the time indices from all tracks at which we haven't stored any data 
-	void RemoveObsoleteTimes();
-
 	// combines tracks that have the exact same data in them
 	void RemoveRedundantTracks(bool bUpdateNameToIndexMap = true, float InTolerance = KINDA_SMALL_NUMBER);
 
@@ -268,7 +299,10 @@ private:
 	void UpdateNameToIndexMap();
 
 	TMap<FName, int32> NameToIndex;
-	TArray<float> TimeSamples;
 	TArray<TSharedPtr<FSampleTrackBase>> Tracks;
 	bool bForceToUseCompression;
+	mutable FSampleTrackIndex TimeSampleTrackIndex;
+
+	static const FLazyName AbsoluteTimeName;
+	static const FLazyName DeltaTimeName;
 };
