Control Rig: Fix layered rigs not setting keys in sequencer when UControlRigSequencerEditorLibrary::SetLocalControlRigEulerTransform is called
#rnx

diff --git a/Engine/Plugins/Animation/ControlRig/Source/ControlRig/Private/ControlRig.cpp b/Engine/Plugins/Animation/ControlRig/Source/ControlRig/Private/ControlRig.cpp
--- a/Engine/Plugins/Animation/ControlRig/Source/ControlRig/Private/ControlRig.cpp
+++ b/Engine/Plugins/Animation/ControlRig/Source/ControlRig/Private/ControlRig.cpp
@@ -2112,6 +2112,11 @@ void UControlRig::SetControlValueImpl(const FName& InControlName, const FRigCont
 		// Store the value to apply it after the backwards solve
 		FRigSetControlValueInfo Info = {InValue, bNotify, Context, bSetupUndo, bPrintPythonCommnds, bFixEulerFlips};
 		ControlValues.Add(ControlElement->GetKey(), Info);
+		if (Context.SetKey == EControlRigSetKey::Always && IsInGameThread())
+		{
+			// Force evaluation to set keys in game thread
+			Evaluate_AnyThread();
+		}
 	}
 	else
 	{
