Control Rig Replay: Further work

#rnx

diff --git a/Engine/Plugins/Animation/ControlRig/Source/ControlRig/Private/ControlRig.cpp b/Engine/Plugins/Animation/ControlRig/Source/ControlRig/Private/ControlRig.cpp
--- a/Engine/Plugins/Animation/ControlRig/Source/ControlRig/Private/ControlRig.cpp
+++ b/Engine/Plugins/Animation/ControlRig/Source/ControlRig/Private/ControlRig.cpp
@@ -38,6 +38,7 @@
 #include "RigVMCore/RigVMAssetUserData.h"
 #include "ModularRig.h"
 #include "Units/Modules/RigUnit_ConnectorExecution.h"
+#include "ControlRigReplay.h"
 
 #include UE_INLINE_GENERATED_CPP_BY_NAME(ControlRig)
 
@@ -95,6 +96,7 @@ UControlRig::UControlRig(const FObjectInitializer& ObjectInitializer)
 	, bRecordSelectionPoseForConstructionMode(true)
 	, bIsClearingTransientControls(false)
 #endif
+	, ReplayTimeIndex(INDEX_NONE)
 {
 	EventQueue.Add(FRigUnit_BeginExecution::EventName);
 
@@ -1633,6 +1635,40 @@ bool UControlRig::SupportsBackwardsSolve() const
 
 void UControlRig::AdaptEventQueueForEvaluate(TArray<FName>& InOutEventQueueToRun)
 {
+	if(IsReplayEnabled())
+	{
+		check(Replay);
+
+		// find all events at this particular frame
+		const int32 NumTimes = Replay->InputTracks.GetNumTimes();
+
+		int32 TimeIndex = GetReplayTimeIndex();
+		if(TimeIndex >= NumTimes || TimeIndex < 0)
+		{
+			TimeIndex = 0;
+			SetReplayTimeIndex(TimeIndex);
+		}
+		const double TimeAtIndex = Replay->InputTracks.GetAbsoluteTime(TimeIndex);
+
+		InOutEventQueueToRun.Reset();
+		while(TimeIndex < NumTimes)
+		{
+			const FName EventFromReplay = Replay->InputTracks.GetRigVMEvent(TimeIndex);
+			InOutEventQueueToRun.Add(EventFromReplay);
+
+			// also queue up the next event to run immediately if it was recorded at the same time
+			if(TimeIndex != NumTimes - 1)
+			{
+				if(FMath::IsNearlyEqual(TimeAtIndex, Replay->InputTracks.GetAbsoluteTime(++TimeIndex)))
+				{
+					continue;
+				}
+			}
+			break;
+		}
+		return;
+	}
+	
 	Super::AdaptEventQueueForEvaluate(InOutEventQueueToRun);
 
 	for (int32 i=0; i<InOutEventQueueToRun.Num(); ++i)
@@ -3725,6 +3761,37 @@ uint32 UControlRig::GetShapeLibraryHash() const
 	return Hash;
 }
 
+int32 UControlRig::GetReplayTimeIndex() const
+{
+	return ReplayTimeIndex;
+}
+
+void UControlRig::SetReplayTimeIndex(int32 InReplayTimeIndex)
+{
+	ReplayTimeIndex = InReplayTimeIndex;
+}
+
+void UControlRig::DisableReplay()
+{
+	SetReplayTimeIndex(INDEX_NONE);
+	Replay = nullptr;
+}
+
+bool UControlRig::IsReplayEnabled() const
+{
+	return GetReplayTimeIndex() != INDEX_NONE && Replay != nullptr;
+}
+
+void UControlRig::SetReplay(UControlRigReplay* InReplay)
+{
+	if(InReplay == nullptr)
+	{
+		DisableReplay();
+		return;
+	}
+	Replay = InReplay;
+}
+
 #endif
  
 #undef LOCTEXT_NAMESPACE
diff --git a/Engine/Plugins/Animation/ControlRig/Source/ControlRig/Private/ControlRigReplay.cpp b/Engine/Plugins/Animation/ControlRig/Source/ControlRig/Private/ControlRigReplay.cpp
--- a/Engine/Plugins/Animation/ControlRig/Source/ControlRig/Private/ControlRigReplay.cpp
+++ b/Engine/Plugins/Animation/ControlRig/Source/ControlRig/Private/ControlRigReplay.cpp
@@ -5,9 +5,11 @@
 #include "ControlRigObjectVersion.h"
 #include "Units/Execution/RigUnit_PrepareForExecution.h"
 #include "Units/Execution/RigUnit_BeginExecution.h"
+#include "Units/RigUnitContext.h"
 #include "HAL/PlatformTime.h"
 #if WITH_EDITOR
 #include "AssetToolsModule.h"
+#include "Kismet/KismetSystemLibrary.h"
 #endif
 
 #include UE_INLINE_GENERATED_CPP_BY_NAME(ControlRigReplay)
@@ -34,30 +36,75 @@ void FControlRigReplayTracks::Reset()
 {
 	FSampleTrackHost::Reset();
 	SampleTrackIndex = FSampleTrackIndex();
-	ParentIndices.Reset();
 }
 
-void FControlRigReplayTracks::StoreHierarchy(URigHierarchy* InHierarchy, const TArrayView<FRigElementKey>& InKeys, bool bStorePose,
-	bool bStoreComponents, bool bStoreMetadata)
+bool FControlRigReplayTracks::IsEmpty() const
 {
-	static const FName TopologyHashName = TEXT("TopologyHash");
-	static const FName MetadataVersionName = TEXT("MetadataVersion");
-	static const FName ParentIndicesName = TEXT("ParentIndices");
+	return GetContainer()->GetNumTimes() == 0 || GetContainer()->NumTracks() == 0;
+}
 
-	const TSharedPtr<FSampleTrack<uint32>> TopologyHashTrack = StaticCastSharedPtr<FSampleTrack<uint32>>(GetContainer()->FindOrAddTrack(TopologyHashName, FSampleTrackBase::ETrackType_Uint32));
-	const TSharedPtr<FSampleTrack<uint32>> MetadataVersionTrack = StaticCastSharedPtr<FSampleTrack<uint32>>(GetContainer()->FindOrAddTrack(MetadataVersionName, FSampleTrackBase::ETrackType_Uint32));
-	const TSharedPtr<FSampleTrack<TArray<int32>>> ParentIndicesTrack = StaticCastSharedPtr<FSampleTrack<TArray<int32>>>(GetContainer()->FindOrAddTrack(ParentIndicesName, FSampleTrackBase::ETrackType_Int32Array));
+void FControlRigReplayTracks::StoreRigVMEvent(const FName& InName)
+{
+	const TSharedPtr<FSampleTrack<FName>> EventTrack = GetContainer()->FindOrAddTrack<FName>(RigVMEventName, FSampleTrackBase::ETrackType_Name);
+	EventTrack->AddSample(InName);
+}
 
-	ParentIndices.SetNum(InHierarchy->Num(), EAllowShrinking::No);
-	InHierarchy->ForEach([this, InHierarchy](const FRigBaseElement* InElement) -> bool
+FName FControlRigReplayTracks::GetRigVMEvent(int32 InTimeIndex) const
+{
+	if(const TSharedPtr<const FSampleTrack<FName>> EventTrack = GetContainer()->FindTrack<FName>(RigVMEventName))
 	{
-		ParentIndices[InElement->GetIndex()] = InHierarchy->GetIndex(InHierarchy->GetDefaultParent(InElement->GetKey()));
-		return true;
-	});
+		return EventTrack->GetValueAtTimeIndex(InTimeIndex, SampleTrackIndex);
+	}
+	return NAME_None;
+}
+
+void FControlRigReplayTracks::StoreInteraction(uint8 InInteractionMode, const TArray<FRigElementKey>& InElementsBeingInteracted)
+{
+	const TSharedPtr<FSampleTrack<uint32>> InteractionTypeTrack = GetContainer()->FindOrAddTrack<uint32>(InteractionTypeName, FSampleTrackBase::ETrackType_Uint32);
+	const TSharedPtr<FSampleTrack<TArray<FRigElementKey>>> ElementsBeingInteractedTrack = GetContainer()->FindOrAddTrack<TArray<FRigElementKey>>(ElementsBeingInteractedName, FSampleTrackBase::ETrackType_ElementKeyArray);
+	InteractionTypeTrack->AddSample(static_cast<uint32>(InInteractionMode));
+	ElementsBeingInteractedTrack->AddSample(InElementsBeingInteracted);
+}
+
+TTuple<uint8, TArray<FRigElementKey>> FControlRigReplayTracks::GetInteraction(int32 InTimeIndex) const
+{
+	const TSharedPtr<const FSampleTrack<uint32>> InteractionTypeTrack = GetContainer()->FindTrack<uint32>(InteractionTypeName);
+	const TSharedPtr<const FSampleTrack<TArray<FRigElementKey>>> ElementsBeingInteractedTrack = GetContainer()->FindTrack<TArray<FRigElementKey>>(ElementsBeingInteractedName);
+	if(InteractionTypeTrack.IsValid() && ElementsBeingInteractedTrack.IsValid())
+	{
+		return {static_cast<uint8>(InteractionTypeTrack->GetValueAtTimeIndex(InTimeIndex, SampleTrackIndex)), ElementsBeingInteractedTrack->GetValueAtTimeIndex(InTimeIndex, SampleTrackIndex)};
+	}
+	return {0, TArray<FRigElementKey>()};
+}
+
+void FControlRigReplayTracks::StoreHierarchy(URigHierarchy* InHierarchy, const TArrayView<FRigElementKey>& InKeys, bool bStorePose,
+                                             bool bStoreComponents, bool bStoreMetadata)
+{
+	// for the first frame store the expected hierarchy topology in the output
+	if(!bIsInput)
+	{
+		const TSharedPtr<FSampleTrack<uint32>> TopologyHashTrack = GetContainer()->FindOrAddTrack<uint32>(TopologyHashName, FSampleTrackBase::ETrackType_Uint32);
+		const TSharedPtr<FSampleTrack<uint32>> MetadataVersionTrack = GetContainer()->FindOrAddTrack<uint32>(MetadataVersionName, FSampleTrackBase::ETrackType_Uint32);
+		TopologyHashTrack->AddSample(InHierarchy->GetTopologyHash());
+		MetadataVersionTrack->AddSample(InHierarchy->GetMetadataTagVersion());
+
+		if(GetNumTimes() == 1)
+		{
+			const TSharedPtr<FSampleTrack<TArray<FRigElementKey>>> ElementKeysTrack = GetContainer()->FindOrAddTrack<TArray<FRigElementKey>>(ElementKeysName, FSampleTrackBase::ETrackType_ElementKeyArray);
+			const TSharedPtr<FSampleTrack<TArray<int32>>> ParentIndicesTrack = GetContainer()->FindOrAddTrack<TArray<int32>>(ParentIndicesName, FSampleTrackBase::ETrackType_Int32Array);
 
-	TopologyHashTrack->AddSample(InHierarchy->GetTopologyHash());
-	MetadataVersionTrack->AddSample(InHierarchy->GetMetadataTagVersion());
-	ParentIndicesTrack->AddSample(ParentIndices);
+			const TArray<FRigElementKey> ElementKeys = InHierarchy->GetAllKeys();
+			TArray<int32> ParentIndices;
+			ParentIndices.Reserve(ElementKeys.Num());
+			for(const FRigElementKey& ElementKey : ElementKeys)
+			{
+				ParentIndices.Add(ElementKeys.Find(InHierarchy->GetDefaultParent(ElementKey)));
+			}
+
+			ElementKeysTrack->AddSample(ElementKeys);
+			ParentIndicesTrack->AddSample(ParentIndices);
+		}
+	}
 
 	ForEachElement(InHierarchy, InKeys, [this, InHierarchy, bStorePose, bStoreComponents, bStoreMetadata](FRigBaseElement* InElement, bool& bSuccess)
 	{
@@ -138,6 +185,25 @@ void FControlRigReplayTracks::StorePose(URigHierarchy* InHierarchy, FRigBaseElem
 	}
 	const FTransform LocalTransform = InHierarchy->GetLocalTransform(InElement->GetIndex());
 	Track->AddSample(FTransform3f(LocalTransform));
+
+	if(bIsInput)
+	{
+		if(const FRigConnectorElement* ConnectorElement = Cast<FRigConnectorElement>(InElement))
+		{
+			if(const UControlRig* ControlRig = Cast<UControlRig>(InHierarchy->GetOuter()))
+			{
+				const FName& ConnectorTrackName = *(TrackName.ToString() + TEXT("ConnectorTargets"));
+				const TSharedPtr<FSampleTrack<TArray<FRigElementKey>>> ConnectorTrack = StaticCastSharedPtr<FSampleTrack<TArray<FRigElementKey>>>(Storage->FindOrAddTrack(ConnectorTrackName, FSampleTrackBase::ETrackType_ElementKeyArray));
+			
+				TArray<FRigElementKey> Targets;
+				if(const FRigElementKeyRedirector::FCachedKeyArray* Cache = ControlRig->GetElementKeyRedirector().Find(ConnectorElement->GetKey()))
+				{
+					Targets.Append(FRigElementKeyRedirector::Convert(*Cache));
+				}
+				ConnectorTrack->AddSample(Targets);
+			}
+		}
+	}
 }
 
 bool FControlRigReplayTracks::RestorePose(int32 InTimeIndex, URigHierarchy* InHierarchy, const TArrayView<FRigElementKey>& InKeys, TReportFunction InReportFunction) const
@@ -301,7 +367,21 @@ void FControlRigReplayTracks::StoreMetaData(URigHierarchy* InHierarchy, const TA
 
 void FControlRigReplayTracks::StoreMetaData(URigHierarchy* InHierarchy, FRigBaseElement* InElement)
 {
+	if(!InHierarchy->HasMetadata(InElement))
+	{
+		return;
+	}
+	
+	FSampleTrackContainer* Storage = GetContainer();
 	const TArray<FName> MetadataNames = InHierarchy->GetMetadataNames(InElement->GetKey());
+	const FName& TrackName = GetTrackName(InElement->GetKey());
+	const FName MetadataNamesTrackName = *(TrackName.ToString() + TEXT("MetadataNames"));
+	TSharedPtr<FSampleTrack<TArray<FName>>> Track = StaticCastSharedPtr<FSampleTrack<TArray<FName>>>(Storage->FindOrAddTrack(MetadataNamesTrackName, FSampleTrackBase::ETrackType_NameArray));
+	if(Track)
+	{
+		Track->AddSample(MetadataNames);
+	}
+
 	for(const FName& MetadataName : MetadataNames)
 	{
 		if(FRigBaseMetadata* Metadata = InElement->GetMetadata(MetadataName))
@@ -315,7 +395,7 @@ void FControlRigReplayTracks::StoreMetaData(URigHierarchy* InHierarchy, FRigBase
 {
 	FSampleTrackContainer* Storage = GetContainer();
 	const FName TrackName = GetTrackName(InElement->GetKey(), InMetadata->GetName());
-	const FSampleTrackBase::ETrackType TrackType = GetTrackTypeFromMetadata(InMetadata->GetType());
+	const FSampleTrackBase::ETrackType TrackType = GetTrackTypeFromMetadataType(InMetadata->GetType());
 	TSharedPtr<FSampleTrackBase> Track = Storage->FindTrack(TrackName);
 	if(Track)
 	{
@@ -493,25 +573,47 @@ bool FControlRigReplayTracks::RestoreMetaData(int32 InTimeIndex, URigHierarchy*
 
 bool FControlRigReplayTracks::RestoreMetaData(int32 InTimeIndex, URigHierarchy* InHierarchy, FRigBaseElement* InElement, TReportFunction InReportFunction) const
 {
+	const TArray<FName>& StoredMetadataNames = GetMetadataNames(InTimeIndex, SampleTrackIndex, InElement->GetKey(), InReportFunction);
+	
 	bool bSuccess = true;
-	const TArray<FName> MetadataNames = InHierarchy->GetMetadataNames(InElement->GetKey());
-	for(const FName& MetadataName : MetadataNames)
+	for(const FName& StoredMetadataName : StoredMetadataNames)
 	{
-		if(FRigBaseMetadata* Metadata = InElement->GetMetadata(MetadataName))
+		if(!RestoreMetaData(InTimeIndex, InHierarchy, InElement, StoredMetadataName, InReportFunction))
 		{
-			if(!RestoreMetaData(InTimeIndex, InHierarchy, InElement, Metadata, InReportFunction))
-			{
-				bSuccess = false;
-			}
+			bSuccess = false;
 		}
 	}
 	return bSuccess;
 }
 
-bool FControlRigReplayTracks::RestoreMetaData(int32 InTimeIndex, URigHierarchy* InHierarchy, FRigBaseElement* InElement, FRigBaseMetadata* InMetadata, TReportFunction InReportFunction) const
+const TArray<FName> FControlRigReplayTracks::GetMetadataNames(int32 InTimeIndex, FSampleTrackIndex& OutSampleTrackIndex, const FRigElementKey& InElementKey, TReportFunction InReportFunction) const
 {
 	const FSampleTrackContainer* Storage = GetContainer();
-	const FName TrackName = GetTrackName(InElement->GetKey(), InMetadata->GetName());
+	static const TArray<FName> EmptyMetadataNames;
+
+	const FName& TrackName = GetTrackName(InElementKey);
+	const FName MetadataNamesTrackName = *(TrackName.ToString() + TEXT("MetadataNames"));
+	const TSharedPtr<const FSampleTrack<TArray<FName>>> Track = Storage->FindTrack<TArray<FName>>(MetadataNamesTrackName);
+	if(!Track.IsValid())
+	{
+		return EmptyMetadataNames;
+	}
+	if(Track->GetTrackType() != FSampleTrackBase::ETrackType_NameArray)
+	{
+		if(InReportFunction)
+		{
+			InReportFunction(EMessageSeverity::Error, TrackName, TEXT("Track has incorrect type."));
+		}
+		return EmptyMetadataNames;
+	}
+
+	return Track->GetValueAtTimeIndex(InTimeIndex, OutSampleTrackIndex);
+}
+
+bool FControlRigReplayTracks::RestoreMetaData(int32 InTimeIndex, URigHierarchy* InHierarchy, FRigBaseElement* InElement, const FName& InMetadataName, TReportFunction InReportFunction) const
+{
+	const FSampleTrackContainer* Storage = GetContainer();
+	const FName TrackName = GetTrackName(InElement->GetKey(), InMetadataName);
 	const TSharedPtr<const FSampleTrackBase> Track = Storage->FindTrack(TrackName);
 	if(!Track.IsValid())
 	{
@@ -521,7 +623,25 @@ bool FControlRigReplayTracks::RestoreMetaData(int32 InTimeIndex, URigHierarchy*
 		}
 		return false;
 	}
-	if(Track->GetTrackType() != GetTrackTypeFromMetadata(InMetadata->GetType()))
+
+	const ERigMetadataType ExpectedMetadataType = GetMetadataTypeFromTrackType(Track->GetTrackType());
+	FRigBaseMetadata* Metadata = InElement->GetMetadata(InMetadataName);
+	if(Metadata == nullptr)
+	{
+		if(ExpectedMetadataType != ERigMetadataType::Invalid)
+		{
+			Metadata = InHierarchy->GetMetadataForElement(InElement, InMetadataName, ExpectedMetadataType, true);
+		}
+		else
+		{
+			if(InReportFunction)
+			{
+				InReportFunction(EMessageSeverity::Warning, TrackName, TEXT("Cannot create metadata. Invalid metadata type."));
+			}
+			return false;
+		}
+	}
+	else if(Metadata->GetType() != ExpectedMetadataType)
 	{
 		if(InReportFunction)
 		{
@@ -530,60 +650,60 @@ bool FControlRigReplayTracks::RestoreMetaData(int32 InTimeIndex, URigHierarchy*
 		return false;
 	}
 	
-	switch(InMetadata->GetType())
+	switch(Metadata->GetType())
 	{
 		case ERigMetadataType::Bool:
 		{
 			const bool& Value = StaticCastSharedPtr<const FSampleTrack<bool>>(Track)->GetValueAtTimeIndex(InTimeIndex, SampleTrackIndex);
-			static_cast<FRigBoolMetadata*>(InMetadata)->SetValue(Value);
+			static_cast<FRigBoolMetadata*>(Metadata)->SetValue(Value);
 			break;
 		}
 		case ERigMetadataType::BoolArray:
 		{
 			const TArray<bool>& Value = StaticCastSharedPtr<const FSampleTrack<TArray<bool>>>(Track)->GetValueAtTimeIndex(InTimeIndex, SampleTrackIndex);
-			static_cast<FRigBoolArrayMetadata*>(InMetadata)->SetValue(Value);
+			static_cast<FRigBoolArrayMetadata*>(Metadata)->SetValue(Value);
 			break;
 		}
 		case ERigMetadataType::Float:
 		{
 			const float& Value = StaticCastSharedPtr<const FSampleTrack<float>>(Track)->GetValueAtTimeIndex(InTimeIndex, SampleTrackIndex);
-			static_cast<FRigFloatMetadata*>(InMetadata)->SetValue(Value);
+			static_cast<FRigFloatMetadata*>(Metadata)->SetValue(Value);
 			break;
 		}
 		case ERigMetadataType::FloatArray:
 		{
 			const TArray<float>& Value = StaticCastSharedPtr<const FSampleTrack<TArray<float>>>(Track)->GetValueAtTimeIndex(InTimeIndex, SampleTrackIndex);
-			static_cast<FRigFloatArrayMetadata*>(InMetadata)->SetValue(Value);
+			static_cast<FRigFloatArrayMetadata*>(Metadata)->SetValue(Value);
 			break;
 		}
 		case ERigMetadataType::Int32:
 		{
 			const int32& Value = StaticCastSharedPtr<const FSampleTrack<int32>>(Track)->GetValueAtTimeIndex(InTimeIndex, SampleTrackIndex);
-			static_cast<FRigInt32Metadata*>(InMetadata)->SetValue(Value);
+			static_cast<FRigInt32Metadata*>(Metadata)->SetValue(Value);
 			break;
 		}
 		case ERigMetadataType::Int32Array:
 		{
 			const TArray<int32>& Value = StaticCastSharedPtr<const FSampleTrack<TArray<int32>>>(Track)->GetValueAtTimeIndex(InTimeIndex, SampleTrackIndex);
-			static_cast<FRigInt32ArrayMetadata*>(InMetadata)->SetValue(Value);
+			static_cast<FRigInt32ArrayMetadata*>(Metadata)->SetValue(Value);
 			break;
 		}
 		case ERigMetadataType::Name:
 		{
 			const FName& Value = StaticCastSharedPtr<const FSampleTrack<FName>>(Track)->GetValueAtTimeIndex(InTimeIndex, SampleTrackIndex);
-			static_cast<FRigNameMetadata*>(InMetadata)->SetValue(Value);
+			static_cast<FRigNameMetadata*>(Metadata)->SetValue(Value);
 			break;
 		}
 		case ERigMetadataType::NameArray:
 		{
 			const TArray<FName>& Value = StaticCastSharedPtr<const FSampleTrack<TArray<FName>>>(Track)->GetValueAtTimeIndex(InTimeIndex, SampleTrackIndex);
-			static_cast<FRigNameArrayMetadata*>(InMetadata)->SetValue(Value);
+			static_cast<FRigNameArrayMetadata*>(Metadata)->SetValue(Value);
 			break;
 		}
 		case ERigMetadataType::Vector:
 		{
 			const FVector3f& Value = StaticCastSharedPtr<const FSampleTrack<FVector3f>>(Track)->GetValueAtTimeIndex(InTimeIndex, SampleTrackIndex);
-			static_cast<FRigVectorMetadata*>(InMetadata)->SetValue(FVector(Value));
+			static_cast<FRigVectorMetadata*>(Metadata)->SetValue(FVector(Value));
 			break;
 		}
 		case ERigMetadataType::VectorArray:
@@ -595,13 +715,13 @@ bool FControlRigReplayTracks::RestoreMetaData(int32 InTimeIndex, URigHierarchy*
 			{
 				ConvertedValue.Emplace(SingleValue);
 			}
-			static_cast<FRigVectorArrayMetadata*>(InMetadata)->SetValue(ConvertedValue);
+			static_cast<FRigVectorArrayMetadata*>(Metadata)->SetValue(ConvertedValue);
 			break;
 		}
 		case ERigMetadataType::Rotator:
 		{
 			const FVector3f& Value = StaticCastSharedPtr<const FSampleTrack<FVector3f>>(Track)->GetValueAtTimeIndex(InTimeIndex, SampleTrackIndex);
-			static_cast<FRigRotatorMetadata*>(InMetadata)->SetValue(FRotator::MakeFromEuler(FVector(Value)));
+			static_cast<FRigRotatorMetadata*>(Metadata)->SetValue(FRotator::MakeFromEuler(FVector(Value)));
 			break;
 		}
 		case ERigMetadataType::RotatorArray:
@@ -613,67 +733,67 @@ bool FControlRigReplayTracks::RestoreMetaData(int32 InTimeIndex, URigHierarchy*
 			{
 				ConvertedValue.Add(FRotator::MakeFromEuler(FVector(SingleValue)));
 			}
-			static_cast<FRigRotatorArrayMetadata*>(InMetadata)->SetValue(ConvertedValue);
+			static_cast<FRigRotatorArrayMetadata*>(Metadata)->SetValue(ConvertedValue);
 			break;
 		}
 		case ERigMetadataType::Quat:
 		{
 			const FQuat4f& Value = StaticCastSharedPtr<const FSampleTrack<FQuat4f>>(Track)->GetValueAtTimeIndex(InTimeIndex, SampleTrackIndex);
-			static_cast<FRigQuatMetadata*>(InMetadata)->SetValue(FQuat(Value)); 
+			static_cast<FRigQuatMetadata*>(Metadata)->SetValue(FQuat(Value)); 
 			break;
 		}
 		case ERigMetadataType::QuatArray:
 		{
-			const TArray<FQuat4f>& Value = StaticCastSharedPtr<const FSampleTrack<TArray<FQuat4f>>>(Track)->GetValueAtSampleIndex(InTimeIndex);
+			const TArray<FQuat4f>& Value = StaticCastSharedPtr<const FSampleTrack<TArray<FQuat4f>>>(Track)->GetValueAtTimeIndex(InTimeIndex, SampleTrackIndex);
 			TArray<FQuat> ConvertedValue;
 			ConvertedValue.Reserve(Value.Num());
 			for(const FQuat4f& SingleValue : Value)
 			{
 				ConvertedValue.Emplace(SingleValue);
 			}
-			static_cast<FRigQuatArrayMetadata*>(InMetadata)->SetValue(ConvertedValue); 
+			static_cast<FRigQuatArrayMetadata*>(Metadata)->SetValue(ConvertedValue); 
 			break;
 		}
 		case ERigMetadataType::Transform:
 		{
-			const FTransform3f& Value = StaticCastSharedPtr<const FSampleTrack<FTransform3f>>(Track)->GetValueAtSampleIndex(InTimeIndex);
-			static_cast<FRigTransformMetadata*>(InMetadata)->SetValue(FTransform(Value)); 
+			const FTransform3f& Value = StaticCastSharedPtr<const FSampleTrack<FTransform3f>>(Track)->GetValueAtTimeIndex(InTimeIndex, SampleTrackIndex);
+			static_cast<FRigTransformMetadata*>(Metadata)->SetValue(FTransform(Value)); 
 			break;
 		}
 		case ERigMetadataType::TransformArray:
 		{
-			const TArray<FTransform3f>& Value = StaticCastSharedPtr<const FSampleTrack<TArray<FTransform3f>>>(Track)->GetValueAtSampleIndex(InTimeIndex);
+			const TArray<FTransform3f>& Value = StaticCastSharedPtr<const FSampleTrack<TArray<FTransform3f>>>(Track)->GetValueAtTimeIndex(InTimeIndex, SampleTrackIndex);
 			TArray<FTransform> ConvertedValue;
 			ConvertedValue.Reserve(Value.Num());
 			for(const FTransform3f& SingleValue : Value)
 			{
 				ConvertedValue.Emplace(SingleValue);
 			}
-			static_cast<FRigTransformArrayMetadata*>(InMetadata)->SetValue(ConvertedValue); 
+			static_cast<FRigTransformArrayMetadata*>(Metadata)->SetValue(ConvertedValue); 
 			break;
 		}
 		case ERigMetadataType::LinearColor:
 		{
-			const FLinearColor& Value = StaticCastSharedPtr<const FSampleTrack<FLinearColor>>(Track)->GetValueAtSampleIndex(InTimeIndex);
-			static_cast<FRigLinearColorMetadata*>(InMetadata)->SetValue(Value); 
+			const FLinearColor& Value = StaticCastSharedPtr<const FSampleTrack<FLinearColor>>(Track)->GetValueAtTimeIndex(InTimeIndex, SampleTrackIndex);
+			static_cast<FRigLinearColorMetadata*>(Metadata)->SetValue(Value); 
 			break;
 		}
 		case ERigMetadataType::LinearColorArray:
 		{
-			const TArray<FLinearColor>& Value = StaticCastSharedPtr<const FSampleTrack<TArray<FLinearColor>>>(Track)->GetValueAtSampleIndex(InTimeIndex);
-			static_cast<FRigLinearColorArrayMetadata*>(InMetadata)->SetValue(Value); 
+			const TArray<FLinearColor>& Value = StaticCastSharedPtr<const FSampleTrack<TArray<FLinearColor>>>(Track)->GetValueAtTimeIndex(InTimeIndex, SampleTrackIndex);
+			static_cast<FRigLinearColorArrayMetadata*>(Metadata)->SetValue(Value); 
 			break;
 		}
 		case ERigMetadataType::RigElementKey:
 		{
-			const FRigElementKey& Value = StaticCastSharedPtr<const FSampleTrack<FRigElementKey>>(Track)->GetValueAtSampleIndex(InTimeIndex);
-			static_cast<FRigElementKeyMetadata*>(InMetadata)->SetValue(Value); 
+			const FRigElementKey& Value = StaticCastSharedPtr<const FSampleTrack<FRigElementKey>>(Track)->GetValueAtTimeIndex(InTimeIndex, SampleTrackIndex);
+			static_cast<FRigElementKeyMetadata*>(Metadata)->SetValue(Value); 
 			break;
 		}
 		case ERigMetadataType::RigElementKeyArray:
 		{
-			const TArray<FRigElementKey>& Value = StaticCastSharedPtr<const FSampleTrack<TArray<FRigElementKey>>>(Track)->GetValueAtSampleIndex(InTimeIndex);
-			static_cast<FRigElementKeyArrayMetadata*>(InMetadata)->SetValue(Value); 
+			const TArray<FRigElementKey>& Value = StaticCastSharedPtr<const FSampleTrack<TArray<FRigElementKey>>>(Track)->GetValueAtTimeIndex(InTimeIndex, SampleTrackIndex);
+			static_cast<FRigElementKeyArrayMetadata*>(Metadata)->SetValue(Value); 
 			break;
 		}
 		case ERigMetadataType::Invalid:
@@ -690,6 +810,8 @@ bool FControlRigReplayTracks::RestoreMetaData(int32 InTimeIndex, URigHierarchy*
 
 void FControlRigReplayTracks::StoreVariables(URigVMHost* InHost)
 {
+	TArray<FName> VariableNames;
+	
 	FSampleTrackContainer* Storage = GetContainer();
 	for (TFieldIterator<FProperty> PropertyIt(InHost->GetClass()); PropertyIt; ++PropertyIt)
 	{
@@ -714,18 +836,17 @@ void FControlRigReplayTracks::StoreVariables(URigVMHost* InHost)
 
 		const FName& TrackName = GetTrackName(Property);
 		const FSampleTrackBase::ETrackType TrackType = GetTrackTypeFromProperty(Property);
-		TSharedPtr<FSampleTrackBase> Track = Storage->FindTrack(TrackName);
-		if(Track)
-		{
-			check(Track->GetTrackType() == TrackType);
-			check(Track->GetScriptStruct() == ScriptStruct);
-		}
-		else
-		{
-			Track = Storage->AddTrack(TrackName, TrackType, ScriptStruct);
-		}
+		TSharedPtr<FSampleTrackBase> Track = Storage->FindOrAddTrack(TrackName, TrackType, ScriptStruct);
 		const uint8* Memory = Property->ContainerPtrToValuePtr<uint8>(InHost);
 		Track->AddSampleFromProperty(Property, Memory);
+
+		VariableNames.Add(Property->GetFName());
+	}
+
+	if(GetNumTimes() == 1 && !VariableNames.IsEmpty())
+	{
+		const TSharedPtr<FSampleTrack<TArray<FName>>> Track = StaticCastSharedPtr<FSampleTrack<TArray<FName>>>(Storage->FindOrAddTrack(VariableNamesName, FSampleTrackBase::ETrackType_NameArray, nullptr));
+		Track->AddSample(VariableNames);
 	}
 }
 
@@ -857,7 +978,7 @@ const FName& FControlRigReplayTracks::GetTrackName(const FProperty* InProperty)
 	return PropertyNameToTrackName.Add(InProperty->GetFName(), TrackName);
 }
 
-FSampleTrackBase::ETrackType FControlRigReplayTracks::GetTrackTypeFromMetadata(ERigMetadataType InMetadataType)
+FSampleTrackBase::ETrackType FControlRigReplayTracks::GetTrackTypeFromMetadataType(ERigMetadataType InMetadataType)
 {
 	switch(InMetadataType)
 	{
@@ -943,6 +1064,98 @@ FSampleTrackBase::ETrackType FControlRigReplayTracks::GetTrackTypeFromMetadata(E
 	return FSampleTrackBase::ETrackType_Unknown;
 }
 
+ERigMetadataType FControlRigReplayTracks::GetMetadataTypeFromTrackType(FSampleTrackBase::ETrackType InTrackType)
+{
+	switch(InTrackType)
+	{
+		case FSampleTrackBase::ETrackType_Bool:
+		{
+			return ERigMetadataType::Bool;
+		}
+		case FSampleTrackBase::ETrackType_BoolArray:
+		{
+			return ERigMetadataType::BoolArray;
+		}
+		case FSampleTrackBase::ETrackType_Float:
+		{
+			return ERigMetadataType::Float;
+		}
+		case FSampleTrackBase::ETrackType_FloatArray:
+		{
+			return ERigMetadataType::FloatArray;
+		}
+		case FSampleTrackBase::ETrackType_Int32:
+		{
+			return ERigMetadataType::Int32;
+		}
+		case FSampleTrackBase::ETrackType_Int32Array:
+		{
+			return ERigMetadataType::Int32Array;
+		}
+		case FSampleTrackBase::ETrackType_Name:
+		{
+			return ERigMetadataType::Name;
+		}
+		case FSampleTrackBase::ETrackType_NameArray:
+		{
+			return ERigMetadataType::NameArray;
+		}
+		case FSampleTrackBase::ETrackType_Vector3f:
+		{
+			return ERigMetadataType::Vector;
+		}
+		case FSampleTrackBase::ETrackType_Vector3fArray:
+		{
+			return ERigMetadataType::VectorArray;
+		}
+		case FSampleTrackBase::ETrackType_Quatf:
+		{
+			return ERigMetadataType::Quat;
+		}
+		case FSampleTrackBase::ETrackType_QuatfArray:
+		{
+			return ERigMetadataType::QuatArray;
+		}
+		case FSampleTrackBase::ETrackType_Transformf:
+		{
+			return ERigMetadataType::Transform;
+		}
+		case FSampleTrackBase::ETrackType_TransformfArray:
+		{
+			return ERigMetadataType::TransformArray;
+		}
+		case FSampleTrackBase::ETrackType_LinearColor:
+		{
+			return ERigMetadataType::LinearColor;
+		}
+		case FSampleTrackBase::ETrackType_LinearColorArray:
+		{
+			return ERigMetadataType::LinearColorArray;
+		}
+		case FSampleTrackBase::ETrackType_ElementKey:
+		{
+			return ERigMetadataType::RigElementKey;
+		}
+		case FSampleTrackBase::ETrackType_ElementKeyArray:
+		{
+			return ERigMetadataType::RigElementKeyArray;
+		}
+		case FSampleTrackBase::ETrackType_Uint32:
+		case FSampleTrackBase::ETrackType_String:
+		case FSampleTrackBase::ETrackType_ComponentKey:
+		case FSampleTrackBase::ETrackType_Struct:
+		case FSampleTrackBase::ETrackType_Uint32Array:
+		case FSampleTrackBase::ETrackType_StringArray:
+		case FSampleTrackBase::ETrackType_ComponentKeyArray:
+		case FSampleTrackBase::ETrackType_StructArray:
+		case FSampleTrackBase::ETrackType_Unknown:
+		{
+			break;
+		}
+	}
+	return ERigMetadataType::Invalid;
+}
+
 FSampleTrackBase::ETrackType FControlRigReplayTracks::GetTrackTypeFromProperty(const FProperty* InProperty)
 {
 	check(InProperty);
@@ -995,6 +1208,18 @@ FSampleTrackBase::ETrackType FControlRigReplayTracks::GetTrackTypeFromProperty(c
 	return TrackType;
 }
 
+const TArray<FRigElementKey> FControlRigReplayTracks::GetElementKeys() const
+{
+	const TSharedPtr<const FSampleTrack<TArray<FRigElementKey>>> ElementKeysTrack = GetContainer()->FindTrack<TArray<FRigElementKey>>(ElementKeysName);
+	if(ElementKeysTrack)
+	{
+		static FSampleTrackIndex SingletonTrackIndex = FSampleTrackIndex::MakeSingleton();
+		return ElementKeysTrack->GetValueAtTimeIndex(0, SingletonTrackIndex);
+	}
+	static const TArray<FRigElementKey> EmptyElementKeys;
+	return EmptyElementKeys;
+}
+
 //////////////////////////////////////////////////////////////////////////////////////////////////////////////////////////////
 /////// UControlRigReplay
 //////////////////////////////////////////////////////////////////////////////////////////////////////////////////////////////
@@ -1008,11 +1233,6 @@ void UControlRigReplay::Serialize(FArchive& Ar)
 {
 	Ar.UsingCustomVersion(FControlRigObjectVersion::GUID);
 	UObject::Serialize(Ar);
-
-	if (EventQueue.IsEmpty())
-	{
-		EventQueue = {FRigUnit_BeginExecution::EventName};
-	}
 }
 
 UControlRigReplay* UControlRigReplay::CreateNewAsset(FString InDesiredPackagePath, FString InBlueprintPathName, UClass* InAssetClass)
@@ -1047,7 +1267,7 @@ UControlRigReplay* UControlRigReplay::CreateNewAsset(FString InDesiredPackagePat
 
 FVector2D UControlRigReplay::GetTimeRange() const
 {
-	return FVector2D(Tracks.GetTimeRange());
+	return FVector2D(OutputTracks.GetContainer()->GetTimeRange());
 }
 
 bool UControlRigReplay::StartRecording(UControlRig* InControlRig)
@@ -1060,17 +1280,27 @@ bool UControlRigReplay::StartRecording(UControlRig* InControlRig)
 	(void)StopRecording();
 	(void)StopReplay();
 
-	Tracks.Reset();
+	InputTracks.Reset();
+	OutputTracks.Reset();
+	
 	RecordControlRig = InControlRig;
 	ClearDelegates(InControlRig);
 
-	EventQueue = InControlRig->EventQueue;
 	TimeAtStartOfRecording = TimeOfLastFrame = FPlatformTime::Seconds();
+	bStoreVariablesDuringPreEvent = true;
+	InControlRig->SetAbsoluteAndDeltaTime(0.f, InControlRig->GetDeltaTime());
 
 	PreConstructionHandle = InControlRig->OnPreConstruction_AnyThread().AddLambda(
 		[this](UControlRig* InControlRig, const FName& InEventName)
 		{
-			// nothing to do here
+			if (InEventName != FRigUnit_PrepareForExecution::EventName)
+			{
+				return;
+			}
+			// store the first frame of variables
+			InputTracks.AddTimeSample(InControlRig->GetAbsoluteTime(), InControlRig->GetDeltaTime());
+			InputTracks.StoreVariables(InControlRig);
+			bStoreVariablesDuringPreEvent = false;
 		}
 	);
 
@@ -1087,13 +1317,24 @@ bool UControlRigReplay::StartRecording(UControlRig* InControlRig)
 			{
 				return;
 			}
-			
-			if (InEventName != EventQueue[0])
+
+			if(!ControlRig->SupportsEvent(InEventName))
 			{
 				return;
 			}
 
-			// todo: record a frame if we are recording backwards?
+			if(bStoreVariablesDuringPreEvent)
+			{
+				InputTracks.AddTimeSample(ControlRig->GetAbsoluteTime(), ControlRig->GetDeltaTime());
+				InputTracks.StoreVariables(ControlRig);
+			}
+			InputTracks.StoreRigVMEvent(InEventName);
+			InputTracks.StoreInteraction(ControlRig->InteractionType, ControlRig->ElementsBeingInteracted);
+			InputTracks.StoreHierarchy(ControlRig->GetHierarchy());
+			
+			OutputTracks.AddTimeSample(ControlRig->GetAbsoluteTime(), ControlRig->GetDeltaTime());
+			
+			bStoreVariablesDuringPreEvent = true;
 		}
 	);
 
@@ -1105,28 +1346,36 @@ bool UControlRigReplay::StartRecording(UControlRig* InControlRig)
 			{
 				return;
 			}
-
+			
 			if (InEventName == FRigUnit_PrepareForExecution::EventName)
 			{
 				return;
 			}
-
-			if (InEventName != EventQueue.Last())
-			{
-				return;
-			}
-
+			
 			const double CurrentTime = FPlatformTime::Seconds(); 
-			const double DeltaTime = CurrentTime - TimeOfLastFrame;
+			const double DisplayDeltaTime = CurrentTime - TimeOfLastFrame;
 			TimeOfLastFrame = CurrentTime;
 
-			Tracks.AddTimeSampleFromDeltaTime(static_cast<float>(DeltaTime));
-			Tracks.StoreHierarchy(ControlRig->GetHierarchy());
-			Tracks.StoreVariables(ControlRig);
+			OutputTracks.StoreVariables(ControlRig);
+			OutputTracks.StoreHierarchy(ControlRig->GetHierarchy());
+			const double RecordingDuration = OutputTracks.GetLastAbsoluteTime();
+
+#if WITH_EDITOR
+			static constexpr TCHAR PrintStringFormat[] = TEXT("Recorded time... %.02f");
+			UKismetSystemLibrary::PrintString(ControlRig->GetWorld(), FString::Printf(PrintStringFormat, static_cast<float>(RecordingDuration)), true, false, FLinearColor::White, DisplayDeltaTime * 0.5f);
+#endif
+
+			if(DesiredRecordingDuration >= -KINDA_SMALL_NUMBER)
+			{
+				if(RecordingDuration >= DesiredRecordingDuration)
+				{
+					StopRecording();
+				}
+			}
 		}
 	);
 
-	if(Tracks.GetContainer()->NumTracks() == 0)
+	if(InputTracks.GetContainer()->NumTracks() == 0)
 	{
 		InControlRig->RequestInit();
 	}
@@ -1140,7 +1389,10 @@ bool UControlRigReplay::StopRecording()
 	{
 		ClearDelegates(ControlRig);
 		RecordControlRig.Reset();
-		Tracks.Compact();
+		InputTracks.Compact();
+		OutputTracks.Compact();
+		TimeAtStartOfRecording = -1.0;
+		DesiredRecordingDuration = -1.0;
 		return true;
 	}
 	return false;
@@ -1157,7 +1409,146 @@ EControlRigReplayPlaybackMode UControlRigReplay::GetPlaybackMode() const
 
 bool UControlRigReplay::StartReplay(UControlRig* InControlRig)
 {
-	return false;
+	StopRecording();
+	StopReplay();
+
+	if(InControlRig == nullptr)
+	{
+		return false;
+	}
+
+	if(InputTracks.GetNumTimes() != OutputTracks.GetNumTimes())
+	{
+		return false;
+	}
+
+	bFirstPlayThrough = true;
+	bStoreVariablesDuringPreEvent = true;
+
+	if(InputTracks.IsEmpty() || OutputTracks.IsEmpty())
+	{
+		return false;
+	}
+
+	// reset the control rig's playback index
+	InControlRig->SetReplay(this);
+	InControlRig->SetReplayTimeIndex(0);
+	
+	PreConstructionHandle = InControlRig->OnPreConstruction_AnyThread().AddLambda(
+		[this](UControlRig* InControlRig, const FName& InEventName)
+		{
+			if (InEventName != FRigUnit_PrepareForExecution::EventName)
+			{
+				return;
+			}
+			// restore the first set of variables
+			InputTracks.RestoreVariables(0, InControlRig);
+			bStoreVariablesDuringPreEvent = false;
+		}
+	);
+
+	PreEventHandle = InControlRig->OnPreExecuted_AnyThread().AddLambda(
+		[this](URigVMHost* InRigVMHost, const FName& InEventName)
+		{
+			UControlRig* ControlRig = Cast<UControlRig>(InRigVMHost);
+			if (!ControlRig || InEventName == FRigUnit_PrepareForExecution::EventName)
+			{
+				return;
+			}
+
+			const int32 ReplayTimeIndex = ControlRig->GetReplayTimeIndex();
+			check(ReplayTimeIndex != INDEX_NONE);
+
+			const FName EventNameFromReplay = InputTracks.GetRigVMEvent(ReplayTimeIndex);
+			if (InEventName != EventNameFromReplay)
+			{
+				return;
+			}
+
+			const TTuple<uint8, TArray<FRigElementKey>> Interaction = InputTracks.GetInteraction(ReplayTimeIndex);
+			ControlRig->InteractionType = Interaction.Get<0>();
+			ControlRig->ElementsBeingInteracted = Interaction.Get<1>();
+
+			ControlRig->SetAbsoluteAndDeltaTime(
+				InputTracks.GetAbsoluteTime(ReplayTimeIndex),
+				InputTracks.GetDeltaTime(ReplayTimeIndex)
+			);
+
+			if(bStoreVariablesDuringPreEvent)
+			{
+				InputTracks.RestoreVariables(ReplayTimeIndex, ControlRig);
+			}
+			InputTracks.RestoreHierarchy(ReplayTimeIndex, ControlRig->GetHierarchy());
+			bStoreVariablesDuringPreEvent = true;
+		}
+	);
+
+	PostEventHandle = InControlRig->OnExecuted_AnyThread().AddLambda(
+		[this](URigVMHost* InRigVMHost, const FName& InEventName)
+		{
+			UControlRig* ControlRig = Cast<UControlRig>(InRigVMHost);
+			if (!ControlRig || InEventName == FRigUnit_PrepareForExecution::EventName)
+			{
+				return;
+			}
+
+			const int32 ReplayTimeIndex = ControlRig->GetReplayTimeIndex();
+
+			ControlRig->SetAbsoluteAndDeltaTime(
+				InputTracks.GetAbsoluteTime(ReplayTimeIndex),
+				InputTracks.GetDeltaTime(ReplayTimeIndex)
+			);
+
+			// only validate results during the first play through
+			const FName EventNameFromReplay = InputTracks.GetRigVMEvent(ReplayTimeIndex);
+			if(bFirstPlayThrough && InEventName == EventNameFromReplay)
+			{
+				const int32 OutputTimeIndex = ReplayTimeIndex;
+				
+				(void)ValidateExpectedResults(OutputTimeIndex, OutputTracks.SampleTrackIndex, ControlRig,
+					[ControlRig](EMessageSeverity::Type InSeverity, const FString InMessage)
+					{
+						switch(InSeverity)
+						{
+							case EMessageSeverity::Warning:
+							case EMessageSeverity::PerformanceWarning:
+							{
+								UE_LOG(LogControlRig, Warning, TEXT("%s"), *InMessage);
+								break;
+							}
+							case EMessageSeverity::Error:
+							{
+								UE_LOG(LogControlRig, Error, TEXT("%s"), *InMessage);
+								break;
+							}
+							case EMessageSeverity::Info:
+							default:
+							{
+								UE_LOG(LogControlRig, Display, TEXT("%s"), *InMessage);
+								break;
+							}
+						}
+					}
+				);
+			}
+
+			// loop the animation
+			if(ReplayTimeIndex >= InputTracks.GetNumTimes() - 1)
+			{
+				ControlRig->SetReplayTimeIndex(0);
+				bFirstPlayThrough = false;
+			}
+			else
+			{
+				ControlRig->SetReplayTimeIndex(ReplayTimeIndex + 1);
+			}
+		}
+	);
+
+	InControlRig->RequestInit();
+
+	ReplayControlRig = InControlRig;
+	return true;
 }
 
 bool UControlRigReplay::StopReplay()
@@ -1165,6 +1556,9 @@ bool UControlRigReplay::StopReplay()
 	if(UControlRig* ControlRig = ReplayControlRig.Get())
 	{
 		ClearDelegates(ControlRig);
+		ControlRig->DisableReplay();
+		ControlRig->InteractionType = 0;
+		ControlRig->ElementsBeingInteracted.Reset();
 		ReplayControlRig.Reset();
 		return true;
 	}
@@ -1183,14 +1577,760 @@ bool UControlRigReplay::IsRecording() const
 
 bool UControlRigReplay::IsValidForTesting() const
 {
-	// todo
-	return false;
+	return !InputTracks.IsEmpty() &&
+		!OutputTracks.IsEmpty() &&
+		InputTracks.GetNumTimes() == OutputTracks.GetNumTimes() &&
+		ControlRigObjectPath.IsValid();
 }
 
 bool UControlRigReplay::PerformTest(UControlRig* InSubject, TFunction<void(EMessageSeverity::Type, const FString&)> InLogFunction) const
 {
-	// todo
-	return false;
+	// if we have nothing to check we can consider this test successful
+	if(InputTracks.IsEmpty() || OutputTracks.IsEmpty())
+	{
+		return true;
+	}
+
+	if(!FMath::IsNearlyEqual(InputTracks.GetAbsoluteTime(0), OutputTracks.GetAbsoluteTime(0)) ||
+		!FMath::IsNearlyEqual(InputTracks.GetLastAbsoluteTime(), OutputTracks.GetLastAbsoluteTime()))
+	{
+		if(InLogFunction)
+		{
+			InLogFunction(EMessageSeverity::Error, TEXT("Test Replay is corrupt. Input and Output time ranges don't match."));
+		}
+		return false;
+	}
+
+	bool bSuccess = true;
+
+	auto ReportFunction = [&bSuccess, InLogFunction](EMessageSeverity::Type Severity, const FName& Key, const FString& Message)
+	{
+		if(InLogFunction)
+		{
+			if(Key.IsNone())
+			{
+				InLogFunction(Severity, Message);
+			}
+			else
+			{
+				InLogFunction(Severity, FString::Printf(TEXT("%s: %s"), *Key.ToString(), *Message));
+			}
+		}
+		if(Severity == EMessageSeverity::Error)
+		{
+			bSuccess = false;
+		}
+	};
+
+	URigHierarchy* Hierarchy = InSubject->GetHierarchy();
+
+	// set up the rig by restoring variables and then running construction 
+	InSubject->RequestInit();
+	InputTracks.RestoreVariables(0, InSubject, ReportFunction);
+
+	// restore the connectors
+	const TArray<FRigElementKey> ConnectorKeys = Hierarchy->GetConnectorKeys();
+	FRigElementKeyRedirector::FKeyMap ConnectorMap;
+	for(const FRigElementKey& ConnectorKey : ConnectorKeys)
+	{
+		const FName& TrackName = InputTracks.GetTrackName(ConnectorKey);
+		const FName& ConnectorTrackName = *(TrackName.ToString() + TEXT("ConnectorTargets"));
+
+		if(const TSharedPtr<const FSampleTrack<TArray<FRigElementKey>>> Track = InputTracks.GetContainer()->FindTrack<TArray<FRigElementKey>>(ConnectorTrackName))
+		{
+			FRigElementKeyRedirector::FKeyArray Targets;
+			Targets.Append(Track->GetValueAtTimeIndex(0, InputTracks.SampleTrackIndex));
+			ConnectorMap.Add(ConnectorKey, Targets);
+		}
+	}
+	if(!ConnectorMap.IsEmpty())
+	{
+		InSubject->SetElementKeyRedirector(FRigElementKeyRedirector(ConnectorMap, Hierarchy));
+	}
+	
+	InSubject->SetAbsoluteAndDeltaTime(InputTracks.GetAbsoluteTime(0), InputTracks.GetDeltaTime(0));
+	
+	InSubject->EventQueue = {FRigUnit_BeginExecution::EventName};
+	InSubject->Evaluate_AnyThread();
+
+	FSampleTrackIndex SampleTrackIndex(OutputTracks.GetContainer()->NumTracks());
+	
+	for(int32 InputTimeIndex = 0; InputTimeIndex < InputTracks.GetNumTimes(); InputTimeIndex++)
+	{
+		InSubject->EventQueue = {InputTracks.GetRigVMEvent(InputTimeIndex)};
+		const TTuple<uint8, TArray<FRigElementKey>> Interaction = InputTracks.GetInteraction(InputTimeIndex);
+		InSubject->InteractionType = Interaction.Get<0>();
+		InSubject->ElementsBeingInteracted = Interaction.Get<1>();
+		
+		if(InputTimeIndex > 0)
+		{
+			InputTracks.RestoreVariables(InputTimeIndex, InSubject, ReportFunction);
+		}
+		
+		InputTracks.RestoreHierarchy(InputTimeIndex, Hierarchy, {}, ReportFunction);
+		
+		InSubject->SetAbsoluteAndDeltaTime(
+			InputTracks.GetAbsoluteTime(InputTimeIndex),
+			InputTracks.GetDeltaTime(InputTimeIndex));
+		
+		InSubject->Evaluate_AnyThread();
+
+		const int OutTimeIndex = InputTimeIndex;
+		if(!ValidateExpectedResults(OutTimeIndex, SampleTrackIndex, InSubject, InLogFunction))
+		{
+			bSuccess = false;
+		}
+	}
+
+	return bSuccess;
+}
+
+bool UControlRigReplay::ValidateExpectedResults(
+	int32 InPlaybackTimeIndex,
+	FSampleTrackIndex& OutSampleTrackIndex,
+	UControlRig* InSubject,
+	TFunction<void(EMessageSeverity::Type, const FString&)> InLogFunction) const
+{
+	URigHierarchy* Hierarchy = InSubject->GetHierarchy();
+	const TArray<FRigElementKey>& StoredElementKeys = OutputTracks.GetElementKeys();
+
+	bool bSuccess = true;
+	if(bValidateHierarchyTopology && InPlaybackTimeIndex == 0)
+	{
+		TArray<FRigElementKey> CurrentElementKeys = Hierarchy->GetAllKeys();
+		if(StoredElementKeys.Num() != CurrentElementKeys.Num())
+		{
+			if(InLogFunction)
+			{
+				InLogFunction(EMessageSeverity::Error, FString::Printf(TEXT("Number of elements in hierarchy (%d) and number of elements in replay (%d) don't match."), CurrentElementKeys.Num(), StoredElementKeys.Num()));
+			}
+			bSuccess = false;
+		}
+
+		for(const FRigElementKey& StoredKey : StoredElementKeys)
+		{
+			if(!CurrentElementKeys.Contains(StoredKey))
+			{
+				if(InLogFunction)
+				{
+					InLogFunction(EMessageSeverity::Error, FString::Printf(TEXT("Hierarchy is missing element %s expected by the replay."), *StoredKey.ToString()));
+				}
+				bSuccess = false;
+			}
+		}
+
+		for(const FRigElementKey& CurrentKey : CurrentElementKeys)
+		{
+			if(!StoredElementKeys.Contains(CurrentKey))
+			{
+				if(InLogFunction)
+				{
+					InLogFunction(EMessageSeverity::Error, FString::Printf(TEXT("Hierarchy contains element %s which is not part of the replay."), *CurrentKey.ToString()));
+				}
+				bSuccess = false;
+			}
+		}
+
+		const TSharedPtr<const FSampleTrack<TArray<int32>>> ParentIndicesTrack = OutputTracks.GetContainer()->FindTrack<TArray<int32>>(FControlRigReplayTracks::ParentIndicesName);
+		const TArray<int32> StoredParentIndices = ParentIndicesTrack->GetValueAtTimeIndex(InPlaybackTimeIndex, OutSampleTrackIndex);
+
+		for(int32 Index = 0; Index < CurrentElementKeys.Num(); Index++)
+		{
+			const FRigElementKey& CurrentKey = CurrentElementKeys[Index];
+			const int32 StoredIndex = StoredElementKeys.Find(CurrentKey);
+			if(StoredIndex == INDEX_NONE)
+			{
+				continue;
+			}
+			if(Hierarchy->Find<FRigTransformElement>(CurrentKey) == nullptr)
+			{
+				continue;
+			}
+			const int32 CurrentParentIndex = CurrentElementKeys.Find(Hierarchy->GetDefaultParent(CurrentKey));
+			const int32 StoredParentIndex = StoredParentIndices[StoredIndex];
+
+			if(CurrentParentIndex == INDEX_NONE && StoredParentIndex != INDEX_NONE)
+			{
+				if(InLogFunction)
+				{
+					InLogFunction(EMessageSeverity::Error, FString::Printf(TEXT("Element %s is at root currently, but is parented to %s in the replay."), *CurrentKey.ToString(), *StoredElementKeys[StoredParentIndex].ToString()));
+				}
+				bSuccess = false;
+			}
+			else if(CurrentParentIndex != INDEX_NONE && StoredParentIndex == INDEX_NONE)
+			{
+				if(InLogFunction)
+				{
+					InLogFunction(EMessageSeverity::Error, FString::Printf(TEXT("Element %s is parent to %s currently, but was root in the replay."), *CurrentKey.ToString(), *CurrentElementKeys[CurrentParentIndex].ToString()));
+				}
+				bSuccess = false;
+			}
+			else if(CurrentParentIndex != INDEX_NONE && StoredParentIndex != INDEX_NONE)
+			{
+				if(CurrentElementKeys[CurrentParentIndex] != StoredElementKeys[StoredParentIndex])
+				{
+					if(InLogFunction)
+					{
+						InLogFunction(EMessageSeverity::Error, FString::Printf(TEXT("Element %s is parent to %s currently, but was parented to %s in the replay."), *CurrentKey.ToString(), *CurrentElementKeys[CurrentParentIndex].ToString(), *StoredElementKeys[StoredParentIndex].ToString()));
+					}
+					bSuccess = false;
+				}
+			}
+		}
+	}
+
+	if(bValidatePose)
+	{
+		for(const FRigElementKey& StoredKey : StoredElementKeys)
+		{
+			if(const FRigBaseElement* Element = Hierarchy->Find(StoredKey))
+			{
+				const FName& TrackName = OutputTracks.GetTrackName(StoredKey);
+				if(Element->IsA<FRigCurveElement>())
+				{
+					if(const TSharedPtr<const FSampleTrack<float>> Track = OutputTracks.GetContainer()->FindTrack<float>(TrackName))
+					{
+						const float& StoredValue = Track->GetValueAtTimeIndex(InPlaybackTimeIndex, OutSampleTrackIndex);
+						const float CurrentValue = Hierarchy->GetCurveValueByIndex(Element->GetIndex());
+						if(!FMath::IsNearlyEqual(StoredValue, CurrentValue, static_cast<float>(Tolerance)))
+						{
+							if(InLogFunction)
+							{
+								InLogFunction(EMessageSeverity::Error, FString::Printf(TEXT("Frame [%04d]: Curve %s has value %.03f, expected by the replay: %.03f."), InPlaybackTimeIndex, *StoredKey.ToString(), CurrentValue, StoredValue));
+							}
+							bSuccess = false;
+						}
+					}
+				}
+				else if(Element->IsA<FRigTransformElement>())
+				{
+					if(const TSharedPtr<const FSampleTrack<FTransform3f>> Track = OutputTracks.GetContainer()->FindTrack<FTransform3f>(TrackName))
+					{
+						const FTransform StoredValue = FTransform(Track->GetValueAtTimeIndex(InPlaybackTimeIndex, OutSampleTrackIndex));
+						const FTransform CurrentValue = Hierarchy->GetLocalTransform(Element->GetIndex());
+						
+						if(!StoredValue.GetLocation().Equals(CurrentValue.GetLocation(), static_cast<float>(Tolerance)))
+						{
+							if(InLogFunction)
+							{
+								InLogFunction(EMessageSeverity::Error, FString::Printf(TEXT("Frame [%04d]: Element %s has position %s, expected %s by the replay."), InPlaybackTimeIndex, *StoredKey.ToString(), *CurrentValue.GetLocation().ToString(), *StoredValue.GetLocation().ToString()));
+							}
+							bSuccess = false;
+						}
+						if(!StoredValue.Rotator().EqualsOrientation(CurrentValue.Rotator(), static_cast<float>(Tolerance)))
+						{
+							if(InLogFunction)
+							{
+								InLogFunction(EMessageSeverity::Error, FString::Printf(TEXT("Frame [%04d]: Element %s has rotation %s, expected %s by the replay."), InPlaybackTimeIndex, *StoredKey.ToString(), *CurrentValue.Rotator().ToString(), *StoredValue.Rotator().ToString()));
+							}
+							bSuccess = false;
+						}
+						if(!StoredValue.GetScale3D().Equals(CurrentValue.GetScale3D(), static_cast<float>(Tolerance)))
+						{
+							if(InLogFunction)
+							{
+								InLogFunction(EMessageSeverity::Error, FString::Printf(TEXT("Frame [%04d]: Element %s has scale %s, expected %s by the replay."), InPlaybackTimeIndex, *StoredKey.ToString(), *CurrentValue.GetScale3D().ToString(), *StoredValue.GetScale3D().ToString()));
+							}
+							bSuccess = false;
+						}
+					}
+				}
+			}
+		}
+	}
+
+	if(bValidateMetadata)
+	{
+		for(const FRigElementKey& StoredKey : StoredElementKeys)
+		{
+			if(const FRigBaseElement* Element = Hierarchy->Find(StoredKey))
+			{
+				auto TrackNameBasedLogFunction = [InLogFunction](EMessageSeverity::Type InSeverity, const FName& InTrackName, const FString& InMessage)
+				{
+					if(InLogFunction)
+					{
+						InLogFunction(InSeverity, FString::Printf(TEXT("%s: %s"), *InTrackName.ToString(), *InMessage));
+					}
+				};
+				
+				const TArray<FName>& StoredMetadataNames = OutputTracks.GetMetadataNames(InPlaybackTimeIndex, OutSampleTrackIndex, StoredKey, TrackNameBasedLogFunction);
+
+				for(const FName& StoredMetadataName : StoredMetadataNames)
+				{
+					const FName& TrackName = OutputTracks.GetTrackName(StoredKey, StoredMetadataName);
+					const TSharedPtr<const FSampleTrackBase> MetadataTrack = OutputTracks.GetContainer()->FindTrack(TrackName);
+					if(MetadataTrack.IsValid())
+					{
+						ERigMetadataType ExpectedMetadataType = OutputTracks.GetMetadataTypeFromTrackType(MetadataTrack->GetTrackType());
+						if(const FRigBaseMetadata* Metadata = Element->GetMetadata(StoredMetadataName, ExpectedMetadataType))
+						{
+							switch(ExpectedMetadataType)
+							{
+								case ERigMetadataType::Bool:
+								{
+									const bool& StoredValue = StaticCastSharedPtr<const FSampleTrack<bool>>(MetadataTrack)->GetValueAtTimeIndex(InPlaybackTimeIndex, OutSampleTrackIndex);
+									const bool& CurrentValue = static_cast<const FRigBoolMetadata*>(Metadata)->GetValue();
+									if(StoredValue != CurrentValue)
+									{
+										TrackNameBasedLogFunction(EMessageSeverity::Error, TrackName, TEXT("Metadata value doesn't match."));
+										bSuccess = false;
+									}
+									break;
+								}
+								case ERigMetadataType::BoolArray:
+								{
+									const TArray<bool>& StoredValue = StaticCastSharedPtr<const FSampleTrack<TArray<bool>>>(MetadataTrack)->GetValueAtTimeIndex(InPlaybackTimeIndex, OutSampleTrackIndex);
+									const TArray<bool>& CurrentValue = static_cast<const FRigBoolArrayMetadata*>(Metadata)->GetValue();
+									if(StoredValue.Num() != CurrentValue.Num())
+									{
+										TrackNameBasedLogFunction(EMessageSeverity::Error, TrackName, TEXT("Metadata array value num doesn't match."));
+										bSuccess = false;
+									}
+									else
+									{
+										for(int32 StoredValueIndex = 0; StoredValueIndex < StoredValue.Num(); StoredValueIndex++)
+										{
+											if(StoredValue[StoredValueIndex] != CurrentValue[StoredValueIndex])
+											{
+												TrackNameBasedLogFunction(EMessageSeverity::Error, TrackName, TEXT("Metadata array value element doesn't match."));
+												bSuccess = false;
+											}
+										}
+									}
+									break;
+								}
+								case ERigMetadataType::Float:
+								{
+									const float& StoredValue = StaticCastSharedPtr<const FSampleTrack<float>>(MetadataTrack)->GetValueAtTimeIndex(InPlaybackTimeIndex, OutSampleTrackIndex);
+									const float& CurrentValue = static_cast<const FRigFloatMetadata*>(Metadata)->GetValue();
+									if(FMath::IsNearlyEqual(StoredValue, CurrentValue, static_cast<float>(Tolerance)))
+									{
+										TrackNameBasedLogFunction(EMessageSeverity::Error, TrackName, TEXT("Metadata value doesn't match."));
+										bSuccess = false;
+									}
+									break;
+								}
+								case ERigMetadataType::FloatArray:
+								{
+									const TArray<float>& StoredValue = StaticCastSharedPtr<const FSampleTrack<TArray<float>>>(MetadataTrack)->GetValueAtTimeIndex(InPlaybackTimeIndex, OutSampleTrackIndex);
+									const TArray<float>& CurrentValue = static_cast<const FRigFloatArrayMetadata*>(Metadata)->GetValue();
+										if(StoredValue.Num() != CurrentValue.Num())
+										{
+											TrackNameBasedLogFunction(EMessageSeverity::Error, TrackName, TEXT("Metadata array value num doesn't match."));
+											bSuccess = false;
+										}
+										else
+										{
+											for(int32 StoredValueIndex = 0; StoredValueIndex < StoredValue.Num(); StoredValueIndex++)
+											{
+												if(FMath::IsNearlyEqual(StoredValue[StoredValueIndex], CurrentValue[StoredValueIndex], static_cast<float>(Tolerance)))
+												{
+													TrackNameBasedLogFunction(EMessageSeverity::Error, TrackName, TEXT("Metadata array value element doesn't match."));
+													bSuccess = false;
+												}
+											}
+										}
+										break;
+								}
+								case ERigMetadataType::Int32:
+								{
+									const int32& StoredValue = StaticCastSharedPtr<const FSampleTrack<int32>>(MetadataTrack)->GetValueAtTimeIndex(InPlaybackTimeIndex, OutSampleTrackIndex);
+									const int32& CurrentValue = static_cast<const FRigInt32Metadata*>(Metadata)->GetValue();
+									if(StoredValue != CurrentValue)
+									{
+										TrackNameBasedLogFunction(EMessageSeverity::Error, TrackName, TEXT("Metadata value doesn't match."));
+										bSuccess = false;
+									}
+									break;
+								}
+								case ERigMetadataType::Int32Array:
+								{
+									const TArray<int32>& StoredValue = StaticCastSharedPtr<const FSampleTrack<TArray<int32>>>(MetadataTrack)->GetValueAtTimeIndex(InPlaybackTimeIndex, OutSampleTrackIndex);
+									const TArray<int32>& CurrentValue = static_cast<const FRigInt32ArrayMetadata*>(Metadata)->GetValue();
+										if(StoredValue.Num() != CurrentValue.Num())
+										{
+											TrackNameBasedLogFunction(EMessageSeverity::Error, TrackName, TEXT("Metadata array value num doesn't match."));
+											bSuccess = false;
+										}
+										else
+										{
+											for(int32 StoredValueIndex = 0; StoredValueIndex < StoredValue.Num(); StoredValueIndex++)
+											{
+												if(StoredValue[StoredValueIndex] != CurrentValue[StoredValueIndex])
+												{
+													TrackNameBasedLogFunction(EMessageSeverity::Error, TrackName, TEXT("Metadata array value element doesn't match."));
+													bSuccess = false;
+												}
+											}
+										}
+										break;
+								}
+								case ERigMetadataType::Name:
+								{
+									const FName& StoredValue = StaticCastSharedPtr<const FSampleTrack<FName>>(MetadataTrack)->GetValueAtTimeIndex(InPlaybackTimeIndex, OutSampleTrackIndex);
+									const FName& CurrentValue = static_cast<const FRigNameMetadata*>(Metadata)->GetValue();
+									if(StoredValue != CurrentValue)
+									{
+										TrackNameBasedLogFunction(EMessageSeverity::Error, TrackName, TEXT("Metadata value doesn't match."));
+										bSuccess = false;
+									}
+									break;
+								}
+								case ERigMetadataType::NameArray:
+								{
+									const TArray<FName>& StoredValue = StaticCastSharedPtr<const FSampleTrack<TArray<FName>>>(MetadataTrack)->GetValueAtTimeIndex(InPlaybackTimeIndex, OutSampleTrackIndex);
+									const TArray<FName>& CurrentValue = static_cast<const FRigNameArrayMetadata*>(Metadata)->GetValue();
+									if(StoredValue.Num() != CurrentValue.Num())
+									{
+										TrackNameBasedLogFunction(EMessageSeverity::Error, TrackName, TEXT("Metadata array value num doesn't match."));
+										bSuccess = false;
+									}
+									else
+									{
+										for(int32 StoredValueIndex = 0; StoredValueIndex < StoredValue.Num(); StoredValueIndex++)
+										{
+											if(StoredValue[StoredValueIndex] != CurrentValue[StoredValueIndex])
+											{
+												TrackNameBasedLogFunction(EMessageSeverity::Error, TrackName, TEXT("Metadata array value element doesn't match."));
+												bSuccess = false;
+											}
+										}
+									}
+									break;
+								}
+								case ERigMetadataType::Vector:
+								{
+									const FVector3f& StoredValue = StaticCastSharedPtr<const FSampleTrack<FVector3f>>(MetadataTrack)->GetValueAtTimeIndex(InPlaybackTimeIndex, OutSampleTrackIndex);
+									const FVector3f& CurrentValue = FVector3f(static_cast<const FRigVectorMetadata*>(Metadata)->GetValue());
+									if(StoredValue.Equals(CurrentValue, static_cast<float>(Tolerance)))
+									{
+										TrackNameBasedLogFunction(EMessageSeverity::Error, TrackName, TEXT("Metadata value doesn't match."));
+										bSuccess = false;
+									}
+									break;
+								}
+								case ERigMetadataType::VectorArray:
+								{
+									const TArray<FVector3f>& StoredValue = StaticCastSharedPtr<const FSampleTrack<TArray<FVector3f>>>(MetadataTrack)->GetValueAtTimeIndex(InPlaybackTimeIndex, OutSampleTrackIndex);
+									const TArray<FVector>& CurrentValue = static_cast<const FRigVectorArrayMetadata*>(Metadata)->GetValue();
+									if(StoredValue.Num() != CurrentValue.Num())
+									{
+										TrackNameBasedLogFunction(EMessageSeverity::Error, TrackName, TEXT("Metadata array value num doesn't match."));
+										bSuccess = false;
+									}
+									else
+									{
+										for(int32 StoredValueIndex = 0; StoredValueIndex < StoredValue.Num(); StoredValueIndex++)
+										{
+											if(StoredValue[StoredValueIndex].Equals(FVector3f(CurrentValue[StoredValueIndex]), static_cast<float>(Tolerance)))
+											{
+												TrackNameBasedLogFunction(EMessageSeverity::Error, TrackName, TEXT("Metadata array value element doesn't match."));
+												bSuccess = false;
+											}
+										}
+									}
+									break;
+								}
+								case ERigMetadataType::Rotator:
+								{
+									FRotator StoredValue = FRotator::MakeFromEuler(FVector(StaticCastSharedPtr<const FSampleTrack<FVector3f>>(MetadataTrack)->GetValueAtTimeIndex(InPlaybackTimeIndex, OutSampleTrackIndex)));
+									const FRotator& CurrentValue = static_cast<const FRigRotatorMetadata*>(Metadata)->GetValue();
+									if(StoredValue.EqualsOrientation(CurrentValue, static_cast<float>(Tolerance)))
+									{
+										TrackNameBasedLogFunction(EMessageSeverity::Error, TrackName, TEXT("Metadata value doesn't match."));
+										bSuccess = false;
+									}
+									break;
+								}
+								case ERigMetadataType::RotatorArray:
+								{
+									const TArray<FVector3f>& StoredValue = StaticCastSharedPtr<const FSampleTrack<TArray<FVector3f>>>(MetadataTrack)->GetValueAtTimeIndex(InPlaybackTimeIndex, OutSampleTrackIndex);
+									const TArray<FRotator>& CurrentValue = static_cast<const FRigRotatorArrayMetadata*>(Metadata)->GetValue();
+									if(StoredValue.Num() != CurrentValue.Num())
+									{
+										TrackNameBasedLogFunction(EMessageSeverity::Error, TrackName, TEXT("Metadata array value num doesn't match."));
+										bSuccess = false;
+									}
+									else
+									{
+										for(int32 StoredValueIndex = 0; StoredValueIndex < StoredValue.Num(); StoredValueIndex++)
+										{
+											if(FRotator::MakeFromEuler(FVector(StoredValue[StoredValueIndex])).EqualsOrientation((CurrentValue[StoredValueIndex]), static_cast<float>(Tolerance)))
+											{
+												TrackNameBasedLogFunction(EMessageSeverity::Error, TrackName, TEXT("Metadata array value element doesn't match."));
+												bSuccess = false;
+											}
+										}
+									}
+									break;
+								}
+								case ERigMetadataType::Quat:
+								{
+									const FQuat4f& StoredValue = StaticCastSharedPtr<const FSampleTrack<FQuat4f>>(MetadataTrack)->GetValueAtTimeIndex(InPlaybackTimeIndex, OutSampleTrackIndex);
+									const FQuat4f& CurrentValue = FQuat4f(static_cast<const FRigQuatMetadata*>(Metadata)->GetValue()); 
+									if(StoredValue.Equals(CurrentValue, static_cast<float>(Tolerance)))
+									{
+										TrackNameBasedLogFunction(EMessageSeverity::Error, TrackName, TEXT("Metadata value doesn't match."));
+										bSuccess = false;
+									}
+									break;
+								}
+								case ERigMetadataType::QuatArray:
+								{
+									const TArray<FQuat4f>& StoredValue = StaticCastSharedPtr<const FSampleTrack<TArray<FQuat4f>>>(MetadataTrack)->GetValueAtTimeIndex(InPlaybackTimeIndex, OutSampleTrackIndex);
+									const TArray<FQuat>& CurrentValue = static_cast<const FRigQuatArrayMetadata*>(Metadata)->GetValue();
+									if(StoredValue.Num() != CurrentValue.Num())
+									{
+										TrackNameBasedLogFunction(EMessageSeverity::Error, TrackName, TEXT("Metadata array value num doesn't match."));
+										bSuccess = false;
+									}
+									else
+									{
+										for(int32 StoredValueIndex = 0; StoredValueIndex < StoredValue.Num(); StoredValueIndex++)
+										{
+											if(StoredValue[StoredValueIndex].Equals(FQuat4f(CurrentValue[StoredValueIndex]), static_cast<float>(Tolerance)))
+											{
+												TrackNameBasedLogFunction(EMessageSeverity::Error, TrackName, TEXT("Metadata array value element doesn't match."));
+												bSuccess = false;
+											}
+										}
+									}
+									break;
+								}
+								case ERigMetadataType::Transform:
+								{
+									const FTransform3f& StoredValue = StaticCastSharedPtr<const FSampleTrack<FTransform3f>>(MetadataTrack)->GetValueAtTimeIndex(InPlaybackTimeIndex, OutSampleTrackIndex);
+									const FTransform3f& CurrentValue = FTransform3f(static_cast<const FRigTransformMetadata*>(Metadata)->GetValue()); 
+									if(StoredValue.GetLocation().Equals(CurrentValue.GetLocation(), static_cast<float>(Tolerance)))
+									{
+										TrackNameBasedLogFunction(EMessageSeverity::Error, TrackName, TEXT("Metadata location value doesn't match."));
+										bSuccess = false;
+									}
+									if(StoredValue.Rotator().EqualsOrientation(CurrentValue.Rotator(), static_cast<float>(Tolerance)))
+									{
+										TrackNameBasedLogFunction(EMessageSeverity::Error, TrackName, TEXT("Metadata rotation value doesn't match."));
+										bSuccess = false;
+									}
+									if(StoredValue.GetScale3D().Equals(CurrentValue.GetScale3D(), static_cast<float>(Tolerance)))
+									{
+										TrackNameBasedLogFunction(EMessageSeverity::Error, TrackName, TEXT("Metadata scale value doesn't match."));
+										bSuccess = false;
+									}
+									break;
+								}
+								case ERigMetadataType::TransformArray:
+								{
+									const TArray<FTransform3f>& StoredValue = StaticCastSharedPtr<const FSampleTrack<TArray<FTransform3f>>>(MetadataTrack)->GetValueAtTimeIndex(InPlaybackTimeIndex, OutSampleTrackIndex);
+									const TArray<FTransform>& CurrentValue = static_cast<const FRigTransformArrayMetadata*>(Metadata)->GetValue();
+									if(StoredValue.Num() != CurrentValue.Num())
+									{
+										TrackNameBasedLogFunction(EMessageSeverity::Error, TrackName, TEXT("Metadata array value num doesn't match."));
+										bSuccess = false;
+									}
+									else
+									{
+										for(int32 StoredValueIndex = 0; StoredValueIndex < StoredValue.Num(); StoredValueIndex++)
+										{
+											const FTransform3f CurrentTransform(CurrentValue[StoredValueIndex]);
+											if(StoredValue[StoredValueIndex].GetLocation().Equals(CurrentTransform.GetLocation(), static_cast<float>(Tolerance)))
+											{
+												TrackNameBasedLogFunction(EMessageSeverity::Error, TrackName, TEXT("Metadata array locationvalue element doesn't match."));
+												bSuccess = false;
+											}
+											if(StoredValue[StoredValueIndex].Rotator().EqualsOrientation(CurrentTransform.Rotator(), static_cast<float>(Tolerance)))
+											{
+												TrackNameBasedLogFunction(EMessageSeverity::Error, TrackName, TEXT("Metadata array rotation value element doesn't match."));
+												bSuccess = false;
+											}
+											if(StoredValue[StoredValueIndex].GetScale3D().Equals(CurrentTransform.GetScale3D(), static_cast<float>(Tolerance)))
+											{
+												TrackNameBasedLogFunction(EMessageSeverity::Error, TrackName, TEXT("Metadata array scale value element doesn't match."));
+												bSuccess = false;
+											}
+										}
+									}
+									break;
+								}
+								case ERigMetadataType::LinearColor:
+								{
+									const FLinearColor& StoredValue = StaticCastSharedPtr<const FSampleTrack<FLinearColor>>(MetadataTrack)->GetValueAtTimeIndex(InPlaybackTimeIndex, OutSampleTrackIndex);
+									const FLinearColor& CurrentValue = static_cast<const FRigLinearColorMetadata*>(Metadata)->GetValue(); 
+									if(StoredValue.Equals(CurrentValue, static_cast<float>(Tolerance)))
+									{
+										TrackNameBasedLogFunction(EMessageSeverity::Error, TrackName, TEXT("Metadata value doesn't match."));
+										bSuccess = false;
+									}
+									break;
+								}
+								case ERigMetadataType::LinearColorArray:
+								{
+									const TArray<FLinearColor>& StoredValue = StaticCastSharedPtr<const FSampleTrack<TArray<FLinearColor>>>(MetadataTrack)->GetValueAtTimeIndex(InPlaybackTimeIndex, OutSampleTrackIndex);
+									const TArray<FLinearColor>& CurrentValue = static_cast<const FRigLinearColorArrayMetadata*>(Metadata)->GetValue(); 
+									if(StoredValue.Num() != CurrentValue.Num())
+									{
+										TrackNameBasedLogFunction(EMessageSeverity::Error, TrackName, TEXT("Metadata array value num doesn't match."));
+										bSuccess = false;
+									}
+									else
+									{
+										for(int32 StoredValueIndex = 0; StoredValueIndex < StoredValue.Num(); StoredValueIndex++)
+										{
+											if(StoredValue[StoredValueIndex].Equals(CurrentValue[StoredValueIndex], static_cast<float>(Tolerance)))
+											{
+												TrackNameBasedLogFunction(EMessageSeverity::Error, TrackName, TEXT("Metadata array value element doesn't match."));
+												bSuccess = false;
+											}
+										}
+									}
+									break;
+								}
+								case ERigMetadataType::RigElementKey:
+								{
+									const FRigElementKey& StoredValue = StaticCastSharedPtr<const FSampleTrack<FRigElementKey>>(MetadataTrack)->GetValueAtTimeIndex(InPlaybackTimeIndex, OutSampleTrackIndex);
+									const FRigElementKey& CurrentValue = static_cast<const FRigElementKeyMetadata*>(Metadata)->GetValue(); 
+									if(StoredValue != CurrentValue)
+									{
+										TrackNameBasedLogFunction(EMessageSeverity::Error, TrackName, TEXT("Metadata value doesn't match."));
+										bSuccess = false;
+									}
+									break;
+								}
+								case ERigMetadataType::RigElementKeyArray:
+								{
+									const TArray<FRigElementKey>& StoredValue = StaticCastSharedPtr<const FSampleTrack<TArray<FRigElementKey>>>(MetadataTrack)->GetValueAtTimeIndex(InPlaybackTimeIndex, OutSampleTrackIndex);
+									const TArray<FRigElementKey>& CurrentValue = static_cast<const FRigElementKeyArrayMetadata*>(Metadata)->GetValue(); 
+									if(StoredValue.Num() != CurrentValue.Num())
+									{
+										TrackNameBasedLogFunction(EMessageSeverity::Error, TrackName, TEXT("Metadata array value num doesn't match."));
+										bSuccess = false;
+									}
+									else
+									{
+										for(int32 StoredValueIndex = 0; StoredValueIndex < StoredValue.Num(); StoredValueIndex++)
+										{
+											if(StoredValue[StoredValueIndex] != CurrentValue[StoredValueIndex])
+											{
+												TrackNameBasedLogFunction(EMessageSeverity::Error, TrackName, TEXT("Metadata array value element doesn't match."));
+												bSuccess = false;
+											}
+										}
+									}
+									break;
+								}
+								case ERigMetadataType::Invalid:
+								{
+									TrackNameBasedLogFunction(EMessageSeverity::Error, TrackName, TEXT("Unsupported Metadata Type."));
+									return false;
+								}
+								
+							}
+						}
+						else
+						{
+							TrackNameBasedLogFunction(EMessageSeverity::Warning, TrackName, TEXT("Cannot find metadata on element."));
+							bSuccess = false;
+						}
+					}
+					else
+					{
+						TrackNameBasedLogFunction(EMessageSeverity::Warning, TrackName, TEXT("Track not found."));
+						bSuccess = false;
+					}
+				}
+			}
+		}
+	}
+
+	if(bValidateVariables)
+	{
+		TArray<FName> CurrentVariableNames;
+		for (TFieldIterator<FProperty> PropertyIt(InSubject->GetClass()); PropertyIt; ++PropertyIt)
+		{
+			const FProperty* Property = *PropertyIt;
+			if(Property->IsNative())
+			{
+				continue;
+			}
+
+			const UScriptStruct* ScriptStruct = nullptr;
+			if(const FStructProperty* StructProperty = CastField<FStructProperty>(Property))
+			{
+				ScriptStruct = StructProperty->Struct;
+			}
+			else if(const FArrayProperty* ArrayProperty = CastField<FArrayProperty>(Property))
+			{
+				if(const FStructProperty* InnerStructProperty = CastField<FStructProperty>(ArrayProperty->Inner))
+				{
+					ScriptStruct = InnerStructProperty->Struct;
+				}
+			}
+
+			const FName& TrackName = OutputTracks.GetTrackName(Property);
+			const TSharedPtr<const FSampleTrackBase> Track = OutputTracks.GetContainer()->FindTrack(TrackName);
+
+			if(Track)
+			{
+				const FSampleTrackBase::ETrackType TrackType = OutputTracks.GetTrackTypeFromProperty(Property);
+				if(Track->GetTrackType() == TrackType && Track->GetScriptStruct() == ScriptStruct)
+				{
+					TArray<uint8> ValueMemory;
+					ValueMemory.AddUninitialized(Property->GetSize());
+					Property->InitializeValue(ValueMemory.GetData());
+					Track->GetSampleForProperty(InPlaybackTimeIndex, OutSampleTrackIndex, Property, ValueMemory.GetData());
+
+					const uint8* CurrentMemory = Property->ContainerPtrToValuePtr<uint8>(InSubject);
+					if(!Property->Identical(ValueMemory.GetData(), CurrentMemory, PPF_None))
+					{
+						if(InLogFunction)
+						{
+							InLogFunction(EMessageSeverity::Error, FString::Printf(TEXT("Frame [%04d]: Variable '%s' value doesn't match replay."), InPlaybackTimeIndex, *Property->GetName()));
+						}
+						bSuccess = false;
+					}
+					Property->DestroyValue(ValueMemory.GetData());
+				}
+				else
+				{
+					if(InPlaybackTimeIndex == 0 && InLogFunction)
+					{
+						InLogFunction(EMessageSeverity::Error, FString::Printf(TEXT("Frame [%04d]: Track for Variable '%s' has incorrect type."), InPlaybackTimeIndex, *Property->GetName()));
+					}
+					bSuccess = false;
+				}
+			}
+			else
+			{
+				if(InPlaybackTimeIndex == 0 && InLogFunction)
+				{
+					InLogFunction(EMessageSeverity::Error, FString::Printf(TEXT("Frame [%04d]: Track for Variable '%s' is missing in replay."), InPlaybackTimeIndex, *Property->GetName()));
+				}
+				bSuccess = false;
+			}
+			
+			CurrentVariableNames.Add(Property->GetFName());
+		}
+
+		if(InPlaybackTimeIndex == 0)
+		{
+			if(const TSharedPtr<const FSampleTrack<TArray<FName>>> VariableNamesTrack = OutputTracks.GetContainer()->FindTrack<TArray<FName>>(FControlRigReplayTracks::VariableNamesName))
+			{
+				const TArray<FName>& StoredVariableNames = VariableNamesTrack->GetValueAtTimeIndex(InPlaybackTimeIndex, OutSampleTrackIndex);
+				for(const FName& StoredVariableName : StoredVariableNames)
+				{
+					if(!CurrentVariableNames.Contains(StoredVariableName))
+					{
+						if(InLogFunction)
+						{
+							InLogFunction(EMessageSeverity::Error, FString::Printf(TEXT("Frame [%04d]: Variable '%s' for stored track is missing."), InPlaybackTimeIndex, *StoredVariableName.ToString()));
+						}
+						bSuccess = false;
+					}
+				}
+			}
+		}
+	}
+	
+	return bSuccess;
 }
 
 void UControlRigReplay::ClearDelegates(UControlRig* InControlRig)
diff --git a/Engine/Plugins/Animation/ControlRig/Source/ControlRig/Private/ControlRigTestData.cpp b/Engine/Plugins/Animation/ControlRig/Source/ControlRig/Private/ControlRigTestData.cpp
--- a/Engine/Plugins/Animation/ControlRig/Source/ControlRig/Private/ControlRigTestData.cpp
+++ b/Engine/Plugins/Animation/ControlRig/Source/ControlRig/Private/ControlRigTestData.cpp
@@ -241,6 +241,11 @@ void UControlRigTestData::Serialize(FArchive& Ar)
 			}
 		}
 	}
+
+	if (EventQueue.IsEmpty())
+	{
+		EventQueue = {FRigUnit_BeginExecution::EventName};
+	}
 }
 
 FVector2D UControlRigTestData::GetTimeRange() const
diff --git a/Engine/Plugins/Animation/ControlRig/Source/ControlRig/Private/Tracks/SampleTrackBase.cpp b/Engine/Plugins/Animation/ControlRig/Source/ControlRig/Private/Tracks/SampleTrackBase.cpp
--- a/Engine/Plugins/Animation/ControlRig/Source/ControlRig/Private/Tracks/SampleTrackBase.cpp
+++ b/Engine/Plugins/Animation/ControlRig/Source/ControlRig/Private/Tracks/SampleTrackBase.cpp
@@ -133,6 +133,8 @@ int32 FSampleTrackBase::NumSamples() const
 
 bool FSampleTrackBase::GetSampleIndexForTimeIndex(int32 InTimeIndex, FSampleTrackIndex& InOutIndex) const
 {
+	InTimeIndex = FMath::Clamp(InTimeIndex, 0, NumTimes() - 1);
+	
 	int32& SampleIndex = InOutIndex.GetSample(TrackIndex);
 	
 	switch(Mode)
@@ -497,7 +499,7 @@ const TArray<int32>& FSampleTrackBase::GetChildTracks() const
 
 int32 FSampleTrackBase::RemoveObsoleteTimes(int32 InNumSamplesToRemove, const TArray<bool>& InRemoveIndexMap, const TArray<int32>& InNewIndexMap)
 {
-	if(!IsValid() || IsReferenced())
+	if(!IsValid() || IsReferenced() || GetMode() == EMode_Singleton)
 	{
 		return 0;
 	}
diff --git a/Engine/Plugins/Animation/ControlRig/Source/ControlRig/Private/Tracks/SampleTrackContainer.cpp b/Engine/Plugins/Animation/ControlRig/Source/ControlRig/Private/Tracks/SampleTrackContainer.cpp
--- a/Engine/Plugins/Animation/ControlRig/Source/ControlRig/Private/Tracks/SampleTrackContainer.cpp
+++ b/Engine/Plugins/Animation/ControlRig/Source/ControlRig/Private/Tracks/SampleTrackContainer.cpp
@@ -100,7 +100,7 @@ FVector2f FSampleTrackHost::GetTimeRange() const
 	return GetContainer()->GetTimeRange();
 }
 
-void FSampleTrackHost::GetTimeIndex(float InAbsoluteTime, FSampleTrackIndex& InOutTrackIndex) const
+int32 FSampleTrackHost::GetTimeIndex(float InAbsoluteTime, FSampleTrackIndex& InOutTrackIndex) const
 {
 	return GetContainer()->GetTimeIndex(InAbsoluteTime, InOutTrackIndex);
 }
@@ -120,6 +120,11 @@ float FSampleTrackHost::GetAbsoluteTime(int32 InTimeIndex) const
 	return GetContainer()->GetAbsoluteTime(InTimeIndex);
 }
 
+int32 FSampleTrackHost::GetTimeIndex(float InAbsoluteTime) const
+{
+	return GetContainer()->GetTimeIndex(InAbsoluteTime);
+}
+
 float FSampleTrackHost::GetDeltaTime(int32 InTimeIndex) const
 {
 	return GetContainer()->GetDeltaTime(InTimeIndex);
@@ -643,7 +648,7 @@ float FSampleTrackContainer::GetLastDeltaTime() const
 	return DeltaTimeTrack->GetValueAtSampleIndex(DeltaTimeTrack->NumSamples()-1);
 }
 
-void FSampleTrackContainer::GetTimeIndex(float InAbsoluteTime, FSampleTrackIndex& InOutTrackIndex) const
+int32 FSampleTrackContainer::GetTimeIndex(float InAbsoluteTime, FSampleTrackIndex& InOutTrackIndex) const
 {
 	const TSharedPtr<const FSampleTrack<float>> AbsoluteTimeTrack = FindTrack<float>(AbsoluteTimeName);
 	check(AbsoluteTimeTrack.IsValid() && AbsoluteTimeTrack->NumTimes() > 0);
@@ -653,7 +658,7 @@ void FSampleTrackContainer::GetTimeIndex(float InAbsoluteTime, FSampleTrackIndex
 	if(AbsoluteTimeTrack->GetMode() == FSampleTrackBase::EMode_Singleton)
 	{
 		InOutTimeIndex = 0;
-		return;
+		return InOutTimeIndex;
 	}
 
 	const int32 NumSamples = AbsoluteTimeTrack->NumSamples();
@@ -672,6 +677,12 @@ void FSampleTrackContainer::GetTimeIndex(float InAbsoluteTime, FSampleTrackIndex
 	{
 		InOutTimeIndex = AbsoluteTimeTrack->TimeIndices[InOutTimeIndex];
 	}
+	return InOutTimeIndex;
+}
+
+int32 FSampleTrackContainer::GetTimeIndex(float InAbsoluteTime) const
+{
+	return GetTimeIndex(InAbsoluteTime, TimeSampleTrackIndex);
 }
 
 int32 FSampleTrackContainer::GetTrackIndex(const FName& InName) const
diff --git a/Engine/Plugins/Animation/ControlRig/Source/ControlRig/Public/ControlRig.h b/Engine/Plugins/Animation/ControlRig/Source/ControlRig/Public/ControlRig.h
--- a/Engine/Plugins/Animation/ControlRig/Source/ControlRig/Public/ControlRig.h
+++ b/Engine/Plugins/Animation/ControlRig/Source/ControlRig/Public/ControlRig.h
@@ -41,6 +41,7 @@ class USkeletalMesh;
 class USkeletalMeshComponent;
 class AActor;
 class UTransformableControlHandle;
+class UControlRigReplay;
 
 struct FReferenceSkeleton;
 struct FRigUnit;
@@ -908,6 +909,17 @@ private:
 	TSubclassOf<UControlRig> InteractionRigClass_DEPRECATED;
 #endif
 
+public:
+	int32 GetReplayTimeIndex() const;
+	void SetReplayTimeIndex(int32 InReplayTimeIndex);
+	void DisableReplay();
+	bool IsReplayEnabled() const;
+	void SetReplay(UControlRigReplay* InReplay);
+
+private:
+	int32 ReplayTimeIndex;
+	TObjectPtr<UControlRigReplay> Replay;
+
 	friend class FControlRigBlueprintCompilerContext;
 	friend struct FRigHierarchyRef;
 	friend class FControlRigEditor;
@@ -933,6 +945,7 @@ private:
 	friend class UModularRig;
 	friend class UModularRigController;
 	friend struct FControlRigHierarchyMappings;
+	friend class UControlRigReplay;
 };
 
 class CONTROLRIG_API FControlRigBracketScope
diff --git a/Engine/Plugins/Animation/ControlRig/Source/ControlRig/Public/ControlRigReplay.h b/Engine/Plugins/Animation/ControlRig/Source/ControlRig/Public/ControlRigReplay.h
--- a/Engine/Plugins/Animation/ControlRig/Source/ControlRig/Public/ControlRigReplay.h
+++ b/Engine/Plugins/Animation/ControlRig/Source/ControlRig/Public/ControlRigReplay.h
@@ -57,6 +57,12 @@ public:
 	virtual bool Serialize(FArchive& InArchive) override;
 
 	virtual void Reset() override;
+	bool IsEmpty() const;
+
+	void StoreRigVMEvent(const FName& InEvent);
+	FName GetRigVMEvent(int32 InTimeIndex) const;
+	void StoreInteraction(uint8 InInteractionMode, const TArray<FRigElementKey>& InElementsBeingInteracted);
+	TTuple<uint8, TArray<FRigElementKey>> GetInteraction(int32 InTimeIndex) const;
 
 	void StoreHierarchy(
 		URigHierarchy* InHierarchy,
@@ -108,9 +114,12 @@ private:
 	const FName& GetTrackName(const FRigComponentKey& InComponentKey) const;
 	const FName& GetTrackName(const FRigElementKey& InElementKey, const FName& InMetadataName) const;
 	const FName& GetTrackName(const FProperty* InProperty) const;
-	static FSampleTrackBase::ETrackType GetTrackTypeFromMetadata(ERigMetadataType InMetadataType);
+	static FSampleTrackBase::ETrackType GetTrackTypeFromMetadataType(ERigMetadataType InMetadataType);
+	static ERigMetadataType GetMetadataTypeFromTrackType(FSampleTrackBase::ETrackType InTrackType);
 	static FSampleTrackBase::ETrackType GetTrackTypeFromProperty(const FProperty* InProperty) ;
 
+	const TArray<FRigElementKey> GetElementKeys() const;
+
 	void StorePose(URigHierarchy* InHierarchy, FRigBaseElement* InElement);
 	bool RestorePose(int32 InTimeIndex, URigHierarchy* InHierarchy, FRigBaseElement* InElement, TReportFunction InReportFunction) const;
 
@@ -122,9 +131,10 @@ private:
 
 	void StoreMetaData(URigHierarchy* InHierarchy, FRigBaseElement* InElement);
 	bool RestoreMetaData(int32 InTimeIndex, URigHierarchy* InHierarchy, FRigBaseElement* InElement, TReportFunction InReportFunction) const;
+	const TArray<FName> GetMetadataNames(int32 InTimeIndex, FSampleTrackIndex& OutSampleTrackIndex, const FRigElementKey& InElementKey, TReportFunction InReportFunction = nullptr) const;
 
 	void StoreMetaData(URigHierarchy* InHierarchy, FRigBaseElement* InElement, FRigBaseMetadata* InMetadata);
-	bool RestoreMetaData(int32 InTimeIndex, URigHierarchy* InHierarchy, FRigBaseElement* InElement, FRigBaseMetadata* InMetadata, TReportFunction InReportFunction) const;
+	bool RestoreMetaData(int32 InTimeIndex, URigHierarchy* InHierarchy, FRigBaseElement* InElement, const FName& InMetadataName, TReportFunction InReportFunction) const;
 
 	static bool ForEachElement(URigHierarchy* InHierarchy, const TArrayView<FRigElementKey>& InKeys, TFunction<void(FRigBaseElement*, bool&)> InFunction);
 	static bool ForEachComponent(URigHierarchy* InHierarchy, const TArrayView<FRigComponentKey>& InKeys, TFunction<void(FRigBaseComponent*, bool&)> InFunction);
@@ -136,7 +146,18 @@ private:
 	mutable TMap<TTuple<FRigElementKey, FName>, FName> MetadataToTrackName;
 	mutable TMap<FName, FName> PropertyNameToTrackName;
 	mutable FSampleTrackIndex SampleTrackIndex;
-	TArray<int32> ParentIndices;
+	bool bIsInput = true;
+
+	static inline const FLazyName RigVMEventName = FLazyName(TEXT("RigVMEvent"));
+	static inline const FLazyName InteractionTypeName = FLazyName(TEXT("InteractionType"));
+	static inline const FLazyName ElementsBeingInteractedName = FLazyName(TEXT("ElementsBeingInteracted"));
+	static inline const FLazyName TopologyHashName = FLazyName(TEXT("TopologyHash"));
+	static inline const FLazyName MetadataVersionName = FLazyName(TEXT("MetadataVersion"));
+	static inline const FLazyName ElementKeysName = FLazyName(TEXT("ElementKeys"));
+	static inline const FLazyName ParentIndicesName = FLazyName(TEXT("ParentIndices"));
+	static inline const FLazyName VariableNamesName = FLazyName(TEXT("VariableNames"));
+
+	friend class UControlRigReplay;
 };
 
 template<>
@@ -159,11 +180,18 @@ public:
 
 	UControlRigReplay()
 		: Tolerance(0.001)
+		, bValidateHierarchyTopology(true)
+		, bValidatePose(true)
+		, bValidateMetadata(true)
+		, bValidateVariables(true)
 		, EnableTest(true)
 		, TimeAtStartOfRecording(0.0)
 		, TimeOfLastFrame(0.0)
-		, DesiredRecordingDuration(0.0)
-	{}
+		, DesiredRecordingDuration(-1.0)
+	{
+		InputTracks.bIsInput = true;
+		OutputTracks.bIsInput = false;
+	}
 
 	virtual void BeginDestroy() override;
 
@@ -179,16 +207,28 @@ public:
 	FSoftObjectPath ControlRigObjectPath;
 
 	UPROPERTY(EditAnywhere, Category = "ControlRigReplay")
-	FControlRigReplayTracks Tracks;
-	
+	FControlRigReplayTracks InputTracks;
+
+	UPROPERTY(EditAnywhere, Category = "ControlRigReplay")
+	FControlRigReplayTracks OutputTracks;
+
 	UPROPERTY(EditAnywhere, BlueprintReadOnly, Category = "ControlRigReplay")
 	double Tolerance;
 
 	UPROPERTY(EditAnywhere, BlueprintReadOnly, Category = "ControlRigReplay")
-	TArray<int32> FramesToSkip;
+	bool bValidateHierarchyTopology;
+
+	UPROPERTY(EditAnywhere, BlueprintReadOnly, Category = "ControlRigReplay")
+	bool bValidatePose;
+
+	UPROPERTY(EditAnywhere, BlueprintReadOnly, Category = "ControlRigReplay")
+	bool bValidateMetadata;
+
+	UPROPERTY(EditAnywhere, BlueprintReadOnly, Category = "ControlRigReplay")
+	bool bValidateVariables;
 
 	UPROPERTY(EditAnywhere, BlueprintReadOnly, Category = "ControlRigReplay")
-	TArray<FName> EventQueue;
+	TArray<int32> FramesToSkip;
 
 	UPROPERTY(EditAnywhere, BlueprintReadOnly, Category = "ControlRigReplay", AssetRegistrySearchable)
 	bool EnableTest;
@@ -222,6 +262,11 @@ public:
 	virtual bool IsValidForTesting() const;
 
 	virtual bool PerformTest(UControlRig* InSubject, TFunction<void(EMessageSeverity::Type, const FString&)> InLogFunction) const;
+	virtual bool ValidateExpectedResults(
+		int32 InPlaybackTimeIndex,
+		FSampleTrackIndex& OutSampleTrackIndex,
+		UControlRig* InSubject,
+		TFunction<void(EMessageSeverity::Type, const FString&)> InLogFunction) const;
 
 protected:
 
@@ -235,6 +280,8 @@ protected:
 	double TimeAtStartOfRecording;
 	double TimeOfLastFrame;
 	double DesiredRecordingDuration;
+	bool bFirstPlayThrough = false;
+	bool bStoreVariablesDuringPreEvent = false;
 
 	friend class FControlRigEditor;
 };
\ No newline at end of file
diff --git a/Engine/Plugins/Animation/ControlRig/Source/ControlRig/Public/ControlRigTestData.h b/Engine/Plugins/Animation/ControlRig/Source/ControlRig/Public/ControlRigTestData.h
--- a/Engine/Plugins/Animation/ControlRig/Source/ControlRig/Public/ControlRigTestData.h
+++ b/Engine/Plugins/Animation/ControlRig/Source/ControlRig/Public/ControlRigTestData.h
@@ -83,6 +83,9 @@ public:
 	UPROPERTY(VisibleAnywhere, BlueprintReadOnly, Category = "ControlRigTestData")
 	TArray<FControlRigTestDataFrame> OutputFrames;
 
+	UPROPERTY(EditAnywhere, BlueprintReadOnly, Category = "ControlRigTestData")
+	TArray<FName> EventQueue;
+
 	virtual bool IsValidForTesting() const override;
 	
 	virtual bool PerformTest(UControlRig* InSubject, TFunction<void(EMessageSeverity::Type, const FString&)> InLogFunction) const override;
diff --git a/Engine/Plugins/Animation/ControlRig/Source/ControlRig/Public/Rigs/RigHierarchy.h b/Engine/Plugins/Animation/ControlRig/Source/ControlRig/Public/Rigs/RigHierarchy.h
--- a/Engine/Plugins/Animation/ControlRig/Source/ControlRig/Public/Rigs/RigHierarchy.h
+++ b/Engine/Plugins/Animation/ControlRig/Source/ControlRig/Public/Rigs/RigHierarchy.h
@@ -5006,6 +5006,7 @@ private:
 	void OnMetadataChanged(const FRigElementKey& InKey, const FName& InName);
 	void OnMetadataTagChanged(const FRigElementKey& InKey, const FName& InTag, bool bAdded);
 
+public:
 	/** Returns a metadata ptr to the given element's metadata. If the meta data, with the same name, doesn't exist already a new entry
 	    is created for that element. If the name matches but the type differs, the existing metadata is destroyed and a new one with the
 	    matching type is created instead.
diff --git a/Engine/Plugins/Animation/ControlRig/Source/ControlRig/Public/Tracks/SampleTrackContainer.h b/Engine/Plugins/Animation/ControlRig/Source/ControlRig/Public/Tracks/SampleTrackContainer.h
--- a/Engine/Plugins/Animation/ControlRig/Source/ControlRig/Public/Tracks/SampleTrackContainer.h
+++ b/Engine/Plugins/Animation/ControlRig/Source/ControlRig/Public/Tracks/SampleTrackContainer.h
@@ -35,10 +35,11 @@ public:
 	int32 AddTimeSampleFromDeltaTime(float InDeltaTime);
 	int32 GetNumTimes() const;
 	FVector2f GetTimeRange() const;
-	void GetTimeIndex(float InAbsoluteTime, FSampleTrackIndex& InOutTrackIndex) const;
+	int32 GetTimeIndex(float InAbsoluteTime, FSampleTrackIndex& InOutTrackIndex) const;
 	float GetAbsoluteTime(int32 InTimeIndex, FSampleTrackIndex& InOutTrackIndex) const;
 	float GetDeltaTime(int32 InTimeIndex, FSampleTrackIndex& InOutTrackIndex) const;
 	float GetAbsoluteTime(int32 InTimeIndex) const;
+	int32 GetTimeIndex(float InAbsoluteTime) const;
 	float GetDeltaTime(int32 InTimeIndex) const;
 	float GetLastAbsoluteTime() const;
 	float GetLastDeltaTime() const;
@@ -123,7 +124,8 @@ public:
 
 	int32 GetNumTimes() const;
 	FVector2f GetTimeRange() const;
-	void GetTimeIndex(float InAbsoluteTime, FSampleTrackIndex& InOutTrackIndex) const;
+	int32 GetTimeIndex(float InAbsoluteTime, FSampleTrackIndex& InOutTrackIndex) const;
+	int32 GetTimeIndex(float InAbsoluteTime) const;
 	float GetAbsoluteTime(int32 InTimeIndex, FSampleTrackIndex& InOutTrackIndex) const;
 	float GetDeltaTime(int32 InTimeIndex, FSampleTrackIndex& InOutTrackIndex) const;
 	float GetAbsoluteTime(int32 InTimeIndex) const;
diff --git a/Engine/Plugins/Animation/ControlRig/Source/ControlRigEditor/Private/Editor/ControlRigEditor.cpp b/Engine/Plugins/Animation/ControlRig/Source/ControlRigEditor/Private/Editor/ControlRigEditor.cpp
--- a/Engine/Plugins/Animation/ControlRig/Source/ControlRigEditor/Private/Editor/ControlRigEditor.cpp
+++ b/Engine/Plugins/Animation/ControlRig/Source/ControlRigEditor/Private/Editor/ControlRigEditor.cpp
@@ -101,6 +101,7 @@
 #include "Types/ISlateMetaData.h"
 #include "Kismet2/KismetDebugUtilities.h"
 #include "Kismet2/WatchedPin.h"
+#include "Kismet/KismetSystemLibrary.h"
 #include "ToolMenus.h"
 #include "Styling/AppStyle.h"
 #include "AssetRegistry/AssetRegistryModule.h"
@@ -762,15 +763,34 @@ bool FControlRigEditor::RecordReplay(double InRecordingDuration)
 	static const FString Folder = TEXT("/Game/Animation/ControlRig/NoCook/");
 	const FString DesiredPackagePath =  FString::Printf(TEXT("%s/%s_Replay"), *Folder, *GetControlRigBlueprint()->GetName());
 
-	if(const UControlRigReplay* Replay = UControlRigReplay::CreateNewAsset(DesiredPackagePath, GetControlRigBlueprint()->GetPathName(), UControlRigTestData::StaticClass()))
+	if(const UControlRigReplay* Replay = UControlRigReplay::CreateNewAsset(DesiredPackagePath, GetControlRigBlueprint()->GetPathName(), UControlRigReplay::StaticClass()))
 	{
 		SetReplayAssetPath(Replay->GetPathName());
 	}
-	
+
+
 	if(UControlRigReplay* Replay = ReplayStrongPtr.Get())
 	{
-		Replay->StartRecording(GetControlRig());
 		Replay->DesiredRecordingDuration = InRecordingDuration;
+
+		// 3 second preroll
+		TSharedPtr<int32> TimeLeft = MakeShared<int32>(4);
+		GEditor->GetTimerManager()->SetTimer(RecordReplayTimerHandle, FTimerDelegate::CreateLambda([this, Replay, TimeLeft]()
+		{
+			int32& SecondsLeft = *TimeLeft.Get();
+			SecondsLeft--;
+
+			if(SecondsLeft == 0)
+			{
+				Replay->StartRecording(GetControlRig());
+				GEditor->GetTimerManager()->ClearTimer(RecordReplayTimerHandle);
+			}
+			else
+			{
+				static constexpr TCHAR PrintStringFormat[] = TEXT("Recording starts in... %d");
+				UKismetSystemLibrary::PrintString(GetPreviewScene()->GetWorld(), FString::Printf(PrintStringFormat, SecondsLeft), true, false, FLinearColor::Red, 1.f);
+			}
+		}), 1.f, true, 0.01f);
 	}
 	return true;
 }
@@ -912,7 +932,7 @@ void FControlRigEditor::FillToolbar(FToolBarBuilder& ToolbarBuilder, bool bEndSe
 				{
 					if(ReplayStrongPtr.IsValid())
 					{
-						return !ReplayStrongPtr->IsReplaying() && !ReplayStrongPtr->IsRecording();
+						return !ReplayStrongPtr->IsReplaying() && !ReplayStrongPtr->IsRecording() && !RecordReplayTimerHandle.IsValid();
 					}
 					return true;
 				})),
@@ -931,7 +951,7 @@ void FControlRigEditor::FillToolbar(FToolBarBuilder& ToolbarBuilder, bool bEndSe
 					{
 						if(ReplayStrongPtr.IsValid())
 						{
-							return !ReplayStrongPtr->IsReplaying() && !ReplayStrongPtr->IsRecording();
+							return !ReplayStrongPtr->IsReplaying() && !ReplayStrongPtr->IsRecording() && !RecordReplayTimerHandle.IsValid();
 						}
 						return true;
 					})
diff --git a/Engine/Plugins/Animation/ControlRig/Source/ControlRigEditor/Private/Editor/ControlRigEditor.h b/Engine/Plugins/Animation/ControlRig/Source/ControlRigEditor/Private/Editor/ControlRigEditor.h
--- a/Engine/Plugins/Animation/ControlRig/Source/ControlRigEditor/Private/Editor/ControlRigEditor.h
+++ b/Engine/Plugins/Animation/ControlRig/Source/ControlRigEditor/Private/Editor/ControlRigEditor.h
@@ -371,6 +371,7 @@ protected:
 	FSimpleMulticastDelegate RequestNavigateToConnectorWarningDelegate;
 	TSharedPtr<SSchematicGraphPanel> SchematicViewport;
 	bool bSchematicViewPortIsHidden;
+	FTimerHandle RecordReplayTimerHandle;
 
 	static const TArray<FName> ForwardsSolveEventQueue;
 	static const TArray<FName> BackwardsSolveEventQueue;
