Control Rig Edit Mode: general level editor's animation mode improvements

- URigHierarchy: transient controls do not exist in the level editor's animation mode, so we can use URigHierarchy::OrderedSelection directly, which is the only thing we're interested in
- URigHierarchy: added HasAnythingSelectedByPredicate and GetSelectedKeysByPredicate to filter selection in one simgle pass
- ControlProxy: no need to update animation details during playback
- AControlRigShapeActor: no need to update shape color when it hasn't change

#rnx

diff --git a/Engine/Plugins/Animation/ControlRig/Source/ControlRig/Private/ControlRigGizmoActor.cpp b/Engine/Plugins/Animation/ControlRig/Source/ControlRig/Private/ControlRigGizmoActor.cpp
--- a/Engine/Plugins/Animation/ControlRig/Source/ControlRig/Private/ControlRigGizmoActor.cpp
+++ b/Engine/Plugins/Animation/ControlRig/Source/ControlRig/Private/ControlRigGizmoActor.cpp
@@ -21,8 +21,8 @@ AControlRigShapeActor::AControlRigShapeActor(const FObjectInitializer& ObjectIni
 	, bSelected(false)
 	, bHovered(false)
 	, bSelectable(true)
+	, CurrentColor(FLinearColor::Red)
 {
-
 	ActorRootComponent = CreateDefaultSubobject<USceneComponent>(TEXT("SceneComponent0"));
 	StaticMeshComponent = CreateDefaultSubobject<UStaticMeshComponent>(TEXT("StaticMeshComponent0"));
 	StaticMeshComponent->SetCollisionProfileName(UCollisionProfile::NoCollision_ProfileName);
@@ -100,13 +100,17 @@ bool AControlRigShapeActor::IsHovered() const
 }
 
 
-void AControlRigShapeActor::SetShapeColor(const FLinearColor& InColor)
+void AControlRigShapeActor::SetShapeColor(const FLinearColor& InColor, const bool bForce)
 {
 	if (StaticMeshComponent && !ColorParameterName.IsNone())
 	{
-		if (UMaterialInstanceDynamic* MaterialInstance = Cast<UMaterialInstanceDynamic>(StaticMeshComponent->GetMaterial(0)))
+		if (bForce || InColor != CurrentColor)
 		{
-			MaterialInstance->SetVectorParameterValue(ColorParameterName, FVector(InColor));
+			if (UMaterialInstanceDynamic* MaterialInstance = Cast<UMaterialInstanceDynamic>(StaticMeshComponent->GetMaterial(0)))
+			{
+				MaterialInstance->SetVectorParameterValue(ColorParameterName, FVector(InColor));
+			}
+			CurrentColor = InColor;
 		}
 	}
 }
@@ -257,9 +261,11 @@ namespace FControlRigShapeHelper
 			{
 				CreationParam.StaticMesh.LoadSynchronous();
 			}
+
 			if (CreationParam.StaticMesh.IsValid())
 			{
 				MeshComponent->SetStaticMesh(CreationParam.StaticMesh.Get());
+				MeshComponent->DestroyPhysicsState();
 				MeshComponent->SetRelativeTransform(CreationParam.MeshTransform * CreationParam.ShapeTransform);
 			}
 
diff --git a/Engine/Plugins/Animation/ControlRig/Source/ControlRig/Private/ModularRig.cpp b/Engine/Plugins/Animation/ControlRig/Source/ControlRig/Private/ModularRig.cpp
--- a/Engine/Plugins/Animation/ControlRig/Source/ControlRig/Private/ModularRig.cpp
+++ b/Engine/Plugins/Animation/ControlRig/Source/ControlRig/Private/ModularRig.cpp
@@ -1089,13 +1089,15 @@ FString UModularRig::GetParentPath(const FString& InModulePath) const
 	return FString();
 }
 
-FString UModularRig::GetShortestDisplayPathForElement(const FRigElementKey& InElementKey, bool bAlwaysShowNameSpace) const
+const FString& UModularRig::GetShortestDisplayPathForElement(const FRigElementKey& InElementKey, bool bAlwaysShowNameSpace) const
 {
 	if(const TTuple<FString, FString>* ShortestDisplayPaths = ElementKeyToShortestDisplayPath.Find(InElementKey))
 	{
 		return bAlwaysShowNameSpace ? ShortestDisplayPaths->Get<0>() : ShortestDisplayPaths->Get<1>();
 	}
-	return FString();
+
+	static const FString EmptyPath;
+	return EmptyPath;
 }
 
 void UModularRig::ForEachModule(TFunctionRef<bool(FRigModuleInstance*)> PerModuleFunction)
diff --git a/Engine/Plugins/Animation/ControlRig/Source/ControlRig/Private/Rigs/RigHierarchy.cpp b/Engine/Plugins/Animation/ControlRig/Source/ControlRig/Private/Rigs/RigHierarchy.cpp
--- a/Engine/Plugins/Animation/ControlRig/Source/ControlRig/Private/Rigs/RigHierarchy.cpp
+++ b/Engine/Plugins/Animation/ControlRig/Source/ControlRig/Private/Rigs/RigHierarchy.cpp
@@ -1854,6 +1854,32 @@ TArray<FRigElementKey> URigHierarchy::GetSelectedKeys(ERigElementType InTypeFilt
 	return Selection;
 }
 
+bool URigHierarchy::HasAnythingSelectedByPredicate(TFunctionRef<bool(const FRigElementKey&)> InPredicate) const
+{
+	if (URigHierarchy* HierarchyForSelection = HierarchyForSelectionPtr.Get())
+	{
+		return HierarchyForSelection->HasAnythingSelectedByPredicate(InPredicate);
+	}
+
+	return OrderedSelection.ContainsByPredicate([InPredicate](const FRigElementKey& SelectedKey)
+	{
+		return InPredicate(SelectedKey);
+	});
+}
+
+TArray<FRigElementKey> URigHierarchy::GetSelectedKeysByPredicate(TFunctionRef<bool(const FRigElementKey&)> InPredicate) const
+{
+	if(URigHierarchy* HierarchyForSelection = HierarchyForSelectionPtr.Get())
+	{
+		return HierarchyForSelection->GetSelectedKeysByPredicate(InPredicate);
+	}
+
+	return OrderedSelection.FilterByPredicate([InPredicate](const FRigElementKey& SelectedKey)
+	{
+		return InPredicate(SelectedKey);
+	});
+}
+
 FRigName URigHierarchy::JoinNameSpace(const FRigName& InLeft, const FRigName& InRight)
 {
 	return FRigName(JoinNameSpace(InLeft.ToString(), InRight.ToString()));
@@ -2222,10 +2248,10 @@ FText URigHierarchy::GetDisplayNameForUI(const FRigBaseElement* InElement, bool
 {
 	check(InElement);
 
-	if(const UModularRig* ModularRig = GetTypedOuter<UModularRig>())
+	if (const UModularRig* ModularRig = GetTypedOuter<UModularRig>())
 	{
-		const FString ShortestPath = ModularRig->GetShortestDisplayPathForElement(InElement->GetKey(), false);
-		if(!ShortestPath.IsEmpty())
+		const FString& ShortestPath = ModularRig->GetShortestDisplayPathForElement(InElement->GetKey(), false);
+		if (!ShortestPath.IsEmpty())
 		{
 			return FText::FromString(ShortestPath);
 		}
diff --git a/Engine/Plugins/Animation/ControlRig/Source/ControlRig/Public/ControlRigGizmoActor.h b/Engine/Plugins/Animation/ControlRig/Source/ControlRig/Public/ControlRigGizmoActor.h
--- a/Engine/Plugins/Animation/ControlRig/Source/ControlRig/Public/ControlRigGizmoActor.h
+++ b/Engine/Plugins/Animation/ControlRig/Source/ControlRig/Public/ControlRigGizmoActor.h
@@ -129,7 +129,7 @@ public:
 	virtual void TickControl() {};
 
 	/** changes the shape color */
-	virtual void SetShapeColor(const FLinearColor& InColor);
+	virtual void SetShapeColor(const FLinearColor& InColor, const bool bForce = false);
 
 	/** Event called when the transform of this control has changed */
 	UFUNCTION(BlueprintImplementableEvent)
@@ -166,6 +166,9 @@ public:
 		const FRigControlElement* InControlElement,
 		bool bHideManipulators,
 		bool bIsInLevelEditor);
+
+	UPROPERTY()
+	FCachedRigElement CachedIndex;
 	
 private:
 
@@ -180,6 +183,8 @@ private:
 	/** Whether or not this control is selectable*/
 	uint8 bSelectable : 1;
 
+	/** Current color used to avoid color updates when unnecessary */
+	FLinearColor CurrentColor;
 };
 
 /**
diff --git a/Engine/Plugins/Animation/ControlRig/Source/ControlRig/Public/ModularRig.h b/Engine/Plugins/Animation/ControlRig/Source/ControlRig/Public/ModularRig.h
--- a/Engine/Plugins/Animation/ControlRig/Source/ControlRig/Public/ModularRig.h
+++ b/Engine/Plugins/Animation/ControlRig/Source/ControlRig/Public/ModularRig.h
@@ -212,7 +212,7 @@ public:
 	UFUNCTION(BlueprintPure, Category=ModularRig, meta=(ScriptName="GetParentPath", DisplayName="Get Parent Path"))
 	FString GetParentPathForBP(FString InModulePath) const;
 	FString GetParentPath(const FString& InModulePath) const;
-	FString GetShortestDisplayPathForElement(const FRigElementKey& InElementKey, bool bAlwaysShowNameSpace) const;
+	const FString& GetShortestDisplayPathForElement(const FRigElementKey& InElementKey, bool bAlwaysShowNameSpace) const;
 
 	void ForEachModule(TFunctionRef<bool(FRigModuleInstance*)> PerModuleFunction);
 	void ForEachModule(TFunctionRef<bool(const FRigModuleInstance*)> PerModuleFunction) const;
diff --git a/Engine/Plugins/Animation/ControlRig/Source/ControlRig/Public/Rigs/RigHierarchy.h b/Engine/Plugins/Animation/ControlRig/Source/ControlRig/Public/Rigs/RigHierarchy.h
--- a/Engine/Plugins/Animation/ControlRig/Source/ControlRig/Public/Rigs/RigHierarchy.h
+++ b/Engine/Plugins/Animation/ControlRig/Source/ControlRig/Public/Rigs/RigHierarchy.h
@@ -1793,6 +1793,20 @@ public:
 	UFUNCTION(BlueprintCallable, Category = URigHierarchy)
 	TArray<FRigElementKey> GetSelectedKeys(ERigElementType InTypeFilter = ERigElementType::All) const;
 
+	/**
+	 * Returns true if any element is selected and respects the provided predicate 
+	 * @InPredicate A function that should return true when the element is desirable
+	 * @return true if any element is selected and respects the provided predicate
+	 */
+	bool HasAnythingSelectedByPredicate(TFunctionRef<bool(const FRigElementKey&)> InPredicate) const;
+
+	/**
+	 * Returns the keys of selected elements respecting the provided predicate 
+	 * @InPredicate A function that should return true when the element is desirable
+	 * @return An array of any selected element respecting the provided predicate
+	 */
+	TArray<FRigElementKey> GetSelectedKeysByPredicate(TFunctionRef<bool(const FRigElementKey&)> InPredicate) const;
+
 	/**
 	 * Returns true if a given element is selected
 	 * @param InKey The key to check
diff --git a/Engine/Plugins/Animation/ControlRig/Source/ControlRigEditor/Private/EditMode/ControlRigEditMode.cpp b/Engine/Plugins/Animation/ControlRig/Source/ControlRigEditor/Private/EditMode/ControlRigEditMode.cpp
--- a/Engine/Plugins/Animation/ControlRig/Source/ControlRigEditor/Private/EditMode/ControlRigEditMode.cpp
+++ b/Engine/Plugins/Animation/ControlRig/Source/ControlRigEditor/Private/EditMode/ControlRigEditMode.cpp
@@ -7,14 +7,12 @@
 #include "EditMode/ControlRigEditModeToolkit.h"
 #include "Toolkits/ToolkitManager.h"
 #include "EditMode/SControlRigEditModeTools.h"
-#include "Algo/Transform.h"
 #include "ControlRig.h"
 #include "HitProxies.h"
 #include "EditMode/ControlRigEditModeSettings.h"
 #include "EditMode/SControlRigDetails.h"
 #include "EditMode/SControlRigOutliner.h"
 #include "ISequencer.h"
-#include "ISequencer.h"
 #include "MVVM/ViewModels/SequencerEditorViewModel.h"
 #include "MVVM/Selection/Selection.h"
 #include "MovieScene.h"
@@ -27,31 +25,24 @@
 #include "EditMode/ControlRigEditModeCommands.h"
 #include "Framework/Application/SlateApplication.h"
 #include "Modules/ModuleManager.h"
-#include "ISequencerModule.h"
 #include "Toolkits/AssetEditorToolkit.h"
 #include "ControlRigEditorModule.h"
 #include "Constraint.h"
 #include "EngineUtils.h"
-#include "RigVMBlueprintGeneratedClass.h"
 #include "IControlRigObjectBinding.h"
 #include "Kismet2/BlueprintEditorUtils.h"
 #include "ControlRigBlueprint.h"
 #include "ControlRigGizmoActor.h"
-#include "Framework/MultiBox/MultiBoxBuilder.h"
 #include "SEditorViewport.h"
 #include "EditMode/AnimDetailsProxy.h"
 #include "ScopedTransaction.h"
-#include "RigVMModel/RigVMController.h"
 #include "Rigs/AdditiveControlRig.h"
 #include "Rigs/FKControlRig.h"
 #include "ControlRigComponent.h"
-#include "EngineUtils.h"
 #include "Units/Execution/RigUnit_BeginExecution.h"
 #include "Units/Execution/RigUnit_InteractionExecution.h"
-#include "IPersonaPreviewScene.h"
 #include "PersonaSelectionProxies.h"
 #include "PropertyHandle.h"
-#include "Framework/Application/SlateApplication.h"
 #include "UnrealEdGlobals.h"
 #include "Editor/UnrealEdEngine.h"
 #include "Settings/ControlRigSettings.h"
@@ -68,7 +59,6 @@
 #include "EdModeInteractiveToolsContext.h"
 #include "Constraints/MovieSceneConstraintChannelHelper.h"
 #include "Editor/EditorPerProjectUserSettings.h"
-#include "Widgets/Docking/SDockTab.h"
 #include "TransformConstraint.h"
 #include "Animation/DebugSkelMeshComponent.h"
 #include "Tools/ControlRigPose.h"
@@ -196,7 +186,7 @@ FControlRigEditMode::FControlRigEditMode()
 	StoredPose = NewObject<UControlRigPoseAsset>(GetTransientPackage(), NAME_None);
 	DetailKeyFrameCache = MakeShared<FDetailKeyFrameCacheAndHandler>();
 
-	UControlRigEditModeSettings* Settings = GetMutableDefault<UControlRigEditModeSettings>();
+	UControlRigEditModeSettings* Settings = GetMutableSettings();
 	bShowControlsAsOverlay = Settings->bShowControlsAsOverlay;
 
 	Settings->GizmoScaleDelegate.AddLambda([this](float GizmoScale)
@@ -253,7 +243,7 @@ bool FControlRigEditMode::SetSequencer(TWeakPtr<ISequencer> InSequencer)
 		
 		WeakSequencer = InSequencer;
 
-		DetailKeyFrameCache->UnsetDelegates();
+		UnsetSequencerDelegates();
 
 		DestroyShapesActors(nullptr);
 		TArray<TWeakObjectPtr<UControlRig>> PreviousRuntimeRigs = RuntimeControlRigs;
@@ -264,7 +254,9 @@ bool FControlRigEditMode::SetSequencer(TWeakPtr<ISequencer> InSequencer)
 				RemoveControlRig(PreviousRuntimeRigs[PreviousRuntimeRigIndex].Get());
 			}
 		}
+		
 		RuntimeControlRigs.Reset();
+
 		if (InSequencer.IsValid())
 		{
 			TSharedPtr<ISequencer> Sequencer = WeakSequencer.Pin();
@@ -279,8 +271,11 @@ bool FControlRigEditMode::SetSequencer(TWeakPtr<ISequencer> InSequencer)
 					}
 				}
 			}
+
 			LastMovieSceneSig = Sequencer->GetFocusedMovieSceneSequence()->GetMovieScene()->GetSignature();
-			DetailKeyFrameCache->SetDelegates(WeakSequencer, this);
+			
+			SetSequencerDelegates(WeakSequencer);
+			
 			ControlProxy->SetSequencer(WeakSequencer);
 
 			{
@@ -288,6 +283,7 @@ bool FControlRigEditMode::SetSequencer(TWeakPtr<ISequencer> InSequencer)
 				SequencerWidget->EnablePendingFocusOnHovering(CRModeLocals::bFocusMode);
 			}
 		}
+		
 		SetObjects_Internal();
 		if (FControlRigEditModeToolkit::Details.IsValid())
 		{
@@ -366,15 +362,19 @@ void FControlRigEditMode::SetObjects(UControlRig* ControlRig,  UObject* BindingO
 
 bool FControlRigEditMode::IsInLevelEditor() const
 {
-	return GetModeManager() == &GLevelEditorModeTools();
+	return !AreEditingControlRigDirectly() && GetModeManager() == &GLevelEditorModeTools();
 }
 
-void FControlRigEditMode::SetUpDetailPanel()
+void FControlRigEditMode::SetUpDetailPanel() const
 {
 	if (!AreEditingControlRigDirectly() && Toolkit)
 	{
-		StaticCastSharedPtr<SControlRigEditModeTools>(Toolkit->GetInlineContent())->SetSequencer(WeakSequencer.Pin());
-		StaticCastSharedPtr<SControlRigEditModeTools>(Toolkit->GetInlineContent())->SetSettingsDetailsObject(GetMutableDefault<UControlRigEditModeSettings>());	
+		TSharedPtr<SControlRigEditModeTools> ModeTools = StaticCastSharedPtr<SControlRigEditModeTools>(Toolkit->GetInlineContent());
+		if (ModeTools.IsValid())
+		{
+			ModeTools->SetSequencer(WeakSequencer.Pin());
+			ModeTools->SetSettingsDetailsObject(GetMutableSettings());
+		}	
 	}
 }
 
@@ -480,7 +480,8 @@ void FControlRigEditMode::Enter()
 
 	//set up gizmo scale to what we had last and save what it was.
 	PreviousGizmoScale = GetModeManager()->GetWidgetScale();
-	if (const UControlRigEditModeSettings* Settings = GetDefault<UControlRigEditModeSettings>())
+	
+	if (const UControlRigEditModeSettings* Settings = GetSettings())
 	{
 		GetModeManager()->SetWidgetScale(Settings->GizmoScale);
 
@@ -590,7 +591,7 @@ void FControlRigEditMode::Exit()
 	//make sure the widget is reset
 	ResetControlShapeSize();
 
-	if (const UControlRigEditModeSettings* Settings = GetDefault<UControlRigEditModeSettings>())
+	if (const UControlRigEditModeSettings* Settings = GetSettings())
 	{
 		Settings->OnSettingsChange.RemoveAll(this);
 	}
@@ -608,6 +609,7 @@ void FControlRigEditMode::Tick(FEditorViewportClient* ViewportClient, float Delt
 	{
 		return;
 	}
+	
 	CheckMovieSceneSig();
 
 	if(DeferredItemsToFrame.Num() > 0)
@@ -839,193 +841,188 @@ TSet<FName> FControlRigEditMode::GetActiveControlsFromSequencer(UControlRig* Con
 	return ActiveControls;
 }
 
-
 void FControlRigEditMode::Render(const FSceneView* View, FViewport* Viewport, FPrimitiveDrawInterface* PDI)
 {
 	DragToolHandler.Render3DDragTool(View, PDI);
 
-	const UControlRigEditModeSettings* Settings = GetDefault<UControlRigEditModeSettings>();
+	const UControlRigEditModeSettings* Settings = GetSettings();
+	if (Settings->bHideControlShapes)
+	{
+		return;
+	}
+
 	FEditorViewportClient* EditorViewportClient = (FEditorViewportClient*)Viewport->GetClient();
 	const bool bIsInGameView = !AreEditingControlRigDirectly() ? EditorViewportClient && EditorViewportClient->IsInGameView() : false;
-	bool bRender = !Settings->bHideControlShapes;
+	if (bIsInGameView)
+	{
+		//only draw stuff if not in game view
+		return;
+	}
+
+	const EWorldType::Type WorldType = Viewport->GetClient()->GetWorld()->WorldType;
+	const bool bIsAssetEditor = WorldType == EWorldType::Editor || WorldType == EWorldType::EditorPreview;
+	
 	for (TWeakObjectPtr<UControlRig>& ControlRigPtr : RuntimeControlRigs)
 	{
 		UControlRig* ControlRig = ControlRigPtr.Get();
 		//actor game view drawing is handled by not drawing in game via SetActorHiddenInGame().
-		if (bRender && ControlRig && ControlRig->GetControlsVisible())
+		if (ControlRig && ControlRig->GetControlsVisible())
 		{
 			FTransform ComponentTransform = FTransform::Identity;
 			if (!AreEditingControlRigDirectly())
 			{
 				ComponentTransform = GetHostingSceneComponentTransform(ControlRig);
 			}
-			/*
-			const TArray<AControlRigShapeActor*>* ShapeActors = ControlRigShapeActors.Find(ControlRig);
-			if (ShapeActors)
+		
+			URigHierarchy* Hierarchy = ControlRig->GetHierarchy();
+			const bool bHasFKRig = (ControlRig->IsA<UAdditiveControlRig>() || ControlRig->IsA<UFKControlRig>());
+
+			if (Settings->bDisplayHierarchy || bHasFKRig)
 			{
-				for (AControlRigShapeActor* Actor : *ShapeActors)
+				const bool bBoolSetHitProxies = PDI && PDI->IsHitTesting() && bHasFKRig;
+				TSet<FName> ActiveControlName;
+				if (bHasFKRig)
 				{
-					if (GIsEditor && Actor->GetWorld() != nullptr && !Actor->GetWorld()->IsPlayInEditor())
-					{
-						Actor->SetIsTemporarilyHiddenInEditor(false);
-					}
+					ActiveControlName = GetActiveControlsFromSequencer(ControlRig);
 				}
-			}
-			*/
-			
-			//only draw stuff if not in game view
-			if (!bIsInGameView)
-			{
-				URigHierarchy* Hierarchy = ControlRig->GetHierarchy();
-				const bool bHasFKRig = (ControlRig->IsA<UAdditiveControlRig>() || ControlRig->IsA<UFKControlRig>());
-				if (Settings->bDisplayHierarchy || bHasFKRig)
-				{
-					const bool bBoolSetHitProxies = PDI && PDI->IsHitTesting() && bHasFKRig;
-					TSet<FName> ActiveControlName;
-					if (bHasFKRig)
+				Hierarchy->ForEach<FRigTransformElement>([PDI, Hierarchy, ComponentTransform, ControlRig, bHasFKRig, bBoolSetHitProxies, ActiveControlName](FRigTransformElement* TransformElement) -> bool
 					{
-						ActiveControlName = GetActiveControlsFromSequencer(ControlRig);
-					}
-					Hierarchy->ForEach<FRigTransformElement>([PDI, Hierarchy, ComponentTransform, ControlRig, bHasFKRig, bBoolSetHitProxies, ActiveControlName](FRigTransformElement* TransformElement) -> bool
+						if(const FRigControlElement* ControlElement = Cast<FRigControlElement>(TransformElement))
 						{
-							if(const FRigControlElement* ControlElement = Cast<FRigControlElement>(TransformElement))
+							if(ControlElement->Settings.AnimationType != ERigControlAnimationType::AnimationControl)
 							{
-								if(ControlElement->Settings.AnimationType != ERigControlAnimationType::AnimationControl)
-								{
-									return true;
-								}
+								return true;
 							}
-						
-							const FTransform Transform = Hierarchy->GetTransform(TransformElement, ERigTransformType::CurrentGlobal);
+						}
+				
+						const FTransform Transform = Hierarchy->GetTransform(TransformElement, ERigTransformType::CurrentGlobal);
 
-							FRigBaseElementParentArray Parents = Hierarchy->GetParents(TransformElement);
-							for (FRigBaseElement* ParentElement : Parents)
+						FRigBaseElementParentArray Parents = Hierarchy->GetParents(TransformElement);
+						for (FRigBaseElement* ParentElement : Parents)
+						{
+							if (FRigTransformElement* ParentTransformElement = Cast<FRigTransformElement>(ParentElement))
 							{
-								if (FRigTransformElement* ParentTransformElement = Cast<FRigTransformElement>(ParentElement))
+								FLinearColor Color = FLinearColor::White;
+								if (bHasFKRig)
 								{
-									FLinearColor Color = FLinearColor::White;
-									if (bHasFKRig)
+									FName ControlName = UFKControlRig::GetControlName(ParentTransformElement->GetFName(), ParentTransformElement->GetType());
+									if (ActiveControlName.Num() > 0 && ActiveControlName.Contains(ControlName) == false)
 									{
-										FName ControlName = UFKControlRig::GetControlName(ParentTransformElement->GetFName(), ParentTransformElement->GetType());
-										if (ActiveControlName.Num() > 0 && ActiveControlName.Contains(ControlName) == false)
-										{
-											continue;
-										}
-										if (ControlRig->IsControlSelected(ControlName))
-										{
-											Color = FLinearColor::Yellow;
-										}
+										continue;
 									}
-									const FTransform ParentTransform = Hierarchy->GetTransform(ParentTransformElement, ERigTransformType::CurrentGlobal);
-									const bool bHitTesting = bBoolSetHitProxies && (ParentTransformElement->GetType() == ERigElementType::Bone);
-									if (PDI)
+									if (ControlRig->IsControlSelected(ControlName))
 									{
-										if (bHitTesting)
-										{
-											PDI->SetHitProxy(new HFKRigBoneProxy(ParentTransformElement->GetFName(), ControlRig));
-										}
-										PDI->DrawLine(ComponentTransform.TransformPosition(Transform.GetLocation()), ComponentTransform.TransformPosition(ParentTransform.GetLocation()), Color, SDPG_Foreground);
-										if (bHitTesting)
-										{
-											PDI->SetHitProxy(nullptr);
-										}
+										Color = FLinearColor::Yellow;
+									}
+								}
+								const FTransform ParentTransform = Hierarchy->GetTransform(ParentTransformElement, ERigTransformType::CurrentGlobal);
+								const bool bHitTesting = bBoolSetHitProxies && (ParentTransformElement->GetType() == ERigElementType::Bone);
+								if (PDI)
+								{
+									if (bHitTesting)
+									{
+										PDI->SetHitProxy(new HFKRigBoneProxy(ParentTransformElement->GetFName(), ControlRig));
+									}
+									PDI->DrawLine(ComponentTransform.TransformPosition(Transform.GetLocation()), ComponentTransform.TransformPosition(ParentTransform.GetLocation()), Color, SDPG_Foreground);
+									if (bHitTesting)
+									{
+										PDI->SetHitProxy(nullptr);
 									}
 								}
 							}
+						}
 
-							FLinearColor Color = FLinearColor::White;
-							if (bHasFKRig)
+						FLinearColor Color = FLinearColor::White;
+						if (bHasFKRig)
+						{
+							FName ControlName = UFKControlRig::GetControlName(TransformElement->GetFName(), TransformElement->GetType());
+							if (ActiveControlName.Num() > 0 && ActiveControlName.Contains(ControlName) == false)
 							{
-								FName ControlName = UFKControlRig::GetControlName(TransformElement->GetFName(), TransformElement->GetType());
-								if (ActiveControlName.Num() > 0 && ActiveControlName.Contains(ControlName) == false)
-								{
-									return true;
-								}
-								if (ControlRig->IsControlSelected(ControlName))
-								{
-									Color = FLinearColor::Yellow;
-								}
+								return true;
 							}
-							if (PDI)
+							if (ControlRig->IsControlSelected(ControlName))
 							{
-								const bool bHitTesting = PDI->IsHitTesting() && bBoolSetHitProxies && (TransformElement->GetType() == ERigElementType::Bone);
-								if (bHitTesting)
-								{
-									PDI->SetHitProxy(new HFKRigBoneProxy(TransformElement->GetFName(), ControlRig));
-								}
-								PDI->DrawPoint(ComponentTransform.TransformPosition(Transform.GetLocation()), Color, 5.0f, SDPG_Foreground);
-
-								if (bHitTesting)
-								{
-									PDI->SetHitProxy(nullptr);
-								}
+								Color = FLinearColor::Yellow;
 							}
-
-							return true;
-						});
-				}
-
-				EWorldType::Type WorldType = Viewport->GetClient()->GetWorld()->WorldType;
-				const bool bIsAssetEditor = WorldType == EWorldType::Editor || WorldType == EWorldType::EditorPreview;
-				if (bIsAssetEditor && (Settings->bDisplayNulls || ControlRig->IsConstructionModeEnabled()))
-				{
-					TArray<FTransform> SpaceTransforms;
-					TArray<FTransform> SelectedSpaceTransforms;
-					Hierarchy->ForEach<FRigNullElement>([&SpaceTransforms, &SelectedSpaceTransforms, Hierarchy](FRigNullElement* NullElement) -> bool
+						}
+						if (PDI)
 						{
-							if (Hierarchy->IsSelected(NullElement->GetIndex()))
+							const bool bHitTesting = PDI->IsHitTesting() && bBoolSetHitProxies && (TransformElement->GetType() == ERigElementType::Bone);
+							if (bHitTesting)
 							{
-								SelectedSpaceTransforms.Add(Hierarchy->GetTransform(NullElement, ERigTransformType::CurrentGlobal));
+								PDI->SetHitProxy(new HFKRigBoneProxy(TransformElement->GetFName(), ControlRig));
 							}
-							else
+							PDI->DrawPoint(ComponentTransform.TransformPosition(Transform.GetLocation()), Color, 5.0f, SDPG_Foreground);
+
+							if (bHitTesting)
 							{
-								SpaceTransforms.Add(Hierarchy->GetTransform(NullElement, ERigTransformType::CurrentGlobal));
+								PDI->SetHitProxy(nullptr);
 							}
-							return true;
-						});
+						}
 
-					ControlRig->DrawInterface.DrawAxes(FTransform::Identity, SpaceTransforms, Settings->AxisScale);
-					ControlRig->DrawInterface.DrawAxes(FTransform::Identity, SelectedSpaceTransforms, FLinearColor(1.0f, 0.34f, 0.0f, 1.0f), Settings->AxisScale);
-				}
+						return true;
+					});
+			}
 
-				if (bIsAssetEditor && (Settings->bDisplayAxesOnSelection && Settings->AxisScale > SMALL_NUMBER))
+			if (bIsAssetEditor && (Settings->bDisplayNulls || ControlRig->IsConstructionModeEnabled()))
+			{
+				TArray<FTransform> SpaceTransforms;
+				TArray<FTransform> SelectedSpaceTransforms;
+				Hierarchy->ForEach<FRigNullElement>([&SpaceTransforms, &SelectedSpaceTransforms, Hierarchy](FRigNullElement* NullElement) -> bool
 				{
-					TArray<FRigElementKey> SelectedRigElements = GetSelectedRigElements(ControlRig);
-					const float Scale = Settings->AxisScale;
-					PDI->AddReserveLines(SDPG_Foreground, SelectedRigElements.Num() * 3);
-
-					for (const FRigElementKey& SelectedElement : SelectedRigElements)
+					if (Hierarchy->IsSelected(NullElement->GetIndex()))
 					{
-						FTransform ElementTransform = Hierarchy->GetGlobalTransform(SelectedElement);
-						ElementTransform = ElementTransform * ComponentTransform;
-
-						PDI->DrawLine(ElementTransform.GetTranslation(), ElementTransform.TransformPosition(FVector(Scale, 0.f, 0.f)), FLinearColor::Red, SDPG_Foreground);
-						PDI->DrawLine(ElementTransform.GetTranslation(), ElementTransform.TransformPosition(FVector(0.f, Scale, 0.f)), FLinearColor::Green, SDPG_Foreground);
-						PDI->DrawLine(ElementTransform.GetTranslation(), ElementTransform.TransformPosition(FVector(0.f, 0.f, Scale)), FLinearColor::Blue, SDPG_Foreground);
+						SelectedSpaceTransforms.Add(Hierarchy->GetTransform(NullElement, ERigTransformType::CurrentGlobal));
 					}
-				}
+					else
+					{
+						SpaceTransforms.Add(Hierarchy->GetTransform(NullElement, ERigTransformType::CurrentGlobal));
+					}
+					return true;
+				});
 
-				// temporary implementation to draw sockets in 3D
-				if (bIsAssetEditor && (Settings->bDisplaySockets || ControlRig->IsConstructionModeEnabled()) && Settings->AxisScale > SMALL_NUMBER)
-				{
-					const float Scale = Settings->AxisScale;
-					PDI->AddReserveLines(SDPG_Foreground, Hierarchy->Num(ERigElementType::Socket) * 3);
-					static const FLinearColor SocketColor = FControlRigEditorStyle::Get().SocketUserInterfaceColor;
+				ControlRig->DrawInterface.DrawAxes(FTransform::Identity, SpaceTransforms, Settings->AxisScale);
+				ControlRig->DrawInterface.DrawAxes(FTransform::Identity, SelectedSpaceTransforms, FLinearColor(1.0f, 0.34f, 0.0f, 1.0f), Settings->AxisScale);
+			}
 
-					Hierarchy->ForEach<FRigSocketElement>([this, Hierarchy, PDI, ComponentTransform, Scale](FRigSocketElement* Socket)
-					{
-						FTransform ElementTransform = Hierarchy->GetGlobalTransform(Socket->GetIndex());
-						ElementTransform = ElementTransform * ComponentTransform;
+			if (bIsAssetEditor && (Settings->bDisplayAxesOnSelection && Settings->AxisScale > SMALL_NUMBER))
+			{
+				TArray<FRigElementKey> SelectedRigElements = GetSelectedRigElements(ControlRig);
+				const float Scale = Settings->AxisScale;
+				PDI->AddReserveLines(SDPG_Foreground, SelectedRigElements.Num() * 3);
 
-						PDI->DrawLine(ElementTransform.GetTranslation(), ElementTransform.TransformPosition(FVector(Scale, 0.f, 0.f)), SocketColor, SDPG_Foreground);
-						PDI->DrawLine(ElementTransform.GetTranslation(), ElementTransform.TransformPosition(FVector(0.f, Scale, 0.f)), SocketColor, SDPG_Foreground);
-						PDI->DrawLine(ElementTransform.GetTranslation(), ElementTransform.TransformPosition(FVector(0.f, 0.f, Scale)), SocketColor, SDPG_Foreground);
+				for (const FRigElementKey& SelectedElement : SelectedRigElements)
+				{
+					FTransform ElementTransform = Hierarchy->GetGlobalTransform(SelectedElement);
+					ElementTransform = ElementTransform * ComponentTransform;
 
-						return true;
-					});
+					PDI->DrawLine(ElementTransform.GetTranslation(), ElementTransform.TransformPosition(FVector(Scale, 0.f, 0.f)), FLinearColor::Red, SDPG_Foreground);
+					PDI->DrawLine(ElementTransform.GetTranslation(), ElementTransform.TransformPosition(FVector(0.f, Scale, 0.f)), FLinearColor::Green, SDPG_Foreground);
+					PDI->DrawLine(ElementTransform.GetTranslation(), ElementTransform.TransformPosition(FVector(0.f, 0.f, Scale)), FLinearColor::Blue, SDPG_Foreground);
 				}
-				ControlRig->DrawIntoPDI(PDI, ComponentTransform);
 			}
+
+			// temporary implementation to draw sockets in 3D
+			if (bIsAssetEditor && (Settings->bDisplaySockets || ControlRig->IsConstructionModeEnabled()) && Settings->AxisScale > SMALL_NUMBER)
+			{
+				const float Scale = Settings->AxisScale;
+				PDI->AddReserveLines(SDPG_Foreground, Hierarchy->Num(ERigElementType::Socket) * 3);
+				static const FLinearColor SocketColor = FControlRigEditorStyle::Get().SocketUserInterfaceColor;
+
+				Hierarchy->ForEach<FRigSocketElement>([this, Hierarchy, PDI, ComponentTransform, Scale](FRigSocketElement* Socket)
+				{
+					FTransform ElementTransform = Hierarchy->GetGlobalTransform(Socket->GetIndex());
+					ElementTransform = ElementTransform * ComponentTransform;
+
+					PDI->DrawLine(ElementTransform.GetTranslation(), ElementTransform.TransformPosition(FVector(Scale, 0.f, 0.f)), SocketColor, SDPG_Foreground);
+					PDI->DrawLine(ElementTransform.GetTranslation(), ElementTransform.TransformPosition(FVector(0.f, Scale, 0.f)), SocketColor, SDPG_Foreground);
+					PDI->DrawLine(ElementTransform.GetTranslation(), ElementTransform.TransformPosition(FVector(0.f, 0.f, Scale)), SocketColor, SDPG_Foreground);
+
+					return true;
+				});
+			}
+			
+			ControlRig->DrawIntoPDI(PDI, ComponentTransform);
 		}
 	}
 }
@@ -1641,44 +1638,11 @@ bool FControlRigEditMode::HandleClick(FEditorViewportClient* InViewportClient, H
 		ClearRigElementSelection(FRigElementTypeHelper::ToMask(ERigElementType::All));
 	}
 
-	const UControlRigEditModeSettings* Settings = GetDefault<UControlRigEditModeSettings>();
-
+	const UControlRigEditModeSettings* Settings = GetSettings();
 	if (Settings && Settings->bOnlySelectRigControls)
 	{
 		return true;
 	}
-	/*
-	if(!InViewportClient->IsLevelEditorClient() && !InViewportClient->IsSimulateInEditorViewport())
-	{
-		bool bHandled = false;
-		const bool bSelectingSections = GetAnimPreviewScene().AllowMeshHitProxies();
-
-		USkeletalMeshComponent* MeshComponent = GetAnimPreviewScene().GetPreviewMeshComponent();
-
-		if ( HitProxy )
-		{
-			if ( HitProxy->IsA( HPersonaBoneProxy::StaticGetType() ) )
-			{			
-				SetRigElementSelection(ERigElementType::Bone, static_cast<HPersonaBoneProxy*>(HitProxy)->BoneName, true);
-				bHandled = true;
-			}
-		}
-		
-		if ( !bHandled && !bSelectingSections )
-		{
-			// Cast for phys bodies if we didn't get any hit proxies
-			FHitResult Result(1.0f);
-			UDebugSkelMeshComponent* PreviewMeshComponent = GetAnimPreviewScene().GetPreviewMeshComponent();
-			bool bHit = PreviewMeshComponent->LineTraceComponent(Result, Click.GetOrigin(), Click.GetOrigin() + Click.GetDirection() * 10000.0f, FCollisionQueryParams(NAME_None, FCollisionQueryParams::GetUnknownStatId(),true));
-			
-			if(bHit)
-			{
-				SetRigElementSelection(ERigElementType::Bone, Result.BoneName, true);
-				bHandled = true;
-			}
-		}
-	}
-	*/
 	
 	return FEdMode::HandleClick(InViewportClient, HitProxy, Click);
 }
@@ -1798,8 +1762,7 @@ bool IntersectsBox( AActor& InActor, const FBox& InBox, FLevelEditorViewportClie
 
 bool FControlRigEditMode::BoxSelect(FBox& InBox, bool InSelect)
 {
-
-	const UControlRigEditModeSettings* Settings = GetDefault<UControlRigEditModeSettings>();
+	const UControlRigEditModeSettings* Settings = GetSettings();
 	FLevelEditorViewportClient* LevelViewportClient = GCurrentLevelEditingViewportClient;
 	if (LevelViewportClient->IsInGameView() == true || Settings->bHideControlShapes)
 	{
@@ -1862,7 +1825,7 @@ bool FControlRigEditMode::BoxSelect(FBox& InBox, bool InSelect)
 
 bool FControlRigEditMode::FrustumSelect(const FConvexVolume& InFrustum, FEditorViewportClient* InViewportClient, bool InSelect)
 {
-	const UControlRigEditModeSettings* Settings = GetDefault<UControlRigEditModeSettings>();
+	const UControlRigEditModeSettings* Settings = GetSettings();
 	if (!Settings)
 	{
 		return false;
@@ -2137,7 +2100,7 @@ bool FControlRigEditMode::InputDelta(FEditorViewportClient* InViewportClient, FV
 				}
 				else
 				{
-					const UControlRigEditModeSettings* Settings = GetDefault<UControlRigEditModeSettings>();
+					const UControlRigEditModeSettings* Settings = GetSettings();
 					bool bDoLocal = (CoordSystem == ECoordSystem::COORD_Local && Settings && Settings->bLocalTransformsInEachLocalSpace);
 					bool bUseLocal = false;
 					bool bCalcLocal = bDoLocal;
@@ -2448,41 +2411,31 @@ TArray<FRigElementKey> FControlRigEditMode::GetSelectedRigElements(UControlRig*
 
 bool FControlRigEditMode::AreRigElementsSelected(uint32 InTypes, UControlRig* InControlRig) const
 {
-	TArray<FRigElementKey> SelectedRigElements = GetSelectedRigElements(InControlRig);
-
-	for (const FRigElementKey& Ele : SelectedRigElements)
+	if (IsInLevelEditor())
 	{
-		if (FRigElementTypeHelper::DoesHave(InTypes, Ele.Type))
+		// no need to look for transient controls when animating in the level editor
+		const URigHierarchy* Hierarchy = InControlRig->GetHierarchy();
+		if (!Hierarchy)
 		{
-			return true;
+			return false;
 		}
+			
+		return Hierarchy->HasAnythingSelectedByPredicate([InTypes](const FRigElementKey& InSelectedKey)
+		{
+			return FRigElementTypeHelper::DoesHave(InTypes, InSelectedKey.Type);
+		});
 	}
 
-	return false;
-}
-
-int32 FControlRigEditMode::GetNumSelectedRigElements(uint32 InTypes, UControlRig* InControlRig) const
-{
-	TArray<FRigElementKey> SelectedRigElements = GetSelectedRigElements(InControlRig);
-	if (FRigElementTypeHelper::DoesHave(InTypes, ERigElementType::All))
-	{
-		return SelectedRigElements.Num();
-	}
-	else
+	const TArray<FRigElementKey> SelectedRigElements = GetSelectedRigElements(InControlRig);
+	for (const FRigElementKey& Ele : SelectedRigElements)
 	{
-		int32 NumSelected = 0;
-		for (const FRigElementKey& Ele : SelectedRigElements)
+		if (FRigElementTypeHelper::DoesHave(InTypes, Ele.Type))
 		{
-			if (FRigElementTypeHelper::DoesHave(InTypes, Ele.Type))
-			{
-				++NumSelected;
-			}
+			return true;
 		}
-
-		return NumSelected;
 	}
 
-	return 0;
+	return false;
 }
 
 void FControlRigEditMode::RefreshObjects()
@@ -2511,7 +2464,7 @@ bool FControlRigEditMode::CanRemoveFromPreviewScene(const USceneComponent* InCom
 
 ECoordSystem FControlRigEditMode::GetCoordSystemSpace() const
 {
-	const UControlRigEditModeSettings* Settings = GetDefault<UControlRigEditModeSettings>();
+	const UControlRigEditModeSettings* Settings = GetSettings();
 	if (Settings && Settings->bCoordSystemPerWidgetMode)
 	{
 		const int32 WidgetMode = static_cast<int32>(GetModeManager()->GetWidgetMode());
@@ -2667,7 +2620,7 @@ bool FControlRigEditMode::ComputePivotFromElements(UControlRig* InControlRig, FT
 
 void FControlRigEditMode::UpdatePivotTransforms()
 {
-	const UControlRigEditModeSettings* Settings = GetDefault<UControlRigEditModeSettings>();
+	const UControlRigEditModeSettings* Settings = GetSettings();
 	const bool bEachLocalSpace = Settings && Settings->bLocalTransformsInEachLocalSpace;
 	const bool bIsParentSpace = GetCoordSystemSpace() == COORD_Parent;
 
@@ -2679,6 +2632,7 @@ void FControlRigEditMode::UpdatePivotTransforms()
 		{
 			bool bAdd = false;
 			FTransform Pivot = FTransform::Identity;
+			
 			if (AreRigElementsSelected(ValidControlTypeMask(), ControlRig))
 			{
 				if (bIsChangingControlShapeTransform)
@@ -2696,6 +2650,7 @@ void FControlRigEditMode::UpdatePivotTransforms()
 				// we will enter the if first as AreRigElementsSelected will return true before AreRigElementSelectedAndMovable does...
 				bAdd = ComputePivotFromElements(ControlRig, Pivot);
 			}
+			
 			if (bAdd)
 			{
 				PivotTransforms.Add(ControlRig, MoveTemp(Pivot));
@@ -2796,7 +2751,7 @@ void FControlRigEditMode::UpdatePivotTransformsIfNeeded(UControlRig* InControlRi
 		}
 		else
 		{
-			const UControlRigEditModeSettings* Settings = GetDefault<UControlRigEditModeSettings>();
+			const UControlRigEditModeSettings* Settings = GetSettings();
 			const bool bEachLocalSpace = Settings && Settings->bLocalTransformsInEachLocalSpace;
 			const bool bIsParentSpace = GetCoordSystemSpace() == COORD_Parent;
 			ComputePivotFromShapeActors(InControlRig, bEachLocalSpace, bIsParentSpace, InOutTransform);			
@@ -3064,21 +3019,20 @@ void FControlRigEditMode::FrameItems(const TArray<FRigElementKey>& InItems)
 
 void FControlRigEditMode::IncreaseShapeSize()
 {
-	UControlRigEditModeSettings* Settings = GetMutableDefault<UControlRigEditModeSettings>();
+	UControlRigEditModeSettings* Settings = GetMutableSettings();
 	Settings->GizmoScale += 0.1f;
 	GetModeManager()->SetWidgetScale(Settings->GizmoScale);
 }
 
 void FControlRigEditMode::DecreaseShapeSize()
 {
-	UControlRigEditModeSettings* Settings = GetMutableDefault<UControlRigEditModeSettings>();
+	UControlRigEditModeSettings* Settings = GetMutableSettings();
 	Settings->GizmoScale -= 0.1f;
 	GetModeManager()->SetWidgetScale(Settings->GizmoScale);
 }
 
 void FControlRigEditMode::ResetControlShapeSize()
 {
-	UControlRigEditModeSettings* Settings = GetMutableDefault<UControlRigEditModeSettings>();
 	GetModeManager()->SetWidgetScale(PreviousGizmoScale);
 }
 
@@ -3420,14 +3374,14 @@ void FControlRigEditMode::ToggleManipulators()
 	}
 	else
 	{
-		UControlRigEditModeSettings* Settings = GetMutableDefault<UControlRigEditModeSettings>();
+		UControlRigEditModeSettings* Settings = GetMutableSettings();
 		Settings->bHideControlShapes = !Settings->bHideControlShapes;
 	}
 }
 
 void FControlRigEditMode::ToggleAllManipulators()
 {	
-	UControlRigEditModeSettings* Settings = GetMutableDefault<UControlRigEditModeSettings>();
+	UControlRigEditModeSettings* Settings = GetMutableSettings();
 	Settings->bHideControlShapes = !Settings->bHideControlShapes;
 
 	//turn on all if in level editor in case any where off
@@ -3464,7 +3418,7 @@ bool FControlRigEditMode::AreControlsVisible() const
 		return true;
 	}
 	
-	UControlRigEditModeSettings* Settings = GetMutableDefault<UControlRigEditModeSettings>();
+	const UControlRigEditModeSettings* Settings = GetSettings();
 	return !Settings->bHideControlShapes;
 }
 
@@ -3766,11 +3720,16 @@ bool FControlRigEditMode::MouseMove(FEditorViewportClient* InViewportClient, FVi
 		{
 			return false;
 		}
-		
+
 		if (!InViewportClient || InViewportClient->IsInGameView())
 		{
 			return false;
 		}
+
+		if (bSequencerPlaying)
+		{
+			return false;
+		}
 		
 		for (auto&[ControlRig, ShapeActors]: ControlRigShapeActors)
 		{
@@ -3897,6 +3856,54 @@ void FControlRigEditMode::UnregisterPendingFocusMode()
 	}
 }
 
+void FControlRigEditMode::SetSequencerDelegates(const TWeakPtr<ISequencer>& InWeakSequencer)
+{
+	if (ensure(InWeakSequencer == WeakSequencer))
+	{
+		if (WeakSequencer.IsValid())
+		{
+			DetailKeyFrameCache->SetDelegates(WeakSequencer, this);
+
+			if (TSharedPtr<ISequencer> Sequencer = WeakSequencer.Pin())
+			{
+				auto AddSequencerStatusBinding = [this](TMulticastDelegate<void()>& InDelegate)
+				{
+					if (!InDelegate.IsBoundToObject(this))
+					{
+						InDelegate.AddRaw(this, &FControlRigEditMode::UpdateSequencerStatus);
+					}		
+				};
+
+				AddSequencerStatusBinding(Sequencer->OnPlayEvent());
+				AddSequencerStatusBinding(Sequencer->OnStopEvent());
+				// NOTE this is needed as status changes are not triggered
+				AddSequencerStatusBinding(Sequencer->OnGlobalTimeChanged());
+				AddSequencerStatusBinding(Sequencer->OnEndScrubbingEvent());
+			}
+		}
+	}
+	UpdateSequencerStatus();
+}
+
+void FControlRigEditMode::UnsetSequencerDelegates() const
+{
+	DetailKeyFrameCache->UnsetDelegates();
+
+	if (TSharedPtr<ISequencer> Sequencer = WeakSequencer.Pin())
+	{
+		Sequencer->OnPlayEvent().RemoveAll(this);
+		Sequencer->OnStopEvent().RemoveAll(this);
+		Sequencer->OnGlobalTimeChanged().RemoveAll(this);
+		Sequencer->OnEndScrubbingEvent().RemoveAll(this);
+	}
+}
+
+void FControlRigEditMode::UpdateSequencerStatus()
+{
+	const TSharedPtr<ISequencer> Sequencer = WeakSequencer.IsValid() ? WeakSequencer.Pin() : nullptr;
+	bSequencerPlaying = Sequencer.IsValid() && Sequencer->GetPlaybackStatus() == EMovieScenePlayerStatus::Playing;
+}
+
 bool FControlRigEditMode::CheckMovieSceneSig()
 {
 	bool bSomethingChanged = false;
@@ -4381,37 +4388,58 @@ bool FControlRigEditMode::IsTransformDelegateAvailable() const
 	return (OnGetRigElementTransformDelegate.IsBound() && OnSetRigElementTransformDelegate.IsBound());
 }
 
-bool FControlRigEditMode::AreRigElementSelectedAndMovable(UControlRig* ControlRig) const
+namespace UE::Private
 {
-	const UControlRigEditModeSettings* Settings = GetDefault<UControlRigEditModeSettings>();
 
-	if (!Settings || !ControlRig)
+bool IsControlSelectedAndTransformable(const URigHierarchy* InHierarchy, const FRigElementKey& InSelectedKey)
+{
+    if (!FRigElementTypeHelper::DoesHave(FControlRigEditMode::ValidControlTypeMask(), InSelectedKey.Type))
+    {
+    	return false;
+    }
+    
+    const FRigControlElement* ControlElement = InHierarchy ? InHierarchy->Find<FRigControlElement>(InSelectedKey) : nullptr;
+    if (!ControlElement)
+    {
+    	return false;
+    }
+    
+    // can a control non selectable in the viewport be movable?  
+    return ControlElement->Settings.IsSelectable();
+}
+	
+}
+
+bool FControlRigEditMode::AreRigElementSelectedAndMovable(UControlRig* InControlRig) const
+{
+	if (!InControlRig)
 	{
 		return false;
 	}
 
-	auto IsAnySelectedControlMovable = [this, ControlRig]()
+	// no need to look for transient controls when animating in the level editor
+	if (IsInLevelEditor())
 	{
-		URigHierarchy* Hierarchy = ControlRig->GetHierarchy();
+		const URigHierarchy* Hierarchy = InControlRig->GetHierarchy();
 		if (!Hierarchy)
 		{
 			return false;
 		}
-		
-		const TArray<FRigElementKey> SelectedRigElements = GetSelectedRigElements(ControlRig);
-		return SelectedRigElements.ContainsByPredicate([Hierarchy](const FRigElementKey& Element)
+			
+		return Hierarchy->HasAnythingSelectedByPredicate([Hierarchy](const FRigElementKey& InSelectedKey)
 		{
-			if (!FRigElementTypeHelper::DoesHave(ValidControlTypeMask(), Element.Type))
-			{
-				return false;
-			}
-			const FRigControlElement* ControlElement = Hierarchy->Find<FRigControlElement>(Element);
-			if (!ControlElement)
-			{
-				return false;
-			}
-			// can a control non selectable in the viewport be movable?  
-			return ControlElement->Settings.IsSelectable();
+			return UE::Private::IsControlSelectedAndTransformable(Hierarchy, InSelectedKey);
+		});
+	}
+
+	auto IsAnySelectedControlMovable = [this, InControlRig]()
+	{
+		const TArray<FRigElementKey> SelectedRigElements = GetSelectedRigElements(InControlRig);
+
+		const URigHierarchy* Hierarchy = InControlRig->GetHierarchy();
+		return SelectedRigElements.ContainsByPredicate([Hierarchy](const FRigElementKey& InSelectedKey)
+		{
+			return UE::Private::IsControlSelectedAndTransformable(Hierarchy, InSelectedKey);
 		});
 	};
 	
@@ -4477,6 +4505,7 @@ void FControlRigEditMode::OnHierarchyModified(ERigHierarchyNotification InNotif,
 		case ERigHierarchyNotification::ControlVisibilityChanged:
 		case ERigHierarchyNotification::ControlShapeTransformChanged:
 		{
+			const UControlRigEditModeSettings* Settings = GetSettings();
 			const FRigElementKey Key = InElement->GetKey();
 			UControlRig* ControlRig = InHierarchy->GetTypedOuter<UControlRig>();
 			if (Key.Type == ERigElementType::Control)
@@ -4486,7 +4515,6 @@ void FControlRigEditMode::OnHierarchyModified(ERigHierarchyNotification InNotif,
 					if (AControlRigShapeActor* ShapeActor = GetControlShapeFromControlName(ControlRig,Key.Name))
 					{
 						// try to lazily apply the changes to the actor
-						const UControlRigEditModeSettings* Settings = GetDefault<UControlRigEditModeSettings>();
 						if (ShapeActor->UpdateControlSettings(InNotif, ControlRig, ControlElement, Settings->bHideControlShapes, !AreEditingControlRigDirectly()))
 						{
 							break;
@@ -4591,7 +4619,7 @@ void FControlRigEditMode::OnHierarchyModified(ERigHierarchyNotification InNotif,
 
 								if(ControlElement->CanDriveControls())
 								{
-									const UControlRigEditModeSettings* Settings = GetDefault<UControlRigEditModeSettings>();
+									const UControlRigEditModeSettings* Settings = GetSettings();
 
 									const TArray<FRigElementKey>& DrivenKeys = ControlElement->Settings.DrivenControls;
 									for(const FRigElementKey& DrivenKey : DrivenKeys)
@@ -4731,26 +4759,15 @@ void FControlRigEditMode::OnHierarchyModified_AnyThread(ERigHierarchyNotificatio
 
 void FControlRigEditMode::OnControlModified(UControlRig* Subject, FRigControlElement* InControlElement, const FRigControlModifiedContext& Context)
 {
-	//this makes sure the details panel ui get's updated, don't remove
-	const bool bModify = Context.SetKey != EControlRigSetKey::Never;
-	ControlProxy->ProxyChanged(Subject, InControlElement, bModify);
-
-	bPivotsNeedUpdate = true;
-	
-	/*
-	FScopedTransaction ScopedTransaction(LOCTEXT("ModifyControlTransaction", "Modify Control"),!GIsTransacting && Context.SetKey != EControlRigSetKey::Never);
-	ControlProxy->Modify();
-	RecalcPivotTransform();
-
-	if (UControlRig* ControlRig = static_cast<UControlRig*>(Subject))
+	// do not propagate the change to the proxy when playing for performances reasons
+	if (!bSequencerPlaying)
 	{
-		FTransform ComponentTransform = GetHostingSceneComponentTransform();
-		if (AControlRigShapeActor* const* Actor = GizmoToControlMap.FindKey(InControl.Index))
-		{
-			TickControlShape(*Actor, ComponentTransform);
-		}
+		//this makes sure the details panel ui get's updated, don't remove
+		const bool bModify = Context.SetKey != EControlRigSetKey::Never;
+		ControlProxy->ProxyChanged(Subject, InControlElement, bModify);
 	}
-	*/
+
+	bPivotsNeedUpdate = true;
 }
 
 void FControlRigEditMode::OnPreConstruction_AnyThread(UControlRig* InRig, const FName& InEventName)
@@ -4803,7 +4820,7 @@ void FControlRigEditMode::OnPostConstruction_AnyThread(UControlRig* InRig, const
 
 void FControlRigEditMode::OnWidgetModeChanged(UE::Widget::EWidgetMode InWidgetMode)
 {
-	const UControlRigEditModeSettings* Settings = GetDefault<UControlRigEditModeSettings>();
+	const UControlRigEditModeSettings* Settings = GetSettings();
 	if (Settings && Settings->bCoordSystemPerWidgetMode)
 	{
 		TGuardValue<bool> ReentrantGuardSelf(bIsChangingCoordSystem, true);
@@ -5729,13 +5746,13 @@ void FControlRigEditMode::UpdateSelectabilityOnSkeletalMeshes(UControlRig* InCon
 
 void FControlRigEditMode::SetOnlySelectRigControls(bool Val)
 {
-	UControlRigEditModeSettings* Settings = GetMutableDefault<UControlRigEditModeSettings>();
+	UControlRigEditModeSettings* Settings = GetMutableSettings();
 	Settings->bOnlySelectRigControls = Val;
 }
 
 bool FControlRigEditMode::GetOnlySelectRigControls()const
 {
-	const UControlRigEditModeSettings* Settings = GetDefault<UControlRigEditModeSettings>();
+	const UControlRigEditModeSettings* Settings = GetSettings();
 	return Settings->bOnlySelectRigControls;
 }
 
@@ -5922,6 +5939,24 @@ void FControlRigEditMode::PastePose(bool bMirror, int32 PoseNum)
 	}
 }
 
+const UControlRigEditModeSettings* FControlRigEditMode::GetSettings() const
+{
+	if (!WeakSettings.IsValid())
+	{
+		WeakSettings = GetMutableDefault<UControlRigEditModeSettings>();
+	}
+	return WeakSettings.Get();
+}
+
+UControlRigEditModeSettings* FControlRigEditMode::GetMutableSettings() const
+{
+	if (!WeakSettings.IsValid())
+	{
+		WeakSettings = GetMutableDefault<UControlRigEditModeSettings>();
+	}
+	return WeakSettings.Get();
+}
+
 /**
 * FDetailKeyFrameCacheAndHandler
 */
diff --git a/Engine/Plugins/Animation/ControlRig/Source/ControlRigEditor/Private/EditMode/SControlRigDetails.cpp b/Engine/Plugins/Animation/ControlRig/Source/ControlRigEditor/Private/EditMode/SControlRigDetails.cpp
--- a/Engine/Plugins/Animation/ControlRig/Source/ControlRigEditor/Private/EditMode/SControlRigDetails.cpp
+++ b/Engine/Plugins/Animation/ControlRig/Source/ControlRigEditor/Private/EditMode/SControlRigDetails.cpp
@@ -448,7 +448,6 @@ bool SControlRigDetails::ShouldShowPropertyOnDetailCustomization(const FProperty
 		return bShow;
 	};
 
-	bool bContainsVisibleProperty = false;
 	if (InPropertyAndParent.Property.IsA<FStructProperty>())
 	{
 		const FStructProperty* StructProperty = CastField<FStructProperty>(&InPropertyAndParent.Property);
@@ -467,13 +466,17 @@ bool SControlRigDetails::ShouldShowPropertyOnDetailCustomization(const FProperty
 
 bool SControlRigDetails::IsReadOnlyPropertyOnDetailCustomization(const FPropertyAndParent& InPropertyAndParent) const
 {
+	if (SequencerTracker.IsSequencerPlaying())
+	{
+		return true;
+	}
+	
 	auto ShouldPropertyBeEnabled = [](const FProperty& InProperty)
 	{
 		bool bShow = InProperty.HasAnyPropertyFlags(CPF_Interp) || InProperty.HasMetaData(FRigVMStruct::InputMetaName);
 		return bShow;
 	};
 
-	bool bContainsVisibleProperty = false;
 	if (InPropertyAndParent.Property.IsA<FStructProperty>())
 	{
 		const FStructProperty* StructProperty = CastField<FStructProperty>(&InPropertyAndParent.Property);
@@ -496,17 +499,33 @@ FSequencerTracker::~FSequencerTracker()
 	RemoveDelegates();
 }
 
+bool FSequencerTracker::IsSequencerPlaying() const
+{
+	return bSequencerPlaying;
+}
+
 void FSequencerTracker::RemoveDelegates()
 {
 	if (TSharedPtr<ISequencer> Sequencer = WeakSequencer.Pin())
 	{
 		Sequencer->GetSelectionChangedObjectGuids().RemoveAll(this);
+		Sequencer->OnPlayEvent().RemoveAll(this);
+		Sequencer->OnStopEvent().RemoveAll(this);
+		Sequencer->OnGlobalTimeChanged().RemoveAll(this);
+		Sequencer->OnEndScrubbingEvent().RemoveAll(this);
 	}
 }
 
+void FSequencerTracker::UpdateSequencerStatus()
+{
+	const TSharedPtr<ISequencer> Sequencer = WeakSequencer.IsValid() ? WeakSequencer.Pin() : nullptr;
+	bSequencerPlaying = Sequencer.IsValid() && Sequencer->GetPlaybackStatus() == EMovieScenePlayerStatus::Playing;
+}
+
 void FSequencerTracker::SetSequencerAndDetails(TWeakPtr<ISequencer> InWeakSequencer, SControlRigDetails* InControlRigDetails)
 {
 	RemoveDelegates();
+
 	WeakSequencer = InWeakSequencer;
 	ControlRigDetails = InControlRigDetails;
 	if (WeakSequencer.IsValid() == false || InControlRigDetails == nullptr)
@@ -520,7 +539,22 @@ void FSequencerTracker::SetSequencerAndDetails(TWeakPtr<ISequencer> InWeakSequen
 	UpdateSequencerBindings(SequencerSelectedObjects);
 
 	Sequencer->GetSelectionChangedObjectGuids().AddRaw(this, &FSequencerTracker::UpdateSequencerBindings);
+
+	auto AddSequencerStatusBinding = [this](TMulticastDelegate<void()>& InDelegate)
+	{
+		if (!InDelegate.IsBoundToObject(this))
+		{
+			InDelegate.AddRaw(this, &FSequencerTracker::UpdateSequencerStatus);
+		}		
+	};
+
+	AddSequencerStatusBinding(Sequencer->OnPlayEvent());
+	AddSequencerStatusBinding(Sequencer->OnStopEvent());
+	// NOTE this is needed as status changes are not triggered
+	AddSequencerStatusBinding(Sequencer->OnGlobalTimeChanged());
+	AddSequencerStatusBinding(Sequencer->OnEndScrubbingEvent());
 	
+	UpdateSequencerStatus();
 }
 
 void FSequencerTracker::UpdateSequencerBindings(TArray<FGuid> SequencerBindings)
diff --git a/Engine/Plugins/Animation/ControlRig/Source/ControlRigEditor/Private/EditMode/SControlRigDetails.h b/Engine/Plugins/Animation/ControlRig/Source/ControlRigEditor/Private/EditMode/SControlRigDetails.h
--- a/Engine/Plugins/Animation/ControlRig/Source/ControlRigEditor/Private/EditMode/SControlRigDetails.h
+++ b/Engine/Plugins/Animation/ControlRig/Source/ControlRigEditor/Private/EditMode/SControlRigDetails.h
@@ -37,10 +37,17 @@ public:
 	~FSequencerTracker();
 	void SetSequencerAndDetails(TWeakPtr<ISequencer> InWeakSequencer, SControlRigDetails* InControlRigDetails);
 	TMap<UObject*, FArrayOfPropertyTracks>& GetObjectsTracked() { return ObjectsTracked; }
+
+	bool IsSequencerPlaying() const;
+	
 private:
 
 	void RemoveDelegates();
 	void UpdateSequencerBindings(TArray<FGuid> SequencerBindings);
+
+	void UpdateSequencerStatus();
+	bool bSequencerPlaying = false;
+	
 	TWeakPtr<ISequencer> WeakSequencer;
 	TMap<UObject*, FArrayOfPropertyTracks> ObjectsTracked;
 	SControlRigDetails* ControlRigDetails = nullptr;
diff --git a/Engine/Plugins/Animation/ControlRig/Source/ControlRigEditor/Private/EditMode/SControlRigEditModeTools.cpp b/Engine/Plugins/Animation/ControlRig/Source/ControlRigEditor/Private/EditMode/SControlRigEditModeTools.cpp
--- a/Engine/Plugins/Animation/ControlRig/Source/ControlRigEditor/Private/EditMode/SControlRigEditModeTools.cpp
+++ b/Engine/Plugins/Animation/ControlRig/Source/ControlRigEditor/Private/EditMode/SControlRigEditModeTools.cpp
@@ -150,10 +150,8 @@ void SControlRigEditModeTools::Construct(const FArguments& InArgs, TSharedPtr<FC
 	ModeTools = InEditMode.GetModeManager();
 
 	SettingsDetailsView = FModuleManager::GetModuleChecked<FPropertyEditorModule>("PropertyEditor").CreateDetailView(DetailsViewArgs);
-	SettingsDetailsView->SetKeyframeHandler(SharedThis(this));
-	SettingsDetailsView->SetIsPropertyVisibleDelegate(FIsPropertyVisible::CreateSP(this, &SControlRigEditModeTools::ShouldShowPropertyOnDetailCustomization));
-	SettingsDetailsView->SetIsPropertyReadOnlyDelegate(FIsPropertyReadOnly::CreateSP(this, &SControlRigEditModeTools::IsReadOnlyPropertyOnDetailCustomization));
 	SettingsDetailsView->SetGenericLayoutDetailsDelegate(FOnGetDetailCustomizationInstance::CreateStatic(&FControlRigEditModeGenericDetails::MakeInstance, ModeTools));
+
 #if USE_LOCAL_DETAILS
 	ControlEulerTransformDetailsView = FModuleManager::GetModuleChecked<FPropertyEditorModule>("PropertyEditor").CreateDetailView(DetailsViewArgs);
 	ControlEulerTransformDetailsView->SetKeyframeHandler(SharedThis(this));
@@ -520,10 +518,28 @@ void SControlRigEditModeTools::SetSettingsDetailsObject(const TWeakObjectPtr<>&
 	if (SettingsDetailsView)
 	{
 		TArray<TWeakObjectPtr<>> Objects;
+
+		bool bIsControlRigSettings = false;
 		if (InObject.IsValid())
 		{
 			Objects.Add(InObject);
+			bIsControlRigSettings = InObject->IsA<UControlRigEditModeSettings>();
+		}
+
+		if (bIsControlRigSettings)
+		{
+			// no need to override anything when as are control rig settings properties are always visible/editable
+			SettingsDetailsView->SetKeyframeHandler(nullptr);
+			SettingsDetailsView->SetIsPropertyVisibleDelegate(nullptr);
+			SettingsDetailsView->SetIsPropertyReadOnlyDelegate(nullptr);
 		}
+		else
+		{
+			SettingsDetailsView->SetKeyframeHandler(SharedThis(this));
+			SettingsDetailsView->SetIsPropertyVisibleDelegate(FIsPropertyVisible::CreateStatic(&SControlRigEditModeTools::ShouldShowPropertyOnDetailCustomization));
+			SettingsDetailsView->SetIsPropertyReadOnlyDelegate(FIsPropertyReadOnly::CreateStatic(&SControlRigEditModeTools::IsReadOnlyPropertyOnDetailCustomization));
+		}
+		
 		SettingsDetailsView->SetObjects(Objects);
 	}
 }
@@ -716,24 +732,19 @@ void SControlRigEditModeTools::OnKeyPropertyClicked(const IPropertyHandle& Keyed
 	}
 }
 
-bool SControlRigEditModeTools::ShouldShowPropertyOnDetailCustomization(const FPropertyAndParent& InPropertyAndParent) const
+bool SControlRigEditModeTools::ShouldShowPropertyOnDetailCustomization(const FPropertyAndParent& InPropertyAndParent)
 {
 	auto ShouldPropertyBeVisible = [](const FProperty& InProperty)
 	{
-		bool bShow = InProperty.HasAnyPropertyFlags(CPF_Interp) || InProperty.HasMetaData(FRigVMStruct::InputMetaName) || InProperty.HasMetaData(FRigVMStruct::OutputMetaName);
-
-	/*	// Show 'PickerIKTogglePos' properties
-		bShow |= (InProperty.GetFName() == GET_MEMBER_NAME_CHECKED(FLimbControl, PickerIKTogglePos));
-		bShow |= (InProperty.GetFName() == GET_MEMBER_NAME_CHECKED(FSpineControl, PickerIKTogglePos));
-*/
-
 		// Always show settings properties
-		const UClass* OwnerClass = InProperty.GetOwner<UClass>();
-		bShow |= OwnerClass == UControlRigEditModeSettings::StaticClass();
-		return bShow;
+		if (InProperty.GetOwner<UClass>() == UControlRigEditModeSettings::StaticClass())
+		{
+			return true;
+		}
+		
+		return InProperty.HasAnyPropertyFlags(CPF_Interp) || InProperty.HasMetaData(FRigVMStruct::InputMetaName) || InProperty.HasMetaData(FRigVMStruct::OutputMetaName);
 	};
 
-	bool bContainsVisibleProperty = false;
 	if (InPropertyAndParent.Property.IsA<FStructProperty>())
 	{
 		const FStructProperty* StructProperty = CastField<FStructProperty>(&InPropertyAndParent.Property);
@@ -750,19 +761,18 @@ bool SControlRigEditModeTools::ShouldShowPropertyOnDetailCustomization(const FPr
 		(InPropertyAndParent.ParentProperties.Num() > 0 && ShouldPropertyBeVisible(*InPropertyAndParent.ParentProperties[0]));
 }
 
-bool SControlRigEditModeTools::IsReadOnlyPropertyOnDetailCustomization(const FPropertyAndParent& InPropertyAndParent) const
+bool SControlRigEditModeTools::IsReadOnlyPropertyOnDetailCustomization(const FPropertyAndParent& InPropertyAndParent)
 {
 	auto ShouldPropertyBeEnabled = [](const FProperty& InProperty)
 	{
-		bool bShow = InProperty.HasAnyPropertyFlags(CPF_Interp) || InProperty.HasMetaData(FRigVMStruct::InputMetaName);
-
 		// Always show settings properties
-		const UClass* OwnerClass = InProperty.GetOwner<UClass>();
-		bShow |= OwnerClass == UControlRigEditModeSettings::StaticClass();
-		return bShow;
+		if (InProperty.GetOwner<UClass>() == UControlRigEditModeSettings::StaticClass())
+		{
+			return true;
+		}
+		return InProperty.HasAnyPropertyFlags(CPF_Interp) || InProperty.HasMetaData(FRigVMStruct::InputMetaName);
 	};
 
-	bool bContainsVisibleProperty = false;
 	if (InPropertyAndParent.Property.IsA<FStructProperty>())
 	{
 		const FStructProperty* StructProperty = CastField<FStructProperty>(&InPropertyAndParent.Property);
@@ -1094,6 +1104,8 @@ void SControlRigEditModeTools::Compensate(TOptional<FFrameNumber> OptionalKeyTim
 
 FReply SControlRigEditModeTools::OnBakeControlsToNewSpaceButtonClicked()
 {
+	using namespace UE::MovieScene;
+	
 	if (ReadyForBakeOrCompensation() == false)
 	{
 		return FReply::Unhandled();
@@ -1109,11 +1121,6 @@ FReply SControlRigEditModeTools::OnBakeControlsToNewSpaceButtonClicked()
 			FSpaceChannelAndSection SpaceChannelAndSection = FControlRigSpaceChannelHelpers::FindSpaceChannelAndSectionForControl(ControlRig.Get(), ControlKeys[0].Name, Sequencer, false /*bCreateIfNeeded*/);
 			if (SpaceChannelAndSection.SpaceChannel != nullptr)
 			{
-				const FFrameRate TickResolution = Sequencer->GetFocusedTickResolution();
-				const FFrameTime FrameTime = Sequencer->GetLocalTime().ConvertTo(TickResolution);
-				FFrameNumber CurrentTime = FrameTime.GetFrame();
-				FMovieSceneControlRigSpaceBaseKey Value;
-				using namespace UE::MovieScene;
 				//set up settings if not setup
 				if (BakeSpaceSettings.TargetSpace == FRigElementKey())
 				{
@@ -1130,7 +1137,7 @@ FReply SControlRigEditModeTools::OnBakeControlsToNewSpaceButtonClicked()
 					.Controls(ControlKeys) // use the cached controls here since the selection is not recovered until next tick.
 					.Sequencer(Sequencer)
 					.GetControlCustomization(this, &SControlRigEditModeTools::HandleGetControlElementCustomization)
-					.OnBake_Lambda([Sequencer, ControlRig, TickResolution](URigHierarchy* InHierarchy, TArray<FRigElementKey> InControls, const FRigSpacePickerBakeSettings& InSettings)
+					.OnBake_Lambda([Sequencer, ControlRig](URigHierarchy* InHierarchy, TArray<FRigElementKey> InControls, const FRigSpacePickerBakeSettings& InSettings)
 						{		
 							FScopedTransaction Transaction(LOCTEXT("BakeControlToSpace", "Bake Control In Space"));
 							for (const FRigElementKey& ControlKey : InControls)
diff --git a/Engine/Plugins/Animation/ControlRig/Source/ControlRigEditor/Private/EditMode/SControlRigEditModeTools.h b/Engine/Plugins/Animation/ControlRig/Source/ControlRigEditor/Private/EditMode/SControlRigEditModeTools.h
--- a/Engine/Plugins/Animation/ControlRig/Source/ControlRigEditor/Private/EditMode/SControlRigEditModeTools.h
+++ b/Engine/Plugins/Animation/ControlRig/Source/ControlRigEditor/Private/EditMode/SControlRigEditModeTools.h
@@ -104,14 +104,9 @@ private:
 	TSharedPtr<SConstraintsEditionWidget> ConstraintsEditionWidget = nullptr; 
 
 	/** Display or edit set up for property */
-	bool ShouldShowPropertyOnDetailCustomization(const struct FPropertyAndParent& InPropertyAndParent) const;
-	bool IsReadOnlyPropertyOnDetailCustomization(const struct FPropertyAndParent& InPropertyAndParent) const;
+	static bool ShouldShowPropertyOnDetailCustomization(const struct FPropertyAndParent& InPropertyAndParent);
+	static bool IsReadOnlyPropertyOnDetailCustomization(const struct FPropertyAndParent& InPropertyAndParent);
 
-	/** Called when a manipulator is selected in the picker */
-	void OnManipulatorsPicked(const TArray<FName>& Manipulators);
-
-	void HandleModifiedEvent(ERigVMGraphNotifType InNotifType, URigVMGraph* InGraph, UObject* InSubject);
-	void HandleSelectionChanged(TSharedPtr<FRigTreeElement> Selection, ESelectInfo::Type SelectInfo);
 	void OnRigElementSelected(UControlRig* Subject, FRigControlElement* ControlElement, bool bSelected);
 
 	const FRigControlElementCustomization* HandleGetControlElementCustomization(URigHierarchy* InHierarchy, const FRigElementKey& InControlKey);
diff --git a/Engine/Plugins/Animation/ControlRig/Source/ControlRigEditor/Private/EditMode/SControlRigOutliner.cpp b/Engine/Plugins/Animation/ControlRig/Source/ControlRigEditor/Private/EditMode/SControlRigOutliner.cpp
--- a/Engine/Plugins/Animation/ControlRig/Source/ControlRigEditor/Private/EditMode/SControlRigOutliner.cpp
+++ b/Engine/Plugins/Animation/ControlRig/Source/ControlRigEditor/Private/EditMode/SControlRigOutliner.cpp
@@ -50,55 +50,64 @@ FText FMultiRigData::GetName() const
 {
 	if (Key.IsSet())
 	{
-		return FText::FromName(Key.GetValue().Name);
+		return FText::FromName(Key->Name);
 	}
-	else if (ControlRig.IsValid())
+	
+	if (UControlRig* ControlRigPtr = ControlRig.Get())
 	{
-		FString ControlRigName = ControlRig.Get()->GetName();
+		FString ControlRigName = ControlRigPtr->GetName();
 		FString BoundObjectName;
-		if (TSharedPtr<IControlRigObjectBinding> ObjectBinding = ControlRig.Get()->GetObjectBinding())
+
+		if (TSharedPtr<IControlRigObjectBinding> ObjectBinding = ControlRigPtr->GetObjectBinding())
 		{
-			if (ObjectBinding->GetBoundObject())
+			if (UObject* BoundObject = ObjectBinding->GetBoundObject())
 			{
-				AActor* Actor = ObjectBinding->GetBoundObject()->GetTypedOuter<AActor>();
-				if (Actor)
+				if (AActor* Actor = BoundObject->GetTypedOuter<AActor>())
 				{
 					BoundObjectName = Actor->GetActorLabel();
 				}
 			}
 		}
-		FText AreaTitle = FText::Format(LOCTEXT("ControlTitle", "{0}  ({1})"), FText::AsCultureInvariant(ControlRigName), FText::AsCultureInvariant((BoundObjectName)));
-		return AreaTitle;
+		
+		return FText::Format(LOCTEXT("ControlTitle", "{0}  ({1})"), FText::AsCultureInvariant(ControlRigName), FText::AsCultureInvariant(BoundObjectName));
 	}
-	FName None = NAME_None;
-	return FText::FromName(None);
+	
+	return FText::FromName(NAME_None);
 }
 
 FText FMultiRigData::GetDisplayName() const
 {
-	if (Key.IsSet())
+	if (!CachedDisplayName.IsSet())
 	{
-		if(ControlRig.IsValid())
+		if (Key.IsSet())
 		{
-			if(const URigHierarchy* Hierarchy = ControlRig->GetHierarchy())
+			if (const URigHierarchy* Hierarchy = GetHierarchy())
 			{
-				const FText DisplayNameForUI = Hierarchy->GetDisplayNameForUI(Key.GetValue(), false);
-				if(!DisplayNameForUI.IsEmpty())
+				const FRigElementKey& RigElementKey = Key.GetValue();
+
+				static constexpr bool bIncludeNameSpace = false;
+				const FText DisplayNameForUI = Hierarchy->GetDisplayNameForUI(RigElementKey, bIncludeNameSpace);
+				if (!DisplayNameForUI.IsEmpty())
 				{
-					return DisplayNameForUI;
+					CachedDisplayName = DisplayNameForUI;
 				}
-				
-				if(const FRigControlElement* ControlElement = Hierarchy->Find<FRigControlElement>(Key.GetValue()))
+				else if (const FRigControlElement* ControlElement = Hierarchy->Find<FRigControlElement>(RigElementKey))
 				{
-					if(!ControlElement->Settings.DisplayName.IsNone())
+					if (!ControlElement->Settings.DisplayName.IsNone())
 					{
-						return FText::FromName(ControlElement->Settings.DisplayName);
+						CachedDisplayName = FText::FromName(ControlElement->Settings.DisplayName);
 					}
 				}
 			}
 		}
+
+		if (!CachedDisplayName.IsSet())
+		{
+			CachedDisplayName = GetName();
+		}
 	}
-	return GetName();
+	
+	return *CachedDisplayName;
 }
 
 bool FMultiRigData::operator== (const FMultiRigData & Other) const
@@ -131,9 +140,9 @@ bool FMultiRigData::IsValid() const
 
 URigHierarchy* FMultiRigData::GetHierarchy() const
 {
-	if (ControlRig.IsValid())
+	if (UControlRig* ControlRigPtr = ControlRig.Get())
 	{
-		return ControlRig.Get()->GetHierarchy();
+		return ControlRigPtr->GetHierarchy();
 	}
 	return nullptr;
 }
@@ -292,11 +301,6 @@ FReply SMultiRigHierarchyItem::OnGetSelectedClicked()
 	return FReply::Unhandled();  //should flow to selection instead
 }
 
-FText SMultiRigHierarchyItem::GetName() const
-{
-	return (WeakRigTreeElement.Pin()->Data.GetName());
-}
-
 FText SMultiRigHierarchyItem::GetDisplayName() const
 {
 	return (WeakRigTreeElement.Pin()->Data.GetDisplayName());
diff --git a/Engine/Plugins/Animation/ControlRig/Source/ControlRigEditor/Private/EditMode/SControlRigOutliner.h b/Engine/Plugins/Animation/ControlRig/Source/ControlRigEditor/Private/EditMode/SControlRigOutliner.h
--- a/Engine/Plugins/Animation/ControlRig/Source/ControlRigEditor/Private/EditMode/SControlRigOutliner.h
+++ b/Engine/Plugins/Animation/ControlRig/Source/ControlRigEditor/Private/EditMode/SControlRigOutliner.h
@@ -83,9 +83,12 @@ public:
 	bool operator == (const FMultiRigData & Other) const;
 	bool IsValid() const;
 	URigHierarchy* GetHierarchy() const;
-public:
+
 	TWeakObjectPtr<UControlRig> ControlRig;
 	TOptional<FRigElementKey> Key;
+
+private:
+	mutable TOptional<FText> CachedDisplayName;
 };
 
 /** An item in the tree */
@@ -126,7 +129,6 @@ private:
 	TWeakPtr<FMultiRigTreeElement> WeakRigTreeElement;
 	FMultiRigTreeDelegates Delegates;
 
-	FText GetName() const;
 	FText GetDisplayName() const;
 	FReply OnGetSelectedClicked();
 
diff --git a/Engine/Plugins/Animation/ControlRig/Source/ControlRigEditor/Public/EditMode/ControlRigEditMode.h b/Engine/Plugins/Animation/ControlRig/Source/ControlRigEditor/Public/EditMode/ControlRigEditMode.h
--- a/Engine/Plugins/Animation/ControlRig/Source/ControlRigEditor/Public/EditMode/ControlRigEditMode.h
+++ b/Engine/Plugins/Animation/ControlRig/Source/ControlRigEditor/Public/EditMode/ControlRigEditMode.h
@@ -243,7 +243,7 @@ protected:
 	void SetObjects_Internal();
 
 	/** Set up Details Panel based upon Selected Objects*/
-	void SetUpDetailPanel();
+	void SetUpDetailPanel() const;
 
 	/** Updates cached pivot transforms */
 	void UpdatePivotTransforms();
@@ -354,6 +354,12 @@ private:
 	void RegisterPendingFocusMode();
 	void UnregisterPendingFocusMode();
 	FDelegateHandle PendingFocusHandle;
+
+	/** Listen to sequencer and optimize notifications and performances when playing. */
+	void SetSequencerDelegates(const TWeakPtr<ISequencer>& InWeakSequencer);
+	void UnsetSequencerDelegates() const;
+	void UpdateSequencerStatus();
+	bool bSequencerPlaying = false;
 	
 public:
 	
@@ -464,9 +470,6 @@ public:
 	/** Check if any RigElements are selected */
 	bool AreRigElementsSelected(uint32 InTypes, UControlRig* InControlRig) const;
 
-	/** Get the number of selected RigElements */
-	int32 GetNumSelectedRigElements(uint32 InTypes, UControlRig* InControlRig) const;
-
 	/** Get all of the selected Controls*/
 	void GetAllSelectedControls(TMap<UControlRig*, TArray<FRigElementKey>>& OutSelectedControls) const;
 
@@ -646,6 +649,14 @@ private:
 		FControlRigEditMode* EditMode;
 	};
 private:
+
+	// get default/mutable settings
+	// todo: have a local setting object that listen to property changes to send updates
+	// instead of getting data directly from the CDO
+	const UControlRigEditModeSettings* GetSettings() const;
+	UControlRigEditModeSettings* GetMutableSettings() const;
+	mutable TWeakObjectPtr<UControlRigEditModeSettings> WeakSettings;
+
 	//pose used by the hotkeys
 	TObjectPtr<UControlRigPoseAsset> StoredPose;
 
