Control Rig: Fix automation test crash

#rnx

diff --git a/Engine/Plugins/Animation/ControlRig/Source/ControlRig/Private/Tracks/SampleTrack.cpp b/Engine/Plugins/Animation/ControlRig/Source/ControlRig/Private/Tracks/SampleTrack.cpp
--- a/Engine/Plugins/Animation/ControlRig/Source/ControlRig/Private/Tracks/SampleTrack.cpp
+++ b/Engine/Plugins/Animation/ControlRig/Source/ControlRig/Private/Tracks/SampleTrack.cpp
@@ -11,7 +11,7 @@ void FSampleTrack<FInstancedStruct>::Serialize(FArchive& InArchive)
 {
 	FSampleTrackBase::Serialize(InArchive);
 	
-	int32 ArraySize = Values.Num();
+	int32 ArraySize = ValuesStorage.Num();
 	InArchive << ArraySize;
 	
 	if(InArchive.IsLoading())
