Anim In Engine: Clear out save/paste pose hotkeys, could cause issue with Redo(ctrl-y).

diff --git a/Engine/Plugins/Animation/ControlRig/Source/ControlRigEditor/Private/EditMode/ControlRigEditModeCommands.cpp b/Engine/Plugins/Animation/ControlRig/Source/ControlRigEditor/Private/EditMode/ControlRigEditModeCommands.cpp
--- a/Engine/Plugins/Animation/ControlRig/Source/ControlRigEditor/Private/EditMode/ControlRigEditModeCommands.cpp
+++ b/Engine/Plugins/Animation/ControlRig/Source/ControlRigEditor/Private/EditMode/ControlRigEditModeCommands.cpp
@@ -32,11 +32,11 @@ void FControlRigEditModeCommands::RegisterCommands()
 
 	UI_COMMAND(SelectAllControls, "Select All Controls", "Select All Controls", EUserInterfaceActionType::Button, FInputChord(EKeys::Escape, EModifierKey::Shift));
 
-	UI_COMMAND(SavePose, "Save Current Pose", "Save current Pose", EUserInterfaceActionType::Button, FInputChord(EKeys::Y, EModifierKey::Shift));
-	UI_COMMAND(SelectPose, "Select Saved Pose", "Select Controls on saved Pose", EUserInterfaceActionType::Button, FInputChord(EKeys::Y, EModifierKey::Control |EModifierKey::Shift));
-	UI_COMMAND(PastePose, "Paste Saved Pose", "Paste Controls from saved Pose", EUserInterfaceActionType::Button, FInputChord(EKeys::Y, EModifierKey::Control));
-	UI_COMMAND(SelectMirrorPose, "Select Mirrored Saved Pose", "Select mirrored Controls on saved Pose", EUserInterfaceActionType::Button, FInputChord(EKeys::Y, EModifierKey::Alt | EModifierKey::Control | EModifierKey::Shift));
-	UI_COMMAND(PasteMirrorPose, "Paste Mirrored Saved Pose", "Paste mirrored Controls from saved Pose", EUserInterfaceActionType::Button, FInputChord(EKeys::Y, EModifierKey::Alt | EModifierKey::Control));
+	UI_COMMAND(SavePose, "Save Current Pose", "Save current Pose", EUserInterfaceActionType::Button, FInputChord());
+	UI_COMMAND(SelectPose, "Select Saved Pose", "Select Controls on saved Pose", EUserInterfaceActionType::Button, FInputChord());
+	UI_COMMAND(PastePose, "Paste Saved Pose", "Paste Controls from saved Pose", EUserInterfaceActionType::Button, FInputChord());
+	UI_COMMAND(SelectMirrorPose, "Select Mirrored Saved Pose", "Select mirrored Controls on saved Pose", EUserInterfaceActionType::Button, FInputChord());
+	UI_COMMAND(PasteMirrorPose, "Paste Mirrored Saved Pose", "Paste mirrored Controls from saved Pose", EUserInterfaceActionType::Button, FInputChord());
 
 	UI_COMMAND(ToggleControlShapeTransformEdit, "Toggle Shape Transform Edit", "Toggle Editing Selected Control's Shape Transform", EUserInterfaceActionType::ToggleButton, FInputChord(EKeys::Period, EModifierKey::Control)); 
 	UI_COMMAND(OpenSpacePickerWidget, "Open the Space Picker", "Allows space switching on the control", EUserInterfaceActionType::Button, FInputChord(EKeys::Tab)); 
