Control Rig Replay: Support for Rig Modules

#rnx

diff --git a/Engine/Plugins/Animation/ControlRig/Source/ControlRig/Private/ControlRigReplay.cpp b/Engine/Plugins/Animation/ControlRig/Source/ControlRig/Private/ControlRigReplay.cpp
--- a/Engine/Plugins/Animation/ControlRig/Source/ControlRig/Private/ControlRigReplay.cpp
+++ b/Engine/Plugins/Animation/ControlRig/Source/ControlRig/Private/ControlRigReplay.cpp
@@ -2,7 +2,9 @@
 
 #include "ControlRigReplay.h"
 #include "ControlRig.h"
+#include "Rigs/RigHierarchyController.h"
 #include "ControlRigObjectVersion.h"
+#include "Engine/SkeletalMesh.h"
 #include "Units/Execution/RigUnit_PrepareForExecution.h"
 #include "Units/Execution/RigUnit_BeginExecution.h"
 #include "Units/RigUnitContext.h"
@@ -98,7 +100,14 @@ void FControlRigReplayTracks::StoreHierarchy(URigHierarchy* InHierarchy, const T
 			ParentIndices.Reserve(ElementKeys.Num());
 			for(const FRigElementKey& ElementKey : ElementKeys)
 			{
-				ParentIndices.Add(ElementKeys.Find(InHierarchy->GetDefaultParent(ElementKey)));
+				if(ElementKey.Type == ERigElementType::Connector)
+				{
+					ParentIndices.Add(INDEX_NONE);
+				}
+				else
+				{
+					ParentIndices.Add(ElementKeys.Find(InHierarchy->GetDefaultParent(ElementKey)));
+				}
 			}
 
 			ElementKeysTrack->AddSample(ElementKeys);
@@ -163,32 +172,17 @@ void FControlRigReplayTracks::StorePose(URigHierarchy* InHierarchy, const TArray
 
 void FControlRigReplayTracks::StorePose(URigHierarchy* InHierarchy, FRigBaseElement* InElement)
 {
-	if(!InElement->IsA<FRigTransformElement>() && !InElement->IsA<FRigCurveElement>())
+	if(!InElement->IsA<FRigTransformElement>() && !InElement->IsA<FRigCurveElement>() && !InElement->IsA<FRigConnectorElement>())
 	{
 		return;
 	}
 
 	FSampleTrackContainer* Storage = GetContainer();
 	const FName& TrackName = GetTrackName(InElement->GetKey());
-	if(FRigCurveElement* Curve = Cast<FRigCurveElement>(InElement))
-	{
-		const TSharedPtr<FSampleTrack<float>> Track = StaticCastSharedPtr<FSampleTrack<float>>(Storage->FindOrAddTrack(TrackName, FSampleTrackBase::ETrackType_Float));
-		Track->AddSample(InHierarchy->GetCurveValue(Curve));
-		return;
-	}
-	
-	TSharedPtr<FSampleTrack<FTransform3f>> Track = Storage->FindTrack<FTransform3f>(TrackName);
-	if(!Track.IsValid())
-	{
-		Track = Storage->AddTransformTrack(TrackName);
-		SampleTrackIndex.Update(*Storage);
-	}
-	const FTransform LocalTransform = InHierarchy->GetLocalTransform(InElement->GetIndex());
-	Track->AddSample(FTransform3f(LocalTransform));
 
-	if(bIsInput)
+	if(const FRigConnectorElement* ConnectorElement = Cast<FRigConnectorElement>(InElement))
 	{
-		if(const FRigConnectorElement* ConnectorElement = Cast<FRigConnectorElement>(InElement))
+		if(bIsInput)
 		{
 			if(const UControlRig* ControlRig = Cast<UControlRig>(InHierarchy->GetOuter()))
 			{
@@ -196,14 +190,32 @@ void FControlRigReplayTracks::StorePose(URigHierarchy* InHierarchy, FRigBaseElem
 				const TSharedPtr<FSampleTrack<TArray<FRigElementKey>>> ConnectorTrack = StaticCastSharedPtr<FSampleTrack<TArray<FRigElementKey>>>(Storage->FindOrAddTrack(ConnectorTrackName, FSampleTrackBase::ETrackType_ElementKeyArray));
 			
 				TArray<FRigElementKey> Targets;
-				if(const FRigElementKeyRedirector::FCachedKeyArray* Cache = ControlRig->GetElementKeyRedirector().Find(ConnectorElement->GetKey()))
+				const FRigElementKeyRedirector Redirector = ControlRig->GetElementKeyRedirector();
+				if(const FRigElementKeyRedirector::FCachedKeyArray* Cache = Redirector.Find(ConnectorElement->GetKey()))
 				{
 					Targets.Append(FRigElementKeyRedirector::Convert(*Cache));
 				}
 				ConnectorTrack->AddSample(Targets);
 			}
 		}
+		return;
+	}
+
+	if(FRigCurveElement* Curve = Cast<FRigCurveElement>(InElement))
+	{
+		const TSharedPtr<FSampleTrack<float>> Track = StaticCastSharedPtr<FSampleTrack<float>>(Storage->FindOrAddTrack(TrackName, FSampleTrackBase::ETrackType_Float));
+		Track->AddSample(InHierarchy->GetCurveValue(Curve));
+		return;
+	}
+	
+	TSharedPtr<FSampleTrack<FTransform3f>> Track = Storage->FindTrack<FTransform3f>(TrackName);
+	if(!Track.IsValid())
+	{
+		Track = Storage->AddTransformTrack(TrackName);
+		SampleTrackIndex.Update(*Storage);
 	}
+	const FTransform LocalTransform = InHierarchy->GetLocalTransform(InElement->GetIndex());
+	Track->AddSample(FTransform3f(LocalTransform));
 }
 
 bool FControlRigReplayTracks::RestorePose(int32 InTimeIndex, URigHierarchy* InHierarchy, const TArrayView<FRigElementKey>& InKeys, TReportFunction InReportFunction) const
@@ -1331,9 +1343,9 @@ bool UControlRigReplay::StartRecording(UControlRig* InControlRig)
 			InputTracks.StoreRigVMEvent(InEventName);
 			InputTracks.StoreInteraction(ControlRig->InteractionType, ControlRig->ElementsBeingInteracted);
 			InputTracks.StoreHierarchy(ControlRig->GetHierarchy());
-			
-			OutputTracks.AddTimeSample(ControlRig->GetAbsoluteTime(), ControlRig->GetDeltaTime());
-			
+
+			OutputTracks.AddTimeSample(InputTracks.GetLastAbsoluteTime(), InputTracks.GetLastDeltaTime());
+
 			bStoreVariablesDuringPreEvent = true;
 		}
 	);
@@ -1628,29 +1640,13 @@ bool UControlRigReplay::PerformTest(UControlRig* InSubject, TFunction<void(EMess
 	InSubject->RequestInit();
 	InputTracks.RestoreVariables(0, InSubject, ReportFunction);
 
-	// restore the connectors
-	const TArray<FRigElementKey> ConnectorKeys = Hierarchy->GetConnectorKeys();
-	FRigElementKeyRedirector::FKeyMap ConnectorMap;
-	for(const FRigElementKey& ConnectorKey : ConnectorKeys)
-	{
-		const FName& TrackName = InputTracks.GetTrackName(ConnectorKey);
-		const FName& ConnectorTrackName = *(TrackName.ToString() + TEXT("ConnectorTargets"));
-
-		if(const TSharedPtr<const FSampleTrack<TArray<FRigElementKey>>> Track = InputTracks.GetContainer()->FindTrack<TArray<FRigElementKey>>(ConnectorTrackName))
-		{
-			FRigElementKeyRedirector::FKeyArray Targets;
-			Targets.Append(Track->GetValueAtTimeIndex(0, InputTracks.SampleTrackIndex));
-			ConnectorMap.Add(ConnectorKey, Targets);
-		}
-	}
-	if(!ConnectorMap.IsEmpty())
-	{
-		InSubject->SetElementKeyRedirector(FRigElementKeyRedirector(ConnectorMap, Hierarchy));
-	}
-	
 	InSubject->SetAbsoluteAndDeltaTime(InputTracks.GetAbsoluteTime(0), InputTracks.GetDeltaTime(0));
 	
 	InSubject->EventQueue = {FRigUnit_BeginExecution::EventName};
+
+	// make sure to import the hierarchy the same way it is imported in the control rig editor and apply the connectors
+	InSubject->OnPreConstruction_AnyThread().AddUObject(this, &UControlRigReplay::HandlePreconstructionForTest);
+	
 	InSubject->Evaluate_AnyThread();
 
 	FSampleTrackIndex SampleTrackIndex(OutputTracks.GetContainer()->NumTracks());
@@ -1746,7 +1742,7 @@ bool UControlRigReplay::ValidateExpectedResults(
 			{
 				continue;
 			}
-			const int32 CurrentParentIndex = CurrentElementKeys.Find(Hierarchy->GetDefaultParent(CurrentKey));
+			const int32 CurrentParentIndex = CurrentKey.Type == ERigElementType::Connector ? INDEX_NONE : CurrentElementKeys.Find(Hierarchy->GetDefaultParent(CurrentKey));
 			const int32 StoredParentIndex = StoredParentIndices[StoredIndex];
 
 			if(CurrentParentIndex == INDEX_NONE && StoredParentIndex != INDEX_NONE)
@@ -2002,7 +1998,7 @@ bool UControlRigReplay::ValidateExpectedResults(
 								{
 									const FVector3f& StoredValue = StaticCastSharedPtr<const FSampleTrack<FVector3f>>(MetadataTrack)->GetValueAtTimeIndex(InPlaybackTimeIndex, OutSampleTrackIndex);
 									const FVector3f& CurrentValue = FVector3f(static_cast<const FRigVectorMetadata*>(Metadata)->GetValue());
-									if(StoredValue.Equals(CurrentValue, static_cast<float>(Tolerance)))
+									if(!StoredValue.Equals(CurrentValue, static_cast<float>(Tolerance)))
 									{
 										TrackNameBasedLogFunction(EMessageSeverity::Error, TrackName, TEXT("Metadata value doesn't match."));
 										bSuccess = false;
@@ -2022,7 +2018,7 @@ bool UControlRigReplay::ValidateExpectedResults(
 									{
 										for(int32 StoredValueIndex = 0; StoredValueIndex < StoredValue.Num(); StoredValueIndex++)
 										{
-											if(StoredValue[StoredValueIndex].Equals(FVector3f(CurrentValue[StoredValueIndex]), static_cast<float>(Tolerance)))
+											if(!StoredValue[StoredValueIndex].Equals(FVector3f(CurrentValue[StoredValueIndex]), static_cast<float>(Tolerance)))
 											{
 												TrackNameBasedLogFunction(EMessageSeverity::Error, TrackName, TEXT("Metadata array value element doesn't match."));
 												bSuccess = false;
@@ -2035,7 +2031,7 @@ bool UControlRigReplay::ValidateExpectedResults(
 								{
 									FRotator StoredValue = FRotator::MakeFromEuler(FVector(StaticCastSharedPtr<const FSampleTrack<FVector3f>>(MetadataTrack)->GetValueAtTimeIndex(InPlaybackTimeIndex, OutSampleTrackIndex)));
 									const FRotator& CurrentValue = static_cast<const FRigRotatorMetadata*>(Metadata)->GetValue();
-									if(StoredValue.EqualsOrientation(CurrentValue, static_cast<float>(Tolerance)))
+									if(!StoredValue.EqualsOrientation(CurrentValue, static_cast<float>(Tolerance)))
 									{
 										TrackNameBasedLogFunction(EMessageSeverity::Error, TrackName, TEXT("Metadata value doesn't match."));
 										bSuccess = false;
@@ -2055,7 +2051,7 @@ bool UControlRigReplay::ValidateExpectedResults(
 									{
 										for(int32 StoredValueIndex = 0; StoredValueIndex < StoredValue.Num(); StoredValueIndex++)
 										{
-											if(FRotator::MakeFromEuler(FVector(StoredValue[StoredValueIndex])).EqualsOrientation((CurrentValue[StoredValueIndex]), static_cast<float>(Tolerance)))
+											if(!FRotator::MakeFromEuler(FVector(StoredValue[StoredValueIndex])).EqualsOrientation((CurrentValue[StoredValueIndex]), static_cast<float>(Tolerance)))
 											{
 												TrackNameBasedLogFunction(EMessageSeverity::Error, TrackName, TEXT("Metadata array value element doesn't match."));
 												bSuccess = false;
@@ -2068,7 +2064,7 @@ bool UControlRigReplay::ValidateExpectedResults(
 								{
 									const FQuat4f& StoredValue = StaticCastSharedPtr<const FSampleTrack<FQuat4f>>(MetadataTrack)->GetValueAtTimeIndex(InPlaybackTimeIndex, OutSampleTrackIndex);
 									const FQuat4f& CurrentValue = FQuat4f(static_cast<const FRigQuatMetadata*>(Metadata)->GetValue()); 
-									if(StoredValue.Equals(CurrentValue, static_cast<float>(Tolerance)))
+									if(!StoredValue.Equals(CurrentValue, static_cast<float>(Tolerance)))
 									{
 										TrackNameBasedLogFunction(EMessageSeverity::Error, TrackName, TEXT("Metadata value doesn't match."));
 										bSuccess = false;
@@ -2088,7 +2084,7 @@ bool UControlRigReplay::ValidateExpectedResults(
 									{
 										for(int32 StoredValueIndex = 0; StoredValueIndex < StoredValue.Num(); StoredValueIndex++)
 										{
-											if(StoredValue[StoredValueIndex].Equals(FQuat4f(CurrentValue[StoredValueIndex]), static_cast<float>(Tolerance)))
+											if(!StoredValue[StoredValueIndex].Equals(FQuat4f(CurrentValue[StoredValueIndex]), static_cast<float>(Tolerance)))
 											{
 												TrackNameBasedLogFunction(EMessageSeverity::Error, TrackName, TEXT("Metadata array value element doesn't match."));
 												bSuccess = false;
@@ -2101,17 +2097,17 @@ bool UControlRigReplay::ValidateExpectedResults(
 								{
 									const FTransform3f& StoredValue = StaticCastSharedPtr<const FSampleTrack<FTransform3f>>(MetadataTrack)->GetValueAtTimeIndex(InPlaybackTimeIndex, OutSampleTrackIndex);
 									const FTransform3f& CurrentValue = FTransform3f(static_cast<const FRigTransformMetadata*>(Metadata)->GetValue()); 
-									if(StoredValue.GetLocation().Equals(CurrentValue.GetLocation(), static_cast<float>(Tolerance)))
+									if(!StoredValue.GetLocation().Equals(CurrentValue.GetLocation(), static_cast<float>(Tolerance)))
 									{
 										TrackNameBasedLogFunction(EMessageSeverity::Error, TrackName, TEXT("Metadata location value doesn't match."));
 										bSuccess = false;
 									}
-									if(StoredValue.Rotator().EqualsOrientation(CurrentValue.Rotator(), static_cast<float>(Tolerance)))
+									if(!StoredValue.Rotator().EqualsOrientation(CurrentValue.Rotator(), static_cast<float>(Tolerance)))
 									{
 										TrackNameBasedLogFunction(EMessageSeverity::Error, TrackName, TEXT("Metadata rotation value doesn't match."));
 										bSuccess = false;
 									}
-									if(StoredValue.GetScale3D().Equals(CurrentValue.GetScale3D(), static_cast<float>(Tolerance)))
+									if(!StoredValue.GetScale3D().Equals(CurrentValue.GetScale3D(), static_cast<float>(Tolerance)))
 									{
 										TrackNameBasedLogFunction(EMessageSeverity::Error, TrackName, TEXT("Metadata scale value doesn't match."));
 										bSuccess = false;
@@ -2132,17 +2128,17 @@ bool UControlRigReplay::ValidateExpectedResults(
 										for(int32 StoredValueIndex = 0; StoredValueIndex < StoredValue.Num(); StoredValueIndex++)
 										{
 											const FTransform3f CurrentTransform(CurrentValue[StoredValueIndex]);
-											if(StoredValue[StoredValueIndex].GetLocation().Equals(CurrentTransform.GetLocation(), static_cast<float>(Tolerance)))
+											if(!StoredValue[StoredValueIndex].GetLocation().Equals(CurrentTransform.GetLocation(), static_cast<float>(Tolerance)))
 											{
 												TrackNameBasedLogFunction(EMessageSeverity::Error, TrackName, TEXT("Metadata array locationvalue element doesn't match."));
 												bSuccess = false;
 											}
-											if(StoredValue[StoredValueIndex].Rotator().EqualsOrientation(CurrentTransform.Rotator(), static_cast<float>(Tolerance)))
+											if(!StoredValue[StoredValueIndex].Rotator().EqualsOrientation(CurrentTransform.Rotator(), static_cast<float>(Tolerance)))
 											{
 												TrackNameBasedLogFunction(EMessageSeverity::Error, TrackName, TEXT("Metadata array rotation value element doesn't match."));
 												bSuccess = false;
 											}
-											if(StoredValue[StoredValueIndex].GetScale3D().Equals(CurrentTransform.GetScale3D(), static_cast<float>(Tolerance)))
+											if(!StoredValue[StoredValueIndex].GetScale3D().Equals(CurrentTransform.GetScale3D(), static_cast<float>(Tolerance)))
 											{
 												TrackNameBasedLogFunction(EMessageSeverity::Error, TrackName, TEXT("Metadata array scale value element doesn't match."));
 												bSuccess = false;
@@ -2155,7 +2151,7 @@ bool UControlRigReplay::ValidateExpectedResults(
 								{
 									const FLinearColor& StoredValue = StaticCastSharedPtr<const FSampleTrack<FLinearColor>>(MetadataTrack)->GetValueAtTimeIndex(InPlaybackTimeIndex, OutSampleTrackIndex);
 									const FLinearColor& CurrentValue = static_cast<const FRigLinearColorMetadata*>(Metadata)->GetValue(); 
-									if(StoredValue.Equals(CurrentValue, static_cast<float>(Tolerance)))
+									if(!StoredValue.Equals(CurrentValue, static_cast<float>(Tolerance)))
 									{
 										TrackNameBasedLogFunction(EMessageSeverity::Error, TrackName, TEXT("Metadata value doesn't match."));
 										bSuccess = false;
@@ -2175,7 +2171,7 @@ bool UControlRigReplay::ValidateExpectedResults(
 									{
 										for(int32 StoredValueIndex = 0; StoredValueIndex < StoredValue.Num(); StoredValueIndex++)
 										{
-											if(StoredValue[StoredValueIndex].Equals(CurrentValue[StoredValueIndex], static_cast<float>(Tolerance)))
+											if(!StoredValue[StoredValueIndex].Equals(CurrentValue[StoredValueIndex], static_cast<float>(Tolerance)))
 											{
 												TrackNameBasedLogFunction(EMessageSeverity::Error, TrackName, TEXT("Metadata array value element doesn't match."));
 												bSuccess = false;
@@ -2354,3 +2350,49 @@ void UControlRigReplay::ClearDelegates(UControlRig* InControlRig)
 		}
 	}
 }
+
+void UControlRigReplay::HandlePreconstructionForTest(UControlRig* InRig, const FName& InEventName) const
+{
+	InRig->OnPreConstruction_AnyThread().RemoveAll(this);
+	
+	if(InRig->IsRigModule())
+	{
+		if(Cast<USkeletalMesh>(PreviewSkeletalMeshObjectPath.TryLoad()))
+		{
+			if(USkeletalMesh* PreviewSkeletalMesh = Cast<USkeletalMesh>(PreviewSkeletalMeshObjectPath.TryLoad()))
+			{
+				if(URigHierarchy* Hierarchy = InRig->GetHierarchy())
+				{
+					if(URigHierarchyController* Controller = Hierarchy->GetController(true))
+					{
+						const TArray<FRigSocketState> SocketStates = InRig->GetHierarchy()->GetSocketStates();
+						Controller->ImportPreviewSkeletalMesh(PreviewSkeletalMesh, false, false, false, false);
+						InRig->GetHierarchy()->RestoreSocketsFromStates(SocketStates);
+					}
+				}
+			}
+		}
+	}
+
+	// restore the connectors
+	const TArray<FRigElementKey> ConnectorKeys = InRig->GetHierarchy()->GetConnectorKeys();
+	FRigElementKeyRedirector::FKeyMap ConnectorMap;
+	for(const FRigElementKey& ConnectorKey : ConnectorKeys)
+	{
+		const FName& TrackName = InputTracks.GetTrackName(ConnectorKey);
+		const FName& ConnectorTrackName = *(TrackName.ToString() + TEXT("ConnectorTargets"));
+
+		if(const TSharedPtr<const FSampleTrack<TArray<FRigElementKey>>> Track = InputTracks.GetContainer()->FindTrack<TArray<FRigElementKey>>(ConnectorTrackName))
+		{
+			FRigElementKeyRedirector::FKeyArray Targets;
+			Targets.Append(Track->GetValueAtTimeIndex(0, InputTracks.SampleTrackIndex));
+			ConnectorMap.Add(ConnectorKey, Targets);
+		}
+	}
+	if(!ConnectorMap.IsEmpty())
+	{
+		InRig->SetElementKeyRedirector(FRigElementKeyRedirector(ConnectorMap, InRig->GetHierarchy()));
+		InRig->GetHierarchy()->ElementKeyRedirector = &InRig->GetElementKeyRedirector(); 
+	}
+
+}
diff --git a/Engine/Plugins/Animation/ControlRig/Source/ControlRig/Private/Rigs/RigHierarchyController.cpp b/Engine/Plugins/Animation/ControlRig/Source/ControlRig/Private/Rigs/RigHierarchyController.cpp
--- a/Engine/Plugins/Animation/ControlRig/Source/ControlRig/Private/Rigs/RigHierarchyController.cpp
+++ b/Engine/Plugins/Animation/ControlRig/Source/ControlRig/Private/Rigs/RigHierarchyController.cpp
@@ -7,6 +7,7 @@
 #include "ModularRig.h"
 #include "HelperUtil.h"
 #include "Engine/SkeletalMesh.h"
+#include "Units/Execution/RigUnit_PrepareForExecution.h"
 
 #if WITH_EDITOR
 #include "Framework/Notifications/NotificationManager.h"
@@ -1256,6 +1257,51 @@ TArray<FRigElementKey> URigHierarchyController::ImportCurvesFromAsset(FString In
 	return TArray<FRigElementKey>();
 }
 
+TArray<FRigElementKey> URigHierarchyController::ImportPreviewSkeletalMesh(USkeletalMesh* InSkeletalMesh, bool bReplaceExistingBones,
+	bool bRemoveObsoleteBones, bool bSelectBones, bool bSetupUndo)
+{
+	URigHierarchy* Hierarchy = GetHierarchy();
+	check(Hierarchy);
+	
+	// find the instruction index for the construction event
+	int32 InstructionIndex = INDEX_NONE;
+	if(UControlRig* ControlRig = Cast<UControlRig>(Hierarchy->GetOuter()))
+	{
+		if(URigVM* VM = ControlRig->GetVM())
+		{
+			const int32 EntryIndex = VM->GetByteCode().FindEntryIndex(FRigUnit_PrepareForExecution::EventName);
+			if(EntryIndex != INDEX_NONE)
+			{
+				InstructionIndex = VM->GetByteCode().GetEntry(EntryIndex).InstructionIndex;
+			}
+		}
+	}
+
+	// import the bones for the preview hierarchy
+	// use the ref skeleton so we'll only see the bones that are actually part of the mesh
+	const TArray<FRigElementKey> Bones = ImportBones(InSkeletalMesh->GetRefSkeleton(), NAME_None, bReplaceExistingBones, bRemoveObsoleteBones, bSelectBones, bSetupUndo);
+	for(const FRigElementKey& Bone : Bones)
+	{
+		if(FRigBaseElement* Element = Hierarchy->Find(Bone))
+		{
+			Element->CreatedAtInstructionIndex = InstructionIndex;
+		}
+	}
+
+	// create a null to store controls under
+	static const FRigElementKey ControlParentKey(TEXT("Controls"), ERigElementType::Null);
+	if(!Hierarchy->Contains(ControlParentKey))
+	{
+		const FRigElementKey Null = AddNull(ControlParentKey.Name, FRigElementKey(), FTransform::Identity, true, false, false);
+		if(FRigBaseElement* Element = Hierarchy->Find(Null))
+		{
+			Element->CreatedAtInstructionIndex = InstructionIndex;
+		}
+	}
+
+	return Bones;
+}
+
 USkeletalMesh* URigHierarchyController::GetSkeletalMeshFromAssetPath(const FString& InAssetPath)
 {
 	UObject* AssetObject = StaticLoadObject(UObject::StaticClass(), NULL, *InAssetPath, NULL, LOAD_None, NULL);
diff --git a/Engine/Plugins/Animation/ControlRig/Source/ControlRig/Private/Tracks/SampleTrackBase.cpp b/Engine/Plugins/Animation/ControlRig/Source/ControlRig/Private/Tracks/SampleTrackBase.cpp
--- a/Engine/Plugins/Animation/ControlRig/Source/ControlRig/Private/Tracks/SampleTrackBase.cpp
+++ b/Engine/Plugins/Animation/ControlRig/Source/ControlRig/Private/Tracks/SampleTrackBase.cpp
@@ -216,6 +216,10 @@ bool FSampleTrackBase::IsIdentical(const FSampleTrackBase* InOther, float InTole
 	{
 		return false;
 	}
+	if(GetScriptStruct() != InOther->GetScriptStruct())
+	{
+		return false;
+	}
 	if(GetMode() != InOther->GetMode())
 	{
 		return false;
diff --git a/Engine/Plugins/Animation/ControlRig/Source/ControlRig/Public/ControlRigReplay.h b/Engine/Plugins/Animation/ControlRig/Source/ControlRig/Public/ControlRigReplay.h
--- a/Engine/Plugins/Animation/ControlRig/Source/ControlRig/Public/ControlRigReplay.h
+++ b/Engine/Plugins/Animation/ControlRig/Source/ControlRig/Public/ControlRigReplay.h
@@ -206,6 +206,9 @@ public:
 	UPROPERTY(AssetRegistrySearchable, VisibleAnywhere, BlueprintReadOnly, Category = "ControlRigReplay")
 	FSoftObjectPath ControlRigObjectPath;
 
+	UPROPERTY(AssetRegistrySearchable, VisibleAnywhere, BlueprintReadOnly, Category = "ControlRigReplay")
+	FSoftObjectPath PreviewSkeletalMeshObjectPath;
+
 	UPROPERTY(EditAnywhere, Category = "ControlRigReplay")
 	FControlRigReplayTracks InputTracks;
 
@@ -272,6 +275,8 @@ protected:
 
 	void ClearDelegates(UControlRig* InControlRig);
 
+	void HandlePreconstructionForTest(UControlRig* InRig, const FName& InEventName) const;
+
 	TWeakObjectPtr<UControlRig> ReplayControlRig;
 	TWeakObjectPtr<UControlRig> RecordControlRig;
 	FDelegateHandle PreConstructionHandle;
diff --git a/Engine/Plugins/Animation/ControlRig/Source/ControlRig/Public/Rigs/RigHierarchy.h b/Engine/Plugins/Animation/ControlRig/Source/ControlRig/Public/Rigs/RigHierarchy.h
--- a/Engine/Plugins/Animation/ControlRig/Source/ControlRig/Public/Rigs/RigHierarchy.h
+++ b/Engine/Plugins/Animation/ControlRig/Source/ControlRig/Public/Rigs/RigHierarchy.h
@@ -4489,6 +4489,7 @@ private:
 		}
 	};
 	friend class UControlRigTestData;
+	friend class UControlRigReplay;
 	friend struct FControlRigTestDataFrame;
 
 	// Managed lookup from Element Key to Index
diff --git a/Engine/Plugins/Animation/ControlRig/Source/ControlRig/Public/Rigs/RigHierarchyController.h b/Engine/Plugins/Animation/ControlRig/Source/ControlRig/Public/Rigs/RigHierarchyController.h
--- a/Engine/Plugins/Animation/ControlRig/Source/ControlRig/Public/Rigs/RigHierarchyController.h
+++ b/Engine/Plugins/Animation/ControlRig/Source/ControlRig/Public/Rigs/RigHierarchyController.h
@@ -417,6 +417,21 @@ public:
         bool bSetupUndo = false);
 #endif
 
+	/**
+	 * Imports all bones from a preview skeletal mesh. Used for rig modules and their preview skeleton
+	 * @param bReplaceExistingBones If true existing bones will be removed
+	 * @param bRemoveObsoleteBones If true bones non-existent in the skeleton will be removed from the hierarchy
+	 * @param bSelectBones If true the bones will be selected upon import
+	 * @param bSetupUndo If set to true the stack will record the change for undo / redo
+	 * @return The keys of the imported elements
+	 */
+	UFUNCTION(BlueprintCallable, Category = URigHierarchyController)
+	TArray<FRigElementKey> ImportPreviewSkeletalMesh(USkeletalMesh* InSkeletalMesh,
+		bool bReplaceExistingBones = true,
+		bool bRemoveObsoleteBones = true,
+		bool bSelectBones = false,
+		bool bSetupUndo = false);
+
 	/**
 	 * Exports the selected items to text
 	 * @return The text representation of the selected items
diff --git a/Engine/Plugins/Animation/ControlRig/Source/ControlRigEditor/Private/Editor/ControlRigEditor.cpp b/Engine/Plugins/Animation/ControlRig/Source/ControlRigEditor/Private/Editor/ControlRigEditor.cpp
--- a/Engine/Plugins/Animation/ControlRig/Source/ControlRigEditor/Private/Editor/ControlRigEditor.cpp
+++ b/Engine/Plugins/Animation/ControlRig/Source/ControlRigEditor/Private/Editor/ControlRigEditor.cpp
@@ -763,34 +763,44 @@ bool FControlRigEditor::RecordReplay(double InRecordingDuration)
 	static const FString Folder = TEXT("/Game/Animation/ControlRig/NoCook/");
 	const FString DesiredPackagePath =  FString::Printf(TEXT("%s/%s_Replay"), *Folder, *GetControlRigBlueprint()->GetName());
 
-	if(const UControlRigReplay* Replay = UControlRigReplay::CreateNewAsset(DesiredPackagePath, GetControlRigBlueprint()->GetPathName(), UControlRigReplay::StaticClass()))
+	if(UControlRigReplay* Replay = UControlRigReplay::CreateNewAsset(DesiredPackagePath, GetControlRigBlueprint()->GetPathName(), UControlRigReplay::StaticClass()))
 	{
 		SetReplayAssetPath(Replay->GetPathName());
+		if(const USkeletalMesh* PreviewSkeletalMesh = GetControlRigBlueprint()->GetPreviewMesh())
+		{
+			Replay->PreviewSkeletalMeshObjectPath = FSoftObjectPath(PreviewSkeletalMesh);
+		}
 	}
 
-
 	if(UControlRigReplay* Replay = ReplayStrongPtr.Get())
 	{
 		Replay->DesiredRecordingDuration = InRecordingDuration;
 
-		// 3 second preroll
-		TSharedPtr<int32> TimeLeft = MakeShared<int32>(4);
-		GEditor->GetTimerManager()->SetTimer(RecordReplayTimerHandle, FTimerDelegate::CreateLambda([this, Replay, TimeLeft]()
+		if(InRecordingDuration <= SMALL_NUMBER)
 		{
-			int32& SecondsLeft = *TimeLeft.Get();
-			SecondsLeft--;
-
-			if(SecondsLeft == 0)
-			{
-				Replay->StartRecording(GetControlRig());
-				GEditor->GetTimerManager()->ClearTimer(RecordReplayTimerHandle);
-			}
-			else
+			Replay->StartRecording(GetControlRig());
+		}
+		else
+		{
+			// 3 second preroll
+			TSharedPtr<int32> TimeLeft = MakeShared<int32>(4);
+			GEditor->GetTimerManager()->SetTimer(RecordReplayTimerHandle, FTimerDelegate::CreateLambda([this, Replay, TimeLeft]()
 			{
-				static constexpr TCHAR PrintStringFormat[] = TEXT("Recording starts in... %d");
-				UKismetSystemLibrary::PrintString(GetPreviewScene()->GetWorld(), FString::Printf(PrintStringFormat, SecondsLeft), true, false, FLinearColor::Red, 1.f);
-			}
-		}), 1.f, true, 0.01f);
+				int32& SecondsLeft = *TimeLeft.Get();
+				SecondsLeft--;
+
+				if(SecondsLeft == 0)
+				{
+					Replay->StartRecording(GetControlRig());
+					GEditor->GetTimerManager()->ClearTimer(RecordReplayTimerHandle);
+				}
+				else
+				{
+					static constexpr TCHAR PrintStringFormat[] = TEXT("Recording starts in... %d");
+					UKismetSystemLibrary::PrintString(GetPreviewScene()->GetWorld(), FString::Printf(PrintStringFormat, SecondsLeft), true, false, FLinearColor::Red, 1.f);
+				}
+			}), 1.f, true, 0.01f);
+		}
 	}
 	return true;
 }
@@ -5812,44 +5822,13 @@ void FControlRigEditor::OnPreConstruction_AnyThread(UControlRig* InRig, const FN
 	{
 		if(RigBlueprint->IsControlRigModule())
 		{
-			if(RigBlueprint->PreviewSkeletalMesh)
+			if(USkeletalMesh* PreviewSkeletalMesh = RigBlueprint->PreviewSkeletalMesh.Get())
 			{
 				if(URigHierarchy* Hierarchy = InRig->GetHierarchy())
 				{
 					if(URigHierarchyController* Controller = Hierarchy->GetController(true))
 					{
-						// find the instruction index for the construction event
-						int32 InstructionIndex = INDEX_NONE;
-						if(URigVM* VM = InRig->GetVM())
-						{
-							const int32 EntryIndex = VM->GetByteCode().FindEntryIndex(FRigUnit_PrepareForExecution::EventName);
-							if(EntryIndex != INDEX_NONE)
-							{
-								InstructionIndex = VM->GetByteCode().GetEntry(EntryIndex).InstructionIndex;
-							}
-						}
-
-						// import the bones for the preview hierarchy
-						// use the ref skeleton so we'll only see the bones that are actually part of the mesh
-						const TArray<FRigElementKey> Bones = Controller->ImportBones(RigBlueprint->PreviewSkeletalMesh->GetRefSkeleton(), NAME_None, false, false, false, false);
-						for(const FRigElementKey& Bone : Bones)
-						{
-							if(FRigBaseElement* Element = Hierarchy->Find(Bone))
-							{
-								Element->CreatedAtInstructionIndex = InstructionIndex;
-							}
-						}
-
-						// create a null to store controls under
-						static const FRigElementKey ControlParentKey(TEXT("Controls"), ERigElementType::Null);
-						if(!Hierarchy->Contains(ControlParentKey))
-						{
-							const FRigElementKey Null = Controller->AddNull(ControlParentKey.Name, FRigElementKey(), FTransform::Identity, true, false, false);
-							if(FRigBaseElement* Element = Hierarchy->Find(Null))
-							{
-								Element->CreatedAtInstructionIndex = InstructionIndex;
-							}
-						}
+						Controller->ImportPreviewSkeletalMesh(PreviewSkeletalMesh, false, false, false, false);
 					}
 				}
 
