Control Rig: Fix ldd error - move WITH_EDITOR

#rnx

diff --git a/Engine/Plugins/Animation/ControlRig/Source/ControlRig/Private/Rigs/RigHierarchyController.cpp b/Engine/Plugins/Animation/ControlRig/Source/ControlRig/Private/Rigs/RigHierarchyController.cpp
--- a/Engine/Plugins/Animation/ControlRig/Source/ControlRig/Private/Rigs/RigHierarchyController.cpp
+++ b/Engine/Plugins/Animation/ControlRig/Source/ControlRig/Private/Rigs/RigHierarchyController.cpp
@@ -1257,6 +1257,8 @@ TArray<FRigElementKey> URigHierarchyController::ImportCurvesFromAsset(FString In
 	return TArray<FRigElementKey>();
 }
 
+#endif
+
 TArray<FRigElementKey> URigHierarchyController::ImportPreviewSkeletalMesh(USkeletalMesh* InSkeletalMesh, bool bReplaceExistingBones,
 	bool bRemoveObsoleteBones, bool bSelectBones, bool bSetupUndo)
 {
@@ -1302,6 +1304,8 @@ TArray<FRigElementKey> URigHierarchyController::ImportPreviewSkeletalMesh(USkele
 	return Bones;
 }
 
+#if WITH_EDITOR
+
 USkeletalMesh* URigHierarchyController::GetSkeletalMeshFromAssetPath(const FString& InAssetPath)
 {
 	UObject* AssetObject = StaticLoadObject(UObject::StaticClass(), NULL, *InAssetPath, NULL, LOAD_None, NULL);
