Control Rig: Avoid drifting when manipulating rig in a layered control rig.
Reapplying fixes from 34128760, 34201378, 34375090 and 34809364, which were backed out in 36297026, and in addition, fixing the fact that some Shape Actors did not have an up-to-date transform when manipulating multiple controls at the same time.

#rnx

diff --git a/Engine/Plugins/Animation/ControlRig/Source/ControlRigEditor/Private/EditMode/ControlRigEditMode.cpp b/Engine/Plugins/Animation/ControlRig/Source/ControlRigEditor/Private/EditMode/ControlRigEditMode.cpp
--- a/Engine/Plugins/Animation/ControlRig/Source/ControlRigEditor/Private/EditMode/ControlRigEditMode.cpp
+++ b/Engine/Plugins/Animation/ControlRig/Source/ControlRigEditor/Private/EditMode/ControlRigEditMode.cpp
@@ -5060,8 +5060,9 @@ void FControlRigEditMode::MoveControlShape(AControlRigShapeActor* ShapeActor, co
 			// different if we are interacting on a control on an animated character
 			CurrentTransform = CurrentTransform.GetRelativeTransform(ShapeActor->OffsetTransform);
 			CurrentTransform = CurrentTransform * ControlRig->GetHierarchy()->GetGlobalControlOffsetTransform(ShapeActor->GetElementKey(), false);
-			
-			ShapeActor->SetGlobalTransform(CurrentTransform);
+
+			// Don't set the global transform to the shape actor to avoid drifting
+			//ShapeActor->SetGlobalTransform(CurrentTransform);
 
 			RigsToEvaluateDuringThisTick.AddUnique(ControlRig);
 			return;
@@ -5069,8 +5070,27 @@ void FControlRigEditMode::MoveControlShape(AControlRigShapeActor* ShapeActor, co
 	}
 
 	//not local or doing scale.
-	FTransform CurrentTransform = GetControlShapeTransform(ShapeActor) * ToWorldTransform;
 
+	// Get the global transform from shape actor to avoid drifting
+	FTransform CurrentTransform;
+	if (ShapeActor->StaticMeshComponent && ShapeActor->StaticMeshComponent->GetStaticMesh())
+	{
+		// Update ShapeActor in case we are moving multiple shapes, which affect one another
+		FRigControlElement* ControlElement = ControlRig->FindControl(ShapeActor->ControlName);
+		const FTransform Transform = ControlRig->GetHierarchy()->GetTransform(ControlElement, ERigTransformType::CurrentGlobal);
+		ShapeActor->SetActorTransform(Transform * ToWorldTransform);
+		
+		const FTransform& ParentTransform = ShapeActor->GetAttachParentActor()->GetTransform();
+		CurrentTransform = ShapeActor->GetGlobalTransform() * ParentTransform;
+	}
+	else
+	{
+		// If the static mesh is not valid, we cannot rely on the shape's transform.
+		// This happens for FKControlRigs (and other control types)
+		// We will need to rely on the information we have in the rig hierarchy
+		CurrentTransform = GetControlShapeTransform(ShapeActor) * ToWorldTransform;
+	}
+	
 	if (bRotation)
 	{
 		FQuat CurrentRotation = CurrentTransform.GetRotation();
@@ -5162,8 +5182,9 @@ void FControlRigEditMode::MoveControlShape(AControlRigShapeActor* ShapeActor, co
 		};
 
 		const TFunction<void()> SetGlobalAndUpdateLocalTask = [ShapeActor, CurrentTransform, ControlRig, bCalcLocal, InOutLocal]() {
-			
-			ShapeActor->SetGlobalTransform(CurrentTransform);
+
+			// Don't set the global transform to the shape actor to avoid drifting
+			//ShapeActor->SetGlobalTransform(CurrentTransform);
 
 			if (bCalcLocal && InOutLocal)
 			{
