ControlRig Anim Node: activate optimizations based on Ferris feedback

- ControlRig.AnimNode.UseFunctions true by default
- ControlRig.EnableAnimNodePerformanceOptimizations 1 by default

#rnx

diff --git a/Engine/Plugins/Animation/ControlRig/Source/ControlRig/Private/Tools/ControlRigHierarchyMappings.cpp b/Engine/Plugins/Animation/ControlRig/Source/ControlRig/Private/Tools/ControlRigHierarchyMappings.cpp
--- a/Engine/Plugins/Animation/ControlRig/Source/ControlRig/Private/Tools/ControlRigHierarchyMappings.cpp
+++ b/Engine/Plugins/Animation/ControlRig/Source/ControlRig/Private/Tools/ControlRigHierarchyMappings.cpp
@@ -14,7 +14,10 @@
 #include "Components/SkeletalMeshComponent.h"
 
 // CVar to enable performance optimizations within the anim node pose exchange
-TAutoConsoleVariable<int32> CVarControlRigEnableAnimNodePerformanceOptimizations(TEXT("ControlRig.EnableAnimNodePerformanceOptimizations"), 0, TEXT("if nonzero we enable the (experimental) execution performance optimizations of Control Rig AnimNodes."));
+TAutoConsoleVariable<int32> CVarControlRigEnableAnimNodePerformanceOptimizations(
+	TEXT("ControlRig.EnableAnimNodePerformanceOptimizations"),
+	1,
+	TEXT("if nonzero we enable the (experimental) execution performance optimizations of Control Rig AnimNodes."));
 
 
 void FControlRigHierarchyMappings::InitializeInstance()
diff --git a/Engine/Plugins/Animation/ControlRig/Source/ControlRig/Private/Tools/ControlRigVariableMappings.cpp b/Engine/Plugins/Animation/ControlRig/Source/ControlRig/Private/Tools/ControlRigVariableMappings.cpp
--- a/Engine/Plugins/Animation/ControlRig/Source/ControlRig/Private/Tools/ControlRigVariableMappings.cpp
+++ b/Engine/Plugins/Animation/ControlRig/Source/ControlRig/Private/Tools/ControlRigVariableMappings.cpp
@@ -11,7 +11,7 @@ namespace AnimNodeLocals
 
 const TCHAR* UseFunctionsName = TEXT("ControlRig.AnimNode.UseFunctions");
 
-static bool bUseFunctions = false;
+static bool bUseFunctions = true;
 static FAutoConsoleVariableRef CVarUseCachedFunctions(
    UseFunctionsName,
    bUseFunctions,
