Control Rig: Fix crash when using modular rigs in construction mode
#rnx

diff --git a/Engine/Plugins/Animation/ControlRig/Source/ControlRig/Private/ModularRig.cpp b/Engine/Plugins/Animation/ControlRig/Source/ControlRig/Private/ModularRig.cpp
--- a/Engine/Plugins/Animation/ControlRig/Source/ControlRig/Private/ModularRig.cpp
+++ b/Engine/Plugins/Animation/ControlRig/Source/ControlRig/Private/ModularRig.cpp
@@ -169,15 +169,16 @@ const FRigModuleInstance* FRigModuleInstance::GetRootModule() const
 
 const FRigConnectorElement* FRigModuleInstance::FindPrimaryConnector() const
 {
-	if(PrimaryConnector)
-	{
-		return PrimaryConnector;
-	}
-	
 	if(const UControlRig* Rig = GetRig())
 	{
 		if(const URigHierarchy* Hierarchy = Rig->GetHierarchy())
 		{
+			PrimaryConnector.UpdateCache(Hierarchy);
+			if(PrimaryConnector.IsValid())
+			{
+				return Cast<FRigConnectorElement>(PrimaryConnector.GetElement());
+			}
+	
 			const FString MyModulePath = GetPath();
 			const TArray<FRigConnectorElement*> AllConnectors = Hierarchy->GetConnectors();
 			for(const FRigConnectorElement* Connector : AllConnectors)
@@ -189,8 +190,8 @@ const FRigConnectorElement* FRigModuleInstance::FindPrimaryConnector() const
 					{
 						if(ModulePath.Equals(MyModulePath, ESearchCase::IgnoreCase))
 						{
-							PrimaryConnector = Connector;
-							return PrimaryConnector;
+							PrimaryConnector.UpdateCache(Connector->GetKey(), Hierarchy);
+							return Cast<FRigConnectorElement>(PrimaryConnector.GetElement());
 						}
 					}
 				}
@@ -636,6 +637,7 @@ void UModularRig::ResetModules(bool bDestroyModuleRigs)
 	for (FRigModuleInstance& Module : Modules)
 	{
 		Module.CachedChildren.Reset();
+		Module.PrimaryConnector.Reset();
 
 		if(bDestroyModuleRigs)
 		{
@@ -761,11 +763,19 @@ FRigModuleInstance* UModularRig::AddModuleInstance(const FName& InModuleName, TS
 		return nullptr;
 	}
 
+	// InParent ptr will become invalid once we add the new module in the Modules array
+	// Get the information we need before that.
+	FString ParentPath;
+	if (InParent)
+	{
+		ParentPath = InParent->GetPath();
+	}
+
 	FRigModuleInstance& NewModule = Modules.Add_GetRef(FRigModuleInstance());
 	NewModule.Name = InModuleName;
-	if (InParent)
+	if (!ParentPath.IsEmpty())
 	{
-		NewModule.ParentPath = InParent->GetPath();
+		NewModule.ParentPath = ParentPath;
 	}
 	const FString Name = NewModule.GetPath();
 
diff --git a/Engine/Plugins/Animation/ControlRig/Source/ControlRig/Public/ModularRig.h b/Engine/Plugins/Animation/ControlRig/Source/ControlRig/Public/ModularRig.h
--- a/Engine/Plugins/Animation/ControlRig/Source/ControlRig/Public/ModularRig.h
+++ b/Engine/Plugins/Animation/ControlRig/Source/ControlRig/Public/ModularRig.h
@@ -71,7 +71,7 @@ public:
 	TMap<FName, FRigVMExternalVariable> VariableBindings;
 
 	TArray<FRigModuleInstance*> CachedChildren;
-	mutable const FRigConnectorElement* PrimaryConnector = nullptr;
+	mutable FCachedRigElement PrimaryConnector;
 
 	FString GetShortName() const;
 	FString GetLongName() const
