ControlRig - Pose Adapter: fixed PoseIndexToElementIndex not being reset

#rnx

diff --git a/Engine/Plugins/Animation/ControlRig/Source/ControlRig/Private/Tools/ControlRigPoseAdapter.cpp b/Engine/Plugins/Animation/ControlRig/Source/ControlRig/Private/Tools/ControlRigPoseAdapter.cpp
--- a/Engine/Plugins/Animation/ControlRig/Source/ControlRig/Private/Tools/ControlRigPoseAdapter.cpp
+++ b/Engine/Plugins/Animation/ControlRig/Source/ControlRig/Private/Tools/ControlRigPoseAdapter.cpp
@@ -154,16 +154,17 @@ void FControlRigPoseAdapter::UpdateInputOutputMappingIfRequired(UControlRig* InC
 	ParentPoseIndices.Reset();
 	RequiresHierarchyForSpaceConversion.Reset();
 	ElementIndexToPoseIndex.Reset();
+	PoseIndexToElementIndex.Reset();
+	GlobalPose.Reset();
+	LocalPose.Reset();
+	HierarchyCurveLookup.Reset();
 
 	const int32 NumBonesInPose = InRequiredBones.GetCompactPoseNumBones();
 
-	GlobalPose.Reset();
-	LocalPose.Reset();
 	ParentPoseIndices.Reserve(NumBonesInPose);
 	RequiresHierarchyForSpaceConversion.Reserve(NumBonesInPose);
 	GlobalPose.AddDefaulted(NumBonesInPose);
 	LocalPose.AddDefaulted(NumBonesInPose);
-	HierarchyCurveLookup.Reset();
 
 	bTransferInLocalSpace = !(bInTransferPoseInGlobalSpace || InNodeMappingContainer.IsValid());
 
