Control Rig: Interaction bracket rig hierarhcy notifications have no payloads and so should not fail on queue.

#rnx

diff --git a/Engine/Plugins/Animation/ControlRig/Source/ControlRig/Private/Rigs/RigHierarchy.cpp b/Engine/Plugins/Animation/ControlRig/Source/ControlRig/Private/Rigs/RigHierarchy.cpp
--- a/Engine/Plugins/Animation/ControlRig/Source/ControlRig/Private/Rigs/RigHierarchy.cpp
+++ b/Engine/Plugins/Animation/ControlRig/Source/ControlRig/Private/Rigs/RigHierarchy.cpp
@@ -4003,10 +4003,6 @@ void URigHierarchy::QueueNotification(ERigHierarchyNotification InNotification,
 		Entry.Key = InSubject.Component->GetElementKey();
 		Entry.ComponentName = InSubject.Component->GetFName();
 	}
-	else
-	{
-		checkNoEntry();
-	}
 	QueuedNotifications.Enqueue(Entry);
 }
 
