Make the critical section in RigVMHost transactionally safe.

#rnx

diff --git a/Engine/Plugins/Animation/ControlRig/Source/ControlRig/Private/AnimNode_ControlRigBase.cpp b/Engine/Plugins/Animation/ControlRig/Source/ControlRig/Private/AnimNode_ControlRigBase.cpp
--- a/Engine/Plugins/Animation/ControlRig/Source/ControlRig/Private/AnimNode_ControlRigBase.cpp
+++ b/Engine/Plugins/Animation/ControlRig/Source/ControlRig/Private/AnimNode_ControlRigBase.cpp
@@ -13,7 +13,7 @@
 #include "Algo/Transform.h"
 #include "Animation/AnimCurveUtils.h"
 #include "Tools/ControlRigPoseAdapter.h"
-
+#include "Misc/ScopeLock.h"
 
 #include UE_INLINE_GENERATED_CPP_BY_NAME(AnimNode_ControlRigBase)
 
@@ -233,7 +233,7 @@ void FAnimNode_ControlRigBase::ExecuteControlRig(FPoseContext& InOutput)
 	{
 		// Before we start modifying the RigHierarchy, we need to lock the rig to avoid
 		// corrupted state
-		FScopeLock LockRig(&ControlRig->GetEvaluateMutex());
+		UE::TScopeLock LockRig(ControlRig->GetEvaluateMutex());
 		
 		// temporarily give control rig access to the stack allocated attribute container
 		// control rig may have rig units that can add/get attributes to/from this container
diff --git a/Engine/Plugins/Animation/ControlRig/Source/ControlRig/Private/ControlRig.cpp b/Engine/Plugins/Animation/ControlRig/Source/ControlRig/Private/ControlRig.cpp
--- a/Engine/Plugins/Animation/ControlRig/Source/ControlRig/Private/ControlRig.cpp
+++ b/Engine/Plugins/Animation/ControlRig/Source/ControlRig/Private/ControlRig.cpp
@@ -4,6 +4,7 @@
 #include "Engine/SkeletalMesh.h"
 #include "GameFramework/Actor.h"
 #include "Misc/RuntimeErrors.h"
+#include "Misc/ScopeLock.h"
 #include "IControlRigObjectBinding.h"
 #include "HelperUtil.h"
 #include "ObjectTrace.h"
@@ -384,7 +385,7 @@ void UControlRig::Evaluate_AnyThread()
 		// intermediate stage of evaluation, for example, during evaluate, we can have a call
 		// to copy hierarchy, which empties the hierarchy for a short period of time
 		// and we don't want other systems to see that.
-		FScopeLock EvaluateLock(&GetEvaluateMutex());
+		UE::TScopeLock EvaluateLock(GetEvaluateMutex());
 
 		URigHierarchy* Hierarchy = GetHierarchy();
 
diff --git a/Engine/Plugins/Animation/ControlRig/Source/ControlRig/Private/Sequencer/MovieSceneControlRigParameterTemplate.cpp b/Engine/Plugins/Animation/ControlRig/Source/ControlRig/Private/Sequencer/MovieSceneControlRigParameterTemplate.cpp
--- a/Engine/Plugins/Animation/ControlRig/Source/ControlRig/Private/Sequencer/MovieSceneControlRigParameterTemplate.cpp
+++ b/Engine/Plugins/Animation/ControlRig/Source/ControlRig/Private/Sequencer/MovieSceneControlRigParameterTemplate.cpp
@@ -26,6 +26,7 @@
 #include "TransformableHandle.h"
 #include "AnimationCoreLibrary.h"
 #include "Constraints/ControlRigTransformableHandle.h"
+#include "Misc/ScopeLock.h"
 
 #include UE_INLINE_GENERATED_CPP_BY_NAME(MovieSceneControlRigParameterTemplate)
 
@@ -782,7 +783,7 @@ struct FControlRigParameterPreAnimatedTokenProducer : IMovieScenePreAnimatedToke
 							// at the same time, which can lead to various issues like double-freeing some random array when doing SetControlValue
 							// note: the critical section accepts recursive locking so it is ok that we
 							// are calling evaluate_anythread later within the same scope
-							FScopeLock EvaluateLock(&ControlRig->GetEvaluateMutex());
+							UE::TScopeLock EvaluateLock(ControlRig->GetEvaluateMutex());
 						
 							//Restore control rig first
 							const bool bSetupUndo = false;
diff --git a/Engine/Plugins/Animation/ControlRig/Source/ControlRigEditor/Private/Editor/ControlRigSkeletalMeshComponent.cpp b/Engine/Plugins/Animation/ControlRig/Source/ControlRigEditor/Private/Editor/ControlRigSkeletalMeshComponent.cpp
--- a/Engine/Plugins/Animation/ControlRig/Source/ControlRigEditor/Private/Editor/ControlRigSkeletalMeshComponent.cpp
+++ b/Engine/Plugins/Animation/ControlRig/Source/ControlRigEditor/Private/Editor/ControlRigSkeletalMeshComponent.cpp
@@ -5,6 +5,7 @@
 #include "SkeletalDebugRendering.h"
 #include "ControlRig.h"
 #include "AnimPreviewInstance.h"
+#include "Misc/ScopeLock.h"
 
 #include UE_INLINE_GENERATED_CPP_BY_NAME(ControlRigSkeletalMeshComponent)
 
@@ -87,7 +88,7 @@ void UControlRigSkeletalMeshComponent::RebuildDebugDrawSkeleton()
 			// intermediate stage of evaluation, for example, during evaluate, we can have a call
 			// to copy hierarchy, which empties the hierarchy for a short period of time.
 			// if we did not have this lock and try to grab it, we could get an empty bone array
-			FScopeLock EvaluateLock(&ControlRig->GetEvaluateMutex());
+			UE::TScopeLock EvaluateLock(ControlRig->GetEvaluateMutex());
 			
 			// just copy it because this is not thread safe
 			URigHierarchy* Hierarchy = ControlRig->GetHierarchy();
diff --git a/Engine/Plugins/Runtime/RigVM/Source/RigVM/Private/RigVMHost.cpp b/Engine/Plugins/Runtime/RigVM/Source/RigVM/Private/RigVMHost.cpp
--- a/Engine/Plugins/Runtime/RigVM/Source/RigVM/Private/RigVMHost.cpp
+++ b/Engine/Plugins/Runtime/RigVM/Source/RigVM/Private/RigVMHost.cpp
@@ -11,6 +11,7 @@
 #include "UObject/UObjectIterator.h"
 #include "UObject/ObjectSaveContext.h"
 #include "SceneView.h"
+#include "Misc/ScopeLock.h"
 
 #if WITH_EDITOR
 #include "Kismet2/BlueprintEditorUtils.h"
@@ -335,7 +336,7 @@ void URigVMHost::Evaluate_AnyThread()
 	// intermediate stage of evaluation, for example, during evaluate, we can have a call
 	// to copy hierarchy, which empties the hierarchy for a short period of time
 	// and we don't want other systems to see that.
-	FScopeLock EvaluateLock(&GetEvaluateMutex());
+	UE::TScopeLock EvaluateLock(GetEvaluateMutex());
 	
 	// The EventQueueToRun should only be modified in this function
 	ensureMsgf(EventQueueToRun.IsEmpty(), TEXT("Detected a recursive call to the control rig evaluation function %s"), *GetPackage()->GetPathName());
diff --git a/Engine/Plugins/Runtime/RigVM/Source/RigVM/Public/RigVMHost.h b/Engine/Plugins/Runtime/RigVM/Source/RigVM/Public/RigVMHost.h
--- a/Engine/Plugins/Runtime/RigVM/Source/RigVM/Public/RigVMHost.h
+++ b/Engine/Plugins/Runtime/RigVM/Source/RigVM/Public/RigVMHost.h
@@ -9,6 +9,7 @@
 #include "RigVMCore/RigVMGraphFunctionHost.h"
 #include "Engine/AssetUserData.h"
 #include "Interfaces/Interface_AssetUserData.h"
+#include "Misc/TransactionallySafeCriticalSection.h"
 #if WITH_EDITOR
 #include "RigVMCore/RigVMDebugInfo.h"
 #endif
@@ -83,7 +84,7 @@ public:
 	virtual void Evaluate_AnyThread();
 
 	/** Locks for the scope of Evaluate_AnyThread */
-	FCriticalSection& GetEvaluateMutex() { return EvaluateMutex; };
+	FTransactionallySafeCriticalSection& GetEvaluateMutex() { return EvaluateMutex; };
 
 	/** Returns the member properties as an external variable array */
 	TArray<FRigVMExternalVariable> GetExternalVariables() const;
@@ -529,7 +530,7 @@ public:
 
 private:
 
-	FCriticalSection EvaluateMutex;
+	FTransactionallySafeCriticalSection EvaluateMutex;
 
 protected:
 
