Get Control Offset and Set Control Offset rigUnits

diff --git a/Engine/Plugins/Animation/ControlRig/Source/ControlRig/Private/Units/Hierarchy/RigUnit_GetControlOffset.cpp b/Engine/Plugins/Animation/ControlRig/Source/ControlRig/Private/Units/Hierarchy/RigUnit_GetControlOffset.cpp
new file mode 100644
--- /dev/null
+++ b/Engine/Plugins/Animation/ControlRig/Source/ControlRig/Private/Units/Hierarchy/RigUnit_GetControlOffset.cpp
@@ -0,0 +1,39 @@
+// Copyright Epic Games, Inc. All Rights Reserved.
+
+#include "Units/Hierarchy/RigUnit_GetControlOffset.h"
+#include "Units/RigUnitContext.h"
+
+#include UE_INLINE_GENERATED_CPP_BY_NAME(RigUnit_GetControlOffset)
+
+FRigUnit_GetControlOffset_Execute()
+{
+	DECLARE_SCOPE_HIERARCHICAL_COUNTER_RIGUNIT()
+
+	if (URigHierarchy* Hierarchy = ExecuteContext.Hierarchy)
+	{
+		if (!CachedIndex.UpdateCache(FRigElementKey(Control, ERigElementType::Control), Hierarchy))
+		{
+			UE_CONTROLRIG_RIGUNIT_REPORT_WARNING(TEXT("Control '%s' is not valid."), *Control.ToString());
+		}
+		else
+		{
+			switch (Space)
+			{
+				case ERigVMTransformSpace::GlobalSpace:
+				{
+					OffsetTransform = Hierarchy->GetInitialGlobalTransform(CachedIndex);
+					break;
+				}
+				case ERigVMTransformSpace::LocalSpace:
+				{
+					OffsetTransform = Hierarchy->GetInitialLocalTransform(CachedIndex);
+					break;
+				}
+				default:
+				{
+					break;
+				}
+			}
+		}
+	}
+}
\ No newline at end of file
diff --git a/Engine/Plugins/Animation/ControlRig/Source/ControlRig/Private/Units/Hierarchy/RigUnit_SetControlOffset.cpp b/Engine/Plugins/Animation/ControlRig/Source/ControlRig/Private/Units/Hierarchy/RigUnit_SetControlOffset.cpp
--- a/Engine/Plugins/Animation/ControlRig/Source/ControlRig/Private/Units/Hierarchy/RigUnit_SetControlOffset.cpp
+++ b/Engine/Plugins/Animation/ControlRig/Source/ControlRig/Private/Units/Hierarchy/RigUnit_SetControlOffset.cpp
@@ -1,7 +1,7 @@
 // Copyright Epic Games, Inc. All Rights Reserved.
 
 #include "Units/Hierarchy/RigUnit_SetControlOffset.h"
-
+#include "Units/Hierarchy/RigUnit_GetControlOffset.h"
 #include "GeometryCollection/GeometryCollectionAlgo.h"
 #include "Units/RigUnitContext.h"
 #include "Math/ControlRigMathLibrary.h"
@@ -9,6 +9,11 @@
 
 #include UE_INLINE_GENERATED_CPP_BY_NAME(RigUnit_SetControlOffset)
 
+FString FRigUnit_SetControlOffset::GetUnitLabel() const
+{
+	return FString::Printf(TEXT("Set Control Offset"));
+}
+
 FRigUnit_SetControlOffset_Execute()
 {
     DECLARE_SCOPE_HIERARCHICAL_COUNTER_RIGUNIT()
@@ -36,6 +41,45 @@ FRigUnit_SetControlOffset_Execute()
 	}
 }
 
+FString FRigUnit_SetControlTranslationOffset::GetUnitLabel() const
+{
+	return FString::Printf(TEXT("Set Control Translation Offset"));
+}
+
+FRigUnit_SetControlTranslationOffset_Execute()
+{
+	FTransform Transform = FTransform::Identity;
+	FRigUnit_GetControlOffset::StaticExecute(ExecuteContext, Control, Space, Transform, CachedControlIndex);
+	Transform.SetLocation(Offset);
+	FRigUnit_SetControlOffset::StaticExecute(ExecuteContext, Control, Transform, Space, CachedControlIndex);
+}
+
+FString FRigUnit_SetControlRotationOffset::GetUnitLabel() const
+{
+	return FString::Printf(TEXT("Set Control Rotation Offset"));
+}
+
+FRigUnit_SetControlRotationOffset_Execute()
+{
+	FTransform Transform = FTransform::Identity;
+	FRigUnit_GetControlOffset::StaticExecute(ExecuteContext, Control, Space, Transform, CachedControlIndex);
+	Transform.SetRotation(Offset);
+	FRigUnit_SetControlOffset::StaticExecute(ExecuteContext, Control, Transform, Space, CachedControlIndex);
+}
+
+FString FRigUnit_SetControlScaleOffset::GetUnitLabel() const
+{
+	return FString::Printf(TEXT("Set Control Scale Offset"));
+}
+
+FRigUnit_SetControlScaleOffset_Execute()
+{
+	FTransform Transform = FTransform::Identity;
+	FRigUnit_GetControlOffset::StaticExecute(ExecuteContext, Control, Space, Transform, CachedControlIndex);
+	Transform.SetScale3D(Scale);
+	FRigUnit_SetControlOffset::StaticExecute(ExecuteContext, Control, Transform, Space, CachedControlIndex);
+}
+
 FRigUnit_GetShapeTransform_Execute()
 {
 	DECLARE_SCOPE_HIERARCHICAL_COUNTER_RIGUNIT()
diff --git a/Engine/Plugins/Animation/ControlRig/Source/ControlRig/Public/Units/Hierarchy/RigUnit_GetControlOffset.h b/Engine/Plugins/Animation/ControlRig/Source/ControlRig/Public/Units/Hierarchy/RigUnit_GetControlOffset.h
new file mode 100644
--- /dev/null
+++ b/Engine/Plugins/Animation/ControlRig/Source/ControlRig/Public/Units/Hierarchy/RigUnit_GetControlOffset.h
@@ -0,0 +1,47 @@
+// Copyright Epic Games, Inc. All Rights Reserved.
+
+#pragma once
+
+#include "Units/RigUnit.h"
+#include "RigUnit_GetControlOffset.generated.h"
+
+/**
+ * GetControlOffset is used to perform a change in the hierarchy by setting a single control's transform.
+ * This is typically only used during the Construction Event.
+ */
+USTRUCT(meta=(DisplayName="Get Control Offset", Category="Controls", DocumentationPolicy="Strict", Keywords = "GetControlOffset,Initial,InitialTransform,GetInitialTransform,GetInitialControlTransform", NodeColor="0, 0.364706, 1.0"))
+struct CONTROLRIG_API FRigUnit_GetControlOffset : public FRigUnit
+{
+
+	GENERATED_BODY()
+
+	FRigUnit_GetControlOffset()
+		: Control(NAME_None)
+		, Space(ERigVMTransformSpace::GlobalSpace)
+		, OffsetTransform(FTransform::Identity)
+		, CachedIndex()
+	{}
+
+	RIGVM_METHOD()
+		virtual void Execute() override;
+
+	/**
+	 * The name of the Control to set the transform for.
+	 */
+	UPROPERTY(meta = (Input, CustomWidget = "ControlName" ))
+	FName Control;
+
+	/**
+	 * Defines if the transform should be retrieved in local or global space
+	 */
+	UPROPERTY(meta = (Input))
+	ERigVMTransformSpace Space;
+
+	// The current transform of the given item - or identity in case it wasn't found.
+	UPROPERTY(meta = (Output))
+	FTransform OffsetTransform;
+
+	// Used to cache the internally used index
+	UPROPERTY()
+	FCachedRigElement CachedIndex;
+};
diff --git a/Engine/Plugins/Animation/ControlRig/Source/ControlRig/Public/Units/Hierarchy/RigUnit_SetControlOffset.h b/Engine/Plugins/Animation/ControlRig/Source/ControlRig/Public/Units/Hierarchy/RigUnit_SetControlOffset.h
--- a/Engine/Plugins/Animation/ControlRig/Source/ControlRig/Public/Units/Hierarchy/RigUnit_SetControlOffset.h
+++ b/Engine/Plugins/Animation/ControlRig/Source/ControlRig/Public/Units/Hierarchy/RigUnit_SetControlOffset.h
@@ -9,17 +9,20 @@
  * SetControlOffset is used to perform a change in the hierarchy by setting a single control's transform.
  * This is typically only used during the Construction Event.
  */
-USTRUCT(meta=(DisplayName="Set Control Offset", Category="Controls", DocumentationPolicy="Strict", Keywords = "SetControlOffset,Initial,InitialTransform,SetInitialTransform,SetInitialControlTransform", NodeColor="0, 0.364706, 1.0"))
+USTRUCT(meta=(DisplayName="Set Control Offset", Category="Controls", TemplateName="SetControlOffset", DocumentationPolicy="Strict", Keywords = "SetControlOffset,Initial,InitialTransform,SetInitialTransform,SetInitialControlTransform", NodeColor="0, 0.364706, 1.0", Varying))
 struct CONTROLRIG_API FRigUnit_SetControlOffset : public FRigUnitMutable
 {
 	GENERATED_BODY()
 
 	FRigUnit_SetControlOffset()
 		: Control(NAME_None)
+		, Offset(FTransform::Identity)
 		, Space(ERigVMTransformSpace::GlobalSpace)
 	{
 	}
 
+	virtual FString GetUnitLabel() const override;
+
 	RIGVM_METHOD()
 	virtual void Execute() override;
 
@@ -30,13 +33,149 @@ struct CONTROLRIG_API FRigUnit_SetControlOffset : public FRigUnitMutable
 	FName Control;
 
 	/**
-	 * The offset transform to set for the control
+	 * The offset transform to set on the control
 	 */
 	UPROPERTY(meta = (Input, Output))
 	FTransform Offset;
+	
+	/**
+	 * Defines if the control's transform should be set
+	 * in local or global space.
+	 */
+	UPROPERTY(meta = (Input))
+	ERigVMTransformSpace Space;
+
+	// user to internally cache the index of the bone
+	UPROPERTY()
+	FCachedRigElement CachedControlIndex;
+};
+
+/**
+ * SetControlTranslationOffset is used to perform a change in the hierarchy by setting a single control's translation offset.
+ * This is typically only used during the Construction Event.
+ */
+USTRUCT(meta = (DisplayName = "Set Control Translation Offset", Category = "Controls", TemplateName="SetControlOffset", DocumentationPolicy = "Strict", Keywords = "SetControlOffset,Initial,InitialTransform,SetInitialTransform,SetInitialControlTransform,SetControlTranslationOffset,SetInitialTranslation,SetInitialLocation", NodeColor = "0, 0.364706, 1.0", Varying))
+struct CONTROLRIG_API FRigUnit_SetControlTranslationOffset : public FRigUnitMutable
+{
+	GENERATED_BODY()
+
+	FRigUnit_SetControlTranslationOffset()
+		: Control(NAME_None)
+		, Offset(FVector::ZeroVector)
+		, Space(ERigVMTransformSpace::GlobalSpace)
+	{
+	}
+
+	virtual FString GetUnitLabel() const override;
+
+	RIGVM_METHOD()
+		virtual void Execute() override;
+
+	/**
+	 * The name of the Control to set the transform for.
+	 */
+	UPROPERTY(meta = (Input, CustomWidget = "ControlName"))
+	FName Control;
+
+	/**
+	 * The input translation offset to set on the control
+	 */
+	UPROPERTY(meta = (Input, Output))
+	FVector Offset;
+
+	/**
+	 * Defines if the control's transform should be set
+	 * in local or global space.
+	 */
+	UPROPERTY(meta = (Input))
+	ERigVMTransformSpace Space;
+
+
+	// user to internally cache the index of the bone
+	UPROPERTY()
+	FCachedRigElement CachedControlIndex;
+};
+
+/**
+ * SetControlRotationOffset is used to perform a change in the hierarchy by setting a single control's rotation offset.
+ * This is typically only used during the Construction Event.
+ */
+USTRUCT(meta = (DisplayName = "Set Control Rotation Offset", Category = "Controls", TemplateName = "Set Control Offset", DocumentationPolicy = "Strict", Keywords = "SetControlOffset,Initial,InitialTransform,SetInitialTransform,SetInitialControlTransform,SetControlRotationOffset,SetInitialRotation,SetInitialRotation", NodeColor = "0, 0.364706, 1.0", Varying))
+struct CONTROLRIG_API FRigUnit_SetControlRotationOffset : public FRigUnitMutable
+{
+	GENERATED_BODY()
+
+	FRigUnit_SetControlRotationOffset()
+		: Control(NAME_None)
+		, Offset(FQuat::Identity)
+		, Space(ERigVMTransformSpace::GlobalSpace)
+	{
+	}
+
+	virtual FString GetUnitLabel() const override;
+
+	RIGVM_METHOD()
+		virtual void Execute() override;
+
+	/**
+	 * The name of the Control to set the transform for.
+	 */
+	UPROPERTY(meta = (Input, CustomWidget = "ControlName"))
+	FName Control;
+
+	/**
+	 * The input rotation offset to set on the control
+	 */
+	UPROPERTY(meta = (Input, Output))
+	FQuat Offset;
+
+	/**
+	 * Defines if the control's transform should be set
+	 * in local or global space.
+	 */
+	UPROPERTY(meta = (Input))
+	ERigVMTransformSpace Space;
+
+	// user to internally cache the index of the bone
+	UPROPERTY()
+	FCachedRigElement CachedControlIndex;
+};
+
+/**
+ * SetControlScaleOffset is used to perform a change in the hierarchy by setting a single control's scale offset.
+ * This is typically only used during the Construction Event.
+ */
+USTRUCT(meta = (DisplayName = "Set Control Scale Offset", Category = "Controls", DocumentationPolicy = "Strict", Keywords = "SetControlOffset,Initial,InitialTransform,SetInitialTransform,SetInitialControlTransform,SetControlScaleOffset,SetInitialScale,SetInitialScale", NodeColor = "0, 0.364706, 1.0", Varying))
+struct CONTROLRIG_API FRigUnit_SetControlScaleOffset : public FRigUnitMutable
+{
+	GENERATED_BODY()
+
+	FRigUnit_SetControlScaleOffset()
+		: Control(NAME_None)
+		, Scale(FVector::OneVector)
+		, Space(ERigVMTransformSpace::GlobalSpace)
+	{
+	}
+
+	virtual FString GetUnitLabel() const override;
+
+	RIGVM_METHOD()
+		virtual void Execute() override;
+
+	/**
+	 * The name of the Control to set the transform for.
+	 */
+	UPROPERTY(meta = (Input, CustomWidget = "ControlName"))
+	FName Control;
+
+	/**
+	 * The new scale offset to set on the control
+	 */
+	UPROPERTY(meta = (Input, Output))
+	FVector Scale;
 
 	/**
-	 * Defines if the bone's transform should be set
+	 * Defines if the control's transform should be set
 	 * in local or global space.
 	 */
 	UPROPERTY(meta = (Input))
