Control Rig Edit Mode: fixed local transform calculation when manipulating using bLocalTransformsInEachLocalSpace
- also implemented what it means for parent and explicit modes

#rnx

diff --git a/Engine/Plugins/Animation/ControlRig/Source/ControlRigEditor/Private/EditMode/ControlRigEditMode.cpp b/Engine/Plugins/Animation/ControlRig/Source/ControlRigEditor/Private/EditMode/ControlRigEditMode.cpp
--- a/Engine/Plugins/Animation/ControlRig/Source/ControlRigEditor/Private/EditMode/ControlRigEditMode.cpp
+++ b/Engine/Plugins/Animation/ControlRig/Source/ControlRigEditor/Private/EditMode/ControlRigEditMode.cpp
@@ -81,10 +81,10 @@
 
 TAutoConsoleVariable<bool> CVarClickSelectThroughGizmo(TEXT("ControlRig.Sequencer.ClickSelectThroughGizmo"), false, TEXT("When false you can't click through a gizmo and change selection if you will select the gizmo when in Animation Mode, default to false."));
 
-namespace CRModeLocals
+namespace UE::ControlRigEditMode::Private
 {
 
-	const TCHAR* FocusModeName = TEXT("AnimMode.PendingFocusMode");
+const TCHAR* FocusModeName = TEXT("AnimMode.PendingFocusMode");
 
 static bool bFocusMode = false;
 static FAutoConsoleVariableRef CVarSetFocusOnHover(
@@ -100,6 +100,8 @@ IConsoleVariable* GetFocusModeVariable()
 	
 }
 
+using namespace UE::ControlRigEditMode;
+
 void UControlRigEditModeDelegateHelper::OnPoseInitialized()
 {
 	if (EditMode)
@@ -282,7 +284,7 @@ bool FControlRigEditMode::SetSequencer(TWeakPtr<ISequencer> InSequencer)
 
 			{
 				TSharedRef<SSequencer> SequencerWidget = StaticCastSharedRef<SSequencer>(Sequencer->GetSequencerWidget());
-				SequencerWidget->EnablePendingFocusOnHovering(CRModeLocals::bFocusMode);
+				SequencerWidget->EnablePendingFocusOnHovering(Private::bFocusMode);
 			}
 		}
 		
@@ -2057,10 +2059,6 @@ bool FControlRigEditMode::InputDelta(FEditorViewportClient* InViewportClient, FV
 		return DragToolHandler.InputDelta(InViewportClient, InViewport, InDrag, InRot, InScale);
 	}
 
-	FVector Drag = InDrag;
-	FRotator Rot = InRot;
-	FVector Scale = InScale;
-
 	const bool bCtrlDown = InViewport->KeyState(EKeys::LeftControl) || InViewport->KeyState(EKeys::RightControl);
 	const bool bShiftDown = InViewport->KeyState(EKeys::LeftShift) || InViewport->KeyState(EKeys::RightShift);
 
@@ -2076,15 +2074,39 @@ bool FControlRigEditMode::InputDelta(FEditorViewportClient* InViewportClient, FV
 
 	const UE::Widget::EWidgetMode WidgetMode = InViewportClient->GetWidgetMode();
 	const EAxisList::Type CurrentAxis = InViewportClient->GetCurrentWidgetAxis();
-	const ECoordSystem CoordSystem = InViewportClient->GetWidgetCoordSystemSpace();
-
-	const bool bDoRotation = !Rot.IsZero() && (WidgetMode == UE::Widget::WM_Rotate || WidgetMode == UE::Widget::WM_TranslateRotateZ);
-	const bool bDoTranslation = !Drag.IsZero() && (WidgetMode == UE::Widget::WM_Translate || WidgetMode == UE::Widget::WM_TranslateRotateZ);
-	const bool bDoScale = !Scale.IsZero() && WidgetMode == UE::Widget::WM_Scale;
+	const ETransformSpace TransformSpace = GetTransformSpace();
+	
+	FTransformContext TransformContext;
+	TransformContext.bTranslation = !InDrag.IsZero() && (WidgetMode == UE::Widget::WM_Translate || WidgetMode == UE::Widget::WM_TranslateRotateZ);
+	TransformContext.Drag = InDrag;
+	TransformContext.bRotation = !InRot.IsZero() && (WidgetMode == UE::Widget::WM_Rotate || WidgetMode == UE::Widget::WM_TranslateRotateZ);
+	TransformContext.Rot = InRot;
+	TransformContext.bScale = !InScale.IsZero() && WidgetMode == UE::Widget::WM_Scale;
+	TransformContext.Scale = InScale;
+	TransformContext.Space = TransformSpace;
 
+	const bool EnableLocalTransform = [Settings = GetSettings(), TransformSpace, TransformContext]()
+	{
+		if (Settings && Settings->bLocalTransformsInEachLocalSpace)
+		{
+			switch (TransformSpace)
+			{
+			case ETransformSpace::World:
+				return false;
+			case ETransformSpace::Local:
+			case ETransformSpace::Parent:
+				return true;
+			case ETransformSpace::Explicit:
+				return TransformContext.bTranslation;
+			default:
+				break;
+			}
+		}
+		return false;
+	}();
 
 	if (InteractionScopes.Num() > 0 && bMouseButtonDown && CurrentAxis != EAxisList::None
-		&& (bDoRotation || bDoTranslation || bDoScale))
+		&& TransformContext.CanTransform())
 	{
 		for (auto& Pairs : ControlRigShapeActors)
 		{
@@ -2103,7 +2125,7 @@ bool FControlRigEditMode::InputDelta(FEditorViewportClient* InViewportClient, FV
 								GEditor->BeginTransaction(LOCTEXT("ChangeControlShapeTransaction", "Change Control Shape Transform"));
 							}
 
-							ChangeControlShapeTransform(ShapeActor, bDoTranslation, InDrag, bDoRotation, InRot, bDoScale, InScale, ComponentTransform);
+							ChangeControlShapeTransform(ShapeActor, TransformContext, ComponentTransform);
 							bManipulatorMadeChange = true;
 
 							// break here since we only support changing shape transform of a single control at a time
@@ -2113,8 +2135,7 @@ bool FControlRigEditMode::InputDelta(FEditorViewportClient* InViewportClient, FV
 				}
 				else
 				{
-					const UControlRigEditModeSettings* Settings = GetSettings();
-					bool bDoLocal = (CoordSystem == ECoordSystem::COORD_Local && Settings && Settings->bLocalTransformsInEachLocalSpace);
+					bool bDoLocal = EnableLocalTransform;
 					bool bUseLocal = false;
 					bool bCalcLocal = bDoLocal;
 					bool bFirstTime = true;
@@ -2171,8 +2192,7 @@ bool FControlRigEditMode::InputDelta(FEditorViewportClient* InViewportClient, FV
 							}
 
 							TArray<TFunction<void()>> Tasks;
-							MoveControlShape(ShapeActor, bDoTranslation, InDrag, bDoRotation, InRot, bDoScale, InScale, ComponentTransform,
-								bUseLocal, bDoLocal, &InOutLocal, Tasks);
+							MoveControlShape(ShapeActor, TransformContext, ComponentTransform, bUseLocal, bDoLocal, &InOutLocal, Tasks);
 							TasksPerActor.Add(ShapeActor, Tasks);
 							bManipulatorMadeChange = true;
 						}
@@ -2228,26 +2248,26 @@ bool FControlRigEditMode::InputDelta(FEditorViewportClient* InViewportClient, FV
 					{
 						FTransform NewWorldTransform = OnGetRigElementTransformDelegate.Execute(SelectedRigElements[Index], false, true) * ComponentTransform;
 						bool bTransformChanged = false;
-						if (bDoRotation)
+						if (TransformContext.bRotation)
 						{
 							FQuat CurrentRotation = NewWorldTransform.GetRotation();
-							CurrentRotation = (Rot.Quaternion() * CurrentRotation);
+							CurrentRotation = (TransformContext.Rot.Quaternion() * CurrentRotation);
 							NewWorldTransform.SetRotation(CurrentRotation);
 							bTransformChanged = true;
 						}
 
-						if (bDoTranslation)
+						if (TransformContext.bTranslation)
 						{
 							FVector CurrentLocation = NewWorldTransform.GetLocation();
-							CurrentLocation = CurrentLocation + Drag;
+							CurrentLocation = CurrentLocation + TransformContext.Drag;
 							NewWorldTransform.SetLocation(CurrentLocation);
 							bTransformChanged = true;
 						}
 
-						if (bDoScale)
+						if (TransformContext.bScale)
 						{
 							FVector CurrentScale = NewWorldTransform.GetScale3D();
-							CurrentScale = CurrentScale + Scale;
+							CurrentScale = CurrentScale + TransformContext.Scale;
 							NewWorldTransform.SetScale3D(CurrentScale);
 							bTransformChanged = true;
 						}
@@ -2525,21 +2545,25 @@ bool FControlRigEditMode::ComputePivotFromEditedShape(UControlRig* InControlRig,
 	return true;
 }
 
-FControlRigEditMode::ETransformSpaceSpace FControlRigEditMode::GetTransformSpace() const
+ETransformSpace FControlRigEditMode::GetTransformSpace() const
 {
 	switch (GetCoordSystemSpace())
 	{
+		case COORD_World:
+			return ETransformSpace::World;
+		case COORD_Local:
+			return ETransformSpace::Local;
 		case COORD_Parent:
-			return ETransformSpaceSpace::Parent;
+			return ETransformSpace::Parent;
 		case COORD_Explicit:
-			return ETransformSpaceSpace::Explicit;
+			return ETransformSpace::Explicit;
 		default:
 			break;
 	}
-	return ETransformSpaceSpace::Default;	
+	return ETransformSpace::Local;	
 };
 
-FTransform FControlRigEditMode::GetPivotOrientation(const FRigElementKey& InControlKey, URigHierarchy* InHierarchy, const ETransformSpaceSpace InSpace, const FTransform& InComponentTransform) const
+FTransform FControlRigEditMode::GetPivotOrientation(const FRigElementKey& InControlKey, URigHierarchy* InHierarchy, const ETransformSpace InSpace, const FTransform& InComponentTransform) const
 {
 	FRigControlElement* Control = InHierarchy->Find<FRigControlElement>(InControlKey);
 	if (!Control)
@@ -2549,14 +2573,15 @@ FTransform FControlRigEditMode::GetPivotOrientation(const FRigElementKey& InCont
 		
 	switch (InSpace)
 	{
-	case ETransformSpaceSpace::Default:
+	case ETransformSpace::World:
+	case ETransformSpace::Local:
 		return InHierarchy->GetTransform(Control, ERigTransformType::CurrentGlobal);
-	case ETransformSpaceSpace::Parent:
+	case ETransformSpace::Parent:
 		{
 			const int32 NumParents = InHierarchy->GetNumberOfParents(Control);			
 			return NumParents > 0 ? InHierarchy->GetParentTransform(Control, ERigTransformType::CurrentGlobal) : InHierarchy->GetTransform(Control, ERigTransformType::CurrentGlobal);
 		}
-	case ETransformSpaceSpace::Explicit:
+	case ETransformSpace::Explicit:
 		{
 			FRotationContext& RotationContext = GetRotationContext();
 					
@@ -2583,7 +2608,7 @@ FTransform FControlRigEditMode::GetPivotOrientation(const FRigElementKey& InCont
 	return FTransform::Identity;
 }
 
-bool FControlRigEditMode::ComputePivotFromShapeActors(UControlRig* InControlRig, const bool bEachLocalSpace, const ETransformSpaceSpace InSpace, FTransform& OutTransform) const
+bool FControlRigEditMode::ComputePivotFromShapeActors(UControlRig* InControlRig, const bool bEachLocalSpace, const ETransformSpace InSpace, FTransform& OutTransform) const
 {
 	if (!ensure(!bIsChangingControlShapeTransform))
 	{
@@ -3879,17 +3904,17 @@ void FControlRigEditMode::RegisterPendingFocusMode()
 		return;
 	}
 	
-	static IConsoleVariable* UseFocusMode = CRModeLocals::GetFocusModeVariable();
+	static IConsoleVariable* UseFocusMode = Private::GetFocusModeVariable();
 	if (ensure(UseFocusMode))
 	{
 		auto OnFocusModeChanged = [this](IConsoleVariable*)
 		{
-			PendingFocus.Enable(CRModeLocals::bFocusMode);
+			PendingFocus.Enable(Private::bFocusMode);
 			if (WeakSequencer.IsValid())
 			{
 				TSharedPtr<ISequencer> PreviousSequencer = WeakSequencer.Pin();
 				TSharedRef<SSequencer> PreviousSequencerWidget = StaticCastSharedRef<SSequencer>(PreviousSequencer->GetSequencerWidget());
-				PreviousSequencerWidget->EnablePendingFocusOnHovering(CRModeLocals::bFocusMode);
+				PreviousSequencerWidget->EnablePendingFocusOnHovering(Private::bFocusMode);
 			}
 		};
 		if (!PendingFocusHandle.IsValid())
@@ -3914,7 +3939,7 @@ void FControlRigEditMode::UnregisterPendingFocusMode()
 	
 	if (PendingFocusHandle.IsValid())
 	{
-		static IConsoleVariable* UseFocusMode = CRModeLocals::GetFocusModeVariable();
+		static IConsoleVariable* UseFocusMode = Private::GetFocusModeVariable();
 		if (ensure(UseFocusMode))
 		{
 			UseFocusMode->OnChangedDelegate().Remove(PendingFocusHandle);
@@ -4998,7 +5023,7 @@ void FControlRigEditMode::SetControlShapeTransform(
 
 	static constexpr bool bNotify = true, bUndo = true;
 	
-	const bool bIsExplicitRotation = GetTransformSpace() == ETransformSpaceSpace::Explicit && GetModeManager()->GetWidgetMode() == UE::Widget::EWidgetMode::WM_Rotate;
+	const bool bIsExplicitRotation = GetTransformSpace() == ETransformSpace::Explicit && GetModeManager()->GetWidgetMode() == UE::Widget::EWidgetMode::WM_Rotate;
 	auto ApplyExplicitRotation = [bIsExplicitRotation, ControlRig, Hierarchy, ControlElement, InDeltaRot]()
 	{
 		if (!bIsExplicitRotation)
@@ -5142,9 +5167,16 @@ FTransform FControlRigEditMode::GetControlShapeTransform(const AControlRigShapeA
 	return FTransform::Identity;
 }
 
-void FControlRigEditMode::MoveControlShape(AControlRigShapeActor* ShapeActor, const bool bTranslation, FVector& InDrag, const bool bRotation,
-	FRotator& InRot, const bool bScale, FVector& InScale, const FTransform& ToWorldTransform, bool bUseLocal, bool bCalcLocal, FTransform* InOutLocal, TArray<TFunction<void()>>& OutTasks)
+void FControlRigEditMode::MoveControlShape( AControlRigShapeActor* ShapeActor,
+	const FTransformContext& InTransformContext, const FTransform& ToWorldTransform,
+	const bool bUseLocal, const bool bCalcLocal, FTransform* InOutLocal,
+	TArray<TFunction<void()>>& OutTasks)
 {
+	if (!ensure(InTransformContext.CanTransform()))
+	{
+		return;
+	}
+	
 	UControlRig* ControlRig = ShapeActor ? ShapeActor->ControlRig.Get() : nullptr;
 	if (!ensure(ControlRig))
 	{
@@ -5170,23 +5202,90 @@ void FControlRigEditMode::MoveControlShape(AControlRigShapeActor* ShapeActor, co
 	{
 		FTransform CurrentLocalTransform = ControlRig->GetControlLocalTransform(ShapeActor->ControlName);
 
-		if (bRotation)
+		URigHierarchy* Hierarchy = ControlRig->GetHierarchy();
+		FRigControlElement* Control = Hierarchy ? Hierarchy->Find<FRigControlElement>(ShapeActor->GetElementKey()) : nullptr;
+		
+		if (InTransformContext.bRotation)
 		{
 			FQuat CurrentRotation = CurrentLocalTransform.GetRotation();
-			if(InOutLocal)
+			if (InOutLocal)
 			{
-				CurrentRotation = (CurrentRotation * InOutLocal->GetRotation());
+				FQuat DeltaRot = InOutLocal->GetRotation();
+
+				if (ensure(Control))
+				{
+					switch(InTransformContext.Space)
+					{
+					case ETransformSpace::World:
+					case ETransformSpace::Explicit:
+						ensure(false);
+						break;
+					case ETransformSpace::Local:
+						break;
+					case ETransformSpace::Parent:
+						{
+							const int32 NumParents = Hierarchy->GetNumberOfParents(Control);
+							const FTransform Global = Hierarchy->GetTransform(Control, ERigTransformType::CurrentGlobal);
+							const FTransform Parent = NumParents > 0 ? Hierarchy->GetParentTransform(Control, ERigTransformType::CurrentGlobal) : Global;
+
+							FTransform RelativeToParent = Global.GetRelativeTransform(Parent);
+							RelativeToParent = RelativeToParent * DeltaRot;
+							const FTransform NewGlobal = RelativeToParent * Parent;
+							
+							DeltaRot = NewGlobal.GetRelativeTransform(Global).GetRotation();
+							break;
+						}
+					}
+				}
+
+				CurrentRotation = (CurrentRotation * DeltaRot);
 			}
 			CurrentLocalTransform.SetRotation(CurrentRotation);
 			bTransformChanged = true;
 		}
 
-		if (bTranslation)
+		if (InTransformContext.bTranslation)
 		{
 			FVector CurrentLocation = CurrentLocalTransform.GetLocation();
-			if(InOutLocal)
+			if (InOutLocal)
 			{
-				CurrentLocation = CurrentLocation + InOutLocal->GetLocation();
+				FVector Delta = InOutLocal->GetLocation();
+
+				if (ensure(Control))
+				{
+					switch(InTransformContext.Space)
+					{
+					case ETransformSpace::World:
+						ensure(false);
+						break;
+					case ETransformSpace::Local:
+					{
+						// in local mode, the incoming delta is expressed in the global space so it has to be switched back to the offset space
+						const FTransform Global = Hierarchy->GetTransform(Control, ERigTransformType::CurrentGlobal);
+						const FVector GlobalDelta = Global.TransformVector(Delta);
+						const FTransform Offset = Hierarchy->GetControlOffsetTransform(Control, ERigTransformType::CurrentGlobal);
+						Delta = Offset.InverseTransformVector(GlobalDelta);
+						break;
+					}
+					case ETransformSpace::Parent:
+					{
+						// in parent mode, the incoming delta is expressed in the parent space so it has to be switched back to the offset space
+						const int32 NumParents = Hierarchy->GetNumberOfParents(Control);
+						const FTransform Parent = NumParents > 0 ? Hierarchy->GetParentTransform(Control, ERigTransformType::CurrentGlobal) : Hierarchy->GetTransform(Control, ERigTransformType::CurrentGlobal);
+						const FVector GlobalDelta = Parent.TransformVector(Delta);
+						const FTransform Offset = Hierarchy->GetControlOffsetTransform(Control, ERigTransformType::CurrentGlobal);
+						Delta = Offset.InverseTransformVector(GlobalDelta);
+						break;
+					}
+					case ETransformSpace::Explicit:
+					{
+						// nothing to do as it has already been computed in the right space
+						break;
+					}
+					}
+				}
+    
+				CurrentLocation = CurrentLocation + Delta;
 			}
 			CurrentLocalTransform.SetLocation(CurrentLocation);
 			bTransformChanged = true;
@@ -5221,10 +5320,10 @@ void FControlRigEditMode::MoveControlShape(AControlRigShapeActor* ShapeActor, co
 
 	// Get the global transform from shape actor to avoid drifting
 	FTransform CurrentTransform;
+	FRigControlElement* ControlElement = ControlRig->FindControl(ShapeActor->ControlName);
 	if (ShapeActor->StaticMeshComponent && ShapeActor->StaticMeshComponent->GetStaticMesh())
 	{
 		// Update ShapeActor in case we are moving multiple shapes, which affect one another
-		FRigControlElement* ControlElement = ControlRig->FindControl(ShapeActor->ControlName);
 		const FTransform Transform = ControlRig->GetHierarchy()->GetTransform(ControlElement, ERigTransformType::CurrentGlobal);
 		ShapeActor->SetActorTransform(Transform * ToWorldTransform);
 		
@@ -5238,27 +5337,28 @@ void FControlRigEditMode::MoveControlShape(AControlRigShapeActor* ShapeActor, co
 		// We will need to rely on the information we have in the rig hierarchy
 		CurrentTransform = GetControlShapeTransform(ShapeActor) * ToWorldTransform;
 	}
+	const FTransform GlobalTransform(CurrentTransform);
 	
-	if (bRotation)
+	if (InTransformContext.bRotation)
 	{
 		FQuat CurrentRotation = CurrentTransform.GetRotation();
-		CurrentRotation = (InRot.Quaternion() * CurrentRotation);
+		CurrentRotation = (InTransformContext.Rot.Quaternion() * CurrentRotation);
 		CurrentTransform.SetRotation(CurrentRotation);
 		bTransformChanged = true;
 	}
 
-	if (bTranslation)
+	if (InTransformContext.bTranslation)
 	{
 		FVector CurrentLocation = CurrentTransform.GetLocation();
-		CurrentLocation = CurrentLocation + InDrag;
+		CurrentLocation = CurrentLocation + InTransformContext.Drag;
 		CurrentTransform.SetLocation(CurrentLocation);
 		bTransformChanged = true;
 	}
 
-	if (bScale)
+	if (InTransformContext.bScale)
 	{
 		FVector CurrentScale = CurrentTransform.GetScale3D();
-		CurrentScale = CurrentScale + InScale;
+		CurrentScale = CurrentScale + InTransformContext.Scale;
 		CurrentTransform.SetScale3D(CurrentScale);
 		bTransformChanged = true;
 	}
@@ -5274,15 +5374,45 @@ void FControlRigEditMode::MoveControlShape(AControlRigShapeActor* ShapeActor, co
 		Context.EventName = FRigUnit_BeginExecution::EventName;
 		Context.bConstraintUpdate = true;
 
+		FTransform TransformSpace = FTransform::Identity; 
 		if (bCalcLocal && InOutLocal)
 		{
+			URigHierarchy* Hierarchy = ControlRig->GetHierarchy();
 			*InOutLocal = ControlRig->GetControlLocalTransform(ShapeActor->ControlName);
+
+			switch (InTransformContext.Space)
+			{
+			case ETransformSpace::World:
+				ensure(false);
+				break;
+			case ETransformSpace::Local:
+				TransformSpace = GlobalTransform;
+				break;
+			case ETransformSpace::Parent:
+				{
+					const int32 NumParents = Hierarchy->GetNumberOfParents(ControlElement);
+					if (NumParents > 0)
+					{
+						TransformSpace = Hierarchy->GetParentTransform(ControlElement, ERigTransformType::CurrentGlobal) * ToWorldTransform;
+					}
+					else
+					{
+						TransformSpace = GlobalTransform;
+					}
+				}
+				break;
+			case ETransformSpace::Explicit:
+				if (ensure(InTransformContext.bTranslation))
+				{
+					TransformSpace = Hierarchy->GetControlOffsetTransform(ControlElement, ERigTransformType::CurrentGlobal) * ToWorldTransform;
+				}
+				break;
+			}
 		}
 
 		const UWorld* World = ControlRig->GetWorld();
 		const bool bPrintPythonCommands = World ? World->IsPreviewWorld() : false;
 
-		const FRigControlElement* ControlElement = ControlRig->FindControl(ShapeActor->ControlName);
 		const bool bIsTransientControl = ControlElement ? ControlElement->Settings.bIsTransientControl : false;
 
 		// if we are operating on a PIE instance which is playing we need to reapply the input pose
@@ -5309,11 +5439,11 @@ void FControlRigEditMode::MoveControlShape(AControlRigShapeActor* ShapeActor, co
 		const FTransform ToWorldTransformCopy = ToWorldTransform;
 		
 		const TFunction<void()> SetControlShapeTask = [this, ShapeActor, ControlRig, ControlElement, NewTransform, ToWorldTransformCopy, Context,
-			bPrintPythonCommands, bIsTransientControl, bRotation, InRot, bQueueTasks]()
+			bPrintPythonCommands, bIsTransientControl, bRotation = InTransformContext.bRotation, Rot = InTransformContext.Rot, bQueueTasks]()
 		{
 			//fix flips and do rotation orders only if not additive
 			const bool bFixEulerFlips = !ControlRig->IsAdditive() && bRotation;
-			SetControlShapeTransform(ShapeActor, NewTransform, ToWorldTransformCopy, Context, bPrintPythonCommands, InRot, bFixEulerFlips);
+			SetControlShapeTransform(ShapeActor, NewTransform, ToWorldTransformCopy, Context, bPrintPythonCommands, Rot, bFixEulerFlips);
 			NotifyDrivenControls(ControlRig, ShapeActor->GetElementKey(),Context);
 
 			if (ControlElement && !bIsTransientControl)
@@ -5329,7 +5459,7 @@ void FControlRigEditMode::MoveControlShape(AControlRigShapeActor* ShapeActor, co
 			}
 		};
 
-		const TFunction<void()> SetGlobalAndUpdateLocalTask = [ShapeActor, CurrentTransform, ControlRig, bCalcLocal, InOutLocal]() {
+		const TFunction<void()> SetGlobalAndUpdateLocalTask = [ShapeActor, CurrentTransform, ControlRig, bCalcLocal, InOutLocal, InTransformContext, TransformSpace]() {
 
 			// Don't set the global transform to the shape actor to avoid drifting
 			//ShapeActor->SetGlobalTransform(CurrentTransform);
@@ -5338,6 +5468,36 @@ void FControlRigEditMode::MoveControlShape(AControlRigShapeActor* ShapeActor, co
 			{
 				const FTransform NewLocal = ControlRig->GetControlLocalTransform(ShapeActor->ControlName);
 				*InOutLocal = NewLocal.GetRelativeTransform(*InOutLocal);
+				
+				switch (InTransformContext.Space)
+				{
+					case ETransformSpace::World:
+						break;
+					case ETransformSpace::Local:
+						if (InTransformContext.bTranslation)
+						{
+							InOutLocal->SetLocation(TransformSpace.InverseTransformVector(InTransformContext.Drag));
+						}
+						break;
+					case ETransformSpace::Parent:
+						if (InTransformContext.bTranslation)
+						{
+							InOutLocal->SetLocation(TransformSpace.InverseTransformVector(InTransformContext.Drag));
+						}
+						if (InTransformContext.bRotation)
+						{
+							FQuat SpaceRotation = TransformSpace.GetRotation();
+							SpaceRotation = (InTransformContext.Rot.Quaternion() * SpaceRotation);
+							InOutLocal->SetRotation(TransformSpace.InverseTransformRotation(SpaceRotation));
+						}
+						break;
+					case ETransformSpace::Explicit:
+						if (InTransformContext.bTranslation)
+						{
+							InOutLocal->SetLocation(TransformSpace.InverseTransformVector(InTransformContext.Drag));
+						}
+						break;
+				}
 			}
 		};
 
@@ -5354,70 +5514,65 @@ void FControlRigEditMode::MoveControlShape(AControlRigShapeActor* ShapeActor, co
 	}
 }
 
-void FControlRigEditMode::ChangeControlShapeTransform(AControlRigShapeActor* ShapeActor, const bool bTranslation, FVector& InDrag,
-	const bool bRotation, FRotator& InRot, const bool bScale, FVector& InScale, const FTransform& ToWorldTransform)
+void FControlRigEditMode::ChangeControlShapeTransform(AControlRigShapeActor* InShapeActor, const FTransformContext& InContext, const FTransform& ToWorldTransform)
 {
-	bool bTransformChanged = false; 
-
-	FTransform CurrentTransform = FTransform::Identity;
+	if (!InContext.CanTransform())
+	{
+		return;
+	}
+	
+	UControlRig* ControlRig = InShapeActor->ControlRig.Get();
+	URigHierarchy* Hierarchy = ControlRig ? ControlRig->GetHierarchy() : nullptr;
+	if (!Hierarchy)
+	{
+		return;
+	}
 
-	if (UControlRig* ControlRig = ShapeActor->ControlRig.Get())
+	FRigControlElement* ControlElement = Hierarchy->Find<FRigControlElement>(InShapeActor->GetElementKey());
+	if (!ControlElement)
 	{
-		if (FRigControlElement* ControlElement = ControlRig->GetHierarchy()->Find<FRigControlElement>(ShapeActor->GetElementKey()))
-		{
-			CurrentTransform = ControlRig->GetHierarchy()->GetControlShapeTransform(ControlElement, ERigTransformType::CurrentGlobal);
-			CurrentTransform = CurrentTransform * ToWorldTransform;
-		}
+		return;
 	}
 
-	if (bRotation)
+	FTransform CurrentTransform = Hierarchy->GetControlShapeTransform(ControlElement, ERigTransformType::CurrentGlobal);
+	CurrentTransform = CurrentTransform * ToWorldTransform;
+
+	if (InContext.bRotation)
 	{
 		FQuat CurrentRotation = CurrentTransform.GetRotation();
-		CurrentRotation = (InRot.Quaternion() * CurrentRotation);
+		CurrentRotation = (InContext.Rot.Quaternion() * CurrentRotation);
 		CurrentTransform.SetRotation(CurrentRotation);
-		bTransformChanged = true;
 	}
 
-	if (bTranslation)
+	if (InContext.bTranslation)
 	{
 		FVector CurrentLocation = CurrentTransform.GetLocation();
-		CurrentLocation = CurrentLocation + InDrag;
+		CurrentLocation = CurrentLocation + InContext.Drag;
 		CurrentTransform.SetLocation(CurrentLocation);
-		bTransformChanged = true;
 	}
 
-	if (bScale)
+	if (InContext.bScale)
 	{
 		FVector CurrentScale = CurrentTransform.GetScale3D();
-		CurrentScale = CurrentScale + InScale;
+		CurrentScale = CurrentScale + InContext.Scale;
 		CurrentTransform.SetScale3D(CurrentScale);
-		bTransformChanged = true;
 	}
 
-	if (bTransformChanged)
-	{
-		if (UControlRig* ControlRig = ShapeActor->ControlRig.Get())
-		{
-			FTransform NewTransform = CurrentTransform.GetRelativeTransform(ToWorldTransform);
+	FTransform NewTransform = CurrentTransform.GetRelativeTransform(ToWorldTransform);
 
-			if (FRigControlElement* ControlElement = ControlRig->GetHierarchy()->Find<FRigControlElement>(ShapeActor->GetElementKey()))
-			{
-				// do not setup undo for this first step since it is just used to calculate the local transform
-				ControlRig->GetHierarchy()->SetControlShapeTransform(ControlElement, NewTransform, ERigTransformType::CurrentGlobal, false);
-				FTransform CurrentLocalShapeTransform = ControlRig->GetHierarchy()->GetControlShapeTransform(ControlElement, ERigTransformType::CurrentLocal);
-				// this call should trigger an instance-to-BP update in ControlRigEditor
-				ControlRig->GetHierarchy()->SetControlShapeTransform(ControlElement, CurrentLocalShapeTransform, ERigTransformType::InitialLocal, true);
+	// do not setup undo for this first step since it is just used to calculate the local transform
+	Hierarchy->SetControlShapeTransform(ControlElement, NewTransform, ERigTransformType::CurrentGlobal, false);
+	const FTransform CurrentLocalShapeTransform = Hierarchy->GetControlShapeTransform(ControlElement, ERigTransformType::CurrentLocal);
+	// this call should trigger an instance-to-BP update in ControlRigEditor
+	Hierarchy->SetControlShapeTransform(ControlElement, CurrentLocalShapeTransform, ERigTransformType::InitialLocal, true);
 
-				FTransform MeshTransform = FTransform::Identity;
-				FTransform ShapeTransform = CurrentLocalShapeTransform;
-				if (const FControlRigShapeDefinition* Gizmo = UControlRigShapeLibrary::GetShapeByName(ControlElement->Settings.ShapeName, ControlRig->GetShapeLibraries(), ControlRig->ShapeLibraryNameMap))
-				{
-					MeshTransform = Gizmo->Transform;
-				}
-				ShapeActor->StaticMeshComponent->SetRelativeTransform(MeshTransform * ShapeTransform);
-			}
-		}
-	} 
+	FTransform MeshTransform = FTransform::Identity;
+	FTransform ShapeTransform = CurrentLocalShapeTransform;
+	if (const FControlRigShapeDefinition* Gizmo = UControlRigShapeLibrary::GetShapeByName(ControlElement->Settings.ShapeName, ControlRig->GetShapeLibraries(), ControlRig->ShapeLibraryNameMap))
+	{
+		MeshTransform = Gizmo->Transform;
+	}
+	InShapeActor->StaticMeshComponent->SetRelativeTransform(MeshTransform * ShapeTransform);
 }
 
 
diff --git a/Engine/Plugins/Animation/ControlRig/Source/ControlRigEditor/Public/EditMode/ControlRigEditMode.h b/Engine/Plugins/Animation/ControlRig/Source/ControlRigEditor/Public/EditMode/ControlRigEditMode.h
--- a/Engine/Plugins/Animation/ControlRig/Source/ControlRigEditor/Public/EditMode/ControlRigEditMode.h
+++ b/Engine/Plugins/Animation/ControlRig/Source/ControlRigEditor/Public/EditMode/ControlRigEditMode.h
@@ -124,6 +124,37 @@ private:
 
 };
 
+namespace UE::ControlRigEditMode
+{
+
+enum class ETransformSpace
+{
+	World,
+	Local,
+	Parent,
+	Explicit
+};
+	
+struct FTransformContext
+{
+	bool CanTransform() const
+	{
+		return bTranslation || bRotation || bScale;
+	}
+	
+	bool bTranslation = false;
+	FVector Drag = FVector::ZeroVector;
+	
+	bool bRotation = false;
+	FRotator Rot = FRotator::ZeroRotator;
+	
+	bool bScale = false;
+	FVector Scale = FVector::OneVector;
+	
+	ETransformSpace Space = ETransformSpace::World;
+};
+
+}
 
 
 class CONTROLRIGEDITOR_API FControlRigEditMode : public IPersonaEditMode
@@ -236,13 +267,7 @@ public:
 
 protected:
 
-	enum class ETransformSpaceSpace
-	{
-		Default,
-		Parent,
-		Explicit
-	};
-	ETransformSpaceSpace GetTransformSpace() const;
+	UE::ControlRigEditMode::ETransformSpace GetTransformSpace() const;
 	
 	// shape related functions wrt enable/selection
 	/** Get the node name from the property path */
@@ -258,9 +283,9 @@ protected:
 	/** Updates cached pivot transforms */
 	void UpdatePivotTransforms();
 	bool ComputePivotFromEditedShape(UControlRig* InControlRig, FTransform& OutTransform) const;
-	bool ComputePivotFromShapeActors(UControlRig* InControlRig, const bool bEachLocalSpace, const ETransformSpaceSpace InSpace, FTransform& OutTransform) const;
+	bool ComputePivotFromShapeActors(UControlRig* InControlRig, const bool bEachLocalSpace, const UE::ControlRigEditMode::ETransformSpace InSpace, FTransform& OutTransform) const;
 	bool ComputePivotFromElements(UControlRig* InControlRig, FTransform& OutTransform) const;
-	FTransform GetPivotOrientation(const FRigElementKey& InControlKey, URigHierarchy* InHierarchy, const ETransformSpaceSpace InSpace, const FTransform& InComponentTransform) const;
+	FTransform GetPivotOrientation(const FRigElementKey& InControlKey, URigHierarchy* InHierarchy, const UE::ControlRigEditMode::ETransformSpace InSpace, const FTransform& InComponentTransform) const;
 	
 	/** Get the current coordinate system space */
 	ECoordSystem GetCoordSystemSpace() const;
@@ -550,8 +575,7 @@ private:
 		const FRotator& InRot, const bool bFixEulerFlips = true) const;
 	static FTransform GetControlShapeTransform(const AControlRigShapeActor* ShapeActor);
 
-	void ChangeControlShapeTransform(AControlRigShapeActor* ShapeActor, const bool bTranslation, FVector& InDrag,
-		const bool bRotation, FRotator& InRot, const bool bScale, FVector& InScale, const FTransform& ToWorldTransform);
+	static void ChangeControlShapeTransform(AControlRigShapeActor* ShapeActor, const UE::ControlRigEditMode::FTransformContext& InContext, const FTransform& ToWorldTransform);
 
 	bool ModeSupportedByShapeActor(const AControlRigShapeActor* ShapeActor, UE::Widget::EWidgetMode InMode) const;
 
@@ -561,9 +585,9 @@ public:
 
 protected:
 
-	void MoveControlShape(AControlRigShapeActor* ShapeActor, const bool bTranslation, FVector& InDrag, 
-		const bool bRotation, FRotator& InRot, const bool bScale, FVector& InScale, const FTransform& ToWorldTransform,
-		bool bUseLocal, bool bCalcLocal, FTransform* InOutLocal, TArray<TFunction<void()>>& OutTasks);
+	void MoveControlShape(AControlRigShapeActor* ShapeActor,
+		const UE::ControlRigEditMode::FTransformContext& InContext, const FTransform& ToWorldTransform,
+		const bool bUseLocal, const bool bCalcLocal, FTransform* InOutLocal, TArray<TFunction<void()>>& OutTasks);
 
 	/** Get bindings to a runtime object */
 	//If the passed in ControlRig is nullptr we use the first Control Rig(this can happen from the BP Editors).
