Fix another critical section -> transactionally safe variant in control rig.

diff --git a/Engine/Plugins/Animation/ControlRig/Source/ControlRig/Private/Rigs/RigHierarchy.cpp b/Engine/Plugins/Animation/ControlRig/Source/ControlRig/Private/Rigs/RigHierarchy.cpp
--- a/Engine/Plugins/Animation/ControlRig/Source/ControlRig/Private/Rigs/RigHierarchy.cpp
+++ b/Engine/Plugins/Animation/ControlRig/Source/ControlRig/Private/Rigs/RigHierarchy.cpp
@@ -10,6 +10,7 @@
 #include "ModularRigRuleManager.h"
 #include "Units/RigUnitContext.h"
 #include "Math/ControlRigMathLibrary.h"
+#include "Misc/ScopeLock.h"
 #include "UObject/AnimObjectVersion.h"
 #include "ControlRigObjectVersion.h"
 #include "IControlRigObjectBinding.h"
@@ -39,7 +40,7 @@ static FCriticalSection GRigHierarchyStackTraceMutex;
 static char GRigHierarchyStackTrace[65536];
 static void RigHierarchyCaptureCallStack(FString& OutCallstack, uint32 NumCallsToIgnore)
 {
-	FScopeLock ScopeLock(&GRigHierarchyStackTraceMutex);
+	UE::TScopeLock ScopeLock(GRigHierarchyStackTraceMutex);
 	GRigHierarchyStackTrace[0] = 0;
 	FPlatformStackWalk::StackWalkAndDump(GRigHierarchyStackTrace, 65535, 1 + NumCallsToIgnore);
 	OutCallstack = ANSI_TO_TCHAR(GRigHierarchyStackTrace);
@@ -175,7 +176,7 @@ void URigHierarchy::AddReferencedObjects(UObject* InpThis, FReferenceCollector&
 	Super::AddReferencedObjects(InpThis, Collector);
 
 	URigHierarchy* pThis = static_cast<URigHierarchy*>(InpThis);
-	FScopeLock Lock(&pThis->ElementsLock);
+	UE::TScopeLock Lock(pThis->ElementsLock);
 	for (FRigBaseElement* Element : pThis->Elements)
 	{
 		Collector.AddPropertyReferencesWithStructARO(Element->GetScriptStruct(), Element, pThis);
@@ -184,7 +185,7 @@ void URigHierarchy::AddReferencedObjects(UObject* InpThis, FReferenceCollector&
 
 void URigHierarchy::Save(FArchive& Ar)
 {
-	FScopeLock Lock(&ElementsLock);
+	UE::TScopeLock Lock(ElementsLock);
 	UE_RIGVM_ARCHIVETRACE_SCOPE(Ar, FString::Printf(TEXT("URigHierarchy(%s)"), *GetName()));
 
 	Ar.UsingCustomVersion(FAnimObjectVersion::GUID);
@@ -341,7 +342,7 @@ void URigHierarchy::Save(FArchive& Ar)
 
 void URigHierarchy::Load(FArchive& Ar)
 {
-	FScopeLock Lock(&ElementsLock);
+	UE::TScopeLock Lock(ElementsLock);
 	
 	// unlink any existing pose adapter
 	UnlinkPoseAdapter();
@@ -661,7 +662,7 @@ void URigHierarchy::Reset_Impl(bool bResetElements)
 
 	if(bResetElements)
 	{
-		FScopeLock Lock(&ElementsLock);
+		UE::TScopeLock Lock(ElementsLock);
 		
 		// walk in reverse since certain elements might not have been allocated themselves
 		for(int32 ElementIndex = Elements.Num() - 1; ElementIndex >= 0; ElementIndex--)
@@ -734,7 +735,7 @@ void URigHierarchy::ForEachListeningHierarchy(TFunctionRef<void(const FRigHierar
 
 void URigHierarchy::ResetToDefault()
 {
-	FScopeLock Lock(&ElementsLock);
+	UE::TScopeLock Lock(ElementsLock);
 	
 	if(DefaultHierarchyPtr.IsValid())
 	{
@@ -755,7 +756,7 @@ void URigHierarchy::CopyHierarchy(URigHierarchy* InHierarchy)
 	
 	LLM_SCOPE_BYNAME(TEXT("Animation/ControlRig"));
 
-	FScopeLock Lock(&ElementsLock);
+	UE::TScopeLock Lock(ElementsLock);
 	if(Elements.Num() == 0 && InHierarchy->Elements.Num () == 0)
 	{
 		return;
@@ -1033,7 +1034,7 @@ void URigHierarchy::CopyHierarchy(URigHierarchy* InHierarchy)
 
 uint32 URigHierarchy::GetNameHash() const
 {
-	FScopeLock Lock(&ElementsLock);
+	UE::TScopeLock Lock(ElementsLock);
 	
 	uint32 Hash = GetTypeHash(GetTopologyVersion());
 	for (int32 ElementIndex = 0; ElementIndex < Elements.Num(); ElementIndex++)
@@ -1046,7 +1047,7 @@ uint32 URigHierarchy::GetNameHash() const
 
 uint32 URigHierarchy::GetTopologyHash(bool bIncludeTopologyVersion, bool bIncludeTransientControls) const
 {
-	FScopeLock Lock(&ElementsLock);
+	UE::TScopeLock Lock(ElementsLock);
 	
 	uint32 Hash = bIncludeTopologyVersion ? TopologyVersion : 0;
 
@@ -1275,7 +1276,7 @@ void URigHierarchy::ResetPoseToInitial(ERigElementType InTypeFilter)
 	LLM_SCOPE_BYNAME(TEXT("Animation/ControlRig"));
 	bool bPerformFiltering = InTypeFilter != ERigElementType::All;
 	
-	FScopeLock Lock(&ElementsLock);
+	UE::TScopeLock Lock(ElementsLock);
 	const TGuardValue<bool> DisableRecordingOfCurveChanges(bRecordCurveChanges, false);
 	
 	// if we are resetting the pose on some elements, we need to check if
@@ -3188,7 +3189,7 @@ bool URigHierarchy::SwitchToParent(FRigBaseElement* InChild, FRigBaseElement* In
 	// Keep this in function scope, since we might be taking a pointer to it.
 	TElementDependencyMap DependencyMapFromVM;
 	{
-		FScopeLock Lock(&ExecuteContextLock);
+		UE::TScopeLock Lock(ExecuteContextLock);
 		if(ExecuteContext != nullptr && DependencyMapPtr->IsEmpty())
 		{
 			if(ExecuteContext->VM)
@@ -3903,7 +3904,7 @@ void URigHierarchy::Notify(ERigHierarchyNotification InNotifType, const FRigNoti
 
 	// if we are running a VM right now
 	{
-		FScopeLock Lock(&ExecuteContextLock);
+		UE::TScopeLock Lock(ExecuteContextLock);
 		if(ExecuteContext != nullptr)
 		{
 			QueueNotification(InNotifType, InSubject);
@@ -4020,7 +4021,7 @@ void URigHierarchy::SendQueuedNotifications()
 	}
 
 	{
-		FScopeLock Lock(&ExecuteContextLock);
+		UE::TScopeLock Lock(ExecuteContextLock);
     	if(ExecuteContext != nullptr)
     	{
     		return;
@@ -4178,7 +4179,7 @@ FTransform URigHierarchy::GetTransform(FRigTransformElement* InTransformElement,
 	{
 		if(bRecordTransformsAtRuntime)
 		{
-			FScopeLock Lock(&ExecuteContextLock);
+			UE::TScopeLock Lock(ExecuteContextLock);
 			if (ExecuteContext != nullptr)
 			{
 				ReadTransformsAtRuntime.Emplace(
@@ -4338,7 +4339,7 @@ void URigHierarchy::SetTransform(FRigTransformElement* InTransformElement, const
 	{
 		if(bRecordTransformsAtRuntime)
 		{
-			FScopeLock Lock(&ExecuteContextLock);
+			UE::TScopeLock Lock(ExecuteContextLock);
 			if (ExecuteContext)
 			{
 				const FRigVMExecuteContext& PublicData = ExecuteContext->GetPublicData<>();
@@ -4521,7 +4522,7 @@ FTransform URigHierarchy::GetControlOffsetTransform(FRigControlElement* InContro
 	{
 		if(bRecordTransformsAtRuntime) 
 		{
-			FScopeLock Lock(&ExecuteContextLock);
+			UE::TScopeLock Lock(ExecuteContextLock);
 			if (ExecuteContext)
 			{
 				ReadTransformsAtRuntime.Emplace(
@@ -4614,7 +4615,7 @@ void URigHierarchy::SetControlOffsetTransform(FRigControlElement* InControlEleme
 	{
 		if(bRecordTransformsAtRuntime)
 		{
-			FScopeLock Lock(&ExecuteContextLock);
+			UE::TScopeLock Lock(ExecuteContextLock);
 			if (ExecuteContext)
 			{
 				WrittenTransformsAtRuntime.Emplace(
@@ -5772,7 +5773,7 @@ void URigHierarchy::EnsureCachedChildrenAreCurrent() const
 	
 void URigHierarchy::UpdateCachedChildren()
 {
-	FScopeLock Lock(&ElementsLock);
+	UE::TScopeLock Lock(ElementsLock);
 	
 	// First we tally up how many children each element has, then we allocate for the total
 	// count and do the same loop again.
diff --git a/Engine/Plugins/Animation/ControlRig/Source/ControlRig/Private/Rigs/RigHierarchyController.cpp b/Engine/Plugins/Animation/ControlRig/Source/ControlRig/Private/Rigs/RigHierarchyController.cpp
--- a/Engine/Plugins/Animation/ControlRig/Source/ControlRig/Private/Rigs/RigHierarchyController.cpp
+++ b/Engine/Plugins/Animation/ControlRig/Source/ControlRig/Private/Rigs/RigHierarchyController.cpp
@@ -7,6 +7,7 @@
 #include "ModularRig.h"
 #include "HelperUtil.h"
 #include "Engine/SkeletalMesh.h"
+#include "Misc/ScopeLock.h"
 #include "Units/Execution/RigUnit_PrepareForExecution.h"
 
 #if WITH_EDITOR
@@ -2344,7 +2345,7 @@ int32 URigHierarchyController::AddElement(FRigBaseElement* InElementToAdd, FRigB
 	ensure(IsValid());
 
 	URigHierarchy* Hierarchy = GetHierarchy();
-	FScopeLock Lock(&Hierarchy->ElementsLock);
+	UE::TScopeLock Lock(Hierarchy->ElementsLock);
 
 	InElementToAdd->CachedNameString.Reset();
 	InElementToAdd->SubIndex = Hierarchy->Num(InElementToAdd->Key.Type);
@@ -2475,7 +2476,7 @@ bool URigHierarchyController::RemoveElement(FRigBaseElement* InElement)
 	}
 
 	URigHierarchy* Hierarchy = GetHierarchy();
-	FScopeLock Lock(&Hierarchy->ElementsLock);
+	UE::TScopeLock Lock(Hierarchy->ElementsLock);
 
 	// make sure this element is part of this hierarchy
 	ensure(Hierarchy->FindChecked(InElement->Key) == InElement);
diff --git a/Engine/Plugins/Animation/ControlRig/Source/ControlRig/Public/Rigs/RigHierarchy.h b/Engine/Plugins/Animation/ControlRig/Source/ControlRig/Public/Rigs/RigHierarchy.h
--- a/Engine/Plugins/Animation/ControlRig/Source/ControlRig/Public/Rigs/RigHierarchy.h
+++ b/Engine/Plugins/Animation/ControlRig/Source/ControlRig/Public/Rigs/RigHierarchy.h
@@ -17,6 +17,7 @@
 #include "RigVMPythonUtils.h"
 #endif
 #include "Containers/Queue.h"
+#include "Misc/TransactionallySafeCriticalSection.h"
 #include "Serialization/MemoryReader.h"
 #include "Serialization/MemoryWriter.h"
 #include "RigHierarchy.generated.h"
@@ -4464,7 +4465,7 @@ private:
 
 	// Storage for the elements
 	mutable TArray<FRigBaseElement*> Elements;
-	mutable FCriticalSection ElementsLock;
+	mutable FTransactionallySafeCriticalSection ElementsLock;
 
 	// Storage for the elements
 	mutable TArray<TArray<FRigBaseElement*>> ElementsPerType;
@@ -5073,7 +5074,7 @@ private:
 	// NOTE: it is not safe to read or write the execute context without locking it first as the
 	// FRigHierarchyExecuteContextBracket can change the context from the main thread (from sequencer for example)
 	// at the same time that a control rig is being evaluated on an animation thread.
-	mutable FCriticalSection ExecuteContextLock;
+	mutable FTransactionallySafeCriticalSection ExecuteContextLock;
 	const FRigVMExtendedExecuteContext* ExecuteContext;
 
 #if WITH_EDITOR
