Control Rig: Use animationscrubtimeline for replays

#rnx

diff --git a/Engine/Plugins/Animation/ControlRig/Source/ControlRig/Private/ControlRig.cpp b/Engine/Plugins/Animation/ControlRig/Source/ControlRig/Private/ControlRig.cpp
--- a/Engine/Plugins/Animation/ControlRig/Source/ControlRig/Private/ControlRig.cpp
+++ b/Engine/Plugins/Animation/ControlRig/Source/ControlRig/Private/ControlRig.cpp
@@ -3776,7 +3776,6 @@ void UControlRig::SetReplayTimeIndex(int32 InReplayTimeIndex)
 
 void UControlRig::DisableReplay()
 {
-	SetReplayTimeIndex(INDEX_NONE);
 	Replay = nullptr;
 }
 
diff --git a/Engine/Plugins/Animation/ControlRig/Source/ControlRig/Private/ControlRigReplay.cpp b/Engine/Plugins/Animation/ControlRig/Source/ControlRig/Private/ControlRigReplay.cpp
--- a/Engine/Plugins/Animation/ControlRig/Source/ControlRig/Private/ControlRigReplay.cpp
+++ b/Engine/Plugins/Animation/ControlRig/Source/ControlRig/Private/ControlRigReplay.cpp
@@ -15,6 +15,7 @@
 #endif
 
 #include UE_INLINE_GENERATED_CPP_BY_NAME(ControlRigReplay)
+#define LOCTEXT_NAMESPACE "ControlRigReplay"
 
 //////////////////////////////////////////////////////////////////////////////////////////////////////////////////////////////
 /////// FControlRigReplayTracks
@@ -95,7 +96,9 @@ void FControlRigReplayTracks::StoreHierarchy(URigHierarchy* InHierarchy, const T
 			const TSharedPtr<FSampleTrack<TArray<FRigElementKey>>> ElementKeysTrack = GetContainer()->FindOrAddTrack<TArray<FRigElementKey>>(ElementKeysName, FSampleTrackBase::ETrackType_ElementKeyArray);
 			const TSharedPtr<FSampleTrack<TArray<int32>>> ParentIndicesTrack = GetContainer()->FindOrAddTrack<TArray<int32>>(ParentIndicesName, FSampleTrackBase::ETrackType_Int32Array);
 
-			const TArray<FRigElementKey> ElementKeys = InHierarchy->GetAllKeys();
+			TArray<FRigElementKey> ElementKeys = InHierarchy->GetAllKeys();
+			FilterElementKeys(ElementKeys);
+			
 			TArray<int32> ParentIndices;
 			ParentIndices.Reserve(ElementKeys.Num());
 			for(const FRigElementKey& ElementKey : ElementKeys)
@@ -172,7 +175,7 @@ void FControlRigReplayTracks::StorePose(URigHierarchy* InHierarchy, const TArray
 
 void FControlRigReplayTracks::StorePose(URigHierarchy* InHierarchy, FRigBaseElement* InElement)
 {
-	if(!InElement->IsA<FRigTransformElement>() && !InElement->IsA<FRigCurveElement>() && !InElement->IsA<FRigConnectorElement>())
+	if(!InElement->IsA<FRigTransformElement>() && !InElement->IsA<FRigCurveElement>() && !InElement->IsA<FRigConnectorElement>() && !InElement->IsA<FRigReferenceElement>())
 	{
 		return;
 	}
@@ -231,7 +234,7 @@ bool FControlRigReplayTracks::RestorePose(int32 InTimeIndex, URigHierarchy* InHi
 
 bool FControlRigReplayTracks::RestorePose(int32 InTimeIndex, URigHierarchy* InHierarchy, FRigBaseElement* InElement, TReportFunction InReportFunction) const
 {
-	if(!InElement->IsA<FRigTransformElement>() && !InElement->IsA<FRigCurveElement>())
+	if(!InElement->IsA<FRigTransformElement>() && !InElement->IsA<FRigCurveElement>() && !InElement->IsA<FRigReferenceElement>())
 	{
 		return true;
 	}
@@ -941,6 +944,14 @@ bool FControlRigReplayTracks::ForEachComponent(URigHierarchy* InHierarchy, const
 	return bSuccess;
 }
 
+void FControlRigReplayTracks::FilterElementKeys(TArray<FRigElementKey>& InOutElementKeys)
+{
+	InOutElementKeys.RemoveAll([](const FRigElementKey& ElementKey) -> bool
+	{
+		return ElementKey.Type == ERigElementType::Reference;
+	});
+}
+
 const FName& FControlRigReplayTracks::GetTrackName(const FRigElementKey& InElementKey) const
 {
 	if(const FName* ExistingTrackName = ElementKeyToTrackName.Find(InElementKey))
@@ -1236,6 +1247,13 @@ const TArray<FRigElementKey> FControlRigReplayTracks::GetElementKeys() const
 /////// UControlRigReplay
 //////////////////////////////////////////////////////////////////////////////////////////////////////////////////////////////
 
+const FText UControlRigReplay::LiveStatus = LOCTEXT("LiveStatus", "Live"); 
+const FText UControlRigReplay::LiveStatusTooltip = LOCTEXT("LiveStatusTooltip", "The replay is not affecting the rig."); 
+const FText UControlRigReplay::ReplayInputsStatus = LOCTEXT("ReplayInputsStatus", "Replay Inputs");
+const FText UControlRigReplay::ReplayInputsStatusTooltip = LOCTEXT("ReplayInputsStatusTooltip", "The replay's input data is applied first, then the rig runs.");
+const FText UControlRigReplay::GroundTruthStatus = LOCTEXT("GroundTruthStatus", "Ground Truth");
+const FText UControlRigReplay::GroundTruthStatusTooltip = LOCTEXT("GroundTruthStatusTooltip", "The results from the replay override the rig completely.");
+
 void UControlRigReplay::BeginDestroy()
 {
 	UObject::BeginDestroy();
@@ -1325,7 +1343,8 @@ bool UControlRigReplay::StartRecording(UControlRig* InControlRig)
 				return;
 			}
 
-			if (InEventName == FRigUnit_PrepareForExecution::EventName)
+			if (InEventName == FRigUnit_PrepareForExecution::EventName ||
+				InEventName == FRigUnit_PostPrepareForExecution::EventName)
 			{
 				return;
 			}
@@ -1344,8 +1363,6 @@ bool UControlRigReplay::StartRecording(UControlRig* InControlRig)
 			InputTracks.StoreInteraction(ControlRig->InteractionType, ControlRig->ElementsBeingInteracted);
 			InputTracks.StoreHierarchy(ControlRig->GetHierarchy());
 
-			OutputTracks.AddTimeSample(InputTracks.GetLastAbsoluteTime(), InputTracks.GetLastDeltaTime());
-
 			bStoreVariablesDuringPreEvent = true;
 		}
 	);
@@ -1359,7 +1376,8 @@ bool UControlRigReplay::StartRecording(UControlRig* InControlRig)
 				return;
 			}
 			
-			if (InEventName == FRigUnit_PrepareForExecution::EventName)
+			if (InEventName == FRigUnit_PrepareForExecution::EventName ||
+				InEventName == FRigUnit_PostPrepareForExecution::EventName)
 			{
 				return;
 			}
@@ -1368,6 +1386,7 @@ bool UControlRigReplay::StartRecording(UControlRig* InControlRig)
 			const double DisplayDeltaTime = CurrentTime - TimeOfLastFrame;
 			TimeOfLastFrame = CurrentTime;
 
+			OutputTracks.AddTimeSample(InputTracks.GetLastAbsoluteTime(), InputTracks.GetLastDeltaTime());
 			OutputTracks.StoreVariables(ControlRig);
 			OutputTracks.StoreHierarchy(ControlRig->GetHierarchy());
 			const double RecordingDuration = OutputTracks.GetLastAbsoluteTime();
@@ -1414,16 +1433,58 @@ EControlRigReplayPlaybackMode UControlRigReplay::GetPlaybackMode() const
 {
 	if(IsReplaying())
 	{
-		return EControlRigReplayPlaybackMode::ReplayInputs;
+		return PlaybackMode;
 	}
 	return EControlRigReplayPlaybackMode::Live;
 }
 
-bool UControlRigReplay::StartReplay(UControlRig* InControlRig)
+void UControlRigReplay::SetPlaybackMode(EControlRigReplayPlaybackMode InMode)
 {
+	if(InMode >= EControlRigReplayPlaybackMode::Max)
+	{
+		InMode = EControlRigReplayPlaybackMode::ReplayInputs;
+	}
+
+	if(PlaybackMode == InMode)
+	{
+		return;
+	}
+
+	if(InMode == EControlRigReplayPlaybackMode::Live)
+	{
+		StopReplay();
+	}
+	else
+	{
+		PlaybackMode = InMode; 
+	}
+}
+
+bool UControlRigReplay::StartReplay(UControlRig* InControlRig, EControlRigReplayPlaybackMode InMode)
+{
+	if(IsReplaying() && ReplayControlRig.Get() == InControlRig)
+	{
+		if(InMode != GetPlaybackMode())
+		{
+			SetPlaybackMode(InMode);
+			return true;
+		}
+		if(bReplayPaused)
+		{
+			bReplayPaused = false;
+			return true;
+		}
+	}
+	
 	StopRecording();
 	StopReplay();
 
+	if(InMode == EControlRigReplayPlaybackMode::Live)
+	{
+		SetPlaybackMode(InMode);
+		return true;
+	}
+
 	if(InControlRig == nullptr)
 	{
 		return false;
@@ -1434,7 +1495,6 @@ bool UControlRigReplay::StartReplay(UControlRig* InControlRig)
 		return false;
 	}
 
-	bFirstPlayThrough = true;
 	bStoreVariablesDuringPreEvent = true;
 
 	if(InputTracks.IsEmpty() || OutputTracks.IsEmpty())
@@ -1442,9 +1502,7 @@ bool UControlRigReplay::StartReplay(UControlRig* InControlRig)
 		return false;
 	}
 
-	// reset the control rig's playback index
 	InControlRig->SetReplay(this);
-	InControlRig->SetReplayTimeIndex(0);
 	
 	PreConstructionHandle = InControlRig->OnPreConstruction_AnyThread().AddLambda(
 		[this](UControlRig* InControlRig, const FName& InEventName)
@@ -1463,34 +1521,40 @@ bool UControlRigReplay::StartReplay(UControlRig* InControlRig)
 		[this](URigVMHost* InRigVMHost, const FName& InEventName)
 		{
 			UControlRig* ControlRig = Cast<UControlRig>(InRigVMHost);
-			if (!ControlRig || InEventName == FRigUnit_PrepareForExecution::EventName)
+			if (!ControlRig
+				|| InEventName == FRigUnit_PrepareForExecution::EventName
+				|| InEventName == FRigUnit_PostPrepareForExecution::EventName)
 			{
 				return;
 			}
 
-			const int32 ReplayTimeIndex = ControlRig->GetReplayTimeIndex();
-			check(ReplayTimeIndex != INDEX_NONE);
+			int32 ReplayTimeIndex = ControlRig->GetReplayTimeIndex();
 
-			const FName EventNameFromReplay = InputTracks.GetRigVMEvent(ReplayTimeIndex);
-			if (InEventName != EventNameFromReplay)
+			if(GetPlaybackMode() == EControlRigReplayPlaybackMode::ReplayInputs)
 			{
-				return;
-			}
+				ReplayTimeIndex = FMath::Clamp(ReplayTimeIndex, 0, InputTracks.GetNumTimes() - 1);
+				
+				const FName EventNameFromReplay = InputTracks.GetRigVMEvent(ReplayTimeIndex);
+				if (InEventName != EventNameFromReplay)
+				{
+					return;
+				}
 
-			const TTuple<uint8, TArray<FRigElementKey>> Interaction = InputTracks.GetInteraction(ReplayTimeIndex);
-			ControlRig->InteractionType = Interaction.Get<0>();
-			ControlRig->ElementsBeingInteracted = Interaction.Get<1>();
+				const TTuple<uint8, TArray<FRigElementKey>> Interaction = InputTracks.GetInteraction(ReplayTimeIndex);
+				ControlRig->InteractionType = Interaction.Get<0>();
+				ControlRig->ElementsBeingInteracted = Interaction.Get<1>();
 
-			ControlRig->SetAbsoluteAndDeltaTime(
-				InputTracks.GetAbsoluteTime(ReplayTimeIndex),
-				InputTracks.GetDeltaTime(ReplayTimeIndex)
-			);
+				ControlRig->SetAbsoluteAndDeltaTime(
+					InputTracks.GetAbsoluteTime(ReplayTimeIndex),
+					InputTracks.GetDeltaTime(ReplayTimeIndex)
+				);
 
-			if(bStoreVariablesDuringPreEvent)
-			{
-				InputTracks.RestoreVariables(ReplayTimeIndex, ControlRig);
+				if(bStoreVariablesDuringPreEvent)
+				{
+					InputTracks.RestoreVariables(ReplayTimeIndex, ControlRig);
+				}
+				InputTracks.RestoreHierarchy(ReplayTimeIndex, ControlRig->GetHierarchy());
 			}
-			InputTracks.RestoreHierarchy(ReplayTimeIndex, ControlRig->GetHierarchy());
 			bStoreVariablesDuringPreEvent = true;
 		}
 	);
@@ -1499,60 +1563,94 @@ bool UControlRigReplay::StartReplay(UControlRig* InControlRig)
 		[this](URigVMHost* InRigVMHost, const FName& InEventName)
 		{
 			UControlRig* ControlRig = Cast<UControlRig>(InRigVMHost);
-			if (!ControlRig || InEventName == FRigUnit_PrepareForExecution::EventName)
+			if (!ControlRig
+				 || InEventName == FRigUnit_PrepareForExecution::EventName
+				 || InEventName == FRigUnit_PostPrepareForExecution::EventName)
 			{
 				return;
 			}
 
-			const int32 ReplayTimeIndex = ControlRig->GetReplayTimeIndex();
+			int32 ReplayTimeIndex = ControlRig->GetReplayTimeIndex();
 
-			ControlRig->SetAbsoluteAndDeltaTime(
-				InputTracks.GetAbsoluteTime(ReplayTimeIndex),
-				InputTracks.GetDeltaTime(ReplayTimeIndex)
-			);
-
-			// only validate results during the first play through
-			const FName EventNameFromReplay = InputTracks.GetRigVMEvent(ReplayTimeIndex);
-			if(bFirstPlayThrough && InEventName == EventNameFromReplay)
+			if(GetPlaybackMode() == EControlRigReplayPlaybackMode::ReplayInputs)
 			{
-				const int32 OutputTimeIndex = ReplayTimeIndex;
+				ReplayTimeIndex = FMath::Clamp(ReplayTimeIndex, 0, InputTracks.GetNumTimes() - 1);
 				
-				(void)ValidateExpectedResults(OutputTimeIndex, OutputTracks.SampleTrackIndex, ControlRig,
-					[ControlRig](EMessageSeverity::Type InSeverity, const FString InMessage)
-					{
-						switch(InSeverity)
+				ControlRig->SetAbsoluteAndDeltaTime(
+					InputTracks.GetAbsoluteTime(ReplayTimeIndex),
+					InputTracks.GetDeltaTime(ReplayTimeIndex)
+				);
+
+				// only validate results during the first play through
+				const FName EventNameFromReplay = InputTracks.GetRigVMEvent(ReplayTimeIndex);
+				if(InEventName == EventNameFromReplay)
+				{
+					const int32 OutputTimeIndex = ReplayTimeIndex;
+
+					(void)ValidateExpectedResults(OutputTimeIndex, OutputTracks.SampleTrackIndex, ControlRig,
+						[ControlRig](EMessageSeverity::Type InSeverity, const FString InMessage)
 						{
-							case EMessageSeverity::Warning:
-							case EMessageSeverity::PerformanceWarning:
-							{
-								UE_LOG(LogControlRig, Warning, TEXT("%s"), *InMessage);
-								break;
-							}
-							case EMessageSeverity::Error:
-							{
-								UE_LOG(LogControlRig, Error, TEXT("%s"), *InMessage);
-								break;
-							}
-							case EMessageSeverity::Info:
-							default:
+							switch(InSeverity)
 							{
-								UE_LOG(LogControlRig, Display, TEXT("%s"), *InMessage);
-								break;
+								case EMessageSeverity::Warning:
+								case EMessageSeverity::PerformanceWarning:
+								{
+									UE_LOG(LogControlRig, Warning, TEXT("%s"), *InMessage);
+									break;
+								}
+								case EMessageSeverity::Error:
+								{
+									UE_LOG(LogControlRig, Error, TEXT("%s"), *InMessage);
+									break;
+								}
+								case EMessageSeverity::Info:
+								default:
+								{
+									UE_LOG(LogControlRig, Display, TEXT("%s"), *InMessage);
+									break;
+								}
 							}
 						}
+					);
+				}
+				
+				if(!bReplayPaused)
+				{
+					// loop the animation
+					if(ReplayTimeIndex >= InputTracks.GetNumTimes() - 1)
+					{
+						ControlRig->SetReplayTimeIndex(0);
 					}
-				);
-			}
-
-			// loop the animation
-			if(ReplayTimeIndex >= InputTracks.GetNumTimes() - 1)
-			{
-				ControlRig->SetReplayTimeIndex(0);
-				bFirstPlayThrough = false;
+					else
+					{
+						ControlRig->SetReplayTimeIndex(ReplayTimeIndex + 1);
+					}
+				}
 			}
 			else
 			{
-				ControlRig->SetReplayTimeIndex(ReplayTimeIndex + 1);
+				ReplayTimeIndex = FMath::Clamp(ReplayTimeIndex, 0, OutputTracks.GetNumTimes() - 1);
+
+				ControlRig->SetAbsoluteAndDeltaTime(
+					OutputTracks.GetAbsoluteTime(ReplayTimeIndex),
+					OutputTracks.GetDeltaTime(ReplayTimeIndex)
+				);
+				
+				OutputTracks.RestoreVariables(ReplayTimeIndex, ControlRig);
+				OutputTracks.RestoreHierarchy(ReplayTimeIndex, ControlRig->GetHierarchy());
+
+				if(!bReplayPaused)
+				{
+					// loop the animation
+					if(ReplayTimeIndex >= OutputTracks.GetNumTimes() - 1)
+					{
+						ControlRig->SetReplayTimeIndex(0);
+					}
+					else
+					{
+						ControlRig->SetReplayTimeIndex(ReplayTimeIndex + 1);
+					}
+				}
 			}
 		}
 	);
@@ -1560,6 +1658,8 @@ bool UControlRigReplay::StartReplay(UControlRig* InControlRig)
 	InControlRig->RequestInit();
 
 	ReplayControlRig = InControlRig;
+	SetPlaybackMode(InMode);
+
 	return true;
 }
 
@@ -1572,16 +1672,34 @@ bool UControlRigReplay::StopReplay()
 		ControlRig->InteractionType = 0;
 		ControlRig->ElementsBeingInteracted.Reset();
 		ReplayControlRig.Reset();
+		bReplayPaused = false;
+		PlaybackMode = EControlRigReplayPlaybackMode::Live;
 		return true;
 	}
 	return false;
 }
 
+bool UControlRigReplay::PauseReplay()
+{
+	if(!IsReplaying())
+	{
+		return false;
+	}
+
+	bReplayPaused = true;
+	return false;
+}
+
 bool UControlRigReplay::IsReplaying() const
 {
 	return ReplayControlRig.IsValid();
 }
 
+bool UControlRigReplay::IsPaused() const
+{
+	return bReplayPaused;
+}
+
 bool UControlRigReplay::IsRecording() const
 {
 	return RecordControlRig.IsValid();
@@ -1595,6 +1713,16 @@ bool UControlRigReplay::IsValidForTesting() const
 		ControlRigObjectPath.IsValid();
 }
 
+bool UControlRigReplay::HasValidationErrors() const
+{
+	return !LastValidationWarningsAndErrors.IsEmpty();
+}
+
+const TArray<FString>& UControlRigReplay::GetValidationErrors() const
+{
+	return LastValidationWarningsAndErrors;
+}
+
 bool UControlRigReplay::PerformTest(UControlRig* InSubject, TFunction<void(EMessageSeverity::Type, const FString&)> InLogFunction) const
 {
 	// if we have nothing to check we can consider this test successful
@@ -1689,17 +1817,29 @@ bool UControlRigReplay::ValidateExpectedResults(
 {
 	URigHierarchy* Hierarchy = InSubject->GetHierarchy();
 	const TArray<FRigElementKey>& StoredElementKeys = OutputTracks.GetElementKeys();
+	LastValidationWarningsAndErrors.Reset();
+
+	auto LocalLogFunction = [InLogFunction, this](EMessageSeverity::Type InSeverity, const FString& InMessage)
+	{
+		if(InLogFunction)
+		{
+			InLogFunction(InSeverity, InMessage);
+		}
+		if(InSeverity == EMessageSeverity::Warning || InSeverity == EMessageSeverity::Error)
+		{
+			LastValidationWarningsAndErrors.Add(InMessage);
+		}
+	};
 
 	bool bSuccess = true;
 	if(bValidateHierarchyTopology && InPlaybackTimeIndex == 0)
 	{
 		TArray<FRigElementKey> CurrentElementKeys = Hierarchy->GetAllKeys();
+		FControlRigReplayTracks::FilterElementKeys(CurrentElementKeys);
+		
 		if(StoredElementKeys.Num() != CurrentElementKeys.Num())
 		{
-			if(InLogFunction)
-			{
-				InLogFunction(EMessageSeverity::Error, FString::Printf(TEXT("Number of elements in hierarchy (%d) and number of elements in replay (%d) don't match."), CurrentElementKeys.Num(), StoredElementKeys.Num()));
-			}
+			LocalLogFunction(EMessageSeverity::Error, FString::Printf(TEXT("Number of elements in hierarchy (%d) and number of elements in replay (%d) don't match."), CurrentElementKeys.Num(), StoredElementKeys.Num()));
 			bSuccess = false;
 		}
 
@@ -1707,10 +1847,7 @@ bool UControlRigReplay::ValidateExpectedResults(
 		{
 			if(!CurrentElementKeys.Contains(StoredKey))
 			{
-				if(InLogFunction)
-				{
-					InLogFunction(EMessageSeverity::Error, FString::Printf(TEXT("Hierarchy is missing element %s expected by the replay."), *StoredKey.ToString()));
-				}
+				LocalLogFunction(EMessageSeverity::Error, FString::Printf(TEXT("Hierarchy is missing element %s expected by the replay."), *StoredKey.ToString()));
 				bSuccess = false;
 			}
 		}
@@ -1719,10 +1856,7 @@ bool UControlRigReplay::ValidateExpectedResults(
 		{
 			if(!StoredElementKeys.Contains(CurrentKey))
 			{
-				if(InLogFunction)
-				{
-					InLogFunction(EMessageSeverity::Error, FString::Printf(TEXT("Hierarchy contains element %s which is not part of the replay."), *CurrentKey.ToString()));
-				}
+				LocalLogFunction(EMessageSeverity::Error, FString::Printf(TEXT("Hierarchy contains element %s which is not part of the replay."), *CurrentKey.ToString()));
 				bSuccess = false;
 			}
 		}
@@ -1747,28 +1881,19 @@ bool UControlRigReplay::ValidateExpectedResults(
 
 			if(CurrentParentIndex == INDEX_NONE && StoredParentIndex != INDEX_NONE)
 			{
-				if(InLogFunction)
-				{
-					InLogFunction(EMessageSeverity::Error, FString::Printf(TEXT("Element %s is at root currently, but is parented to %s in the replay."), *CurrentKey.ToString(), *StoredElementKeys[StoredParentIndex].ToString()));
-				}
+				LocalLogFunction(EMessageSeverity::Error, FString::Printf(TEXT("Element %s is at root currently, but is parented to %s in the replay."), *CurrentKey.ToString(), *StoredElementKeys[StoredParentIndex].ToString()));
 				bSuccess = false;
 			}
 			else if(CurrentParentIndex != INDEX_NONE && StoredParentIndex == INDEX_NONE)
 			{
-				if(InLogFunction)
-				{
-					InLogFunction(EMessageSeverity::Error, FString::Printf(TEXT("Element %s is parent to %s currently, but was root in the replay."), *CurrentKey.ToString(), *CurrentElementKeys[CurrentParentIndex].ToString()));
-				}
+				LocalLogFunction(EMessageSeverity::Error, FString::Printf(TEXT("Element %s is parent to %s currently, but was root in the replay."), *CurrentKey.ToString(), *CurrentElementKeys[CurrentParentIndex].ToString()));
 				bSuccess = false;
 			}
 			else if(CurrentParentIndex != INDEX_NONE && StoredParentIndex != INDEX_NONE)
 			{
 				if(CurrentElementKeys[CurrentParentIndex] != StoredElementKeys[StoredParentIndex])
 				{
-					if(InLogFunction)
-					{
-						InLogFunction(EMessageSeverity::Error, FString::Printf(TEXT("Element %s is parent to %s currently, but was parented to %s in the replay."), *CurrentKey.ToString(), *CurrentElementKeys[CurrentParentIndex].ToString(), *StoredElementKeys[StoredParentIndex].ToString()));
-					}
+					LocalLogFunction(EMessageSeverity::Error, FString::Printf(TEXT("Element %s is parent to %s currently, but was parented to %s in the replay."), *CurrentKey.ToString(), *CurrentElementKeys[CurrentParentIndex].ToString(), *StoredElementKeys[StoredParentIndex].ToString()));
 					bSuccess = false;
 				}
 			}
@@ -1779,6 +1904,11 @@ bool UControlRigReplay::ValidateExpectedResults(
 	{
 		for(const FRigElementKey& StoredKey : StoredElementKeys)
 		{
+			if(StoredKey.Type == ERigElementType::Connector)
+			{
+				continue;
+			}
+			
 			if(const FRigBaseElement* Element = Hierarchy->Find(StoredKey))
 			{
 				const FName& TrackName = OutputTracks.GetTrackName(StoredKey);
@@ -1790,10 +1920,7 @@ bool UControlRigReplay::ValidateExpectedResults(
 						const float CurrentValue = Hierarchy->GetCurveValueByIndex(Element->GetIndex());
 						if(!FMath::IsNearlyEqual(StoredValue, CurrentValue, static_cast<float>(Tolerance)))
 						{
-							if(InLogFunction)
-							{
-								InLogFunction(EMessageSeverity::Error, FString::Printf(TEXT("Frame [%04d]: Curve %s has value %.03f, expected by the replay: %.03f."), InPlaybackTimeIndex, *StoredKey.ToString(), CurrentValue, StoredValue));
-							}
+							LocalLogFunction(EMessageSeverity::Error, FString::Printf(TEXT("Frame [%04d]: Curve %s has value %.03f, expected by the replay: %.03f."), InPlaybackTimeIndex, *StoredKey.ToString(), CurrentValue, StoredValue));
 							bSuccess = false;
 						}
 					}
@@ -1807,26 +1934,17 @@ bool UControlRigReplay::ValidateExpectedResults(
 						
 						if(!StoredValue.GetLocation().Equals(CurrentValue.GetLocation(), static_cast<float>(Tolerance)))
 						{
-							if(InLogFunction)
-							{
-								InLogFunction(EMessageSeverity::Error, FString::Printf(TEXT("Frame [%04d]: Element %s has position %s, expected %s by the replay."), InPlaybackTimeIndex, *StoredKey.ToString(), *CurrentValue.GetLocation().ToString(), *StoredValue.GetLocation().ToString()));
-							}
+							LocalLogFunction(EMessageSeverity::Error, FString::Printf(TEXT("Frame [%04d]: Element %s has position %s, expected %s by the replay."), InPlaybackTimeIndex, *StoredKey.ToString(), *CurrentValue.GetLocation().ToString(), *StoredValue.GetLocation().ToString()));
 							bSuccess = false;
 						}
 						if(!StoredValue.Rotator().EqualsOrientation(CurrentValue.Rotator(), static_cast<float>(Tolerance)))
 						{
-							if(InLogFunction)
-							{
-								InLogFunction(EMessageSeverity::Error, FString::Printf(TEXT("Frame [%04d]: Element %s has rotation %s, expected %s by the replay."), InPlaybackTimeIndex, *StoredKey.ToString(), *CurrentValue.Rotator().ToString(), *StoredValue.Rotator().ToString()));
-							}
+							LocalLogFunction(EMessageSeverity::Error, FString::Printf(TEXT("Frame [%04d]: Element %s has rotation %s, expected %s by the replay."), InPlaybackTimeIndex, *StoredKey.ToString(), *CurrentValue.Rotator().ToString(), *StoredValue.Rotator().ToString()));
 							bSuccess = false;
 						}
 						if(!StoredValue.GetScale3D().Equals(CurrentValue.GetScale3D(), static_cast<float>(Tolerance)))
 						{
-							if(InLogFunction)
-							{
-								InLogFunction(EMessageSeverity::Error, FString::Printf(TEXT("Frame [%04d]: Element %s has scale %s, expected %s by the replay."), InPlaybackTimeIndex, *StoredKey.ToString(), *CurrentValue.GetScale3D().ToString(), *StoredValue.GetScale3D().ToString()));
-							}
+							LocalLogFunction(EMessageSeverity::Error, FString::Printf(TEXT("Frame [%04d]: Element %s has scale %s, expected %s by the replay."), InPlaybackTimeIndex, *StoredKey.ToString(), *CurrentValue.GetScale3D().ToString(), *StoredValue.GetScale3D().ToString()));
 							bSuccess = false;
 						}
 					}
@@ -1843,10 +1961,7 @@ bool UControlRigReplay::ValidateExpectedResults(
 			{
 				auto TrackNameBasedLogFunction = [InLogFunction](EMessageSeverity::Type InSeverity, const FName& InTrackName, const FString& InMessage)
 				{
-					if(InLogFunction)
-					{
-						InLogFunction(InSeverity, FString::Printf(TEXT("%s: %s"), *InTrackName.ToString(), *InMessage));
-					}
+					InLogFunction(InSeverity, FString::Printf(TEXT("%s: %s"), *InTrackName.ToString(), *InMessage));
 				};
 				
 				const TArray<FName>& StoredMetadataNames = OutputTracks.GetMetadataNames(InPlaybackTimeIndex, OutSampleTrackIndex, StoredKey, TrackNameBasedLogFunction);
@@ -2277,10 +2392,7 @@ bool UControlRigReplay::ValidateExpectedResults(
 					const uint8* CurrentMemory = Property->ContainerPtrToValuePtr<uint8>(InSubject);
 					if(!Property->Identical(ValueMemory.GetData(), CurrentMemory, PPF_None))
 					{
-						if(InLogFunction)
-						{
-							InLogFunction(EMessageSeverity::Error, FString::Printf(TEXT("Frame [%04d]: Variable '%s' value doesn't match replay."), InPlaybackTimeIndex, *Property->GetName()));
-						}
+						LocalLogFunction(EMessageSeverity::Error, FString::Printf(TEXT("Frame [%04d]: Variable '%s' value doesn't match replay."), InPlaybackTimeIndex, *Property->GetName()));
 						bSuccess = false;
 					}
 					Property->DestroyValue(ValueMemory.GetData());
@@ -2289,7 +2401,7 @@ bool UControlRigReplay::ValidateExpectedResults(
 				{
 					if(InPlaybackTimeIndex == 0 && InLogFunction)
 					{
-						InLogFunction(EMessageSeverity::Error, FString::Printf(TEXT("Frame [%04d]: Track for Variable '%s' has incorrect type."), InPlaybackTimeIndex, *Property->GetName()));
+						LocalLogFunction(EMessageSeverity::Error, FString::Printf(TEXT("Frame [%04d]: Track for Variable '%s' has incorrect type."), InPlaybackTimeIndex, *Property->GetName()));
 					}
 					bSuccess = false;
 				}
@@ -2298,7 +2410,7 @@ bool UControlRigReplay::ValidateExpectedResults(
 			{
 				if(InPlaybackTimeIndex == 0 && InLogFunction)
 				{
-					InLogFunction(EMessageSeverity::Error, FString::Printf(TEXT("Frame [%04d]: Track for Variable '%s' is missing in replay."), InPlaybackTimeIndex, *Property->GetName()));
+					LocalLogFunction(EMessageSeverity::Error, FString::Printf(TEXT("Frame [%04d]: Track for Variable '%s' is missing in replay."), InPlaybackTimeIndex, *Property->GetName()));
 				}
 				bSuccess = false;
 			}
@@ -2315,10 +2427,7 @@ bool UControlRigReplay::ValidateExpectedResults(
 				{
 					if(!CurrentVariableNames.Contains(StoredVariableName))
 					{
-						if(InLogFunction)
-						{
-							InLogFunction(EMessageSeverity::Error, FString::Printf(TEXT("Frame [%04d]: Variable '%s' for stored track is missing."), InPlaybackTimeIndex, *StoredVariableName.ToString()));
-						}
+						LocalLogFunction(EMessageSeverity::Error, FString::Printf(TEXT("Frame [%04d]: Variable '%s' for stored track is missing."), InPlaybackTimeIndex, *StoredVariableName.ToString()));
 						bSuccess = false;
 					}
 				}
@@ -2396,3 +2505,5 @@ void UControlRigReplay::HandlePreconstructionForTest(UControlRig* InRig, const F
 	}
 
 }
+
+#undef LOCTEXT_NAMESPACE
diff --git a/Engine/Plugins/Animation/ControlRig/Source/ControlRig/Private/ControlRigTestData.cpp b/Engine/Plugins/Animation/ControlRig/Source/ControlRig/Private/ControlRigTestData.cpp
--- a/Engine/Plugins/Animation/ControlRig/Source/ControlRig/Private/ControlRigTestData.cpp
+++ b/Engine/Plugins/Animation/ControlRig/Source/ControlRig/Private/ControlRigTestData.cpp
@@ -390,7 +390,7 @@ bool UControlRigTestData::StopRecording()
 	return false;
 }
 
-bool UControlRigTestData::StartReplay(UControlRig* InControlRig)
+bool UControlRigTestData::StartReplay(UControlRig* InControlRig, EControlRigReplayPlaybackMode InMode)
 {
 	StopRecording();
 	StopReplay();
diff --git a/Engine/Plugins/Animation/ControlRig/Source/ControlRig/Public/ControlRigReplay.h b/Engine/Plugins/Animation/ControlRig/Source/ControlRig/Public/ControlRigReplay.h
--- a/Engine/Plugins/Animation/ControlRig/Source/ControlRig/Public/ControlRigReplay.h
+++ b/Engine/Plugins/Animation/ControlRig/Source/ControlRig/Public/ControlRigReplay.h
@@ -139,6 +139,8 @@ private:
 	static bool ForEachElement(URigHierarchy* InHierarchy, const TArrayView<FRigElementKey>& InKeys, TFunction<void(FRigBaseElement*, bool&)> InFunction);
 	static bool ForEachComponent(URigHierarchy* InHierarchy, const TArrayView<FRigComponentKey>& InKeys, TFunction<void(FRigBaseComponent*, bool&)> InFunction);
 
+	static void FilterElementKeys(TArray<FRigElementKey>& InOutElementKeys);
+
 	static inline FLazyName HierarchyTopologyHashName = FLazyName(TEXT("HierarchyTopologyHash"));
 	static inline FLazyName HierarchyTopologyParentIndicesName = FLazyName(TEXT("HierarchyTopologyParentIndices"));
 	mutable TMap<FRigElementKey, FName> ElementKeyToTrackName;
@@ -203,6 +205,9 @@ public:
 	UFUNCTION(BlueprintPure, Category = "ControlRigReplay")
 	virtual FVector2D GetTimeRange() const;
 
+	UPROPERTY(EditAnywhere, BlueprintReadOnly, Category = "ControlRigReplay", meta=(MultiLine=true))
+   	FText Description;
+
 	UPROPERTY(AssetRegistrySearchable, VisibleAnywhere, BlueprintReadOnly, Category = "ControlRigReplay")
 	FSoftObjectPath ControlRigObjectPath;
 
@@ -243,26 +248,37 @@ public:
 	virtual bool StopRecording();
 
 	UFUNCTION(BlueprintCallable, Category = "ControlRigReplay")
-	virtual bool StartReplay(UControlRig* InControlRig);
+	virtual bool StartReplay(UControlRig* InControlRig, EControlRigReplayPlaybackMode InMode = EControlRigReplayPlaybackMode::ReplayInputs);
 
 	UFUNCTION(BlueprintCallable, Category = "ControlRigReplay")
 	virtual bool StopReplay();
 
+	UFUNCTION(BlueprintCallable, Category = "ControlRigReplay")
+	virtual bool PauseReplay();
+
 	UFUNCTION(BlueprintPure, Category = "ControlRigReplay")
 	virtual bool IsReplaying() const;
 
+	UFUNCTION(BlueprintPure, Category = "ControlRigReplay")
+	virtual bool IsPaused() const;
+
 	UFUNCTION(BlueprintPure, Category = "ControlRigReplay")
 	virtual bool IsRecording() const;
 
 	UFUNCTION(BlueprintPure, Category = "ControlRigReplay")
 	EControlRigReplayPlaybackMode GetPlaybackMode() const;
 
+	UFUNCTION(BlueprintCallable, Category = "ControlRigReplay")
+	void SetPlaybackMode(EControlRigReplayPlaybackMode InMode);
+
 	bool IsTestEnabled() const
 	{
 		return EnableTest;
 	}
 
 	virtual bool IsValidForTesting() const;
+	virtual bool HasValidationErrors() const;
+	virtual const TArray<FString>& GetValidationErrors() const;
 
 	virtual bool PerformTest(UControlRig* InSubject, TFunction<void(EMessageSeverity::Type, const FString&)> InLogFunction) const;
 	virtual bool ValidateExpectedResults(
@@ -285,8 +301,17 @@ protected:
 	double TimeAtStartOfRecording;
 	double TimeOfLastFrame;
 	double DesiredRecordingDuration;
-	bool bFirstPlayThrough = false;
 	bool bStoreVariablesDuringPreEvent = false;
+	bool bReplayPaused = false;
+	EControlRigReplayPlaybackMode PlaybackMode = EControlRigReplayPlaybackMode::Live;
+	mutable TArray<FString> LastValidationWarningsAndErrors;
+
+	static const FText LiveStatus; 
+	static const FText LiveStatusTooltip; 
+	static const FText ReplayInputsStatus;
+	static const FText ReplayInputsStatusTooltip;
+	static const FText GroundTruthStatus;
+	static const FText GroundTruthStatusTooltip;
 
 	friend class FControlRigEditor;
 };
\ No newline at end of file
diff --git a/Engine/Plugins/Animation/ControlRig/Source/ControlRig/Public/ControlRigTestData.h b/Engine/Plugins/Animation/ControlRig/Source/ControlRig/Public/ControlRigTestData.h
--- a/Engine/Plugins/Animation/ControlRig/Source/ControlRig/Public/ControlRigTestData.h
+++ b/Engine/Plugins/Animation/ControlRig/Source/ControlRig/Public/ControlRigTestData.h
@@ -71,7 +71,7 @@ public:
 
 	virtual bool StartRecording(UControlRig* InControlRig) override;
 	virtual bool StopRecording() override;
-	virtual bool StartReplay(UControlRig* InControlRig) override;
+	virtual bool StartReplay(UControlRig* InControlRig, EControlRigReplayPlaybackMode InMode = EControlRigReplayPlaybackMode::ReplayInputs) override;
 	virtual bool StopReplay() override;
 
 	UPROPERTY(VisibleAnywhere, BlueprintReadOnly, Category = "ControlRigTestData")
diff --git a/Engine/Plugins/Animation/ControlRig/Source/ControlRigEditor/Private/ControlRigEditorStyle.h b/Engine/Plugins/Animation/ControlRig/Source/ControlRigEditor/Private/ControlRigEditorStyle.h
--- a/Engine/Plugins/Animation/ControlRig/Source/ControlRigEditor/Private/ControlRigEditorStyle.h
+++ b/Engine/Plugins/Animation/ControlRig/Source/ControlRigEditor/Private/ControlRigEditorStyle.h
@@ -53,6 +53,9 @@ public:
 		SetContentRoot(ControlRigPluginContentDir);
 		SetCoreContentRoot(EngineEditorSlateDir);
 
+		const FSlateColor ErrorBorderColor = FAppStyle::GetSlateColor("Colors.Error");
+		FLinearColor ErrorFillColor = ErrorBorderColor.GetSpecifiedColor();
+
 		const FSlateColor DefaultForeground(FLinearColor(0.72f, 0.72f, 0.72f, 1.f));
 
 		// Class Icons
@@ -118,6 +121,8 @@ public:
 				Set( "ControlRig.Viewport.Notification.ChangeShapeTransform", new BOX_BRUSH("Common/RoundedSelection_16x", 4.0f/16.0f, FLinearColor(FColor(169, 0, 148))));
 				// similar style to "AnimViewport.Notification.Warning"
 				Set( "ControlRig.Viewport.Notification.DirectManipulation", new BOX_BRUSH("Common/RoundedSelection_16x", 4.0f/16.0f, FLinearColor(FColor(0, 112, 224))));
+				// similar style to "AnimViewport.Notification.Warning"
+				Set( "ControlRig.Viewport.Notification.ReplayValidation", new BOX_BRUSH("Common/RoundedSelection_16x", 4.0f/16.0f, ErrorFillColor));
 			}
 		}
 
diff --git a/Engine/Plugins/Animation/ControlRig/Source/ControlRigEditor/Private/EditMode/ControlRigEditMode.cpp b/Engine/Plugins/Animation/ControlRig/Source/ControlRigEditor/Private/EditMode/ControlRigEditMode.cpp
--- a/Engine/Plugins/Animation/ControlRig/Source/ControlRigEditor/Private/EditMode/ControlRigEditMode.cpp
+++ b/Engine/Plugins/Animation/ControlRig/Source/ControlRigEditor/Private/EditMode/ControlRigEditMode.cpp
@@ -3344,11 +3344,17 @@ void FControlRigEditMode::OpenSpacePickerWidget()
 			if (RuntimeRig->IsAdditive())
 			{
 				const FTransform Transform = RuntimeRig->GetControlGlobalTransform(InControlKey.Name);
-			   RuntimeRig->SwitchToParent(InControlKey, InSpaceKey, false, true);
-			   RuntimeRig->Evaluate_AnyThread();
-			   FRigControlValue ControlValue = RuntimeRig->GetControlValueFromGlobalTransform(InControlKey.Name, Transform, ERigTransformType::CurrentGlobal);
-			   RuntimeRig->SetControlValue(InControlKey.Name, ControlValue);
-			   RuntimeRig->Evaluate_AnyThread();
+				RuntimeRig->SwitchToParent(InControlKey, InSpaceKey, false, true);
+				{
+					const TGuardValue<float> AbsoluteTimeGuard(RuntimeRig->AbsoluteTime, RuntimeRig->AbsoluteTime);
+				    RuntimeRig->Evaluate_AnyThread();
+				}
+				FRigControlValue ControlValue = RuntimeRig->GetControlValueFromGlobalTransform(InControlKey.Name, Transform, ERigTransformType::CurrentGlobal);
+				RuntimeRig->SetControlValue(InControlKey.Name, ControlValue);
+				{
+					const TGuardValue<float> AbsoluteTimeGuard(RuntimeRig->AbsoluteTime, RuntimeRig->AbsoluteTime);
+					RuntimeRig->Evaluate_AnyThread();
+				}
 			}
 			else
 			{
@@ -3652,7 +3658,10 @@ void FControlRigEditMode::ZeroTransforms(bool bSelectionOnly)
 
 		if (bHasNonDefaultParent == true) //now we have the initial pose setup we need to get the global transforms as specified now then set them in the current parent space
 		{
-			ControlRig->Evaluate_AnyThread();
+			{
+				const TGuardValue<float> AbsoluteTimeGuard(ControlRig->AbsoluteTime, ControlRig->AbsoluteTime);
+				ControlRig->Evaluate_AnyThread();
+			}
 
 			//get global transforms
 			TMap<FRigElementKey, FTransform> GlobalTransforms;
@@ -3788,7 +3797,10 @@ void FControlRigEditMode::InvertInputPose(bool bSelectionOnly)
 		}
 
 		const TArray<FRigControlElement*> ModifiedElements = ControlRig->InvertInputPose(SelectedRigElements, EControlRigSetKey::Never);
-		ControlRig->Evaluate_AnyThread();
+		{
+			const TGuardValue<float> AbsoluteTimeGuard(ControlRig->AbsoluteTime, ControlRig->AbsoluteTime);
+			ControlRig->Evaluate_AnyThread();
+		}
 
 		for (FRigControlElement* ControlElement : ModifiedElements)
 		{
@@ -5102,6 +5114,7 @@ void FControlRigEditMode::SetControlShapeTransform(
 		if (ControlRig->IsAdditive())
 		{
 			TGuardValue<bool> CompensateGuard(FMovieSceneConstraintChannelHelper::bDoNotCompensate, true);
+			const TGuardValue<float> AbsoluteTimeGuard(ControlRig->AbsoluteTime, ControlRig->AbsoluteTime);
 			ControlRig->Evaluate_AnyThread();
 		}
 	};
@@ -5433,6 +5446,7 @@ void FControlRigEditMode::MoveControlShape( AControlRigShapeActor* ShapeActor,
 		}
 		else
 		{
+			const TGuardValue<float> AbsoluteTimeGuard(ControlRig->AbsoluteTime, ControlRig->AbsoluteTime);
 			ControlRig->Evaluate_AnyThread();
 		}
 
@@ -5454,6 +5468,7 @@ void FControlRigEditMode::MoveControlShape( AControlRigShapeActor* ShapeActor,
 				}
 				else
 				{
+					const TGuardValue<float> AbsoluteTimeGuard(ControlRig->AbsoluteTime, ControlRig->AbsoluteTime);
 					ControlRig->Evaluate_AnyThread();
 				}
 			}
@@ -6638,6 +6653,7 @@ FControlRigEditMode::FPendingControlRigEvaluator::~FPendingControlRigEvaluator()
 	{
 		for(UControlRig* ControlRig : EditMode->RigsToEvaluateDuringThisTick)
 		{
+			const TGuardValue<float> AbsoluteTimeGuard(ControlRig->AbsoluteTime, ControlRig->AbsoluteTime);
 			ControlRig->Evaluate_AnyThread();
 		}
 		EditMode->RigsToEvaluateDuringThisTick.Reset();
diff --git a/Engine/Plugins/Animation/ControlRig/Source/ControlRigEditor/Private/Editor/ControlRigEditor.cpp b/Engine/Plugins/Animation/ControlRig/Source/ControlRigEditor/Private/Editor/ControlRigEditor.cpp
--- a/Engine/Plugins/Animation/ControlRig/Source/ControlRigEditor/Private/Editor/ControlRigEditor.cpp
+++ b/Engine/Plugins/Animation/ControlRig/Source/ControlRigEditor/Private/Editor/ControlRigEditor.cpp
@@ -121,6 +121,8 @@
 #include "Editor/Persona/Private/SAnimationEditorViewport.h"
 #include "ViewportToolbar/UnrealEdViewportToolbar.h"
 #include "ViewportToolbar/UnrealEdViewportToolbarContext.h"
+#include "ITransportControl.h"
+#include "PropertyPath.h"
 
 #define LOCTEXT_NAMESPACE "ControlRigEditor"
 
@@ -751,14 +753,65 @@ TSharedRef<SWidget> FControlRigEditor::GenerateReplayAssetRecordMenuContent()
 	return MenuBuilder.MakeWidget();
 }
 
+TSharedRef<SWidget> FControlRigEditor::GenerateReplayAssetPlaybackMenuContent()
+{
+	FMenuBuilder MenuBuilder(true, GetToolkitCommands());
+
+	MenuBuilder.AddMenuEntry(
+	UControlRigReplay::LiveStatus,
+	UControlRigReplay::LiveStatusTooltip,
+	FSlateIcon(FControlRigEditorStyle::Get().GetStyleSetName(),"ClassIcon.ControlRigBlueprint"),
+	FUIAction(
+			FExecuteAction::CreateLambda([this]()
+			{
+				if(ReplayStrongPtr.IsValid())
+				{
+					ReplayStrongPtr->StopReplay();
+				}
+			}
+		)	
+	));
+
+	MenuBuilder.AddMenuEntry(
+	UControlRigReplay::ReplayInputsStatus,
+	UControlRigReplay::ReplayInputsStatusTooltip,
+	FSlateIcon(FControlRigEditorStyle::Get().GetStyleSetName(),"ClassIcon.ControlRigSequence"),
+	FUIAction(
+			FExecuteAction::CreateLambda([this]()
+			{
+				if(ReplayStrongPtr.IsValid())
+				{
+					ReplayStrongPtr->StartReplay(GetControlRig(), EControlRigReplayPlaybackMode::ReplayInputs);
+				}
+			}
+		)	
+	));
+
+	MenuBuilder.AddMenuEntry(
+	UControlRigReplay::GroundTruthStatus,
+	UControlRigReplay::GroundTruthStatusTooltip,
+	FSlateIcon(FControlRigEditorStyle::Get().GetStyleSetName(),"ClassIcon.ControlRigSequence"),
+	FUIAction(
+			FExecuteAction::CreateLambda([this]()
+			{
+				if(ReplayStrongPtr.IsValid())
+				{
+					ReplayStrongPtr->StartReplay(GetControlRig(), EControlRigReplayPlaybackMode::GroundTruth);
+				}
+			}
+		)	
+	));
+
+	return MenuBuilder.MakeWidget();
+}
+
 bool FControlRigEditor::RecordReplay(double InRecordingDuration)
 {
 	if(GetControlRig() == nullptr)
 	{
 		return false;
 	}
-	
-	
+
 	// create a new test asset
 	static const FString Folder = TEXT("/Game/Animation/ControlRig/NoCook/");
 	const FString DesiredPackagePath =  FString::Printf(TEXT("%s/%s_Replay"), *Folder, *GetControlRigBlueprint()->GetName());
@@ -774,6 +827,7 @@ bool FControlRigEditor::RecordReplay(double InRecordingDuration)
 
 	if(UControlRigReplay* Replay = ReplayStrongPtr.Get())
 	{
+		Replay->Modify();
 		Replay->DesiredRecordingDuration = InRecordingDuration;
 
 		if(InRecordingDuration <= SMALL_NUMBER)
@@ -812,13 +866,18 @@ void FControlRigEditor::ToggleReplay()
 		switch(ReplayStrongPtr->GetPlaybackMode())
 		{
 			case EControlRigReplayPlaybackMode::ReplayInputs:
+			{
+				ReplayStrongPtr->StartReplay(GetControlRig(), EControlRigReplayPlaybackMode::GroundTruth);
+				break;
+			}
+			case EControlRigReplayPlaybackMode::GroundTruth:
 			{
 				ReplayStrongPtr->StopReplay();
 				break;
 			}
 			default:
 			{
-				ReplayStrongPtr->StartReplay(GetControlRig());
+				ReplayStrongPtr->StartReplay(GetControlRig(), EControlRigReplayPlaybackMode::ReplayInputs);
 				break;
 			}
 		}
@@ -937,15 +996,29 @@ void FControlRigEditor::FillToolbar(FToolBarBuilder& ToolbarBuilder, bool bEndSe
 		{
 			ToolbarBuilder.AddSeparator();
 
-			ToolbarBuilder.AddComboButton(
-				FUIAction(FExecuteAction(), FCanExecuteAction::CreateLambda([this]()
+			const FCanExecuteAction OnlyWhenNotRecordingAction = FCanExecuteAction::CreateLambda([this]()
+			{
+				if(ReplayStrongPtr.IsValid())
 				{
-					if(ReplayStrongPtr.IsValid())
-					{
-						return !ReplayStrongPtr->IsReplaying() && !ReplayStrongPtr->IsRecording() && !RecordReplayTimerHandle.IsValid();
-					}
-					return true;
-				})),
+					return !ReplayStrongPtr->IsReplaying() && !ReplayStrongPtr->IsRecording() && !RecordReplayTimerHandle.IsValid();
+				}
+				return true;
+			});
+
+			const FCanExecuteAction OnlyWithValidReplayAction = FCanExecuteAction::CreateLambda([this]()
+			{
+				if(ReplayStrongPtr.IsValid())
+				{
+					return !ReplayStrongPtr->IsRecording() && !RecordReplayTimerHandle.IsValid();
+				}
+				return false;
+			});
+
+			const FUIAction EmptyOnlyWhenNotRecordingAction(FExecuteAction(), OnlyWhenNotRecordingAction);
+			const FUIAction EmptyOnlyWithValidReplayAction(FExecuteAction(), OnlyWithValidReplayAction);
+
+			ToolbarBuilder.AddComboButton(
+				EmptyOnlyWhenNotRecordingAction,
 				FOnGetContent::CreateSP(this, &FControlRigEditor::GenerateReplayAssetModeMenuContent),
 				TAttribute<FText>::CreateSP(this, &FControlRigEditor::GetReplayAssetName),
 				TAttribute<FText>::CreateSP(this, &FControlRigEditor::GetReplayAssetTooltip),
@@ -957,14 +1030,7 @@ void FControlRigEditor::FillToolbar(FToolBarBuilder& ToolbarBuilder, bool bEndSe
 					{
 						RecordReplay(0.0);
 					}),
-					FCanExecuteAction::CreateLambda([this]()
-					{
-						if(ReplayStrongPtr.IsValid())
-						{
-							return !ReplayStrongPtr->IsReplaying() && !ReplayStrongPtr->IsRecording() && !RecordReplayTimerHandle.IsValid();
-						}
-						return true;
-					})
+					OnlyWhenNotRecordingAction
 				),
 				NAME_None,
 				LOCTEXT("ReplayRecordButton", "Record"),
@@ -972,7 +1038,7 @@ void FControlRigEditor::FillToolbar(FToolBarBuilder& ToolbarBuilder, bool bEndSe
 				FSlateIcon(FControlRigEditorStyle::Get().GetStyleSetName(),"ControlRig.Replay.Record")
 				);
 			ToolbarBuilder.AddComboButton(
-				FUIAction(),
+				EmptyOnlyWhenNotRecordingAction,
 				FOnGetContent::CreateSP(this, &FControlRigEditor::GenerateReplayAssetRecordMenuContent),
 				LOCTEXT("ReplayRecordMenu_Label", "Recording Modes"),
 				LOCTEXT("ReplayRecordMenu_ToolTip", "Pick between different modes for recording"),
@@ -984,65 +1050,52 @@ void FControlRigEditor::FillToolbar(FToolBarBuilder& ToolbarBuilder, bool bEndSe
 					{
 						ToggleReplay();
 					}),
-					FCanExecuteAction::CreateLambda([this]
-					{
-						if(ReplayStrongPtr.IsValid())
-						{
-							return !ReplayStrongPtr->IsRecording();
-						}
-						return false;
-					})
+					OnlyWithValidReplayAction
 				),
 				NAME_None,
 				TAttribute<FText>::CreateLambda([this]()
 				{
-					static const FText LiveStatus = LOCTEXT("LiveStatus", "Live"); 
-					static const FText ReplayInputsStatus = LOCTEXT("ReplayInputsStatus", "Replay Inputs");
-					//static const FText GroundTruthStatus = LOCTEXT("GroundTruthStatus", "Ground Truth");
 					if(ReplayStrongPtr.IsValid())
 					{
 						switch(ReplayStrongPtr->GetPlaybackMode())
 						{
 							case EControlRigReplayPlaybackMode::ReplayInputs:
 							{
-								return ReplayInputsStatus;
+								return UControlRigReplay::ReplayInputsStatus;
+							}
+							case EControlRigReplayPlaybackMode::GroundTruth:
+							{
+								return UControlRigReplay::GroundTruthStatus;
 							}
-							//case EControlRigReplayPlaybackMode::GroundTruth:
-							//{
-							//	return GroundTruthStatus;
-							//}
 							default:
 							{
 								break;
 							}
 						}
 					}
-					return LiveStatus;
+					return UControlRigReplay::LiveStatus;
 				}),
 				TAttribute<FText>::CreateLambda([this]()
 				{
-					static const FText LiveStatusTooltip = LOCTEXT("LiveStatusTooltip", "The replay is not affecting the rig."); 
-					static const FText ReplayInputsStatusTooltip = LOCTEXT("ReplayInputsStatusTooltip", "The replay inputs are being replayed onto the rig.");
-					//static const FText GroundTruthStatusTooltip = LOCTEXT("GroundTruthStatusTooltip", "The complete result of the rig is being overwritten.");
 					if(ReplayStrongPtr.IsValid())
 					{
 						switch(ReplayStrongPtr->GetPlaybackMode())
 						{
 							case EControlRigReplayPlaybackMode::ReplayInputs:
 							{
-								return ReplayInputsStatusTooltip;
+								return UControlRigReplay::ReplayInputsStatusTooltip;
+							}
+							case EControlRigReplayPlaybackMode::GroundTruth:
+							{
+								return UControlRigReplay::GroundTruthStatusTooltip;
 							}
-							//case EControlRigReplayPlaybackMode::GroundTruth:
-							//{
-							//	return GroundTruthStatusTooltip;
-							//}
 							default:
 							{
 								break;
 							}
 						}
 					}
-					return LiveStatusTooltip;
+					return UControlRigReplay::LiveStatusTooltip;
 				}),
 				TAttribute<FSlateIcon>::CreateLambda([this]()
 				{
@@ -1056,6 +1109,14 @@ void FControlRigEditor::FillToolbar(FToolBarBuilder& ToolbarBuilder, bool bEndSe
 				}),
 				EUserInterfaceActionType::Button
 			);
+
+			ToolbarBuilder.AddComboButton(
+				EmptyOnlyWithValidReplayAction,
+				FOnGetContent::CreateSP(this, &FControlRigEditor::GenerateReplayAssetPlaybackMenuContent),
+				LOCTEXT("ReplayPlaybackModeMenu_Label", "Playback Modes"),
+				LOCTEXT("ReplayPlaybackModeMenu_ToolTip", "Pick between different modes for playback"),
+				FSlateIcon(FAppStyle::GetAppStyleSetName(), "LevelEditor.Recompile"),
+				true);
 		}
 	}
 	
@@ -2328,7 +2389,38 @@ void FControlRigEditor::HandleViewportCreated(const TSharedRef<class IPersonaVie
 			],
 			FPersonaViewportNotificationOptions(TAttribute<EVisibility>::CreateSP(this, &FControlRigEditor::GetConnectorWarningVisibility))
 		);
-	}	
+	}
+
+	if(CVarControlRigShowTestingToolbar.GetValueOnAnyThread())
+	{
+		{
+			FPersonaViewportNotificationOptions ReplayValidationNotificationOptions(TAttribute<EVisibility>::CreateSP<>(this, &FControlRigEditor::GetReplayValidationErrorVisibility));
+			ReplayValidationNotificationOptions.OnGetBrushOverride = TAttribute<const FSlateBrush*>(FControlRigEditorStyle::Get().GetBrush("ControlRig.Viewport.Notification.ReplayValidation"));
+
+			InViewport->AddNotification(
+				EMessageSeverity::Info,
+				false,
+				SNew(SHorizontalBox)
+				.Visibility(this, &FControlRigEditor::GetReplayValidationErrorVisibility)
+				.ToolTipText(this, &FControlRigEditor::GetReplayValidationErrorTooltip)
+				+SHorizontalBox::Slot()
+				.FillWidth(1.0f)
+				.Padding(4.0f, 4.0f)
+				[
+					SNew(SHorizontalBox)
+					+SHorizontalBox::Slot()
+					.VAlign(VAlign_Center)
+					.AutoWidth()
+					[
+						SNew(STextBlock)
+						.TextStyle(FAppStyle::Get(), "AnimViewport.MessageText")
+						.Text(LOCTEXT("ReplayValidationErrorButtonText", "Replay Validation Error"))
+					]
+				],
+				ReplayValidationNotificationOptions
+			);
+		}
+	}
 
 	InViewport->AddNotification(MakeAttributeLambda(GetErrorSeverity),
 		false,
@@ -5925,4 +6017,215 @@ void FControlRigEditor::OnPostConstruction_AnyThread(UControlRig* InRig, const F
 	}
 }
 
+void FControlRigEditor::SetupTimelineDelegates(FAnimationScrubPanelDelegates& InOutDelegates)
+{
+	InOutDelegates.IsRecordingActiveDelegate.BindSP(this, &FControlRigEditor::HandleReplayIsRecordingActive);
+	InOutDelegates.GetRecordingVisibilityDelegate.BindSP(this, &FControlRigEditor::HandleGetReplayRecordButtonVisibility);
+	InOutDelegates.StartRecordingDelegate.BindSP(this, &FControlRigEditor::HandleReplayStartRecording);
+	InOutDelegates.StopRecordingDelegate.BindSP(this, &FControlRigEditor::HandleReplayStopRecording);
+	InOutDelegates.GetPlaybackModeDelegate.BindSP(this, &FControlRigEditor::HandleReplayGetPlaybackMode);
+	InOutDelegates.SetPlaybackModeDelegate.BindSP(this, &FControlRigEditor::HandleReplaySetPlaybackMode);
+	InOutDelegates.GetPlaybackTimeDelegate.BindSP(this, &FControlRigEditor::HandleReplayGetPlaybackTime);
+	InOutDelegates.SetPlaybackTimeDelegate.BindSP(this, &FControlRigEditor::HandleReplaySetPlaybackTime);
+	InOutDelegates.StepForwardDelegate.BindSP(this, &FControlRigEditor::HandleReplayStepForward);
+	InOutDelegates.StepBackwardDelegate.BindSP(this, &FControlRigEditor::HandleReplayStepBackward);
+	InOutDelegates.GetIsLoopingDelegate.BindSP(this, &FControlRigEditor::HandleReplayGetIsLooping);
+	InOutDelegates.SetIsLoopingDelegate.BindSP(this, &FControlRigEditor::HandleReplaySetIsLooping);
+	InOutDelegates.GetPlaybackTimeRangeDelegate.BindSP(this, &FControlRigEditor::HandleReplayGetPlaybackTimeRange);
+	InOutDelegates.GetNumberOfKeysDelegate.BindSP(this, &FControlRigEditor::HandleReplayGetNumberOfKeys);
+}
+
+bool FControlRigEditor::ShowReplayOnTimeline() const
+{
+	if(ReplayStrongPtr.IsValid())
+	{
+		if(GetControlRig()->Replay.Get() == ReplayStrongPtr.Get())
+		{
+			return !ReplayStrongPtr->IsRecording();
+		}
+	}
+	return false;
+}
+
+TOptional<bool> FControlRigEditor::HandleReplayIsRecordingActive() const
+{
+	if(!ShowReplayOnTimeline())
+	{
+		return TOptional<bool>();
+	}
+	return ReplayStrongPtr->IsRecording();
+}
+
+TOptional<EVisibility> FControlRigEditor::HandleGetReplayRecordButtonVisibility() const
+{
+	if(ReplayStrongPtr.IsValid())
+	{
+		return EVisibility::Collapsed;
+	}
+	return TOptional<EVisibility>();
+}
+
+bool FControlRigEditor::HandleReplayStartRecording()
+{
+	// for now this button won't be supported from here.
+	return false;
+}
+
+bool FControlRigEditor::HandleReplayStopRecording()
+{
+	// for now this button won't be supported from here.
+	return false;
+}
+
+TOptional<int32> FControlRigEditor::HandleReplayGetPlaybackMode() const
+{
+	if(!ShowReplayOnTimeline())
+	{
+		return TOptional<int32>();
+	}
+	return (ReplayStrongPtr->IsReplaying() && !ReplayStrongPtr->IsPaused()) ? (int32)EPlaybackMode::PlayingForward : (int32)EPlaybackMode::Stopped; 
+}
+
+bool FControlRigEditor::HandleReplaySetPlaybackMode(int32 InPlaybackMode)
+{
+	if(!ShowReplayOnTimeline())
+	{
+		return false;
+	}
+	if(InPlaybackMode == (int32)EPlaybackMode::Stopped)
+	{
+		ReplayStrongPtr->PauseReplay();
+	}
+	else
+	{
+		ReplayStrongPtr->StartReplay(GetControlRig());
+	}
+	return true;
+}
+
+TOptional<float> FControlRigEditor::HandleReplayGetPlaybackTime() const
+{
+	if(!ShowReplayOnTimeline())
+	{
+		return TOptional<float>();
+	}
+	return GetControlRig()->GetAbsoluteTime();
+}
+
+bool FControlRigEditor::HandleReplaySetPlaybackTime(float InTime, bool bStopPlayback)
+{
+	if(!ShowReplayOnTimeline())
+	{
+		return false;
+	}
+	if(bStopPlayback)
+	{
+		HandleReplaySetPlaybackMode(EPlaybackMode::Stopped);
+	}
+
+	const int32 TimeIndex = ReplayStrongPtr->InputTracks.GetTimeIndex(InTime);
+	if(TimeIndex != INDEX_NONE)
+	{
+		GetControlRig()->SetReplayTimeIndex(TimeIndex);
+	}
+	return true;
+}
+
+bool FControlRigEditor::HandleReplayStepForward()
+{
+	if(!ShowReplayOnTimeline())
+	{
+		return false;
+	}
+
+	HandleReplaySetPlaybackMode(EPlaybackMode::Stopped);
+
+	const int32 PreviousTimeIndex = GetControlRig()->GetReplayTimeIndex();
+	if(PreviousTimeIndex < ReplayStrongPtr->InputTracks.GetNumTimes() - 1)
+	{
+		GetControlRig()->SetReplayTimeIndex(PreviousTimeIndex + 1);
+	}
+	return true;
+}
+
+bool FControlRigEditor::HandleReplayStepBackward()
+{
+	if(!ShowReplayOnTimeline())
+	{
+		return false;
+	}
+	
+	HandleReplaySetPlaybackMode(EPlaybackMode::Stopped);
+
+	const int32 PreviousTimeIndex = GetControlRig()->GetReplayTimeIndex();
+	if(PreviousTimeIndex > 0)
+	{
+		GetControlRig()->SetReplayTimeIndex(PreviousTimeIndex - 1);
+	}
+	return true;
+}
+
+TOptional<bool> FControlRigEditor::HandleReplayGetIsLooping() const
+{
+	if(!ShowReplayOnTimeline())
+	{
+		return TOptional<bool>();
+	}
+	// we are always set to loop
+	return true;
+}
+
+bool FControlRigEditor::HandleReplaySetIsLooping(bool bIsLooping)
+{
+	// we are always set to loop for now.
+	return true;
+}
+
+TOptional<FVector2f> FControlRigEditor::HandleReplayGetPlaybackTimeRange() const
+{
+	if(!ShowReplayOnTimeline())
+	{
+		return TOptional<FVector2f>();
+	}
+	return FVector2f(ReplayStrongPtr->GetTimeRange());
+}
+
+TOptional<uint32> FControlRigEditor::HandleReplayGetNumberOfKeys() const
+{
+	if(!ShowReplayOnTimeline())
+	{
+		return TOptional<uint32>();
+	}
+	return static_cast<uint32>(ReplayStrongPtr->InputTracks.GetNumTimes());
+}
+
+EVisibility FControlRigEditor::GetReplayValidationErrorVisibility() const
+{
+	if(ReplayStrongPtr.IsValid())
+	{
+		if(ReplayStrongPtr->IsReplaying() &&
+			ReplayStrongPtr->GetPlaybackMode() == EControlRigReplayPlaybackMode::ReplayInputs && 
+			ReplayStrongPtr->HasValidationErrors())
+		{
+			return EVisibility::Visible;
+		}
+	}
+	return EVisibility::Hidden;
+}
+
+FText FControlRigEditor::GetReplayValidationErrorTooltip() const
+{
+	if(ReplayStrongPtr.IsValid())
+	{
+		const TArray<FString>& ValidationErrors = ReplayStrongPtr->GetValidationErrors();
+		if(!ValidationErrors.IsEmpty())
+		{
+			static const FText Format = LOCTEXT("ReplayValidationErrorTooltipFormat", "The results from the rig don't match the expected values in the replay.\nSwitch to ground truth to compare or check the output log.\n{0}");
+			const TArrayView<const FString> ClampedValidationErrors(ValidationErrors.GetData(), FMath::Min(20, ValidationErrors.Num()));
+			return FText::Format(Format, FText::FromString(FString::Join(ClampedValidationErrors, TEXT("\n"))));
+		}
+	}
+	return FText();
+}
+
 #undef LOCTEXT_NAMESPACE
diff --git a/Engine/Plugins/Animation/ControlRig/Source/ControlRigEditor/Private/Editor/ControlRigEditor.h b/Engine/Plugins/Animation/ControlRig/Source/ControlRigEditor/Private/Editor/ControlRigEditor.h
--- a/Engine/Plugins/Animation/ControlRig/Source/ControlRigEditor/Private/Editor/ControlRigEditor.h
+++ b/Engine/Plugins/Animation/ControlRig/Source/ControlRigEditor/Private/Editor/ControlRigEditor.h
@@ -27,6 +27,7 @@
 #include "SchematicGraphPanel/SSchematicGraphPanel.h"
 #include "Units/RigUnit.h"
 #include "ControlRigSchematicModel.h"
+#include "PersonaModule.h"
 
 class UControlRigBlueprint;
 class IPersonaToolkit;
@@ -296,6 +297,7 @@ private:
 	bool SetReplayAssetPath(const FString& InAssetPath);
 	TSharedRef<SWidget> GenerateReplayAssetModeMenuContent();
 	TSharedRef<SWidget> GenerateReplayAssetRecordMenuContent();
+	TSharedRef<SWidget> GenerateReplayAssetPlaybackMenuContent();
 	bool RecordReplay(double InRecordingDuration);
 	void ToggleReplay();
 
@@ -373,6 +375,26 @@ protected:
 	bool bSchematicViewPortIsHidden;
 	FTimerHandle RecordReplayTimerHandle;
 
+	void SetupTimelineDelegates(FAnimationScrubPanelDelegates& InOutDelegates);
+	bool ShowReplayOnTimeline() const;
+	TOptional<bool> HandleReplayIsRecordingActive() const;
+	TOptional<EVisibility> HandleGetReplayRecordButtonVisibility() const;
+	bool HandleReplayStartRecording();
+	bool HandleReplayStopRecording();
+	TOptional<int32> HandleReplayGetPlaybackMode() const;
+	bool HandleReplaySetPlaybackMode(int32 InPlaybackMode);
+	TOptional<float> HandleReplayGetPlaybackTime() const;
+	bool HandleReplaySetPlaybackTime(float InTime, bool bStopPlayback);
+	bool HandleReplayStepForward();
+	bool HandleReplayStepBackward();
+	TOptional<bool> HandleReplayGetIsLooping() const;
+	bool HandleReplaySetIsLooping(bool bIsLooping);
+	TOptional<FVector2f> HandleReplayGetPlaybackTimeRange() const;
+	TOptional<uint32> HandleReplayGetNumberOfKeys() const;
+	
+	EVisibility GetReplayValidationErrorVisibility() const;
+	FText GetReplayValidationErrorTooltip() const;
+
 	static const TArray<FName> ForwardsSolveEventQueue;
 	static const TArray<FName> BackwardsSolveEventQueue;
 	static const TArray<FName> ConstructionEventQueue;
diff --git a/Engine/Plugins/Animation/ControlRig/Source/ControlRigEditor/Private/Editor/ControlRigEditorMode.cpp b/Engine/Plugins/Animation/ControlRig/Source/ControlRigEditor/Private/Editor/ControlRigEditorMode.cpp
--- a/Engine/Plugins/Animation/ControlRig/Source/ControlRigEditor/Private/Editor/ControlRigEditorMode.cpp
+++ b/Engine/Plugins/Animation/ControlRig/Source/ControlRigEditor/Private/Editor/ControlRigEditorMode.cpp
@@ -38,6 +38,8 @@ FControlRigEditorMode::FControlRigEditorMode(const TSharedRef<FControlRigEditor>
 	ViewportArgs.bShowTurnTable = false;
 	ViewportArgs.bAlwaysShowTransformToolbar = true;
 	ViewportArgs.OnViewportCreated = FOnViewportCreated::CreateSP(InControlRigEditor, &FControlRigEditor::HandleViewportCreated);
+
+	InControlRigEditor->SetupTimelineDelegates(ViewportArgs.TimelineDelegates);
  
  	TabFactories.RegisterFactory(PersonaModule.CreatePersonaViewportTabFactory(InControlRigEditor, ViewportArgs));
 	TabFactories.RegisterFactory(PersonaModule.CreateAdvancedPreviewSceneTabFactory(InControlRigEditor, InControlRigEditor->GetPersonaToolkit()->GetPreviewScene()));
diff --git a/Engine/Source/Editor/EditorWidgets/Private/STransportControl.cpp b/Engine/Source/Editor/EditorWidgets/Private/STransportControl.cpp
--- a/Engine/Source/Editor/EditorWidgets/Private/STransportControl.cpp
+++ b/Engine/Source/Editor/EditorWidgets/Private/STransportControl.cpp
@@ -27,7 +27,7 @@
 #define LOCTEXT_NAMESPACE "STransportControl"
 
 BEGIN_SLATE_FUNCTION_BUILD_OPTIMIZATION
-TSharedPtr<SWidget> STransportControl::MakeTransportControlWidget(ETransportControlWidgetType WidgetType, bool bAreButtonsFocusable, const FOnMakeTransportWidget& MakeCustomWidgetDelegate)
+TSharedPtr<SWidget> STransportControl::MakeTransportControlWidget(ETransportControlWidgetType WidgetType, const FOnMakeTransportWidget& MakeCustomWidgetDelegate)
 {
 	switch(WidgetType)
 	{
@@ -38,7 +38,7 @@ TSharedPtr<SWidget> STransportControl::MakeTransportControlWidget(ETransportCont
 			. Visibility(TransportControlArgs.OnBackwardEnd.IsBound() ? EVisibility::Visible : EVisibility::Collapsed)
 			. ToolTipText( LOCTEXT("ToFront", "To Front") )
 			. ContentPadding(0.0f)
-			. IsFocusable(bAreButtonsFocusable)
+			. IsFocusable(TransportControlArgs.bAreButtonsFocusable)
 			[
 				SNew(SImage)
 				.ColorAndOpacity(FSlateColor::UseSubduedForeground())
@@ -51,7 +51,7 @@ TSharedPtr<SWidget> STransportControl::MakeTransportControlWidget(ETransportCont
 			. Visibility(TransportControlArgs.OnBackwardStep.IsBound() ? EVisibility::Visible : EVisibility::Collapsed)
 			. ToolTipText( LOCTEXT("ToPrevious", "To Previous") )
 			. ContentPadding(0.0f)
-			. IsFocusable(bAreButtonsFocusable)
+			. IsFocusable(TransportControlArgs.bAreButtonsFocusable)
 			[
 				SNew(SImage)
 				.ColorAndOpacity(FSlateColor::UseSubduedForeground())
@@ -64,7 +64,7 @@ TSharedPtr<SWidget> STransportControl::MakeTransportControlWidget(ETransportCont
 			. ToolTipText( LOCTEXT("Reverse", "Reverse") )
 			. ButtonStyle( FAppStyle::Get(), "Animation.PlayControlsButton" )
 			. ContentPadding(0.0f)
-			. IsFocusable(bAreButtonsFocusable)
+			. IsFocusable(TransportControlArgs.bAreButtonsFocusable)
 			[
 				SNew(SImage)
 				.ColorAndOpacity(FSlateColor::UseSubduedForeground())
@@ -74,10 +74,21 @@ TSharedPtr<SWidget> STransportControl::MakeTransportControlWidget(ETransportCont
 		return SAssignNew(RecordButton, SButton)
 			. ButtonStyle(FAppStyle::Get(), "Animation.PlayControlsButton")
 			. OnClicked(TransportControlArgs.OnRecord)
-			. Visibility(TransportControlArgs.OnRecord.IsBound() ? EVisibility::Visible : EVisibility::Collapsed)
+			. Visibility_Lambda([this]() -> EVisibility
+			{
+				if(!TransportControlArgs.OnRecord.IsBound())
+				{
+					return EVisibility::Collapsed;
+				}
+				if(TransportControlArgs.OnGetRecordButtonVisibility.IsBound())
+				{
+					return TransportControlArgs.OnGetRecordButtonVisibility.Execute();
+				}
+				return EVisibility::Visible;
+			})
 			. ToolTipText( this, &STransportControl::GetRecordStatusTooltip )
 			. ContentPadding(0.0f)
-			. IsFocusable(bAreButtonsFocusable)
+			. IsFocusable(TransportControlArgs.bAreButtonsFocusable)
 			[
 				SNew(SImage)
 				.Image(FAppStyle::Get().GetBrush("Animation.Record"))
@@ -104,7 +115,7 @@ TSharedPtr<SWidget> STransportControl::MakeTransportControlWidget(ETransportCont
 			. ToolTipText( this, &STransportControl::GetForwardStatusTooltip )
 			. ButtonStyle( FAppStyle::Get(), "Animation.PlayControlsButton" )
 			. ContentPadding(0.0f)
-			. IsFocusable(bAreButtonsFocusable)
+			. IsFocusable(TransportControlArgs.bAreButtonsFocusable)
 			[
 				SNew(SImage)
 				.ColorAndOpacity(FSlateColor::UseSubduedForeground())
@@ -117,7 +128,7 @@ TSharedPtr<SWidget> STransportControl::MakeTransportControlWidget(ETransportCont
 			. Visibility(TransportControlArgs.OnForwardStep.IsBound() ? EVisibility::Visible : EVisibility::Collapsed)
 			. ToolTipText( LOCTEXT("ToNext", "To Next") )
 			. ContentPadding(0.0f)
-			. IsFocusable(bAreButtonsFocusable)
+			. IsFocusable(TransportControlArgs.bAreButtonsFocusable)
 			[
 				SNew(SImage)
 				.ColorAndOpacity(FSlateColor::UseSubduedForeground())
@@ -130,7 +141,7 @@ TSharedPtr<SWidget> STransportControl::MakeTransportControlWidget(ETransportCont
 			. Visibility(TransportControlArgs.OnForwardEnd.IsBound() ? EVisibility::Visible : EVisibility::Collapsed)
 			. ToolTipText( LOCTEXT("ToEnd", "To End") )
 			. ContentPadding(0.0f)
-			. IsFocusable(bAreButtonsFocusable)
+			. IsFocusable(TransportControlArgs.bAreButtonsFocusable)
 			[
 				SNew(SImage)
 				.ColorAndOpacity(FSlateColor::UseSubduedForeground())
@@ -143,7 +154,7 @@ TSharedPtr<SWidget> STransportControl::MakeTransportControlWidget(ETransportCont
 			. ToolTipText( this, &STransportControl::GetLoopStatusTooltip )
 			. ButtonStyle( FAppStyle::Get(), "Animation.PlayControlsButton" )
 			. ContentPadding(0.0f)
-			. IsFocusable(bAreButtonsFocusable)
+			. IsFocusable(TransportControlArgs.bAreButtonsFocusable)
 			[
 				SNew(SImage)
 				.ColorAndOpacity(FSlateColor::UseSubduedForeground())
@@ -171,7 +182,7 @@ void STransportControl::Construct( const STransportControl::FArguments& InArgs )
 	{
 		for(FTransportControlWidget WidgetDesc : TransportControlArgs.WidgetsToCreate)
 		{
-			TSharedPtr<SWidget> Widget = MakeTransportControlWidget(WidgetDesc.WidgetType, InArgs._TransportArgs.bAreButtonsFocusable, WidgetDesc.MakeCustomWidgetDelegate);
+			TSharedPtr<SWidget> Widget = MakeTransportControlWidget(WidgetDesc.WidgetType);
 			if(Widget.IsValid())
 			{
 				HorizontalBox->AddSlot()
@@ -188,7 +199,7 @@ void STransportControl::Construct( const STransportControl::FArguments& InArgs )
 	{
 		for(ETransportControlWidgetType WidgetType : TEnumRange<ETransportControlWidgetType>())
 		{
-			TSharedPtr<SWidget> Widget = MakeTransportControlWidget(WidgetType, InArgs._TransportArgs.bAreButtonsFocusable);
+			TSharedPtr<SWidget> Widget = MakeTransportControlWidget(WidgetType);
 			if(Widget.IsValid())
 			{
 				HorizontalBox->AddSlot()
diff --git a/Engine/Source/Editor/EditorWidgets/Private/STransportControl.h b/Engine/Source/Editor/EditorWidgets/Private/STransportControl.h
--- a/Engine/Source/Editor/EditorWidgets/Private/STransportControl.h
+++ b/Engine/Source/Editor/EditorWidgets/Private/STransportControl.h
@@ -56,7 +56,7 @@ private:
 	FReply OnToggleLooping();
 
 	/** Make default transport control widgets */
-	TSharedPtr<SWidget> MakeTransportControlWidget(ETransportControlWidgetType WidgetType, bool bAreButtonsFocusable, const FOnMakeTransportWidget& MakeCustomWidgetDelegate = FOnMakeTransportWidget());
+	TSharedPtr<SWidget> MakeTransportControlWidget(ETransportControlWidgetType WidgetType, const FOnMakeTransportWidget& MakeCustomWidgetDelegate = FOnMakeTransportWidget());
 
 private:
 	/** The handle to the active timer */
diff --git a/Engine/Source/Editor/EditorWidgets/Public/ITransportControl.h b/Engine/Source/Editor/EditorWidgets/Public/ITransportControl.h
--- a/Engine/Source/Editor/EditorWidgets/Public/ITransportControl.h
+++ b/Engine/Source/Editor/EditorWidgets/Public/ITransportControl.h
@@ -23,6 +23,7 @@ DECLARE_DELEGATE_RetVal(bool, FOnGetRecording)
 DECLARE_DELEGATE_RetVal(EPlaybackMode::Type, FOnGetPlaybackMode)
 DECLARE_DELEGATE_TwoParams(FOnTickPlayback, double /*InCurrentTime*/, float /*InDeltaTime*/)
 DECLARE_DELEGATE_RetVal(TSharedRef<SWidget>, FOnMakeTransportWidget)
+DECLARE_DELEGATE_RetVal(EVisibility, FOnGetRecordButtonVisibility)
 
 enum class ETransportControlWidgetType : int32
 {
@@ -74,6 +75,7 @@ struct FTransportControlArgs
 		, OnTickPlayback()
 		, OnGetRecording()
 		, bAreButtonsFocusable(true)
+		, OnGetRecordButtonVisibility()
 	{}
 
 	FOnClicked OnForwardPlay;
@@ -89,6 +91,7 @@ struct FTransportControlArgs
 	FOnTickPlayback OnTickPlayback;
 	FOnGetRecording OnGetRecording;
 	bool bAreButtonsFocusable;
+	FOnGetRecordButtonVisibility OnGetRecordButtonVisibility;
 
 	/** 
 	 * Array of custom widgets to create - if this array is used the default widget ordering will be
diff --git a/Engine/Source/Editor/KismetWidgets/Private/SScrubControlPanel.cpp b/Engine/Source/Editor/KismetWidgets/Private/SScrubControlPanel.cpp
--- a/Engine/Source/Editor/KismetWidgets/Private/SScrubControlPanel.cpp
+++ b/Engine/Source/Editor/KismetWidgets/Private/SScrubControlPanel.cpp
@@ -41,6 +41,7 @@ void SScrubControlPanel::Construct( const SScrubControlPanel::FArguments& InArgs
 	TransportControlArgs.OnGetPlaybackMode = InArgs._OnGetPlaybackMode;
 	TransportControlArgs.OnGetRecording = InArgs._OnGetRecording;
 	TransportControlArgs.OnTickPlayback = InArgs._OnTickPlayback;
+	TransportControlArgs.OnGetRecordButtonVisibility = InArgs._OnGetRecordVisibility;
 	TransportControlArgs.WidgetsToCreate = InArgs._TransportControlWidgetsToCreate.Get(TArray<FTransportControlWidget>());
 
 	FTransportControlArgs TransportControlArgsForRealtimeStreamingMode;
@@ -49,6 +50,7 @@ void SScrubControlPanel::Construct( const SScrubControlPanel::FArguments& InArgs
 	TransportControlArgsForRealtimeStreamingMode.OnGetPlaybackMode = TransportControlArgs.OnGetPlaybackMode;
 	TransportControlArgsForRealtimeStreamingMode.OnGetRecording = TransportControlArgs.OnGetRecording;
 	TransportControlArgsForRealtimeStreamingMode.OnRecord = TransportControlArgs.OnRecord;
+	TransportControlArgsForRealtimeStreamingMode.OnGetRecordButtonVisibility = InArgs._OnGetRecordVisibility;
 
 	this->ChildSlot
 	.Padding( FMargin( 0.0f, 1.0f) )
diff --git a/Engine/Source/Editor/KismetWidgets/Public/SScrubControlPanel.h b/Engine/Source/Editor/KismetWidgets/Public/SScrubControlPanel.h
--- a/Engine/Source/Editor/KismetWidgets/Public/SScrubControlPanel.h
+++ b/Engine/Source/Editor/KismetWidgets/Public/SScrubControlPanel.h
@@ -95,6 +95,7 @@ public:
 		 * ignored in favor of this set of widgets.
 		 */
 		SLATE_ARGUMENT(TOptional<TArray<FTransportControlWidget>>, TransportControlWidgetsToCreate);
+		SLATE_EVENT(FOnGetRecordButtonVisibility, OnGetRecordVisibility);
 
 	SLATE_END_ARGS()
 
diff --git a/Engine/Source/Editor/Persona/Private/SAnimMontageScrubPanel.cpp b/Engine/Source/Editor/Persona/Private/SAnimMontageScrubPanel.cpp
--- a/Engine/Source/Editor/Persona/Private/SAnimMontageScrubPanel.cpp
+++ b/Engine/Source/Editor/Persona/Private/SAnimMontageScrubPanel.cpp
@@ -122,7 +122,7 @@ FReply SAnimMontageScrubPanel::OnClick_Forward_Step()
 	return FReply::Handled();
 }
 
-void SAnimMontageScrubPanel::OnValueChanged(float NewValue)
+void SAnimMontageScrubPanel::OnValueChanged(float NewValue, bool bFireNotifies)
 {
 	if (UAnimPreviewInstance* AnimPreviewInstance = Cast<UAnimPreviewInstance>(GetPreviewInstance()))
 	{
diff --git a/Engine/Source/Editor/Persona/Private/SAnimMontageScrubPanel.h b/Engine/Source/Editor/Persona/Private/SAnimMontageScrubPanel.h
--- a/Engine/Source/Editor/Persona/Private/SAnimMontageScrubPanel.h
+++ b/Engine/Source/Editor/Persona/Private/SAnimMontageScrubPanel.h
@@ -50,7 +50,7 @@ protected:
 
 	virtual FReply OnClick_ToggleLoop() override;
 
-	virtual void OnValueChanged(float NewValue) override;
+	virtual void OnValueChanged(float NewValue, bool bFireNotifies = true) override;
 
 private:
 	TWeakPtr<SMontageEditor> MontageEditor;
diff --git a/Engine/Source/Editor/Persona/Private/SAnimationEditorViewport.cpp b/Engine/Source/Editor/Persona/Private/SAnimationEditorViewport.cpp
--- a/Engine/Source/Editor/Persona/Private/SAnimationEditorViewport.cpp
+++ b/Engine/Source/Editor/Persona/Private/SAnimationEditorViewport.cpp
@@ -750,6 +750,7 @@ void SAnimationEditorViewportTabBody::Construct(const FArguments& InArgs, const
 	OnInvokeTab = InArgs._OnInvokeTab;
 	OnGetViewportText = InArgs._OnGetViewportText;
 	ContextName = InArgs._ContextName;
+	TimelineDelegates = InArgs._TimelineDelegates;
 
 	// register delegates for change notifications
 	InPreviewScene->RegisterOnAnimChanged(FOnAnimChanged::CreateSP(this, &SAnimationEditorViewportTabBody::AnimChanged));
@@ -858,6 +859,7 @@ void SAnimationEditorViewportTabBody::Construct(const FArguments& InArgs, const
 				.ViewInputMin(this, &SAnimationEditorViewportTabBody::GetViewMinInput)
 				.ViewInputMax(this, &SAnimationEditorViewportTabBody::GetViewMaxInput)
 				.bAllowZoom(true)
+				.TimelineDelegates(TimelineDelegates)
 			]
 		];
 
@@ -2153,6 +2155,7 @@ void SAnimationEditorViewportTabBody::UpdateScrubPanel(UAnimationAsset* AnimAsse
 					.ViewInputMax(this, &SAnimationEditorViewportTabBody::GetViewMaxInput)
 					.bAllowZoom(true)
 					.bDisplayAnimScrubBarEditing(false)
+					.TimelineDelegates(TimelineDelegates)
 				];
 		}
 	}
@@ -2160,6 +2163,15 @@ void SAnimationEditorViewportTabBody::UpdateScrubPanel(UAnimationAsset* AnimAsse
 
 float SAnimationEditorViewportTabBody::GetViewMinInput() const
 {
+	if(TimelineDelegates.GetPlaybackTimeRangeDelegate.IsBound())
+	{
+		const TOptional<FVector2f> TimeRange = TimelineDelegates.GetPlaybackTimeRangeDelegate.Execute();
+		if(TimeRange.IsSet())
+		{
+			return TimeRange.GetValue().X;
+		}
+	}
+	
 	if (UDebugSkelMeshComponent* PreviewComponent = GetPreviewScene()->GetPreviewMeshComponent())
 	{
 		UObject* PreviewAsset = GetPreviewScene()->GetPreviewAnimationAsset();
@@ -2177,7 +2189,16 @@ float SAnimationEditorViewportTabBody::GetViewMinInput() const
 }
 
 float SAnimationEditorViewportTabBody::GetViewMaxInput() const
-{ 
+{
+	if(TimelineDelegates.GetPlaybackTimeRangeDelegate.IsBound())
+	{
+		const TOptional<FVector2f> TimeRange = TimelineDelegates.GetPlaybackTimeRangeDelegate.Execute();
+		if(TimeRange.IsSet())
+		{
+			return TimeRange.GetValue().Y;
+		}
+	}
+
 	UDebugSkelMeshComponent* PreviewComponent = GetPreviewScene()->GetPreviewMeshComponent();
 	if (PreviewComponent != NULL)
 	{
diff --git a/Engine/Source/Editor/Persona/Private/SAnimationEditorViewport.h b/Engine/Source/Editor/Persona/Private/SAnimationEditorViewport.h
--- a/Engine/Source/Editor/Persona/Private/SAnimationEditorViewport.h
+++ b/Engine/Source/Editor/Persona/Private/SAnimationEditorViewport.h
@@ -213,6 +213,8 @@ public:
 		SLATE_ARGUMENT(bool, ShowTurnTable)
 
 		SLATE_ARGUMENT(bool, ShowPhysicsMenu)
+
+		SLATE_ARGUMENT(FAnimationScrubPanelDelegates, TimelineDelegates)
 	SLATE_END_ARGS()
 public:
 
@@ -623,6 +625,9 @@ private:
 	TUniquePtr<FScopedTransaction> PendingTransaction;
 
 	FName ContextName;
+
+	FAnimationScrubPanelDelegates TimelineDelegates;
+	
 public:
 	/** UV Channel Selector */
 	TSharedPtr< class STextComboBox > UVChannelCombo;
diff --git a/Engine/Source/Editor/Persona/Private/SAnimationScrubPanel.cpp b/Engine/Source/Editor/Persona/Private/SAnimationScrubPanel.cpp
--- a/Engine/Source/Editor/Persona/Private/SAnimationScrubPanel.cpp
+++ b/Engine/Source/Editor/Persona/Private/SAnimationScrubPanel.cpp
@@ -20,7 +20,10 @@ void SAnimationScrubPanel::Construct( const SAnimationScrubPanel::FArguments& In
 	bSliderBeingDragged = false;
 	LockedSequence = InArgs._LockedSequence;
 	OnSetInputViewRange = InArgs._OnSetInputViewRange;
-
+	TimelineDelegates = InArgs._TimelineDelegates;
+	ViewInputMinAttribute = InArgs._ViewInputMin;
+	ViewInputMaxAttribute = InArgs._ViewInputMax;
+	
 	PreviewScenePtr = InPreviewScene;
 
 	this->ChildSlot
@@ -39,7 +42,7 @@ void SAnimationScrubPanel::Construct( const SAnimationScrubPanel::FArguments& In
 			.NumOfKeys(this, &SAnimationScrubPanel::GetNumberOfKeys)
 			.SequenceLength(this, &SAnimationScrubPanel::GetSequenceLength)
 			.DisplayDrag(this, &SAnimationScrubPanel::GetDisplayDrag)
-			.OnValueChanged(this, &SAnimationScrubPanel::OnValueChanged)
+			.OnValueChanged(this, &SAnimationScrubPanel::OnValueChanged, true)
 			.OnBeginSliderMovement(this, &SAnimationScrubPanel::OnBeginSliderMovement)
 			.OnEndSliderMovement(this, &SAnimationScrubPanel::OnEndSliderMovement)
 			.OnClickedForwardPlay(this, &SAnimationScrubPanel::OnClick_Forward)
@@ -50,11 +53,12 @@ void SAnimationScrubPanel::Construct( const SAnimationScrubPanel::FArguments& In
 			.OnClickedBackwardEnd(this, &SAnimationScrubPanel::OnClick_Backward_End)
 			.OnClickedToggleLoop(this, &SAnimationScrubPanel::OnClick_ToggleLoop)
 			.OnClickedRecord(this, &SAnimationScrubPanel::OnClick_Record)
+			.OnGetRecordVisibility(this, &SAnimationScrubPanel::OnGetRecordVisibility)
 			.OnGetLooping(this, &SAnimationScrubPanel::IsLoopStatusOn)
 			.OnGetPlaybackMode(this, &SAnimationScrubPanel::GetPlaybackMode)
 			.OnGetRecording(this, &SAnimationScrubPanel::IsRecording)
-			.ViewInputMin(InArgs._ViewInputMin)
-			.ViewInputMax(InArgs._ViewInputMax)
+			.ViewInputMin(this, &SAnimationScrubPanel::GetViewInputMin)
+			.ViewInputMax(this, &SAnimationScrubPanel::GetViewInputMax)
 			.bDisplayAnimScrubBarEditing(InArgs._bDisplayAnimScrubBarEditing)
 			.OnSetInputViewRange(InArgs._OnSetInputViewRange)
 			.OnCropAnimSequence( this, &SAnimationScrubPanel::OnCropAnimSequence )
@@ -69,13 +73,22 @@ void SAnimationScrubPanel::Construct( const SAnimationScrubPanel::FArguments& In
 
 FReply SAnimationScrubPanel::OnClick_Forward_Step()
 {
+	SetPlaybackMode(EPlaybackMode::Stopped);
+
+	if(TimelineDelegates.StepForwardDelegate.IsBound())
+	{
+		if(TimelineDelegates.StepForwardDelegate.Execute())
+		{
+			return FReply::Handled();
+		}
+	}
+	
 	UDebugSkelMeshComponent* SMC = GetPreviewScene()->GetPreviewMeshComponent();
 
 	if (UAnimSingleNodeInstance* PreviewInstance = GetPreviewInstance())
 	{
 		bool bShouldStepCloth = FMath::Abs(PreviewInstance->GetLength() - PreviewInstance->GetCurrentTime()) > SMALL_NUMBER;
 
-		PreviewInstance->SetPlaying(false);
 		PreviewInstance->StepForward();
 
 		if(SMC && bShouldStepCloth)
@@ -99,25 +112,29 @@ FReply SAnimationScrubPanel::OnClick_Forward_Step()
 
 FReply SAnimationScrubPanel::OnClick_Forward_End()
 {
-	UAnimSingleNodeInstance* PreviewInstance = GetPreviewInstance();
-	if (PreviewInstance)
-	{
-		PreviewInstance->SetPlaying(false);
-		PreviewInstance->SetPosition(PreviewInstance->GetLength(), false);
-	}
-
+	SetPlaybackMode(EPlaybackMode::Stopped);
+	OnValueChanged(GetViewInputMax(), false);
 	return FReply::Handled();
 }
 
 FReply SAnimationScrubPanel::OnClick_Backward_Step()
 {
+	SetPlaybackMode(EPlaybackMode::Stopped);
+
+	if(TimelineDelegates.StepBackwardDelegate.IsBound())
+	{
+		if(TimelineDelegates.StepBackwardDelegate.Execute())
+		{
+			return FReply::Handled();
+		}
+	}
+
 	UAnimSingleNodeInstance* PreviewInstance = GetPreviewInstance();
 	UDebugSkelMeshComponent* SMC = GetPreviewScene()->GetPreviewMeshComponent();
 	if (PreviewInstance)
 	{
 		bool bShouldStepCloth = PreviewInstance->GetCurrentTime() > SMALL_NUMBER;
 
-		PreviewInstance->SetPlaying(false);
 		PreviewInstance->StepBackward();
 
 		if(SMC && bShouldStepCloth)
@@ -130,12 +147,8 @@ FReply SAnimationScrubPanel::OnClick_Backward_Step()
 
 FReply SAnimationScrubPanel::OnClick_Backward_End()
 {
-	UAnimSingleNodeInstance* PreviewInstance = GetPreviewInstance();
-	if (PreviewInstance)
-	{
-		PreviewInstance->SetPlaying(false);
-		PreviewInstance->SetPosition(0.f, false);
-	}
+	SetPlaybackMode(EPlaybackMode::Stopped);
+	OnValueChanged(GetViewInputMin(), false);
 	return FReply::Handled();
 }
 
@@ -144,44 +157,31 @@ FReply SAnimationScrubPanel::OnClick_Forward()
 	UAnimSingleNodeInstance* PreviewInstance = GetPreviewInstance();
 	UDebugSkelMeshComponent* SMC = GetPreviewScene()->GetPreviewMeshComponent();
 
-	if (PreviewInstance)
+	const EPlaybackMode::Type PlaybackMode = GetPlaybackMode();
+	if(PlaybackMode == EPlaybackMode::PlayingForward)
 	{
-		bool bIsReverse = PreviewInstance->IsReverse();
-		bool bIsPlaying = PreviewInstance->IsPlaying();
-		// if current bIsReverse and bIsPlaying, we'd like to just turn off reverse
-		if (bIsReverse && bIsPlaying)
+		SetPlaybackMode(EPlaybackMode::Stopped);
+
+		if(SMC && SMC->bPauseClothingSimulationWithAnim)
 		{
-			PreviewInstance->SetReverse(false);
+			SMC->SuspendClothingSimulation();
 		}
-		// already playing, simply pause
-		else if (bIsPlaying) 
+	}
+	else
+	{
+		if ( GetScrubValue() >= GetSequenceLength() )
 		{
-			PreviewInstance->SetPlaying(false);
-			
-			if(SMC && SMC->bPauseClothingSimulationWithAnim)
-			{
-				SMC->SuspendClothingSimulation();
-			}
+			OnValueChanged(0.0f, false);
 		}
-		// if not playing, play forward
-		else 
+		SetPlaybackMode(EPlaybackMode::PlayingForward);
+		
+		if(SMC && SMC->bPauseClothingSimulationWithAnim)
 		{
-			//if we're at the end of the animation, jump back to the beginning before playing
-			if ( GetScrubValue() >= GetSequenceLength() )
-			{
-				PreviewInstance->SetPosition(0.0f, false);
-			}
-
-			PreviewInstance->SetReverse(false);
-			PreviewInstance->SetPlaying(true);
-
-			if(SMC && SMC->bPauseClothingSimulationWithAnim)
-			{
-				SMC->ResumeClothingSimulation();
-			}
+			SMC->ResumeClothingSimulation();
 		}
 	}
-	else if(SMC)
+	
+	if(SMC && !PreviewInstance)
 	{
 		SMC->GlobalAnimRateScale = (SMC->GlobalAnimRateScale > 0.0f) ? 0.0f : 1.0f;
 	}
@@ -191,37 +191,36 @@ FReply SAnimationScrubPanel::OnClick_Forward()
 
 FReply SAnimationScrubPanel::OnClick_Backward()
 {
-	UAnimSingleNodeInstance* PreviewInstance = GetPreviewInstance();
-	if (PreviewInstance)
+	const EPlaybackMode::Type PlaybackMode = GetPlaybackMode();
+	if(PlaybackMode == EPlaybackMode::PlayingForward)
 	{
-		bool bIsReverse = PreviewInstance->IsReverse();
-		bool bIsPlaying = PreviewInstance->IsPlaying();
-		// if currently playing forward, just simply turn on reverse
-		if (!bIsReverse && bIsPlaying)
-		{
-			PreviewInstance->SetReverse(true);
-		}
-		else if (bIsPlaying)
-		{
-			PreviewInstance->SetPlaying(false);
-		}
-		else
+		SetPlaybackMode(EPlaybackMode::Stopped);
+	}
+	else
+	{
+		if ( GetScrubValue() >= GetSequenceLength() )
 		{
-			//if we're at the beginning of the animation, jump back to the end before playing
-			if ( GetScrubValue() <= 0.0f )
-			{
-				PreviewInstance->SetPosition(GetSequenceLength(), false);
-			}
-
-			PreviewInstance->SetPlaying(true);
-			PreviewInstance->SetReverse(true);
+			OnValueChanged(0.0f, false);
 		}
+		SetPlaybackMode(EPlaybackMode::PlayingReverse);
 	}
 	return FReply::Handled();
 }
 
 FReply SAnimationScrubPanel::OnClick_ToggleLoop()
 {
+	if(TimelineDelegates.GetIsLoopingDelegate.IsBound() && TimelineDelegates.SetIsLoopingDelegate.IsBound())
+	{
+		const TOptional<bool> PreviousState = TimelineDelegates.GetIsLoopingDelegate.Execute();
+		if(PreviousState.IsSet())
+		{
+			if(TimelineDelegates.SetIsLoopingDelegate.Execute(!PreviousState.GetValue()))
+			{
+				return FReply::Handled();
+			}
+		}
+	}
+	
 	UAnimSingleNodeInstance* PreviewInstance = GetPreviewInstance();
 	if (PreviewInstance)
 	{
@@ -238,14 +237,44 @@ FReply SAnimationScrubPanel::OnClick_Record()
 	return FReply::Handled();
 }
 
+EVisibility SAnimationScrubPanel::OnGetRecordVisibility() const
+{
+	if(TimelineDelegates.GetRecordingVisibilityDelegate.IsBound())
+	{
+		const TOptional<EVisibility> Visibility = TimelineDelegates.GetRecordingVisibilityDelegate.Execute();
+		if(Visibility.IsSet())
+		{
+			return Visibility.GetValue();
+		}
+	}
+	return EVisibility::Visible;
+}
+
 bool SAnimationScrubPanel::IsLoopStatusOn() const
 {
+	if(TimelineDelegates.GetIsLoopingDelegate.IsBound())
+	{
+		const TOptional<bool> PreviousState = TimelineDelegates.GetIsLoopingDelegate.Execute();
+		if(PreviousState.IsSet())
+		{
+			return PreviousState.GetValue();
+		}
+	}
+	
 	UAnimSingleNodeInstance* PreviewInstance = GetPreviewInstance();
 	return (PreviewInstance && PreviewInstance->IsLooping());
 }
 
 EPlaybackMode::Type SAnimationScrubPanel::GetPlaybackMode() const
 {
+	if(TimelineDelegates.GetPlaybackModeDelegate.IsBound())
+	{
+		const TOptional<int32> Mode = TimelineDelegates.GetPlaybackModeDelegate.Execute();
+		if(Mode.IsSet())
+		{
+			return static_cast<EPlaybackMode::Type>(Mode.GetValue());
+		}
+	}
 	if (UAnimSingleNodeInstance* PreviewInstance = GetPreviewInstance())
 	{
 		if (PreviewInstance->IsPlaying())
@@ -264,6 +293,14 @@ EPlaybackMode::Type SAnimationScrubPanel::GetPlaybackMode() const
 
 bool SAnimationScrubPanel::IsRecording() const
 {
+	if(TimelineDelegates.IsRecordingActiveDelegate.IsBound())
+	{
+		const TOptional<bool> State = TimelineDelegates.IsRecordingActiveDelegate.Execute();
+		if(State.IsSet())
+		{
+			return State.GetValue();
+		}
+	}
 	return StaticCastSharedRef<FAnimationEditorPreviewScene>(GetPreviewScene())->IsRecording();
 }
 
@@ -273,11 +310,45 @@ bool SAnimationScrubPanel::IsRealtimeStreamingMode() const
 	return ( ! (PreviewInstance && PreviewInstance->GetCurrentAsset()) );
 }
 
-void SAnimationScrubPanel::OnValueChanged(float NewValue)
+float SAnimationScrubPanel::GetViewInputMin() const
 {
+	if(TimelineDelegates.GetPlaybackTimeRangeDelegate.IsBound())
+	{
+		const TOptional<FVector2f> TimeRange = TimelineDelegates.GetPlaybackTimeRangeDelegate.Execute();
+		if(TimeRange.IsSet())
+		{
+			return TimeRange.GetValue().X;
+		}
+	}
+	return ViewInputMinAttribute.Get();
+}
+
+float SAnimationScrubPanel::GetViewInputMax() const
+{
+	if(TimelineDelegates.GetPlaybackTimeRangeDelegate.IsBound())
+	{
+		const TOptional<FVector2f> TimeRange = TimelineDelegates.GetPlaybackTimeRangeDelegate.Execute();
+		if(TimeRange.IsSet())
+		{
+			return TimeRange.GetValue().Y;
+		}
+	}
+	return ViewInputMaxAttribute.Get();
+}
+
+void SAnimationScrubPanel::OnValueChanged(float NewValue, bool bFireNotifies)
+{
+	if(TimelineDelegates.SetPlaybackTimeDelegate.IsBound())
+	{
+		if(TimelineDelegates.SetPlaybackTimeDelegate.Execute(NewValue, true))
+		{
+			return;
+		}
+	}
+	
 	if (UAnimSingleNodeInstance* PreviewInstance = GetPreviewInstance())
 	{
-		PreviewInstance->SetPosition(NewValue);
+		PreviewInstance->SetPosition(NewValue, bFireNotifies);
 	}
 	else
 	{
@@ -297,7 +368,7 @@ void SAnimationScrubPanel::OnBeginSliderMovement()
 
 	if (UAnimSingleNodeInstance* PreviewInstance = GetPreviewInstance())
 	{
-		PreviewInstance->SetPlaying(false);
+		SetPlaybackMode(EPlaybackMode::Stopped);
 	}
 }
 
@@ -308,6 +379,14 @@ void SAnimationScrubPanel::OnEndSliderMovement(float NewValue)
 
 uint32 SAnimationScrubPanel::GetNumberOfKeys() const
 {
+	if(TimelineDelegates.GetNumberOfKeysDelegate.IsBound())
+	{
+		const TOptional<uint32> NumberOfKeys = TimelineDelegates.GetNumberOfKeysDelegate.Execute();
+		if(NumberOfKeys.IsSet())
+		{
+			return NumberOfKeys.GetValue();
+		}
+	}
 	if (DoesSyncViewport())
 	{
 		UAnimSingleNodeInstance* PreviewInstance = GetPreviewInstance();
@@ -347,6 +426,15 @@ uint32 SAnimationScrubPanel::GetNumberOfKeys() const
 
 float SAnimationScrubPanel::GetSequenceLength() const
 {
+	if(TimelineDelegates.GetPlaybackTimeRangeDelegate.IsBound())
+	{
+		const TOptional<FVector2f> TimeRange = TimelineDelegates.GetPlaybackTimeRangeDelegate.Execute();
+		if(TimeRange.IsSet())
+		{
+			return TimeRange.GetValue().Y - TimeRange.GetValue().X;
+		}
+	}
+	
 	if (DoesSyncViewport())
 	{
 		UAnimSingleNodeInstance* PreviewInstance = GetPreviewInstance();
@@ -392,6 +480,15 @@ class UAnimSingleNodeInstance* SAnimationScrubPanel::GetPreviewInstance() const
 
 float SAnimationScrubPanel::GetScrubValue() const
 {
+	if(TimelineDelegates.GetPlaybackTimeDelegate.IsBound())
+	{
+		const TOptional<float> Time = TimelineDelegates.GetPlaybackTimeDelegate.Execute();
+		if(Time.IsSet())
+		{
+			return Time.GetValue();
+		}
+	}
+	
 	if (DoesSyncViewport())
 	{
 		UAnimSingleNodeInstance* PreviewInstance = GetPreviewInstance();
@@ -608,6 +705,14 @@ void SAnimationScrubPanel::OnReZeroAnimSequence(int32 FrameIndex)
 
 bool SAnimationScrubPanel::GetDisplayDrag() const
 {
+	if(TimelineDelegates.GetPlaybackTimeDelegate.IsBound())
+	{
+		if(TimelineDelegates.GetPlaybackTimeDelegate.Execute().IsSet())
+		{
+			return true;
+		}
+	}
+	
 	UAnimSingleNodeInstance* PreviewInstance = GetPreviewInstance();
 	if (PreviewInstance && PreviewInstance->GetCurrentAsset())
 	{
@@ -617,4 +722,39 @@ bool SAnimationScrubPanel::GetDisplayDrag() const
 	return false;
 }
 
+void SAnimationScrubPanel::SetPlaybackMode(EPlaybackMode::Type InMode)
+{
+	if(TimelineDelegates.SetPlaybackModeDelegate.IsBound())
+	{
+		if(TimelineDelegates.SetPlaybackModeDelegate.Execute(InMode))
+		{
+			return;
+		}
+	}
+	
+	if (UAnimSingleNodeInstance* PreviewInstance = GetPreviewInstance())
+	{
+		switch(InMode)
+		{
+			case EPlaybackMode::Stopped:
+			{
+				PreviewInstance->SetPlaying(false);
+				break;
+			}
+			case EPlaybackMode::PlayingForward:
+			{
+				PreviewInstance->SetReverse(false);
+				PreviewInstance->SetPlaying(true);
+				break;
+			}
+			case EPlaybackMode::PlayingReverse:
+			{
+				PreviewInstance->SetReverse(true);
+				PreviewInstance->SetPlaying(true);
+				break;
+			}
+		}
+	}
+}
+
 #undef LOCTEXT_NAMESPACE
diff --git a/Engine/Source/Editor/Persona/Private/SAnimationScrubPanel.h b/Engine/Source/Editor/Persona/Private/SAnimationScrubPanel.h
--- a/Engine/Source/Editor/Persona/Private/SAnimationScrubPanel.h
+++ b/Engine/Source/Editor/Persona/Private/SAnimationScrubPanel.h
@@ -9,6 +9,7 @@
 #include "Widgets/SCompoundWidget.h"
 #include "SScrubWidget.h"
 #include "ITransportControl.h"
+#include "PersonaModule.h"
 
 class SScrubControlPanel;
 class UAnimationAsset;
@@ -35,6 +36,7 @@ public:
 		/** Called to zero out selected frame's translation from origin */
 		SLATE_EVENT( FSimpleDelegate, OnReZeroAnimSequence )
 		SLATE_ARGUMENT( bool, bAllowZoom )
+		SLATE_ARGUMENT( FAnimationScrubPanelDelegates, TimelineDelegates )
 	SLATE_END_ARGS()
 
 	/**
@@ -68,9 +70,11 @@ protected:
 
 	virtual FReply OnClick_Record();
 
+	virtual EVisibility OnGetRecordVisibility() const;
+
 	void AnimChanged(UAnimationAsset * AnimAsset);
 
-	virtual void OnValueChanged(float NewValue);
+	virtual void OnValueChanged(float NewValue, bool bFireNotifies);
 
 	/** Function to crop animation sequence before/after selected frame */
 	void OnCropAnimSequence( bool bFromStart, float CurrentTime );
@@ -91,12 +95,17 @@ protected:
 	bool IsRecording() const;
 	bool IsLoopStatusOn() const;
 	bool IsRealtimeStreamingMode() const;
+	float GetViewInputMin() const;
+	float GetViewInputMax() const;
 
 	float GetScrubValue() const;
 	class UAnimSingleNodeInstance* GetPreviewInstance() const;
 
 	TSharedPtr<class SScrubControlPanel> ScrubControlPanel;
 	class UAnimSequenceBase* LockedSequence;
+	FAnimationScrubPanelDelegates TimelineDelegates;
+	TAttribute<float> ViewInputMinAttribute;
+	TAttribute<float> ViewInputMaxAttribute;
 
 	/** Do I need to sync with viewport? **/
 	bool DoesSyncViewport() const;
@@ -112,4 +121,6 @@ protected:
 	TSharedRef<IPersonaPreviewScene> GetPreviewScene() const { return PreviewScenePtr.Pin().ToSharedRef(); }
 
 	bool GetDisplayDrag() const;
+
+	void SetPlaybackMode(EPlaybackMode::Type InMode);
 };
diff --git a/Engine/Source/Editor/Persona/Private/TabSpawners.cpp b/Engine/Source/Editor/Persona/Private/TabSpawners.cpp
--- a/Engine/Source/Editor/Persona/Private/TabSpawners.cpp
+++ b/Engine/Source/Editor/Persona/Private/TabSpawners.cpp
@@ -292,6 +292,7 @@ FPreviewViewportSummoner::FPreviewViewportSummoner(TSharedPtr<class FAssetEditor
 	, bShowFloorOptions(InArgs.bShowFloorOptions)
 	, bShowTurnTable(InArgs.bShowTurnTable)
 	, bShowPhysicsMenu(InArgs.bShowPhysicsMenu)
+	, TimelineDelegates(InArgs.TimelineDelegates)
 {
 	TabLabel = FText::Format(LOCTEXT("ViewportTabTitle", "Viewport {0}"), FText::AsNumber(InViewportIndex + 1));
 	TabIcon = FSlateIcon(FAppStyle::GetAppStyleSetName(), "LevelEditor.Tabs.Viewports");
@@ -319,7 +320,8 @@ TSharedRef<SWidget> FPreviewViewportSummoner::CreateTabBody(const FWorkflowTabSp
 		.AlwaysShowTransformToolbar(bAlwaysShowTransformToolbar)
 		.ShowFloorOptions(bShowFloorOptions)
 		.ShowTurnTable(bShowTurnTable)
-		.ShowPhysicsMenu(bShowPhysicsMenu);
+		.ShowPhysicsMenu(bShowPhysicsMenu)
+		.TimelineDelegates(TimelineDelegates);
 
 	OnViewportCreated.ExecuteIfBound(NewViewport);
 
diff --git a/Engine/Source/Editor/Persona/Private/TabSpawners.h b/Engine/Source/Editor/Persona/Private/TabSpawners.h
--- a/Engine/Source/Editor/Persona/Private/TabSpawners.h
+++ b/Engine/Source/Editor/Persona/Private/TabSpawners.h
@@ -235,6 +235,7 @@ struct FPreviewViewportSummoner : public FWorkflowTabFactory
 	bool bShowFloorOptions;
 	bool bShowTurnTable;
 	bool bShowPhysicsMenu;
+	FAnimationScrubPanelDelegates TimelineDelegates;
 };
 
 /////////////////////////////////////////////////////
diff --git a/Engine/Source/Editor/Persona/Public/PersonaModule.h b/Engine/Source/Editor/Persona/Public/PersonaModule.h
--- a/Engine/Source/Editor/Persona/Public/PersonaModule.h
+++ b/Engine/Source/Editor/Persona/Public/PersonaModule.h
@@ -183,6 +183,67 @@ enum class EViewportCorner : uint8
 /** Delegate used to provide custom text for the viewport corners */
 DECLARE_DELEGATE_RetVal_OneParam(FText, FOnGetViewportText, EViewportCorner /*InViewportCorner*/);
 
+DECLARE_DELEGATE_RetVal(TOptional<bool>, FAnimationScrubPanel_IsRecordingActive);
+DECLARE_DELEGATE_RetVal(TOptional<EVisibility>, FAnimationScrubPanel_GetRecordingVisibility);
+DECLARE_DELEGATE_RetVal(bool, FAnimationScrubPanel_StartRecording)
+DECLARE_DELEGATE_RetVal(bool, FAnimationScrubPanel_StopRecording);
+DECLARE_DELEGATE_RetVal(TOptional<int32>, FAnimationScrubPanel_GetPlaybackMode);
+DECLARE_DELEGATE_RetVal_OneParam(bool, FAnimationScrubPanel_SetPlaybackMode, int32 /* playback mode */);
+DECLARE_DELEGATE_RetVal(TOptional<float>, FAnimationScrubPanel_GetPlaybackTime);
+DECLARE_DELEGATE_RetVal_TwoParams(bool, FAnimationScrubPanel_SetPlaybackTime, float /* InTime */, bool /* StopPlayback */);
+DECLARE_DELEGATE_RetVal(bool, FAnimationScrubPanel_StepForward);
+DECLARE_DELEGATE_RetVal(bool, FAnimationScrubPanel_StepBackward);
+DECLARE_DELEGATE_RetVal(TOptional<bool>, FAnimationScrubPanel_GetIsLooping);
+DECLARE_DELEGATE_RetVal_OneParam(bool, FAnimationScrubPanel_SetIsLooping, bool /* IsLooping */);
+DECLARE_DELEGATE_RetVal(TOptional<FVector2f>, FAnimationScrubPanel_GetPlaybackTimeRange);
+DECLARE_DELEGATE_RetVal(TOptional<uint32>, FAnimationScrubPanel_GetNumberOfKeys);
+
+/** Struct to provide delegates to change the timeline behavior for certain clients */
+struct FAnimationScrubPanelDelegates
+{
+	/** Determines if recording is currently active */
+	FAnimationScrubPanel_IsRecordingActive IsRecordingActiveDelegate;
+
+	/** Determines if the recording button should be visible */
+	FAnimationScrubPanel_GetRecordingVisibility GetRecordingVisibilityDelegate;
+
+	/** Starts the recording */
+	FAnimationScrubPanel_StartRecording StartRecordingDelegate;
+
+	/** Stops the recording */	
+	FAnimationScrubPanel_StopRecording StopRecordingDelegate;
+
+	/** Determines if the playback is currently active */
+	FAnimationScrubPanel_GetPlaybackMode GetPlaybackModeDelegate;
+
+	/** Sets the playback mode */
+	FAnimationScrubPanel_SetPlaybackMode SetPlaybackModeDelegate;
+
+	/** Returns the current playback time */
+	FAnimationScrubPanel_GetPlaybackTime GetPlaybackTimeDelegate;
+
+	/** Seeks to a given time */
+	FAnimationScrubPanel_SetPlaybackTime SetPlaybackTimeDelegate;
+
+	/** steps on frame / key forward */
+	FAnimationScrubPanel_StepForward StepForwardDelegate;
+
+	/** steps on frame / key backward */
+	FAnimationScrubPanel_StepBackward StepBackwardDelegate;
+
+	/** Returns the current playback looping flag */
+	FAnimationScrubPanel_GetIsLooping GetIsLoopingDelegate;
+
+	/** Sets the playback looping flag */
+	FAnimationScrubPanel_SetIsLooping SetIsLoopingDelegate;
+
+	/** Returns the playback time range */
+	FAnimationScrubPanel_GetPlaybackTimeRange GetPlaybackTimeRangeDelegate;
+
+	/** Returns the number of keys for this playback */
+	FAnimationScrubPanel_GetNumberOfKeys GetNumberOfKeysDelegate;
+};
+
 /** Arguments used to create a persona viewport tab */
 struct FPersonaViewportArgs
 {
@@ -244,6 +305,9 @@ struct FPersonaViewportArgs
 
 	/** Whether to show options relating to physics */
 	bool bShowPhysicsMenu;
+
+	/** A structure providing delegates to change the timeline behavior */
+	FAnimationScrubPanelDelegates TimelineDelegates;
 };
 
 /**
