Control Rig: Improve SimpleIK when overextending

#rnx

diff --git a/Engine/Plugins/Animation/ControlRig/Source/ControlRig/Private/Math/ControlRigMathLibrary.cpp b/Engine/Plugins/Animation/ControlRig/Source/ControlRig/Private/Math/ControlRigMathLibrary.cpp
--- a/Engine/Plugins/Animation/ControlRig/Source/ControlRig/Private/Math/ControlRigMathLibrary.cpp
+++ b/Engine/Plugins/Animation/ControlRig/Source/ControlRig/Private/Math/ControlRigMathLibrary.cpp
@@ -16,21 +16,19 @@ void FControlRigMathLibrary::SolveBasicTwoBoneIK(FTransform& BoneA, FTransform&
 
 	FVector Axis = BoneA.TransformVectorNoScale(PrimaryAxis);
 	FVector Target1 = BoneB.GetLocation() - BoneA.GetLocation();
+	bool bIsFullyExtended = false;
 
 	FVector BoneBLocation = BoneB.GetLocation();
 	if (!Target1.IsNearlyZero() && !Axis.IsNearlyZero())
 	{
 		Target1 = Target1.GetSafeNormal();
 
-		// If the elbow is placed in the segment between the root and the effector (the limb is completely extended)
-		// offset the elbow towards the pole vector in order to compute the rotations
+		// Compute if the elbow is placed in the segment between the root and the effector (the limb is completely extended)
 		{
 			FVector BaseDirection = (EffectorPos - RootPos).GetSafeNormal();
 			if (FMath::Abs(FVector::DotProduct(BaseDirection, Target1) - 1.0f) < KINDA_SMALL_NUMBER)
 			{
-				// Offset the bone location to use when calculating the rotations
-				BoneBLocation += (PoleVector - BoneBLocation).GetSafeNormal() * KINDA_SMALL_NUMBER*2.f;
-				Target1 = (BoneBLocation - BoneA.GetLocation()).GetSafeNormal();
+				bIsFullyExtended = true;
 			}
 		}
 		
@@ -41,7 +39,7 @@ void FControlRigMathLibrary::SolveBasicTwoBoneIK(FTransform& BoneA, FTransform&
 
 		if (SecondaryAxisWeight > SMALL_NUMBER)
 		{
-			FVector Target2 = BoneBLocation - (Effector.GetLocation() + BoneA.GetLocation()) * 0.5f;
+			FVector Target2 =  bIsFullyExtended ? PoleVector - RootPos : BoneBLocation - (Effector.GetLocation() + BoneA.GetLocation()) * 0.5f;
 			if (!Target2.IsNearlyZero() && !Axis.IsNearlyZero())
 			{
 				Target2 = Target2 - FVector::DotProduct(Target2, Target1) * Target1;
@@ -70,7 +68,7 @@ void FControlRigMathLibrary::SolveBasicTwoBoneIK(FTransform& BoneA, FTransform&
 		if (SecondaryAxisWeight > SMALL_NUMBER)
 		{
 			Axis = BoneB.TransformVectorNoScale(SecondaryAxis);
-			FVector Target2 = BoneBLocation - (Effector.GetLocation() + BoneA.GetLocation()) * 0.5f;
+			FVector Target2 = bIsFullyExtended ?  PoleVector - BoneBLocation : BoneBLocation - (Effector.GetLocation() + BoneA.GetLocation()) * 0.5f;
 			if (!Target2.IsNearlyZero() && !Axis.IsNearlyZero())
 			{
 				Target2 = Target2 - FVector::DotProduct(Target2, Target1) * Target1;
