ControRig: expose rotation order settings to spawn rotator node

diff --git a/Engine/Plugins/Animation/ControlRig/Source/ControlRig/Private/Units/Execution/RigUnit_DynamicHierarchy.cpp b/Engine/Plugins/Animation/ControlRig/Source/ControlRig/Private/Units/Execution/RigUnit_DynamicHierarchy.cpp
--- a/Engine/Plugins/Animation/ControlRig/Source/ControlRig/Private/Units/Execution/RigUnit_DynamicHierarchy.cpp
+++ b/Engine/Plugins/Animation/ControlRig/Source/ControlRig/Private/Units/Execution/RigUnit_DynamicHierarchy.cpp
@@ -780,6 +780,9 @@ void FRigUnit_HierarchyAddControlRotator_Settings::ConfigureFrom(const FRigContr
 	Proxy.ConfigureFrom(InControlElement, InSettings);
 	Limits.ConfigureFrom(InControlElement, InSettings);
 	Shape.ConfigureFrom(InControlElement, InSettings);
+
+	bUsePreferredRotationOrder = InSettings.bUsePreferredRotationOrder;
+	PreferredRotationOrder = InSettings.PreferredRotationOrder;
 }
 
 void FRigUnit_HierarchyAddControlRotator_Settings::Configure(FRigControlSettings& OutSettings) const
@@ -792,6 +795,9 @@ void FRigUnit_HierarchyAddControlRotator_Settings::Configure(FRigControlSettings
 	Proxy.Configure(OutSettings);
 	Limits.Configure(OutSettings);
 	Shape.Configure(OutSettings);
+
+	OutSettings.bUsePreferredRotationOrder = bUsePreferredRotationOrder;
+	OutSettings.PreferredRotationOrder = PreferredRotationOrder;
 }
 
 FRigUnit_HierarchyAddControlRotator_Execute()
diff --git a/Engine/Plugins/Animation/ControlRig/Source/ControlRig/Public/Units/Execution/RigUnit_DynamicHierarchy.h b/Engine/Plugins/Animation/ControlRig/Source/ControlRig/Public/Units/Execution/RigUnit_DynamicHierarchy.h
--- a/Engine/Plugins/Animation/ControlRig/Source/ControlRig/Public/Units/Execution/RigUnit_DynamicHierarchy.h
+++ b/Engine/Plugins/Animation/ControlRig/Source/ControlRig/Public/Units/Execution/RigUnit_DynamicHierarchy.h
@@ -974,6 +974,8 @@ struct CONTROLRIG_API FRigUnit_HierarchyAddControlRotator_Settings : public FRig
 	FRigUnit_HierarchyAddControlRotator_Settings()
 		: FRigUnit_HierarchyAddControl_Settings()
 		, InitialSpace(ERigVMTransformSpace::LocalSpace)
+		, bUsePreferredRotationOrder(false)
+		, PreferredRotationOrder(EEulerRotationOrder::YZX)
 	{
 		FilteredChannels.Reset();
 	}
@@ -996,6 +998,12 @@ struct CONTROLRIG_API FRigUnit_HierarchyAddControlRotator_Settings : public FRig
 
 	UPROPERTY(EditAnywhere, BlueprintReadWrite, Category = "Settings")
 	TArray<ERigControlTransformChannel> FilteredChannels;
+	
+	UPROPERTY(EditAnywhere, BlueprintReadWrite, Category = "Settings")
+	bool bUsePreferredRotationOrder;
+
+	UPROPERTY(EditAnywhere, BlueprintReadWrite, Category = "Settings")
+	EEulerRotationOrder PreferredRotationOrder;
 };
 
 /**
diff --git a/Engine/Plugins/Animation/ControlRig/Source/ControlRigEditor/Private/ControlRigElementDetails.cpp b/Engine/Plugins/Animation/ControlRig/Source/ControlRigEditor/Private/ControlRigElementDetails.cpp
--- a/Engine/Plugins/Animation/ControlRig/Source/ControlRigEditor/Private/ControlRigElementDetails.cpp
+++ b/Engine/Plugins/Animation/ControlRig/Source/ControlRigEditor/Private/ControlRigElementDetails.cpp
@@ -3663,15 +3663,17 @@ void FRigControlElementDetails::CustomizeControl(IDetailLayoutBuilder& DetailBui
 
 	if(CVarControlRigHierarchyEnableRotationOrder.GetValueOnAnyThread())
 	{
-		if(IsAnyControlOfValueType(ERigControlType::EulerTransform))
+		if (IsAnyControlOfValueType(ERigControlType::EulerTransform) || IsAnyControlOfValueType(ERigControlType::Rotator))
 		{
 			const TSharedPtr<IPropertyHandle> UsePreferredRotationOrderHandle = SettingsHandle->GetChildHandle(TEXT("bUsePreferredRotationOrder"));
-			ControlCategory.AddProperty(UsePreferredRotationOrderHandle.ToSharedRef()).DisplayName(FText::FromString(TEXT("Use Preferred Rotation Order")))
+			ControlCategory.AddProperty(UsePreferredRotationOrderHandle.ToSharedRef())
+				.DisplayName(FText::FromString(TEXT("Use Preferred Rotation Order")))
 				.IsEnabled(bIsEnabled);
 
 			const TSharedPtr<IPropertyHandle> PreferredRotationOrderHandle = SettingsHandle->GetChildHandle(TEXT("PreferredRotationOrder"));
-			ControlCategory.AddProperty(PreferredRotationOrderHandle.ToSharedRef()).DisplayName(FText::FromString(TEXT("Preferred Rotation Order")))
-			.IsEnabled(bIsEnabled);
+			ControlCategory.AddProperty(PreferredRotationOrderHandle.ToSharedRef())
+				.DisplayName(FText::FromString(TEXT("Preferred Rotation Order")))
+				.IsEnabled(bIsEnabled);
 		}
 	}
 
