ControlRig: ensure that a control's transform is correctly applied when other selected controls depend on it (follow up to 38077082)

a small "dependency graph" between the selected controls is cached during interaction.
this graph is then used when manipulating to define whether the controls depend on each other and thus determine if the transform should be applied directly or can be deferred.
we're not trying to find out who the parent/children are here, but just storing information on who is a child and who is a parent in what is currently being manipulated.
this map is only valid during interaction and is invalidated if a sequencer is read, as the dependencies may have been modified.

#rnx

diff --git a/Engine/Plugins/Animation/ControlRig/Source/ControlRigEditor/Private/EditMode/ControlRigEditMode.cpp b/Engine/Plugins/Animation/ControlRig/Source/ControlRigEditor/Private/EditMode/ControlRigEditMode.cpp
--- a/Engine/Plugins/Animation/ControlRig/Source/ControlRigEditor/Private/EditMode/ControlRigEditMode.cpp
+++ b/Engine/Plugins/Animation/ControlRig/Source/ControlRigEditor/Private/EditMode/ControlRigEditMode.cpp
@@ -1184,6 +1184,7 @@ bool FControlRigEditMode::HandleBeginTransform(const FEditorViewportClient* InVi
 	
 	InteractionType = GetInteractionType(InViewportClient);
 	bIsTracking = true;
+	InteractionDependencies.Reset();
 	
 	if (InteractionScopes.Num() == 0)
 	{
@@ -1282,6 +1283,7 @@ bool FControlRigEditMode::HandleEndTransform(FEditorViewportClient* InViewportCl
 	
 	InteractionType = (uint8)EControlRigInteractionType::None;
 	bIsTracking = false;
+	InteractionDependencies.Reset();
 	
 	if (InteractionScopes.Num() > 0)
 	{		
@@ -2108,6 +2110,12 @@ bool FControlRigEditMode::InputDelta(FEditorViewportClient* InViewportClient, FV
 	if (InteractionScopes.Num() > 0 && bMouseButtonDown && CurrentAxis != EAxisList::None
 		&& TransformContext.CanTransform())
 	{
+		if (bSequencerPlaying)
+		{
+			// reset the dependency cache as the hierarchy might have changed since the previous frame
+			InteractionDependencies.Reset();
+		}
+		
 		for (auto& Pairs : ControlRigShapeActors)
 		{
 			if (AreRigElementsSelected(ValidControlTypeMask(), Pairs.Key))
@@ -5437,7 +5445,9 @@ void FControlRigEditMode::MoveControlShape( AControlRigShapeActor* ShapeActor,
 		}
 
 		// for readability
-		const bool bSolveImmediately = (bUseLocal == false) && (bCalcLocal == true);
+		const FInteractionDependencyCache& Dependencies = GetInteractionDependencies(ControlRig);
+		const bool bSolveImmediately = (bUseLocal == false) && (bCalcLocal == true) ||
+			Dependencies.HasDownwardDependencies(ShapeActor->GetElementKey());
 		const bool bQueueTasks = !bSolveImmediately;
 
 		if(bQueueTasks)
@@ -6350,6 +6360,68 @@ void FControlRigEditMode::UpdateRotationContext()
 	}
 }
 
+const FInteractionDependencyCache& FControlRigEditMode::GetInteractionDependencies(UControlRig* InControlRig)
+{
+	if (FInteractionDependencyCache* DependencyCache = InteractionDependencies.Find(InControlRig))
+	{
+		return *DependencyCache;
+	}
+
+	static FInteractionDependencyCache DummyDependencies;
+	
+	if (const TArray<TObjectPtr<AControlRigShapeActor>>* ShapeActors = ControlRigShapeActors.Find(InControlRig))
+	{
+		URigHierarchy* Hierarchy = InControlRig->GetHierarchy();
+
+		// get selected controls 
+		TArray<FRigControlElement*> SelectedControls;
+		for (const TObjectPtr<AControlRigShapeActor>& ShapeActor: *ShapeActors)
+		{
+			if (ShapeActor && ShapeActor->IsEnabled() && ShapeActor->IsSelected())
+			{
+				if (FRigControlElement* Control = Hierarchy->Find<FRigControlElement>(ShapeActor->GetElementKey()))
+				{
+					SelectedControls.Add(Control);
+				}
+			}
+		}
+
+		if (SelectedControls.IsEmpty())
+		{
+			return DummyDependencies;
+		}
+
+		// build dependencies between selected controls
+		FInteractionDependencyCache& NewDependency = InteractionDependencies.Add(InControlRig);
+
+		// NOTE: this is not enough for modular rigs since there are several VMs
+		const URigHierarchy::TElementDependencyMap HierarchyDependencyMap = Hierarchy->GetDependenciesForVM(InControlRig->GetVM());
+		
+		for (int32 Index = 0; Index < SelectedControls.Num(); ++Index)
+		{
+			FRigControlElement* Control = SelectedControls[Index];
+			for (int32 NextIndex = Index+1; NextIndex < SelectedControls.Num(); ++NextIndex)
+			{
+				FRigControlElement* OtherControl = SelectedControls[NextIndex];
+				if (Hierarchy->IsParentedTo(Control, OtherControl, HierarchyDependencyMap))
+				{
+					NewDependency.Parents.Add(OtherControl->GetKey());
+					NewDependency.Children.Add(Control->GetKey());
+				}
+				else if (Hierarchy->IsParentedTo(OtherControl, Control, HierarchyDependencyMap))
+				{
+					NewDependency.Parents.Add(Control->GetKey());
+					NewDependency.Children.Add(OtherControl->GetKey());
+				}
+			}
+		}
+
+		return NewDependency;
+	}
+
+	return DummyDependencies;
+}
+
 /**
 * FDetailKeyFrameCacheAndHandler
 */
diff --git a/Engine/Plugins/Animation/ControlRig/Source/ControlRigEditor/Public/EditMode/ControlRigEditMode.h b/Engine/Plugins/Animation/ControlRig/Source/ControlRigEditor/Public/EditMode/ControlRigEditMode.h
--- a/Engine/Plugins/Animation/ControlRig/Source/ControlRigEditor/Public/EditMode/ControlRigEditMode.h
+++ b/Engine/Plugins/Animation/ControlRig/Source/ControlRigEditor/Public/EditMode/ControlRigEditMode.h
@@ -154,6 +154,26 @@ struct FTransformContext
 	ETransformSpace Space = ETransformSpace::World;
 };
 
+/**
+ * FInteractionDependencyCache provides a minimal "dependency graph" between the selected controls during interaction.
+ * it simply stores information on who is a child and who is a parent in whatever is being manipulated.
+ */
+struct FInteractionDependencyCache
+{
+	bool HasDownwardDependencies(const FRigElementKey& InKey) const
+	{
+		return Parents.Contains(InKey);
+	}
+	
+	bool HasUpwardDependencies(const FRigElementKey& InKey) const
+	{
+		return Children.Contains(InKey);
+	}
+	
+	TSet<FRigElementKey> Parents;
+	TSet<FRigElementKey> Children;
+};
+	
 }
 
 
@@ -492,6 +512,12 @@ protected:
 
 	/** previous Gizmo(Widget) scale before we enter this mode, used to set it back*/
 	float PreviousGizmoScale = 1.0f;
+
+	/** Per ControlRig dependencies between the selected controls during interaction. */
+	TMap<UControlRig*, UE::ControlRigEditMode::FInteractionDependencyCache> InteractionDependencies;
+	
+	/** Returns the interaction dependencies of that ControlRig. */
+	const UE::ControlRigEditMode::FInteractionDependencyCache& GetInteractionDependencies(UControlRig* InControlRig);
 	
 public: 
 	/** Clear all selected RigElements */
