fixed bSolveImmediately expression

#rnx

diff --git a/Engine/Plugins/Animation/ControlRig/Source/ControlRigEditor/Private/EditMode/ControlRigEditMode.cpp b/Engine/Plugins/Animation/ControlRig/Source/ControlRigEditor/Private/EditMode/ControlRigEditMode.cpp
--- a/Engine/Plugins/Animation/ControlRig/Source/ControlRigEditor/Private/EditMode/ControlRigEditMode.cpp
+++ b/Engine/Plugins/Animation/ControlRig/Source/ControlRigEditor/Private/EditMode/ControlRigEditMode.cpp
@@ -5446,7 +5446,7 @@ void FControlRigEditMode::MoveControlShape( AControlRigShapeActor* ShapeActor,
 
 		// for readability
 		const FInteractionDependencyCache& Dependencies = GetInteractionDependencies(ControlRig);
-		const bool bSolveImmediately = (bUseLocal == false) && (bCalcLocal == true) ||
+		const bool bSolveImmediately = (!bUseLocal && bCalcLocal) ||
 			Dependencies.HasDownwardDependencies(ShapeActor->GetElementKey());
 		const bool bQueueTasks = !bSolveImmediately;
 
