ControlRig: fixed rotation not correctly being applied on rotators with rotation order set

#rnx

diff --git a/Engine/Plugins/Animation/ControlRig/Source/ControlRig/Private/Sequencer/MovieSceneControlRigParameterTemplate.cpp b/Engine/Plugins/Animation/ControlRig/Source/ControlRig/Private/Sequencer/MovieSceneControlRigParameterTemplate.cpp
--- a/Engine/Plugins/Animation/ControlRig/Source/ControlRig/Private/Sequencer/MovieSceneControlRigParameterTemplate.cpp
+++ b/Engine/Plugins/Animation/ControlRig/Source/ControlRig/Private/Sequencer/MovieSceneControlRigParameterTemplate.cpp
@@ -1755,9 +1755,16 @@ struct TControlRigParameterActuatorVector : TMovieSceneBlendingActuator<FControl
 				{
 					if (ControlElement && ControlElement->Settings.ControlType == ERigControlType::Rotator)
 					{
-						ControlRig->GetHierarchy()->SetControlSpecifiedEulerAngle(ControlElement, InFinalValue.Value);
+						URigHierarchy* Hierarchy = ControlRig->GetHierarchy();
+						const FRotator Rotator(Hierarchy->GetControlQuaternion(ControlElement, InFinalValue.Value));
+						
+						Hierarchy->SetControlSpecifiedEulerAngle(ControlElement, InFinalValue.Value);
+						ControlRig->SetControlValue<FRotator>(ParameterName, Rotator, true, EControlRigSetKey::Never, bSetupUndo);
+					}
+					else
+					{
+						ControlRig->SetControlValue<FVector3f>(ParameterName, (FVector3f)InFinalValue.Value, true, EControlRigSetKey::Never, bSetupUndo);
 					}
-					ControlRig->SetControlValue<FVector3f>(ParameterName, (FVector3f)InFinalValue.Value, true, EControlRigSetKey::Never, bSetupUndo);
 				}
 			}
 
diff --git a/Engine/Plugins/Animation/ControlRig/Source/ControlRigEditor/Private/ControlRigElementDetails.cpp b/Engine/Plugins/Animation/ControlRig/Source/ControlRigEditor/Private/ControlRigElementDetails.cpp
--- a/Engine/Plugins/Animation/ControlRig/Source/ControlRigEditor/Private/ControlRigElementDetails.cpp
+++ b/Engine/Plugins/Animation/ControlRig/Source/ControlRigEditor/Private/ControlRigElementDetails.cpp
@@ -2669,6 +2669,15 @@ FDetailWidgetRow& FRigTransformElementDetails::CreateEulerTransformValueWidgetRo
 						switch(ControlElement->Settings.ControlType)
 						{
 							case ERigControlType::Rotator:
+							{
+								const FVector EulerAngle(InTransform.Rotator().Roll, InTransform.Rotator().Pitch, InTransform.Rotator().Yaw);
+								HierarchyToUpdate->SetControlSpecifiedEulerAngle(ControlElement, EulerAngle, bInitial);
+
+								const ERigControlValueType ValueType = bInitial ? ERigControlValueType::Initial : ERigControlValueType::Current;
+								const FRotator Rotator(Hierarchy->GetControlQuaternion(ControlElement, EulerAngle));
+								HierarchyToUpdate->SetControlValue(ControlElement, FRigControlValue::Make<FRotator>(Rotator), ValueType, bSetupUndoRedo);
+								break;
+							}
 							case ERigControlType::EulerTransform:
 							case ERigControlType::Transform:
 							case ERigControlType::TransformNoScale:
diff --git a/Engine/Plugins/Animation/ControlRig/Source/ControlRigEditor/Private/EditMode/AnimDetailsProxy.cpp b/Engine/Plugins/Animation/ControlRig/Source/ControlRigEditor/Private/EditMode/AnimDetailsProxy.cpp
--- a/Engine/Plugins/Animation/ControlRig/Source/ControlRigEditor/Private/EditMode/AnimDetailsProxy.cpp
+++ b/Engine/Plugins/Animation/ControlRig/Source/ControlRigEditor/Private/EditMode/AnimDetailsProxy.cpp
@@ -1716,12 +1716,17 @@ void UAnimDetailControlsProxyRotation::SetControlRigElementValueFromCurrent(UCon
 {
 	if (ControlElement && ControlRig && ControlElement->Settings.ControlType == ERigControlType::Rotator)
 	{
+		URigHierarchy* Hierarchy = ControlRig->GetHierarchy();
+		
 		FVector3f Val = Rotation.ToVector3f();
 		SetRotationValuesFromContext(ControlRig, ControlElement, Context, Val);
-		FVector EulerAngle(Rotation.ToRotator().Roll, Rotation.ToRotator().Pitch, Rotation.ToRotator().Yaw);
-		ControlRig->GetHierarchy()->SetControlSpecifiedEulerAngle(ControlElement, EulerAngle);
-		ControlRig->SetControlValue<FVector3f>(ControlElement->GetKey().Name, Val, true, Context, false);
-		ControlRig->GetHierarchy()->SetControlSpecifiedEulerAngle(ControlElement, EulerAngle);
+
+		const FVector EulerAngle(Rotation.ToRotator().Roll, Rotation.ToRotator().Pitch, Rotation.ToRotator().Yaw);
+		const FRotator Rotator(Hierarchy->GetControlQuaternion(ControlElement, EulerAngle));
+		
+		Hierarchy->SetControlSpecifiedEulerAngle(ControlElement, EulerAngle);
+		ControlRig->SetControlValue<FRotator>(ControlElement->GetKey().Name, Rotator, true, Context, false);
+		
 		ControlRig->Evaluate_AnyThread();
 	}
 }
diff --git a/Engine/Plugins/Animation/ControlRig/Source/ControlRigEditor/Private/EditMode/ControlRigEditMode.cpp b/Engine/Plugins/Animation/ControlRig/Source/ControlRigEditor/Private/EditMode/ControlRigEditMode.cpp
--- a/Engine/Plugins/Animation/ControlRig/Source/ControlRigEditor/Private/EditMode/ControlRigEditMode.cpp
+++ b/Engine/Plugins/Animation/ControlRig/Source/ControlRigEditor/Private/EditMode/ControlRigEditMode.cpp
@@ -5080,11 +5080,10 @@ void FControlRigEditMode::SetControlShapeTransform(
 		const ERigControlType ControlType = ControlElement->Settings.ControlType;
 		if (ControlType == ERigControlType::Rotator)
 		{
-			const FVector3f NewValue(NewEulerAngle.Y, NewEulerAngle.Z, NewEulerAngle.X);
+			const FRotator Rotator(Hierarchy->GetControlQuaternion(ControlElement, NewEulerAngle));
 			
 			Hierarchy->SetControlSpecifiedEulerAngle(ControlElement, NewEulerAngle);
-			ControlRig->SetControlValue<FVector3f>(ControlName, NewValue, /*bNotify*/true, InContext, /*bUndo*/true);
-			Hierarchy->SetControlSpecifiedEulerAngle(ControlElement, NewEulerAngle);
+			ControlRig->SetControlValue<FRotator>(ControlName, Rotator, /*bNotify*/true, InContext, /*bUndo*/true, bPrintPython);
 		}
 		else if (ControlType == ERigControlType::EulerTransform)
 		{
