Control Rig: Fix crash when resolving module while inspecting the module in the details panel
#rnx

diff --git a/Engine/Plugins/Animation/ControlRig/Source/ControlRigEditor/Private/ControlRigModuleDetails.cpp b/Engine/Plugins/Animation/ControlRig/Source/ControlRigEditor/Private/ControlRigModuleDetails.cpp
--- a/Engine/Plugins/Animation/ControlRig/Source/ControlRigEditor/Private/ControlRigModuleDetails.cpp
+++ b/Engine/Plugins/Animation/ControlRig/Source/ControlRigEditor/Private/ControlRigModuleDetails.cpp
@@ -239,6 +239,8 @@ void FRigModuleInstanceDetails::CustomizeDetails(IDetailLayoutBuilder& DetailBui
 			.IsEnabled(true)
 		];
 
+		TWeakPtr<FRigModuleInstanceDetails> WeakThisPtr = SharedThis(this);
+
 		GeneralCategory.AddCustomRow(FText::FromString(TEXT("Variant Tags")))
 		.NameContent()
 		[
@@ -253,39 +255,43 @@ void FRigModuleInstanceDetails::CustomizeDetails(IDetailLayoutBuilder& DetailBui
 			.Orientation(EOrientation::Orient_Horizontal)
 			.CanAddTags(false)
 			.EnableContextMenu(false)
-			.OnGetTags_Lambda([this]() -> TArray<FRigVMTag>
+			.OnGetTags_Lambda([WeakThisPtr]() -> TArray<FRigVMTag>
 			{
 				TArray<FRigVMTag> Tags;
-				for (int32 InfoIndex=0; InfoIndex<PerModuleInfos.Num(); ++InfoIndex)
+				if (TSharedPtr<FRigModuleInstanceDetails> StrongThis = WeakThisPtr.Pin())
 				{
-					const FPerModuleInfo& ModuleInfo = PerModuleInfos[InfoIndex]; 
-					if(ModuleInfo.Module.IsValid())
+					const TArray<FPerModuleInfo>& PerModuleInfos = StrongThis->PerModuleInfos;
+					for (int32 InfoIndex=0; InfoIndex<PerModuleInfos.Num(); ++InfoIndex)
 					{
-						if (const FRigModuleInstance* Module = ModuleInfo.GetModule())
+						const FPerModuleInfo& ModuleInfo = PerModuleInfos[InfoIndex]; 
+						if(ModuleInfo.Module.IsValid())
 						{
-							if (const UControlRigBlueprint* ModuleBlueprint = Cast<UControlRigBlueprint>(Module->GetRig()->GetClass()->ClassGeneratedBy))
+							if (const FRigModuleInstance* Module = ModuleInfo.GetModule())
 							{
-								if (InfoIndex == 0)
-								{
-									Tags = ModuleBlueprint->GetAssetVariant().Tags;
-								}
-								else
+								if (const UControlRigBlueprint* ModuleBlueprint = Cast<UControlRigBlueprint>(Module->GetRig()->GetClass()->ClassGeneratedBy))
 								{
-									const TArray<FRigVMTag>& OtherTags = ModuleBlueprint->GetAssetVariant().Tags;
-									bool bSameArray = Tags.Num() == OtherTags.Num();
-									if (bSameArray)
+									if (InfoIndex == 0)
 									{
-										for (const FRigVMTag& OtherTag : OtherTags)
+										Tags = ModuleBlueprint->GetAssetVariant().Tags;
+									}
+									else
+									{
+										const TArray<FRigVMTag>& OtherTags = ModuleBlueprint->GetAssetVariant().Tags;
+										bool bSameArray = Tags.Num() == OtherTags.Num();
+										if (bSameArray)
 										{
-											if (!Tags.ContainsByPredicate([OtherTag](const FRigVMTag& Tag) { return OtherTag.Name == Tag.Name; }))
+											for (const FRigVMTag& OtherTag : OtherTags)
 											{
-												return {};
+												if (!Tags.ContainsByPredicate([OtherTag](const FRigVMTag& Tag) { return OtherTag.Name == Tag.Name; }))
+												{
+													return {};
+												}
 											}
 										}
-									}
-									else
-									{
-										return {};
+										else
+										{
+											return {};
+										}
 									}
 								}
 							}
