Control Rig: Fix mac compilation issue
#rnx

diff --git a/Engine/Plugins/Animation/ControlRig/Source/ControlRigEditor/Private/ControlRigModuleDetails.cpp b/Engine/Plugins/Animation/ControlRig/Source/ControlRigEditor/Private/ControlRigModuleDetails.cpp
--- a/Engine/Plugins/Animation/ControlRig/Source/ControlRigEditor/Private/ControlRigModuleDetails.cpp
+++ b/Engine/Plugins/Animation/ControlRig/Source/ControlRigEditor/Private/ControlRigModuleDetails.cpp
@@ -260,10 +260,10 @@ void FRigModuleInstanceDetails::CustomizeDetails(IDetailLayoutBuilder& DetailBui
 				TArray<FRigVMTag> Tags;
 				if (TSharedPtr<FRigModuleInstanceDetails> StrongThis = WeakThisPtr.Pin())
 				{
-					const TArray<FPerModuleInfo>& PerModuleInfos = StrongThis->PerModuleInfos;
-					for (int32 InfoIndex=0; InfoIndex<PerModuleInfos.Num(); ++InfoIndex)
+					const TArray<FPerModuleInfo>& Infos = StrongThis->PerModuleInfos;
+					for (int32 InfoIndex=0; InfoIndex<Infos.Num(); ++InfoIndex)
 					{
-						const FPerModuleInfo& ModuleInfo = PerModuleInfos[InfoIndex]; 
+						const FPerModuleInfo& ModuleInfo = Infos[InfoIndex]; 
 						if(ModuleInfo.Module.IsValid())
 						{
 							if (const FRigModuleInstance* Module = ModuleInfo.GetModule())
