Control Rig: Avoid disconnecting target on own module

#rnx

diff --git a/Engine/Plugins/Animation/ControlRig/Source/ControlRig/Private/ModularRigController.cpp b/Engine/Plugins/Animation/ControlRig/Source/ControlRig/Private/ModularRigController.cpp
--- a/Engine/Plugins/Animation/ControlRig/Source/ControlRig/Private/ModularRigController.cpp
+++ b/Engine/Plugins/Animation/ControlRig/Source/ControlRig/Private/ModularRigController.cpp
@@ -635,7 +635,7 @@ TArray<FRigElementKey> UModularRigController::DisconnectCyclicConnectors(bool bS
 
 			const FRigModuleReference* ConnectorModule = Model->FindModule(ConnectorModulePath);
 			const FRigModuleReference* TargetModule = Model->FindModule(TargetModulePath);
-			if(ConnectorModule == nullptr || TargetModule == nullptr)
+			if(ConnectorModule == nullptr || TargetModule == nullptr || ConnectorModule == TargetModule)
 			{
 				continue;
 			}
