Control Rig: New RigVMEditor, which does not inherit from BlueprintEditor
#rnx

diff --git a/Engine/Plugins/Animation/ControlRig/Source/ControlRigDeveloper/Private/Graph/ControlRigGraphNode.cpp b/Engine/Plugins/Animation/ControlRig/Source/ControlRigDeveloper/Private/Graph/ControlRigGraphNode.cpp
--- a/Engine/Plugins/Animation/ControlRig/Source/ControlRigDeveloper/Private/Graph/ControlRigGraphNode.cpp
+++ b/Engine/Plugins/Animation/ControlRig/Source/ControlRigDeveloper/Private/Graph/ControlRigGraphNode.cpp
@@ -4,6 +4,9 @@
 #include "Graph/ControlRigGraph.h"
 #include "Graph/ControlRigGraphSchema.h"
 #include "IControlRigEditorModule.h"
+#if !WITH_RIGVMLEGACYEDITOR
+#include "RigVMDeveloper/Private/Kismet/RigVMFindInBlueprintManager.h"
+#endif
 
 #include UE_INLINE_GENERATED_CPP_BY_NAME(ControlRigGraphNode)
 
@@ -16,9 +19,15 @@ UControlRigGraphNode::UControlRigGraphNode()
 
 #if WITH_EDITOR
 
+#if WITH_RIGVMLEGACYEDITOR
 void UControlRigGraphNode::AddPinSearchMetaDataInfo(const UEdGraphPin* Pin, TArray<FSearchTagDataPair>& OutTaggedMetaData) const
 {
 	Super::AddPinSearchMetaDataInfo(Pin, OutTaggedMetaData);
+#else
+void UControlRigGraphNode::AddRigVMPinSearchMetaDataInfo(const UEdGraphPin* Pin, TArray<UBlueprintExtension::FSearchTagDataPair>& OutTaggedMetaData) const
+{
+	Super::AddRigVMPinSearchMetaDataInfo(Pin, OutTaggedMetaData);
+#endif
 
 	if(const URigVMPin* ModelPin = FindModelPinFromGraphPin(Pin))
 	{
diff --git a/Engine/Plugins/Animation/ControlRig/Source/ControlRigDeveloper/Public/Graph/ControlRigGraphNode.h b/Engine/Plugins/Animation/ControlRig/Source/ControlRigDeveloper/Public/Graph/ControlRigGraphNode.h
--- a/Engine/Plugins/Animation/ControlRig/Source/ControlRigDeveloper/Public/Graph/ControlRigGraphNode.h
+++ b/Engine/Plugins/Animation/ControlRig/Source/ControlRigDeveloper/Public/Graph/ControlRigGraphNode.h
@@ -18,6 +18,9 @@ class FBlueprintActionDatabaseRegistrar;
 class UEdGraph;
 struct FSlateIcon;
 class UControlRigBlueprint;
+#if WITH_RIGVMLEGACYEDITOR
+struct FSearchTagDataPair;
+#endif
 
 /** Base class for animation ControlRig-related nodes */
 UCLASS()
@@ -40,7 +43,11 @@ public:
 	UControlRigGraphNode();
 
 #if WITH_EDITOR
-	virtual void AddPinSearchMetaDataInfo(const UEdGraphPin* Pin, TArray<FSearchTagDataPair>& OutTaggedMetaData) const override;
+#if WITH_RIGVMLEGACYEDITOR
+	void AddPinSearchMetaDataInfo(const UEdGraphPin* Pin, TArray<FSearchTagDataPair>& OutTaggedMetaData) const;
+#else
+	void AddRigVMPinSearchMetaDataInfo(const UEdGraphPin* Pin, TArray<UBlueprintExtension::FSearchTagDataPair>& OutTaggedMetaData) const;
+#endif
 #endif
 
 private:
diff --git a/Engine/Plugins/Animation/ControlRig/Source/ControlRigEditor/Private/ControlRigDrawingDetails.cpp b/Engine/Plugins/Animation/ControlRig/Source/ControlRigEditor/Private/ControlRigDrawingDetails.cpp
--- a/Engine/Plugins/Animation/ControlRig/Source/ControlRigEditor/Private/ControlRigDrawingDetails.cpp
+++ b/Engine/Plugins/Animation/ControlRig/Source/ControlRigEditor/Private/ControlRigDrawingDetails.cpp
@@ -14,7 +14,12 @@
 #include "FbxImporter.h"
 #include "ScopedTransaction.h"
 #include "Dialogs/Dialogs.h"
+
+#if WITH_RIGVMLEGACYEDITOR
 #include "SKismetInspector.h"
+#else
+#include "Editor/SRigVMDetailsInspector.h"
+#endif
 
 #include UE_INLINE_GENERATED_CPP_BY_NAME(ControlRigDrawingDetails)
 
@@ -107,8 +112,11 @@ FReply FControlRigDrawContainerDetails::OnImportCurvesFromFBXClicked()
 
 		FControlRigDrawContainerImportFbxSettings Settings;
 		TSharedPtr<FStructOnScope> StructToDisplay = MakeShareable(new FStructOnScope(FControlRigDrawContainerImportFbxSettings::StaticStruct(), (uint8*)&Settings));
-
+#if WITH_RIGVMLEGACYEDITOR
 		TSharedRef<SKismetInspector> KismetInspector = SNew(SKismetInspector);
+#else
+		TSharedRef<SRigVMDetailsInspector> KismetInspector = SNew(SRigVMDetailsInspector);
+#endif
 		KismetInspector->ShowSingleStruct(StructToDisplay);
 
 		SGenericDialogWidget::FArguments DialogArguments;
diff --git a/Engine/Plugins/Animation/ControlRig/Source/ControlRigEditor/Private/ControlRigEditorModule.cpp b/Engine/Plugins/Animation/ControlRig/Source/ControlRigEditor/Private/ControlRigEditorModule.cpp
--- a/Engine/Plugins/Animation/ControlRig/Source/ControlRigEditor/Private/ControlRigEditorModule.cpp
+++ b/Engine/Plugins/Animation/ControlRig/Source/ControlRigEditor/Private/ControlRigEditorModule.cpp
@@ -54,7 +54,6 @@
 #include "ControlRigThumbnailRenderer.h"
 #include "RigVMModel/RigVMController.h"
 #include "ControlRigBlueprint.h"
-#include "SKismetInspector.h"
 #include "Dialogs/Dialogs.h"
 #include "Settings/ControlRigSettings.h"
 #include "EditMode/ControlRigControlsProxy.h"
@@ -100,6 +99,12 @@
 #include "SequencerUtilities.h"
 #include "Bindings/MovieSceneSpawnableActorBinding.h"
 
+#if WITH_RIGVMLEGACYEDITOR
+#include "SKismetInspector.h"
+#else
+#include "Editor/SRigVMDetailsInspector.h"
+#endif
+
 #define LOCTEXT_NAMESPACE "ControlRigEditorModule"
 
 DEFINE_LOG_CATEGORY(LogControlRigEditor);
@@ -666,8 +671,11 @@ void FControlRigEditorModule::GetNodeContextMenuActions(IRigVMClientHost* RigVMC
 
 				FRigVMMirrorSettings Settings;
 				TSharedPtr<FStructOnScope> StructToDisplay = MakeShareable(new FStructOnScope(FRigVMMirrorSettings::StaticStruct(), (uint8*)&Settings));
-
+#if WITH_RIGVMLEGACYEDITOR
 				TSharedRef<SKismetInspector> KismetInspector = SNew(SKismetInspector);
+#else
+				TSharedRef<SRigVMDetailsInspector> KismetInspector = SNew(SRigVMDetailsInspector);
+#endif
 				KismetInspector->ShowSingleStruct(StructToDisplay);
 
 				TSharedRef<SCustomDialog> MirrorDialog = SNew(SCustomDialog)
diff --git a/Engine/Plugins/Animation/ControlRig/Source/ControlRigEditor/Private/Editor/ControlRigEditor.cpp b/Engine/Plugins/Animation/ControlRig/Source/ControlRigEditor/Private/Editor/ControlRigEditor.cpp
--- a/Engine/Plugins/Animation/ControlRig/Source/ControlRigEditor/Private/Editor/ControlRigEditor.cpp
+++ b/Engine/Plugins/Animation/ControlRig/Source/ControlRigEditor/Private/Editor/ControlRigEditor.cpp
@@ -6,7 +6,6 @@
 #include "ControlRigBlueprint.h"
 #include "SBlueprintEditorToolbar.h"
 #include "Editor/ControlRigEditorMode.h"
-#include "SKismetInspector.h"
 #include "SEnumCombo.h"
 #include "SceneView.h"
 #include "Widgets/Layout/SScrollBox.h"
@@ -124,6 +123,12 @@
 #include "ITransportControl.h"
 #include "PropertyPath.h"
 
+#if WITH_RIGVMLEGACYEDITOR
+#include "SKismetInspector.h"
+#else
+#include "Editor/SRigVMDetailsInspector.h"
+#endif
+
 #define LOCTEXT_NAMESPACE "ControlRigEditor"
 
 namespace UE::ControlRigEditor::Private
@@ -156,10 +161,10 @@ UControlRigBlueprint* GetControlRigBlueprintFromContext(const UUnrealEdViewportT
 			{
 				if (AnimViewportClient->GetAssetEditorToolkit()->IsBlueprintEditor())
 				{
-					const TSharedRef<FBlueprintEditor>& BlueprintEditor =
-						StaticCastSharedRef<FBlueprintEditor>(AnimViewportClient->GetAssetEditorToolkit());
+					const TSharedRef<FControlRigEditor>& Editor =
+						StaticCastSharedRef<FControlRigEditor>(AnimViewportClient->GetAssetEditorToolkit());
 
-					return Cast<UControlRigBlueprint>(BlueprintEditor->GetBlueprintObj());
+					return Cast<UControlRigBlueprint>(Editor->FRigVMEditorBase::GetBlueprintObj());
 				}
 			}
 		}
@@ -319,7 +324,7 @@ UObject* FControlRigEditor::GetOuterForHost() const
 	{
 		return EditorSkelComp;
 	}
-	return FRigVMEditor::GetOuterForHost();
+	return FRigVMEditorBase::GetOuterForHost();
 }
 
 UClass* FControlRigEditor::GetDetailWrapperClass() const
@@ -327,6 +332,7 @@ UClass* FControlRigEditor::GetDetailWrapperClass() const
 	return UControlRigWrapperObject::StaticClass();
 }
 
+#if WITH_RIGVMLEGACYEDITOR
 bool FControlRigEditor::IsSectionVisible(NodeSectionID::Type InSectionID) const
 {
 	if(!IControlRigEditor::IsSectionVisible(InSectionID))
@@ -358,7 +364,42 @@ bool FControlRigEditor::IsSectionVisible(NodeSectionID::Type InSectionID) const
 	return true;
 }
 
-bool FControlRigEditor::NewDocument_IsVisibleForType(ECreatedDocumentType GraphType) const
+#else
+bool FControlRigEditor::IsSectionVisible(RigVMNodeSectionID::Type InSectionID) const
+{
+	if(!IControlRigEditor::IsSectionVisible(InSectionID))
+    {
+    	return false;
+    }
+
+	if(const UControlRigBlueprint* RigBlueprint = GetControlRigBlueprint())
+	{
+		if(IsModularRig())
+		{
+			switch (InSectionID)
+			{
+				case RigVMNodeSectionID::GRAPH:
+				{
+					return RigBlueprint->SupportsEventGraphs();
+				}
+				case RigVMNodeSectionID::FUNCTION:
+				{
+					return RigBlueprint->SupportsFunctions();
+				}
+				default:
+				{
+					break;
+				}
+			}
+		}
+	}
+	return true;
+}
+#endif
+
+
+#if WITH_RIGVMLEGACYEDITOR
+bool FControlRigEditor::NewDocument_IsVisibleForType(FBlueprintEditor::ECreatedDocumentType GraphType) const
 {
 	if(!IControlRigEditor::NewDocument_IsVisibleForType(GraphType))
 	{
@@ -371,11 +412,11 @@ bool FControlRigEditor::NewDocument_IsVisibleForType(ECreatedDocumentType GraphT
 		{
 			switch(GraphType)
 			{
-				case CGT_NewEventGraph:
+				case FBlueprintEditor::CGT_NewEventGraph:
 				{
 					return RigBlueprint->SupportsEventGraphs();
 				}
-				case CGT_NewFunctionGraph:
+				case FBlueprintEditor::CGT_NewFunctionGraph:
 				{
 					return RigBlueprint->SupportsFunctions();
 				}
@@ -389,6 +430,40 @@ bool FControlRigEditor::NewDocument_IsVisibleForType(ECreatedDocumentType GraphT
 
 	return true;
 }
+#else
+bool FControlRigEditor::NewDocument_IsVisibleForType(FRigVMEditorBase::ECreatedDocumentType GraphType) const
+{
+	if(!IControlRigEditor::NewDocument_IsVisibleForType(GraphType))
+	{
+		return false;
+	}
+
+	if(const UControlRigBlueprint* RigBlueprint = GetControlRigBlueprint())
+	{
+		if(IsModularRig())
+		{
+			switch(GraphType)
+			{
+				case FRigVMEditorBase::CGT_NewEventGraph:
+				{
+					return RigBlueprint->SupportsEventGraphs();
+				}
+				case FRigVMEditorBase::CGT_NewFunctionGraph:
+				{
+					return RigBlueprint->SupportsFunctions();
+				}
+				default:
+				{
+					break;
+				}
+			}
+		}
+	}
+
+	return true;
+}
+#endif
+
 
 FReply FControlRigEditor::OnViewportDrop(const FGeometry& MyGeometry, const FDragDropEvent& DragDropEvent)
 {
@@ -474,7 +549,7 @@ void FControlRigEditor::InitRigVMEditor(const EToolkitMode::Type Mode, const TSh
 {
 	DECLARE_SCOPE_HIERARCHICAL_COUNTER_FUNC()
 
-	FRigVMEditor::InitRigVMEditor(Mode, InitToolkitHost, InRigVMBlueprint);
+	IControlRigEditor::InitRigVMEditor(Mode, InitToolkitHost, InRigVMBlueprint);
 
 	UControlRigBlueprint* ControlRigBlueprint = CastChecked<UControlRigBlueprint>(InRigVMBlueprint);
 
@@ -850,8 +925,10 @@ bool FControlRigEditor::RecordReplay(double InRecordingDuration)
 				}
 				else
 				{
+#if WITH_RIGVMLEGACYEDITOR
 					static constexpr TCHAR PrintStringFormat[] = TEXT("Recording starts in... %d");
 					UKismetSystemLibrary::PrintString(GetPreviewScene()->GetWorld(), FString::Printf(PrintStringFormat, SecondsLeft), true, false, FLinearColor::Red, 1.f);
+#endif
 				}
 			}), 1.f, true, 0.01f);
 		}
@@ -895,7 +972,7 @@ EControlRigReplayPlaybackMode FControlRigEditor::GetReplayPlaybackMode() const
 
 void FControlRigEditor::FillToolbar(FToolBarBuilder& ToolbarBuilder, bool bEndSection)
 {
-	FRigVMEditor::FillToolbar(ToolbarBuilder, false);
+	IControlRigEditor::FillToolbar(ToolbarBuilder, false);
 	
 	{
 		if(CVarControlRigHierarchyEnableModules.GetValueOnAnyThread())
@@ -1139,7 +1216,7 @@ void FControlRigEditor::SetEventQueue(TArray<FName> InEventQueue, bool bCompile)
 	}
 
 	TArray<FRigElementKey> PreviousSelection;
-	if (UControlRigBlueprint* RigBlueprint = Cast<UControlRigBlueprint>(GetBlueprintObj()))
+	if (UControlRigBlueprint* RigBlueprint = Cast<UControlRigBlueprint>(FRigVMEditorBase::GetBlueprintObj()))
 	{
 		if(bCompile)
 		{
@@ -1162,13 +1239,13 @@ void FControlRigEditor::SetEventQueue(TArray<FName> InEventQueue, bool bCompile)
 		}
 	}
 
-	FRigVMEditor::SetEventQueue(InEventQueue, bCompile);
+	IControlRigEditor::SetEventQueue(InEventQueue, bCompile);
 
 	if (UControlRig* ControlRig = GetControlRig())
 	{
 		if (InEventQueue.Num() > 0)
 		{
-			if (UControlRigBlueprint* RigBlueprint = Cast<UControlRigBlueprint>(GetBlueprintObj()))
+			if (UControlRigBlueprint* RigBlueprint = Cast<UControlRigBlueprint>(FRigVMEditorBase::GetBlueprintObj()))
 			{
 				RigBlueprint->Validator->SetControlRig(ControlRig);
 
@@ -1225,7 +1302,7 @@ int32 FControlRigEditor::GetEventQueueComboValue() const
 	{
 		return 3;
 	}
-	return FRigVMEditor::GetEventQueueComboValue();
+	return IControlRigEditor::GetEventQueueComboValue();
 }
 
 FText FControlRigEditor::GetEventQueueLabel() const
@@ -1288,10 +1365,10 @@ FSlateIcon FControlRigEditor::GetEventQueueIcon(const TArray<FName>& InEventQueu
 
 void FControlRigEditor::HandleSetObjectBeingDebugged(UObject* InObject)
 {
-	FRigVMEditor::HandleSetObjectBeingDebugged(InObject);
+	FRigVMEditorBase::HandleSetObjectBeingDebugged(InObject);
 	
 	UControlRig* DebuggedControlRig = Cast<UControlRig>(InObject);
-	if(UControlRig* PreviouslyDebuggedControlRig = Cast<UControlRig>(GetBlueprintObj()->GetObjectBeingDebugged()))
+	if(UControlRig* PreviouslyDebuggedControlRig = Cast<UControlRig>(FRigVMEditorBase::GetBlueprintObj()->GetObjectBeingDebugged()))
 	{
 		if(!URigVMHost::IsGarbageOrDestroyed(PreviouslyDebuggedControlRig))
 		{
@@ -1304,7 +1381,7 @@ void FControlRigEditor::HandleSetObjectBeingDebugged(UObject* InObject)
 		}
 	}
 
-	if (UControlRigBlueprint* RigBlueprint = Cast<UControlRigBlueprint>(GetBlueprintObj()))
+	if (UControlRigBlueprint* RigBlueprint = Cast<UControlRigBlueprint>(FRigVMEditorBase::GetBlueprintObj()))
 	{
 		RigBlueprint->Validator->SetControlRig(DebuggedControlRig);
 	}
@@ -1344,7 +1421,7 @@ void FControlRigEditor::HandleSetObjectBeingDebugged(UObject* InObject)
 			if(bShouldExecute)
 			{
 				DebuggedControlRig->SetBoneInitialTransformsFromSkeletalMeshComponent(EditorSkelComp);
-				if(UControlRigBlueprint* RigBlueprint = Cast<UControlRigBlueprint>(GetBlueprintObj()))
+				if(UControlRigBlueprint* RigBlueprint = Cast<UControlRigBlueprint>(FRigVMEditorBase::GetBlueprintObj()))
 				{
 					// copy the initial transforms back to the blueprint
 					// no need to call modify here since this code only modifies the bp if the preview mesh changed
@@ -1428,7 +1505,7 @@ void FControlRigEditor::SetDetailViewForRigElements(const TArray<FRigElementKey>
 
 	ClearDetailObject();
 
-	UControlRigBlueprint* RigBlueprint = Cast<UControlRigBlueprint>(GetBlueprintObj());
+	UControlRigBlueprint* RigBlueprint = Cast<UControlRigBlueprint>(FRigVMEditorBase::GetBlueprintObj());
 	URigHierarchy* HierarchyBeingDebugged = GetHierarchyBeingDebugged();
 	TArray<UObject*> Objects;
 
@@ -1440,7 +1517,7 @@ void FControlRigEditor::SetDetailViewForRigElements(const TArray<FRigElementKey>
 			continue;
 		}
 
-		URigVMDetailsViewWrapperObject* WrapperObject = URigVMDetailsViewWrapperObject::MakeInstance(GetDetailWrapperClass(), GetBlueprintObj(), Element->GetScriptStruct(), (uint8*)Element, HierarchyBeingDebugged);
+		URigVMDetailsViewWrapperObject* WrapperObject = URigVMDetailsViewWrapperObject::MakeInstance(GetDetailWrapperClass(), FRigVMEditorBase::GetBlueprintObj(), Element->GetScriptStruct(), (uint8*)Element, HierarchyBeingDebugged);
 		WrapperObject->GetWrappedPropertyChangedChainEvent().AddSP(this, &FControlRigEditor::OnWrappedPropertyChangedChainEvent);
 		WrapperObject->AddToRoot();
 
@@ -1477,7 +1554,7 @@ void FControlRigEditor::RefreshDetailView()
 		return;
 	}
 
-	FRigVMEditor::RefreshDetailView();
+	IControlRigEditor::RefreshDetailView();
 }
 
 bool FControlRigEditor::DetailViewShowsAnyRigElement() const
@@ -1487,7 +1564,7 @@ bool FControlRigEditor::DetailViewShowsAnyRigElement() const
 
 bool FControlRigEditor::DetailViewShowsRigElement(FRigElementKey InKey) const
 {
-	TArray< TWeakObjectPtr<UObject> > SelectedObjects = Inspector->GetSelectedObjects();
+	TArray< TWeakObjectPtr<UObject> > SelectedObjects = GetSelectedObjectsFromDetailView();
 	for (TWeakObjectPtr<UObject> SelectedObject : SelectedObjects)
 	{
 		if (SelectedObject.IsValid())
@@ -1514,7 +1591,7 @@ TArray<FRigElementKey> FControlRigEditor::GetSelectedRigElementsFromDetailView()
 {
 	TArray<FRigElementKey> Elements;
 	
-	TArray< TWeakObjectPtr<UObject> > SelectedObjects = Inspector->GetSelectedObjects();
+	TArray< TWeakObjectPtr<UObject> > SelectedObjects = GetSelectedObjectsFromDetailView();
 	for (TWeakObjectPtr<UObject> SelectedObject : SelectedObjects)
 	{
 		if (SelectedObject.IsValid())
@@ -1535,6 +1612,15 @@ TArray<FRigElementKey> FControlRigEditor::GetSelectedRigElementsFromDetailView()
 	return Elements;
 }
 
+TArray<TWeakObjectPtr<UObject>> FControlRigEditor::GetSelectedObjectsFromDetailView() const
+{
+	if (Inspector.IsValid())
+	{
+		return Inspector->GetSelectedObjects();
+	}
+	return TArray<TWeakObjectPtr<UObject>>();
+}
+
 void FControlRigEditor::SetDetailViewForRigModules()
 {
 	SetDetailViewForRigModules(ModulesSelected);
@@ -1549,7 +1635,7 @@ void FControlRigEditor::SetDetailViewForRigModules(const TArray<FString> InKeys)
 
 	ClearDetailObject();
 
-	UControlRigBlueprint* RigBlueprint = Cast<UControlRigBlueprint>(GetBlueprintObj());
+	UControlRigBlueprint* RigBlueprint = Cast<UControlRigBlueprint>(FRigVMEditorBase::GetBlueprintObj());
 	UModularRig* RigBeingDebugged = Cast<UModularRig>(RigBlueprint->GetDebuggedControlRig());
 
 	if (!RigBeingDebugged)
@@ -1590,7 +1676,7 @@ bool FControlRigEditor::DetailViewShowsAnyRigModule() const
 
 bool FControlRigEditor::DetailViewShowsRigModule(FString InKey) const
 {
-	TArray< TWeakObjectPtr<UObject> > SelectedObjects = Inspector->GetSelectedObjects();
+	TArray< TWeakObjectPtr<UObject> > SelectedObjects = GetSelectedObjectsFromDetailView();
 	for (TWeakObjectPtr<UObject> SelectedObject : SelectedObjects)
 	{
 		if (SelectedObject.IsValid())
@@ -1630,7 +1716,7 @@ void FControlRigEditor::Compile()
 			return;
 		}
 
-		TArray< TWeakObjectPtr<UObject> > SelectedObjects = Inspector->GetSelectedObjects();
+		TArray< TWeakObjectPtr<UObject> > SelectedObjects = GetSelectedObjectsFromDetailView();
 
 		if(IsConstructionModeEnabled())
 		{
@@ -1648,7 +1734,7 @@ void FControlRigEditor::Compile()
 		ResetAllBoneModification(); 
 
 		{
-			FRigVMEditor::Compile();
+			FRigVMEditorBase::Compile();
 		}
 
 		ControlRigBlueprint->RecompileModularRig();
@@ -1754,7 +1840,7 @@ void FControlRigEditor::HandleModifiedEvent(ERigVMGraphNotifType InNotifType, UR
 
 void FControlRigEditor::OnCreateGraphEditorCommands(TSharedPtr<FUICommandList> GraphEditorCommandsList)
 {
-	IControlRigEditor::OnCreateGraphEditorCommands(GraphEditorCommandsList);
+	FRigVMEditorBase::OnCreateGraphEditorCommands(GraphEditorCommandsList);
 
 	GraphEditorCommandsList->MapAction(
 		FControlRigEditorCommands::Get().RequestDirectManipulationPosition,
@@ -1788,10 +1874,10 @@ void FControlRigEditor::HandleVMCompiledEvent(UObject* InCompiledObject, URigVM*
 
 void FControlRigEditor::SaveAsset_Execute()
 {
-	FRigVMEditor::SaveAsset_Execute();
+	FRigVMEditorBase::SaveAsset_Execute();
 
 	// Save the new state of the hierarchy in the default object, so that it has the correct values on load
-	UControlRigBlueprint* RigBlueprint = Cast<UControlRigBlueprint>(GetBlueprintObj());
+	UControlRigBlueprint* RigBlueprint = Cast<UControlRigBlueprint>(FRigVMEditorBase::GetBlueprintObj());
 	if(const UControlRig* ControlRig = GetControlRig())
 	{
 		UControlRig* CDO = ControlRig->GetClass()->GetDefaultObject<UControlRig>();
@@ -1806,10 +1892,10 @@ void FControlRigEditor::SaveAsset_Execute()
 
 void FControlRigEditor::SaveAssetAs_Execute()
 {
-	FRigVMEditor::SaveAssetAs_Execute();
+	FRigVMEditorBase::SaveAssetAs_Execute();
 
 	// Save the new state of the hierarchy in the default object, so that it has the correct values on load
-	UControlRigBlueprint* RigBlueprint = Cast<UControlRigBlueprint>(GetBlueprintObj());
+	UControlRigBlueprint* RigBlueprint = Cast<UControlRigBlueprint>(FRigVMEditorBase::GetBlueprintObj());
 	if(const UControlRig* ControlRig = GetControlRig())
 	{
 		UControlRig* CDO = ControlRig->GetClass()->GetDefaultObject<UControlRig>();
@@ -1824,7 +1910,7 @@ void FControlRigEditor::SaveAssetAs_Execute()
 
 bool FControlRigEditor::IsModularRig() const
 {
-	if(UControlRigBlueprint* RigBlueprint = Cast<UControlRigBlueprint>(GetBlueprintObj()))
+	if(UControlRigBlueprint* RigBlueprint = Cast<UControlRigBlueprint>(FRigVMEditorBase::GetBlueprintObj()))
 	{
 		return RigBlueprint->IsModularRig();
 	}
@@ -1833,7 +1919,7 @@ bool FControlRigEditor::IsModularRig() const
 
 bool FControlRigEditor::IsRigModule() const
 {
-	if(UControlRigBlueprint* RigBlueprint = Cast<UControlRigBlueprint>(GetBlueprintObj()))
+	if(UControlRigBlueprint* RigBlueprint = Cast<UControlRigBlueprint>(FRigVMEditorBase::GetBlueprintObj()))
 	{
 		return RigBlueprint->IsControlRigModule();
 	}
@@ -1857,7 +1943,7 @@ FString FControlRigEditor::GetWorldCentricTabPrefix() const
 
 FReply FControlRigEditor::OnSpawnGraphNodeByShortcut(FInputChord InChord, const FVector2D& InPosition, UEdGraph* InGraph)
 {
-	const FReply SuperReply = FRigVMEditor::OnSpawnGraphNodeByShortcut(InChord, InPosition, InGraph);
+	const FReply SuperReply = FRigVMEditorBase::OnSpawnGraphNodeByShortcut(InChord, InPosition, InGraph);
 	if(SuperReply.IsEventHandled())
 	{
 		return SuperReply;
@@ -1904,7 +1990,7 @@ void FControlRigEditor::PostTransaction(bool bSuccess, const FTransaction* Trans
 {
 	EnsureValidRigElementsInDetailPanel();
 
-	if (UControlRigBlueprint* RigBlueprint = Cast<UControlRigBlueprint>(GetBlueprintObj()))
+	if (UControlRigBlueprint* RigBlueprint = Cast<UControlRigBlueprint>(FRigVMEditorBase::GetBlueprintObj()))
 	{
 		// Do not compile here. ControlRigBlueprint::PostTransacted decides when it is necessary to compile depending
 		// on the properties that are affected.
@@ -1942,7 +2028,7 @@ void FControlRigEditor::EnsureValidRigElementsInDetailPanel()
 	UControlRigBlueprint* ControlRigBP = GetControlRigBlueprint();
 	URigHierarchy* Hierarchy = ControlRigBP->Hierarchy; 
 
-	TArray< TWeakObjectPtr<UObject> > SelectedObjects = Inspector->GetSelectedObjects();
+	TArray< TWeakObjectPtr<UObject> > SelectedObjects = GetSelectedObjectsFromDetailView();
 	for (TWeakObjectPtr<UObject> SelectedObject : SelectedObjects)
 	{
 		if (SelectedObject.IsValid())
@@ -1986,13 +2072,13 @@ void FControlRigEditor::OnAnimInitialized()
 
 void FControlRigEditor::HandleVMExecutedEvent(URigVMHost* InHost, const FName& InEventName)
 {
-	FRigVMEditor::HandleVMExecutedEvent(InHost, InEventName);
+	IControlRigEditor::HandleVMExecutedEvent(InHost, InEventName);
 
 	if (UControlRigBlueprint* ControlRigBP = GetControlRigBlueprint())
 	{
 		URigHierarchy* Hierarchy = GetHierarchyBeingDebugged(); 
 
-		TArray< TWeakObjectPtr<UObject> > SelectedObjects = Inspector->GetSelectedObjects();
+		TArray< TWeakObjectPtr<UObject> > SelectedObjects = GetSelectedObjectsFromDetailView();
 		for (TWeakObjectPtr<UObject> SelectedObject : SelectedObjects)
 		{
 			if (SelectedObject.IsValid())
@@ -2101,7 +2187,7 @@ void FControlRigEditor::CreateEditorModeManager()
 
 void FControlRigEditor::Tick(float DeltaTime)
 {
-	FRigVMEditor::Tick(DeltaTime);
+	IControlRigEditor::Tick(DeltaTime);
 
 	bool bDrawHierarchyBones = false;
 
@@ -2165,7 +2251,7 @@ void FControlRigEditor::HandleViewportCreated(const TSharedRef<class IPersonaVie
 	// TODO: this is duplicated code from FAnimBlueprintEditor, would be nice to consolidate. 
 	auto GetCompilationStateText = [this]()
 	{
-		if (UBlueprint* Blueprint = GetBlueprintObj())
+		if (UBlueprint* Blueprint = FRigVMEditorBase::GetBlueprintObj())
 		{
 			switch (Blueprint->Status)
 			{
@@ -2214,7 +2300,7 @@ void FControlRigEditor::HandleViewportCreated(const TSharedRef<class IPersonaVie
 
 	auto CompileBlueprint = [this]()
 	{
-		if (UBlueprint* Blueprint = GetBlueprintObj())
+		if (UBlueprint* Blueprint = FRigVMEditorBase::GetBlueprintObj())
 		{
 			if (!Blueprint->IsUpToDate())
 			{
@@ -2227,7 +2313,7 @@ void FControlRigEditor::HandleViewportCreated(const TSharedRef<class IPersonaVie
 
 	auto GetErrorSeverity = [this]()
 	{
-		if (UBlueprint* Blueprint = GetBlueprintObj())
+		if (UBlueprint* Blueprint = FRigVMEditorBase::GetBlueprintObj())
 		{
 			return (Blueprint->Status == BS_Error) ? EMessageSeverity::Error : EMessageSeverity::Warning;
 		}
@@ -2237,7 +2323,7 @@ void FControlRigEditor::HandleViewportCreated(const TSharedRef<class IPersonaVie
 
 	auto GetIcon = [this]()
 	{
-		if (UBlueprint* Blueprint = GetBlueprintObj())
+		if (UBlueprint* Blueprint = FRigVMEditorBase::GetBlueprintObj())
 		{
 			return (Blueprint->Status == BS_Error) ? FEditorFontGlyphs::Exclamation_Triangle : FEditorFontGlyphs::Eye;
 		}
@@ -2724,7 +2810,7 @@ void FControlRigEditor::HandleViewportCreated(const TSharedRef<class IPersonaVie
 	{
 		AnimViewportClient->OnSetBoneSize.BindLambda([this](float InBoneSize)
 		{
-			if (UControlRigBlueprint* RigBlueprint = Cast<UControlRigBlueprint>(GetBlueprintObj()))
+			if (UControlRigBlueprint* RigBlueprint = Cast<UControlRigBlueprint>(FRigVMEditorBase::GetBlueprintObj()))
 			{
 				RigBlueprint->Modify();
 				RigBlueprint->DebugBoneRadius = InBoneSize;
@@ -2733,7 +2819,7 @@ void FControlRigEditor::HandleViewportCreated(const TSharedRef<class IPersonaVie
 		
 		AnimViewportClient->OnGetBoneSize.BindLambda([this]() -> float
 		{
-			if (UControlRigBlueprint* RigBlueprint = Cast<UControlRigBlueprint>(GetBlueprintObj()))
+			if (UControlRigBlueprint* RigBlueprint = Cast<UControlRigBlueprint>(FRigVMEditorBase::GetBlueprintObj()))
 			{
 				return RigBlueprint->DebugBoneRadius;
 			}
@@ -2930,7 +3016,7 @@ void FControlRigEditor::HandlePreviewSceneCreated(const TSharedRef<IPersonaPrevi
 	if(FWorldContext* WorldContext = GEngine->GetWorldContextFromWorld(InPersonaPreviewScene->GetWorld()))
 	{
 		static constexpr TCHAR Format[] = TEXT("ControlRigEditor (%s)");
-		WorldContext->CustomDescription = FString::Printf(Format, *GetBlueprintObj()->GetName());
+		WorldContext->CustomDescription = FString::Printf(Format, *FRigVMEditorBase::GetBlueprintObj()->GetName());
 	}
 
 	// create ground mesh actor
@@ -2985,9 +3071,9 @@ void FControlRigEditor::UpdateRigVMHost()
 {
 	DECLARE_SCOPE_HIERARCHICAL_COUNTER_FUNC()
 
-	FRigVMEditor::UpdateRigVMHost();
+	IControlRigEditor::UpdateRigVMHost();
 
-	UControlRigBlueprint* Blueprint = Cast<UControlRigBlueprint>(GetBlueprintObj());
+	UControlRigBlueprint* Blueprint = Cast<UControlRigBlueprint>(FRigVMEditorBase::GetBlueprintObj());
 	if(UClass* Class = Blueprint->GeneratedClass)
 	{
 		UControlRigSkeletalMeshComponent* EditorSkelComp = Cast<UControlRigSkeletalMeshComponent>(GetPersonaToolkit()->GetPreviewScene()->GetPreviewMeshComponent());
@@ -3054,7 +3140,7 @@ void FControlRigEditor::CacheNameLists()
 {
 	DECLARE_SCOPE_HIERARCHICAL_COUNTER_FUNC()
 
-	FRigVMEditor::CacheNameLists();
+	IControlRigEditor::CacheNameLists();
 
 	if (UControlRigBlueprint* ControlRigBP = GetControlRigBlueprint())
 	{
@@ -3349,7 +3435,7 @@ void FControlRigEditor::GenerateEventQueueMenuContent(FMenuBuilder& MenuBuilder)
     MenuBuilder.AddMenuEntry(FControlRigEditorCommands::Get().BackwardsAndForwardsSolveEvent, TEXT("BackwardsAndForwards"), TAttribute<FText>(), TAttribute<FText>(), GetEventQueueIcon(BackwardsAndForwardsSolveEventQueue));
     MenuBuilder.EndSection();
 
-    if (const UControlRigBlueprint* RigBlueprint = Cast<UControlRigBlueprint>(GetBlueprintObj()))
+    if (const UControlRigBlueprint* RigBlueprint = Cast<UControlRigBlueprint>(FRigVMEditorBase::GetBlueprintObj()))
     {
     	URigVMEdGraphSchema* Schema = CastChecked<URigVMEdGraphSchema>(RigBlueprint->GetRigVMEdGraphSchemaClass()->GetDefaultObject());
     	
@@ -3379,7 +3465,11 @@ void FControlRigEditor::GenerateEventQueueMenuContent(FMenuBuilder& MenuBuilder)
     			FText::FromString(FString::Printf(TEXT("Runs the user defined %s"), *EventNameStr)),
     			GetEventQueueIcon({EntryName}),
     			FUIAction(
-    				FExecuteAction::CreateSP(this, &FRigVMEditor::SetEventQueue, TArray<FName>({EntryName})),
+#if WITH_RIGVMLEGACYEDITOR
+    				FExecuteAction::CreateSP(this, &FRigVMLegacyEditor::SetEventQueue, TArray<FName>({EntryName})),
+#else
+    				FExecuteAction::CreateSP(this, &FRigVMNewEditor::SetEventQueue, TArray<FName>({EntryName})),
+#endif
     				FCanExecuteAction()
     			)
     		);
@@ -3441,7 +3531,7 @@ FTransform FControlRigEditor::GetRigElementTransform(const FRigElementKey& InEle
 
 	if(bOnDebugInstance)
 	{
-		UControlRig* DebuggedControlRig = Cast<UControlRig>(GetBlueprintObj()->GetObjectBeingDebugged());
+		UControlRig* DebuggedControlRig = Cast<UControlRig>(FRigVMEditorBase::GetBlueprintObj()->GetObjectBeingDebugged());
 		if (DebuggedControlRig == nullptr)
 		{
 			DebuggedControlRig = GetControlRig();
@@ -3551,7 +3641,7 @@ void FControlRigEditor::SetRigElementTransform(const FRigElementKey& InElement,
 
 void FControlRigEditor::OnFinishedChangingProperties(const FPropertyChangedEvent& PropertyChangedEvent)
 {
-	FRigVMEditor::OnFinishedChangingProperties(PropertyChangedEvent);
+	FRigVMEditorBase::OnFinishedChangingProperties(PropertyChangedEvent);
 	
 	DECLARE_SCOPE_HIERARCHICAL_COUNTER_FUNC()
 
@@ -3578,7 +3668,7 @@ void FControlRigEditor::OnFinishedChangingProperties(const FPropertyChangedEvent
 
 void FControlRigEditor::OnWrappedPropertyChangedChainEvent(URigVMDetailsViewWrapperObject* InWrapperObject, const FString& InPropertyPath, FPropertyChangedChainEvent& InPropertyChangedChainEvent)
 {
-	FRigVMEditor::OnWrappedPropertyChangedChainEvent(InWrapperObject, InPropertyPath, InPropertyChangedChainEvent);
+	IControlRigEditor::OnWrappedPropertyChangedChainEvent(InWrapperObject, InPropertyPath, InPropertyChangedChainEvent);
 	
 	check(InWrapperObject);
 	check(!GetWrapperObjects().IsEmpty());
@@ -3821,7 +3911,7 @@ EVisibility FControlRigEditor::GetDirectManipulationVisibility() const
 
 FText FControlRigEditor::GetDirectionManipulationText() const
 {
-	if (UControlRig* DebuggedControlRig = Cast<UControlRig>(GetBlueprintObj()->GetObjectBeingDebugged()))
+	if (UControlRig* DebuggedControlRig = Cast<UControlRig>(FRigVMEditorBase::GetBlueprintObj()->GetObjectBeingDebugged()))
 	{
 		if (URigHierarchy* Hierarchy = DebuggedControlRig->GetHierarchy())
 		{
@@ -3853,7 +3943,7 @@ void FControlRigEditor::OnDirectManipulationChanged(TSharedPtr<FString> NewValue
 		return;
 	}
 	
-	UControlRigBlueprint* ControlRigBlueprint = CastChecked<UControlRigBlueprint>(GetBlueprintObj());
+	UControlRigBlueprint* ControlRigBlueprint = CastChecked<UControlRigBlueprint>(FRigVMEditorBase::GetBlueprintObj());
 	if(ControlRigBlueprint == nullptr)
 	{
 		return;
@@ -3886,7 +3976,7 @@ const TArray<FRigDirectManipulationTarget> FControlRigEditor::GetDirectManipulat
 {
 	if(DirectManipulationSubject.IsValid())
 	{
-		if (UControlRig* DebuggedControlRig = Cast<UControlRig>(GetBlueprintObj()->GetObjectBeingDebugged()))
+		if (UControlRig* DebuggedControlRig = Cast<UControlRig>(FRigVMEditorBase::GetBlueprintObj()->GetObjectBeingDebugged()))
 		{
 			if(const URigVMUnitNode* Node = DirectManipulationSubject.Get())
 			{
@@ -3980,26 +4070,42 @@ FReply FControlRigEditor::OnNavigateToConnectorWarning() const
 
 void FControlRigEditor::BindCommands()
 {
-	FRigVMEditor::BindCommands();
+	IControlRigEditor::BindCommands();
 
 	GetToolkitCommands()->MapAction(
 		FControlRigEditorCommands::Get().ConstructionEvent,
-		FExecuteAction::CreateSP(this, &FRigVMEditor::SetEventQueue, TArray<FName>(ConstructionEventQueue)),
+#if WITH_RIGVMLEGACYEDITOR
+		FExecuteAction::CreateSP(this, &FRigVMLegacyEditor::SetEventQueue, TArray<FName>(ConstructionEventQueue)),
+#else
+		FExecuteAction::CreateSP(this, &FRigVMNewEditor::SetEventQueue, TArray<FName>(ConstructionEventQueue)),
+#endif
 		FCanExecuteAction());
 
 	GetToolkitCommands()->MapAction(
 		FControlRigEditorCommands::Get().ForwardsSolveEvent,
-		FExecuteAction::CreateSP(this, &FRigVMEditor::SetEventQueue, TArray<FName>(ForwardsSolveEventQueue)),
+#if WITH_RIGVMLEGACYEDITOR
+		FExecuteAction::CreateSP(this, &FRigVMLegacyEditor::SetEventQueue, TArray<FName>(ForwardsSolveEventQueue)),
+#else
+		FExecuteAction::CreateSP(this, &FRigVMNewEditor::SetEventQueue, TArray<FName>(ForwardsSolveEventQueue)),
+#endif
 		FCanExecuteAction());
 
 	GetToolkitCommands()->MapAction(
 		FControlRigEditorCommands::Get().BackwardsSolveEvent,
-		FExecuteAction::CreateSP(this, &FRigVMEditor::SetEventQueue, TArray<FName>(BackwardsSolveEventQueue)),
+#if WITH_RIGVMLEGACYEDITOR
+		FExecuteAction::CreateSP(this, &FRigVMLegacyEditor::SetEventQueue, TArray<FName>(BackwardsSolveEventQueue)),
+#else
+		FExecuteAction::CreateSP(this, &FRigVMNewEditor::SetEventQueue, TArray<FName>(BackwardsSolveEventQueue)),
+#endif
 		FCanExecuteAction());
 
 	GetToolkitCommands()->MapAction(
 		FControlRigEditorCommands::Get().BackwardsAndForwardsSolveEvent,
-		FExecuteAction::CreateSP(this, &FRigVMEditor::SetEventQueue, TArray<FName>(BackwardsAndForwardsSolveEventQueue)),
+#if WITH_RIGVMLEGACYEDITOR
+		FExecuteAction::CreateSP(this, &FRigVMLegacyEditor::SetEventQueue, TArray<FName>(BackwardsAndForwardsSolveEventQueue)),
+#else
+		FExecuteAction::CreateSP(this, &FRigVMNewEditor::SetEventQueue, TArray<FName>(BackwardsAndForwardsSolveEventQueue)),
+#endif
 		FCanExecuteAction());
 
 	GetToolkitCommands()->MapAction(
@@ -4110,7 +4216,7 @@ void FControlRigEditor::OnHierarchyChanged()
 
 void FControlRigEditor::OnHierarchyModified(ERigHierarchyNotification InNotif, URigHierarchy* InHierarchy, const FRigNotificationSubject& InSubject)
 {
-	UControlRigBlueprint* RigBlueprint = Cast<UControlRigBlueprint>(GetBlueprintObj());
+	UControlRigBlueprint* RigBlueprint = Cast<UControlRigBlueprint>(FRigVMEditorBase::GetBlueprintObj());
 	if(RigBlueprint == nullptr)
 	{
 		return;
@@ -4408,7 +4514,7 @@ void FControlRigEditor::OnHierarchyModified_AnyThread(ERigHierarchyNotification
 						else
 						{
 							// only clear the details if we are not looking at a transient control
-							if (UControlRig* DebuggedControlRig = Cast<UControlRig>(GetBlueprintObj()->GetObjectBeingDebugged()))
+							if (UControlRig* DebuggedControlRig = Cast<UControlRig>(FRigVMEditorBase::GetBlueprintObj()->GetObjectBeingDebugged()))
 							{
 								if(DebuggedControlRig->RigUnitManipulationInfos.IsEmpty())
 								{
@@ -4749,7 +4855,7 @@ void FControlRigEditor::UpdateBoneModification(FName BoneName, const FTransform&
 		}
 		
 		TMap<FName, FTransform>* TransformOverrideMap = &ControlRig->TransformOverrideForUserCreatedBones;
-		if (UControlRig* DebuggedControlRig = Cast<UControlRig>(GetBlueprintObj()->GetObjectBeingDebugged()))
+		if (UControlRig* DebuggedControlRig = Cast<UControlRig>(FRigVMEditorBase::GetBlueprintObj()->GetObjectBeingDebugged()))
 		{
 			TransformOverrideMap = &DebuggedControlRig->TransformOverrideForUserCreatedBones;
 		}
@@ -4771,7 +4877,7 @@ void FControlRigEditor::RemoveBoneModification(FName BoneName)
 		}
 
 		TMap<FName, FTransform>* TransformOverrideMap = &ControlRig->TransformOverrideForUserCreatedBones;
-		if (UControlRig* DebuggedControlRig = Cast<UControlRig>(GetBlueprintObj()->GetObjectBeingDebugged()))
+		if (UControlRig* DebuggedControlRig = Cast<UControlRig>(FRigVMEditorBase::GetBlueprintObj()->GetObjectBeingDebugged()))
 		{
 			TransformOverrideMap = &DebuggedControlRig->TransformOverrideForUserCreatedBones;
 		}
@@ -4790,7 +4896,7 @@ void FControlRigEditor::ResetAllBoneModification()
 		}
 
 		TMap<FName, FTransform>* TransformOverrideMap = &ControlRig->TransformOverrideForUserCreatedBones;
-		if (UControlRig* DebuggedControlRig = Cast<UControlRig>(GetBlueprintObj()->GetObjectBeingDebugged()))
+		if (UControlRig* DebuggedControlRig = Cast<UControlRig>(FRigVMEditorBase::GetBlueprintObj()->GetObjectBeingDebugged()))
 		{
 			TransformOverrideMap = &DebuggedControlRig->TransformOverrideForUserCreatedBones;
 		}
@@ -5174,7 +5280,7 @@ void FControlRigEditor::OnGraphNodeDropToPerform(TSharedPtr<FDragDropOperation>
 		return;
 	}
 
-	FRigVMEditor::OnGraphNodeDropToPerform(InDragDropOp, InGraph, InNodePosition, InScreenPosition);
+	IControlRigEditor::OnGraphNodeDropToPerform(InDragDropOp, InGraph, InNodePosition, InScreenPosition);
 }
 
 void FControlRigEditor::HandleMakeElementGetterSetter(ERigElementGetterSetterType Type, bool bIsGetter, TArray<FRigElementKey> Keys, UEdGraph* Graph, FVector2D NodePosition)
@@ -5690,13 +5796,13 @@ void FControlRigEditor::HandleMakeElementGetterSetter(ERigElementGetterSetterTyp
 
 void FControlRigEditor::HandleOnControlModified(UControlRig* Subject, FRigControlElement* ControlElement, const FRigControlModifiedContext& Context)
 {
-	UControlRig* DebuggedControlRig = Cast<UControlRig>(GetBlueprintObj()->GetObjectBeingDebugged());
+	UControlRig* DebuggedControlRig = Cast<UControlRig>(FRigVMEditorBase::GetBlueprintObj()->GetObjectBeingDebugged());
 	if (Subject != DebuggedControlRig)
 	{
 		return;
 	}
 
-	UControlRigBlueprint* Blueprint = Cast<UControlRigBlueprint>(GetBlueprintObj());
+	UControlRigBlueprint* Blueprint = Cast<UControlRigBlueprint>(FRigVMEditorBase::GetBlueprintObj());
 	if (Blueprint == nullptr)
 	{
 		return;
@@ -5847,7 +5953,7 @@ void FControlRigEditor::HandleOnControlModified(UControlRig* Subject, FRigContro
 void FControlRigEditor::HandleRefreshEditorFromBlueprint(URigVMBlueprint* InBlueprint)
 {
 	OnHierarchyChanged();
-	FRigVMEditor::HandleRefreshEditorFromBlueprint(InBlueprint);
+	IControlRigEditor::HandleRefreshEditorFromBlueprint(InBlueprint);
 }
 
 UToolMenu* FControlRigEditor::HandleOnGetViewportContextMenuDelegate()
diff --git a/Engine/Plugins/Animation/ControlRig/Source/ControlRigEditor/Private/Editor/ControlRigEditor.h b/Engine/Plugins/Animation/ControlRig/Source/ControlRigEditor/Private/Editor/ControlRigEditor.h
--- a/Engine/Plugins/Animation/ControlRig/Source/ControlRigEditor/Private/Editor/ControlRigEditor.h
+++ b/Engine/Plugins/Animation/ControlRig/Source/ControlRigEditor/Private/Editor/ControlRigEditor.h
@@ -77,7 +77,7 @@ public:
 
 	virtual UObject* GetOuterForHost() const override;
 
-	// FRigVMEditor interface
+	// FRigVMEditorBase interface
 	virtual UClass* GetDetailWrapperClass() const;
 	virtual void Compile() override;
 	virtual void HandleModifiedEvent(ERigVMGraphNotifType InNotifType, URigVMGraph* InGraph, UObject* InSubject) override;
@@ -110,8 +110,13 @@ public:
 
 	// BlueprintEditor interface
 	virtual FReply OnSpawnGraphNodeByShortcut(FInputChord InChord, const FVector2D& InPosition, UEdGraph* InGraph) override;
+#if WITH_RIGVMLEGACYEDITOR
 	virtual bool IsSectionVisible(NodeSectionID::Type InSectionID) const override;
-	virtual bool NewDocument_IsVisibleForType(ECreatedDocumentType GraphType) const override;
+	virtual bool NewDocument_IsVisibleForType(FBlueprintEditor::ECreatedDocumentType GraphType) const override;
+#else
+	virtual bool IsSectionVisible(RigVMNodeSectionID::Type InSectionID) const override;
+	virtual bool NewDocument_IsVisibleForType(FRigVMEditorBase::ECreatedDocumentType GraphType) const override;
+#endif
 
 	virtual void PostTransaction(bool bSuccess, const FTransaction* Transaction, bool bIsRedo) override;
 
@@ -134,6 +139,7 @@ public:
 	bool DetailViewShowsAnyRigElement() const;
 	bool DetailViewShowsRigElement(FRigElementKey InKey) const;
 	TArray<FRigElementKey> GetSelectedRigElementsFromDetailView() const;
+	TArray< TWeakObjectPtr<UObject> > GetSelectedObjectsFromDetailView() const;
 
 	void SetDetailViewForRigModules();
 	void SetDetailViewForRigModules(const TArray<FString> InKeys);
diff --git a/Engine/Plugins/Animation/ControlRig/Source/ControlRigEditor/Private/Editor/ControlRigEditorMode.cpp b/Engine/Plugins/Animation/ControlRig/Source/ControlRigEditor/Private/Editor/ControlRigEditorMode.cpp
--- a/Engine/Plugins/Animation/ControlRig/Source/ControlRigEditor/Private/Editor/ControlRigEditorMode.cpp
+++ b/Engine/Plugins/Animation/ControlRig/Source/ControlRigEditor/Private/Editor/ControlRigEditorMode.cpp
@@ -1,8 +1,10 @@
 // Copyright Epic Games, Inc. All Rights Reserved.
 
-#include "Editor/ControlRigEditorMode.h" 
+#include "Editor/ControlRigEditorMode.h"
+#if WITH_RIGVMLEGACYEDITOR
 #include "BlueprintEditorTabs.h"
 #include "SBlueprintEditorToolbar.h"
+#endif 
 #include "PersonaModule.h"
 #include "IPersonaToolkit.h"
 #include "PersonaTabs.h"
@@ -17,9 +19,13 @@
 #include "ToolMenus.h"
 
 FControlRigEditorMode::FControlRigEditorMode(const TSharedRef<FControlRigEditor>& InControlRigEditor, bool bCreateDefaultLayout)
+#if WITH_RIGVMLEGACYEDITOR
 	: FBlueprintEditorApplicationMode(InControlRigEditor, FControlRigEditorModes::ControlRigEditorMode, FControlRigEditorModes::GetLocalizedMode, false, false)
+#else
+	: FRigVMNewEditorMode(InControlRigEditor)
+#endif
 {
-	ControlRigBlueprintPtr = CastChecked<UControlRigBlueprint>(InControlRigEditor->GetBlueprintObj());
+	ControlRigBlueprintPtr = CastChecked<UControlRigBlueprint>(InControlRigEditor->FRigVMEditorBase::GetBlueprintObj());
 
 	TabFactories.RegisterFactory(MakeShared<FRigHierarchyTabSummoner>(InControlRigEditor));
 	TabFactories.RegisterFactory(MakeShared<FRigVMExecutionStackTabSummoner>(InControlRigEditor));
@@ -31,7 +37,9 @@ FControlRigEditorMode::FControlRigEditorMode(const TSharedRef<FControlRigEditor>
 	FPersonaModule& PersonaModule = FModuleManager::LoadModuleChecked<FPersonaModule>("Persona");
 
  	FPersonaViewportArgs ViewportArgs(InControlRigEditor->GetPersonaToolkit()->GetPreviewScene());
+#if WITH_RIGVMLEGACYEDITOR
  	ViewportArgs.BlueprintEditor = InControlRigEditor;
+#endif
  	ViewportArgs.bShowStats = false;
 	ViewportArgs.bShowPlaySpeedMenu = false;
 	ViewportArgs.bShowTimeline = true;
@@ -78,7 +86,9 @@ FControlRigEditorMode::FControlRigEditorMode(const TSharedRef<FControlRigEditor>
 							->AddTab(FRigHierarchyTabSummoner::TabID, ETabState::OpenedTab)
 							->AddTab(FRigVMExecutionStackTabSummoner::TabID, ETabState::OpenedTab)
 							->AddTab(FRigCurveContainerTabSummoner::TabID, ETabState::OpenedTab)
+#if WITH_RIGVMLEGACYEDITOR
 							->AddTab(FBlueprintEditorTabs::MyBlueprintID, ETabState::OpenedTab)
+#endif
 							->AddTab(FRigVMEditorGraphExplorerTabSummoner::TabID, ETabState::ClosedTab)
 						)
 					)
@@ -100,8 +110,10 @@ FControlRigEditorMode::FControlRigEditorMode(const TSharedRef<FControlRigEditor>
 							// Middle bottom - compiler results & find
 							FTabManager::NewStack()
 							->SetSizeCoefficient(0.2f)
+#if WITH_RIGVMLEGACYEDITOR
 							->AddTab(FBlueprintEditorTabs::CompilerResultsID, ETabState::ClosedTab)
 							->AddTab(FBlueprintEditorTabs::FindResultsID, ETabState::ClosedTab)
+#endif
 						)
 					)
 					->Split
@@ -116,33 +128,45 @@ FControlRigEditorMode::FControlRigEditorMode(const TSharedRef<FControlRigEditor>
 							FTabManager::NewStack()
 							->SetHideTabWell(false)
 							->SetSizeCoefficient(1.f)
+#if WITH_RIGVMLEGACYEDITOR
 							->AddTab(FBlueprintEditorTabs::DetailsID, ETabState::OpenedTab)
+#endif
 							->AddTab(FPersonaTabs::AdvancedPreviewSceneSettingsID, ETabState::OpenedTab)
 							->AddTab(FRigAnimAttributeTabSummoner::TabID, ETabState::OpenedTab)
+#if WITH_RIGVMLEGACYEDITOR
 							->SetForegroundTab(FBlueprintEditorTabs::DetailsID)
+#endif
 						)
 					)
 				)
 			);
 	}
-	
+
+#if WITH_RIGVMLEGACYEDITOR
 	if (UToolMenu* Toolbar = InControlRigEditor->RegisterModeToolbarIfUnregistered(GetModeName()))
 	{
 		InControlRigEditor->GetToolbarBuilder()->AddCompileToolbar(Toolbar);
 		InControlRigEditor->GetToolbarBuilder()->AddScriptingToolbar(Toolbar);
 		InControlRigEditor->GetToolbarBuilder()->AddBlueprintGlobalOptionsToolbar(Toolbar);
 	}
+#endif
 }
 
 void FControlRigEditorMode::RegisterTabFactories(TSharedPtr<FTabManager> InTabManager)
 {
+#if WITH_RIGVMLEGACYEDITOR
 	TSharedPtr<FBlueprintEditor> BP = MyBlueprintEditor.Pin();
+#else
+	TSharedPtr<FRigVMNewEditor> BP = Editor.Pin();
+#endif
 	
 	BP->RegisterToolbarTab(InTabManager.ToSharedRef());
 
 	// Mode-specific setup
+#if WITH_RIGVMLEGACYEDITOR
 	BP->PushTabFactories(CoreTabFactories);
 	BP->PushTabFactories(BlueprintEditorTabFactories);
+#endif
 	BP->PushTabFactories(TabFactories);
 }
 
@@ -175,7 +199,9 @@ FModularRigEditorMode::FModularRigEditorMode(const TSharedRef<FControlRigEditor>
 						->SetHideTabWell(false)
 						->SetSizeCoefficient(0.5f)
 						->AddTab(FRigModuleAssetBrowserTabSummoner::TabID, ETabState::OpenedTab)
+#if WITH_RIGVMLEGACYEDITOR
 						->SetForegroundTab(FBlueprintEditorTabs::DetailsID)
+#endif
 					)
 					->Split
 					(
@@ -185,7 +211,9 @@ FModularRigEditorMode::FModularRigEditorMode(const TSharedRef<FControlRigEditor>
 						->SetSizeCoefficient(0.5f)
 						->SetForegroundTab(FModularRigModelTabSummoner::TabID)
 						->AddTab(FModularRigModelTabSummoner::TabID, ETabState::OpenedTab)
+#if WITH_RIGVMLEGACYEDITOR
 						->AddTab(FBlueprintEditorTabs::MyBlueprintID, ETabState::OpenedTab)
+#endif
 					)
 				)
 				->Split
@@ -202,8 +230,10 @@ FModularRigEditorMode::FModularRigEditorMode(const TSharedRef<FControlRigEditor>
 					FTabManager::NewStack()
 					->SetSizeCoefficient(0.2f)
 					->SetHideTabWell(false)
+#if WITH_RIGVMLEGACYEDITOR
 					->SetForegroundTab(FBlueprintEditorTabs::DetailsID)
 					->AddTab(FBlueprintEditorTabs::DetailsID, ETabState::OpenedTab)
+#endif
 					->AddTab(FPersonaTabs::AdvancedPreviewSceneSettingsID, ETabState::OpenedTab)
 					->AddTab(FRigHierarchyTabSummoner::TabID, ETabState::OpenedTab)
 				)
@@ -213,16 +243,22 @@ FModularRigEditorMode::FModularRigEditorMode(const TSharedRef<FControlRigEditor>
 
 void FModularRigEditorMode::RegisterTabFactories(TSharedPtr<FTabManager> InTabManager)
 {
+#if WITH_RIGVMLEGACYEDITOR
 	TSharedPtr<FBlueprintEditor> BP = MyBlueprintEditor.Pin();
+#else
+	TSharedPtr<FRigVMNewEditor> BP = Editor.Pin();
+#endif
 	
 	BP->RegisterToolbarTab(InTabManager.ToSharedRef());
 
 	static const TArray<FName> DisallowedTabs = {
+#if WITH_RIGVMLEGACYEDITOR
 		FBlueprintEditorTabs::PaletteID,
 		FBlueprintEditorTabs::ReplaceNodeReferencesID,
 		FBlueprintEditorTabs::CompilerResultsID,
 		FBlueprintEditorTabs::FindResultsID,
 		FBlueprintEditorTabs::BookmarksID,
+#endif
 		FRigVMExecutionStackTabSummoner::TabID
 	};
 
@@ -239,7 +275,9 @@ void FModularRigEditorMode::RegisterTabFactories(TSharedPtr<FTabManager> InTabMa
 	};
 
 	// Mode-specific setup
+#if WITH_RIGVMLEGACYEDITOR
 	PushTabFactories(CoreTabFactories);
 	PushTabFactories(BlueprintEditorTabFactories);
+#endif
 	PushTabFactories(TabFactories);
 }
diff --git a/Engine/Plugins/Animation/ControlRig/Source/ControlRigEditor/Private/Editor/ControlRigEditorMode.h b/Engine/Plugins/Animation/ControlRig/Source/ControlRigEditor/Private/Editor/ControlRigEditorMode.h
--- a/Engine/Plugins/Animation/ControlRig/Source/ControlRigEditor/Private/Editor/ControlRigEditorMode.h
+++ b/Engine/Plugins/Animation/ControlRig/Source/ControlRigEditor/Private/Editor/ControlRigEditorMode.h
@@ -6,11 +6,19 @@
 #include "WorkflowOrientedApp/WorkflowTabManager.h"
 #include "Editor/ControlRigEditor.h"
 #include "ControlRigBlueprint.h"
+#if WITH_RIGVMLEGACYEDITOR
 #include "BlueprintEditorModes.h"
+#else
+#include "Editor/RigVMNewEditorMode.h" 
+#endif
 
 class UControlRigBlueprint;
 
+#if WITH_RIGVMLEGACYEDITOR
 class FControlRigEditorMode : public FBlueprintEditorApplicationMode
+#else
+class FControlRigEditorMode : public FRigVMNewEditorMode
+#endif
 {
 public:
 	FControlRigEditorMode(const TSharedRef<FControlRigEditor>& InControlRigEditor, bool bCreateDefaultLayout = true);
diff --git a/Engine/Plugins/Animation/ControlRig/Source/ControlRigEditor/Private/Editor/RigInfluenceMapTabSummoner.cpp b/Engine/Plugins/Animation/ControlRig/Source/ControlRigEditor/Private/Editor/RigInfluenceMapTabSummoner.cpp
--- a/Engine/Plugins/Animation/ControlRig/Source/ControlRigEditor/Private/Editor/RigInfluenceMapTabSummoner.cpp
+++ b/Engine/Plugins/Animation/ControlRig/Source/ControlRigEditor/Private/Editor/RigInfluenceMapTabSummoner.cpp
@@ -4,9 +4,14 @@
 #include "Editor/SRigHierarchy.h"
 #include "Editor/RigVMEditorStyle.h"
 #include "Editor/ControlRigEditor.h"
-#include "SKismetInspector.h"
 #include "ControlRigBlueprint.h"
 
+#if WITH_RIGVMLEGACYEDITOR
+#include "SKismetInspector.h"
+#else
+#include "Editor/SRigVMDetailsInspector.h"
+#endif
+
 #define LOCTEXT_NAMESPACE "RigInfluenceMapTabSummoner"
 
 const FName FRigInfluenceMapTabSummoner::TabID(TEXT("RigInfluenceMap"));
@@ -26,11 +31,15 @@ FRigInfluenceMapTabSummoner::FRigInfluenceMapTabSummoner(const TSharedRef<FContr
 
 TSharedRef<SWidget> FRigInfluenceMapTabSummoner::CreateTabBody(const FWorkflowTabSpawnInfo& Info) const
 {
+#if WITH_RIGVMLEGACYEDITOR
 	TSharedRef<SKismetInspector> KismetInspector = SNew(SKismetInspector);
+#else
+	TSharedRef<SRigVMDetailsInspector> KismetInspector = SNew(SRigVMDetailsInspector);
+#endif
 
-	if (ControlRigEditor.IsValid())
+	if (TSharedPtr<FControlRigEditor> Editor = ControlRigEditor.Pin())
 	{
-		if (UControlRigBlueprint* RigBlueprint = Cast<UControlRigBlueprint>(ControlRigEditor.Pin()->GetBlueprintObj()))
+		if (UControlRigBlueprint* RigBlueprint = Cast<UControlRigBlueprint>(Editor->FRigVMEditorBase::GetBlueprintObj()))
 		{
 			TSharedPtr<FStructOnScope> StructToDisplay = MakeShareable(new FStructOnScope(
 				FRigInfluenceMapPerEvent::StaticStruct(), (uint8*)&RigBlueprint->Influences));
diff --git a/Engine/Plugins/Animation/ControlRig/Source/ControlRigEditor/Private/Editor/RigValidationTabSummoner.cpp b/Engine/Plugins/Animation/ControlRig/Source/ControlRigEditor/Private/Editor/RigValidationTabSummoner.cpp
--- a/Engine/Plugins/Animation/ControlRig/Source/ControlRigEditor/Private/Editor/RigValidationTabSummoner.cpp
+++ b/Engine/Plugins/Animation/ControlRig/Source/ControlRigEditor/Private/Editor/RigValidationTabSummoner.cpp
@@ -27,16 +27,20 @@ FRigValidationTabSummoner::FRigValidationTabSummoner(const TSharedRef<FControlRi
 TSharedRef<SWidget> FRigValidationTabSummoner::CreateTabBody(const FWorkflowTabSpawnInfo& Info) const
 {
 	ensure(WeakControlRigEditor.IsValid());
-
-	UControlRigBlueprint* RigBlueprint = Cast<UControlRigBlueprint>(WeakControlRigEditor.Pin()->GetBlueprintObj());
-	check(RigBlueprint);
+	if (TSharedPtr<FRigVMEditorBase> Editor = WeakControlRigEditor.Pin())
+	{
+		UControlRigBlueprint* RigBlueprint = Cast<UControlRigBlueprint>(Editor->FRigVMEditorBase::GetBlueprintObj());
+		check(RigBlueprint);
 	
-	UControlRigValidator* Validator = RigBlueprint->Validator;
-	check(Validator);
+		UControlRigValidator* Validator = RigBlueprint->Validator;
+		check(Validator);
+
+		TSharedRef<SControlRigValidationWidget> ValidationWidget = SNew(SControlRigValidationWidget, Validator);
+		Validator->SetControlRig(Cast<UControlRig>(RigBlueprint->GetObjectBeingDebugged()));
+		return ValidationWidget;
+	}
 
-	TSharedRef<SControlRigValidationWidget> ValidationWidget = SNew(SControlRigValidationWidget, Validator);
-	Validator->SetControlRig(Cast<UControlRig>(RigBlueprint->GetObjectBeingDebugged()));
-	return ValidationWidget;
+	return SNullWidget::NullWidget;
 }
 
 #undef LOCTEXT_NAMESPACE 
diff --git a/Engine/Plugins/Animation/ControlRig/Source/ControlRigEditor/Private/Editor/SControlRigValidationWidget.cpp b/Engine/Plugins/Animation/ControlRig/Source/ControlRigEditor/Private/Editor/SControlRigValidationWidget.cpp
--- a/Engine/Plugins/Animation/ControlRig/Source/ControlRigEditor/Private/Editor/SControlRigValidationWidget.cpp
+++ b/Engine/Plugins/Animation/ControlRig/Source/ControlRigEditor/Private/Editor/SControlRigValidationWidget.cpp
@@ -16,9 +16,16 @@
 #include "UserInterface/SMessageLogListing.h"
 #include "Framework/Application/SlateApplication.h"
 #include "ControlRigBlueprint.h"
+#include "Editor/RigVMNewEditor.h"
 #include "UObject/UObjectIterator.h"
 #include "Rigs/RigHierarchyController.h"
 
+#if WITH_RIGVMLEGACYEDITOR
+#include "SKismetInspector.h"
+#else
+#include "Editor/SRigVMDetailsInspector.h"
+#endif
+
 #define LOCTEXT_NAMESPACE "SControlRigValidationWidget"
 
 //////////////////////////////////////////////////////////////
@@ -64,7 +71,11 @@ void SControlRigValidationPassTableRow::Construct(const FArguments& InArgs, cons
 			.Padding(24.0f, 4.0f, 0.f, 0.f)
 			.AutoHeight()
 			[
+#if WITH_RIGVMLEGACYEDITOR
 				SAssignNew(KismetInspector, SKismetInspector)
+#else
+				SAssignNew(KismetInspector, SRigVMDetailsInspector)
+#endif
 				.Visibility(InValidationWidget, &SControlRigValidationWidget::IsClassVisible, InPassItem->Class)
 				.ShowTitleArea(false)
 			]
diff --git a/Engine/Plugins/Animation/ControlRig/Source/ControlRigEditor/Private/Editor/SControlRigValidationWidget.h b/Engine/Plugins/Animation/ControlRig/Source/ControlRigEditor/Private/Editor/SControlRigValidationWidget.h
--- a/Engine/Plugins/Animation/ControlRig/Source/ControlRigEditor/Private/Editor/SControlRigValidationWidget.h
+++ b/Engine/Plugins/Animation/ControlRig/Source/ControlRigEditor/Private/Editor/SControlRigValidationWidget.h
@@ -4,11 +4,19 @@
 
 #include "CoreMinimal.h"
 #include "Widgets/Views/SListView.h"
-#include "SKismetInspector.h"
 #include "Model/MessageLogListingModel.h"
 #include "Presentation/MessageLogListingViewModel.h"
 #include "ControlRigValidationPass.h"
 
+#if WITH_RIGVMLEGACYEDITOR
+#include "SKismetInspector.h"
+#else
+#include "Editor/SRigVMDetailsInspector.h"
+#endif
+
+#if !WITH_RIGVMLEGACYEDITOR
+class SRigVMDetailsInspector;
+#endif
 class SControlRigValidationWidget;
 
 class FControlRigValidationPassItem : public TSharedFromThis<FControlRigValidationPassItem>
@@ -37,8 +45,11 @@ class SControlRigValidationPassTableRow : public STableRow<TSharedPtr<FControlRi
 	void RefreshDetails(UControlRigValidator* InValidator, UClass* InClass);
 
 private:
-
+#if WITH_RIGVMLEGACYEDITOR
 	TSharedPtr<SKismetInspector> KismetInspector;
+#else
+	TSharedPtr<SRigVMDetailsInspector> KismetInspector;
+#endif
 };
 
 class SControlRigValidationWidget : public SCompoundWidget
diff --git a/Engine/Plugins/Animation/ControlRig/Source/ControlRigEditor/Private/Editor/SModularRigModel.cpp b/Engine/Plugins/Animation/ControlRig/Source/ControlRigEditor/Private/Editor/SModularRigModel.cpp
--- a/Engine/Plugins/Animation/ControlRig/Source/ControlRigEditor/Private/Editor/SModularRigModel.cpp
+++ b/Engine/Plugins/Animation/ControlRig/Source/ControlRigEditor/Private/Editor/SModularRigModel.cpp
@@ -34,7 +34,6 @@
 #include "HAL/PlatformTime.h"
 #include "Dialogs/Dialogs.h"
 #include "IPersonaToolkit.h"
-#include "SKismetInspector.h"
 #include "Types/WidgetActiveTimerDelegate.h"
 #include "Dialog/SCustomDialog.h"
 #include "EditMode/ControlRigEditMode.h"
@@ -58,6 +57,12 @@
 #include "Widgets/SRigVMBulkEditDialog.h"
 #include "Widgets/SRigVMSwapAssetReferencesWidget.h"
 
+#if WITH_RIGVMLEGACYEDITOR
+#include "SKismetInspector.h"
+#else
+#include "Editor/SRigVMDetailsInspector.h"
+#endif
+
 #define LOCTEXT_NAMESPACE "SModularRigModel"
 
 ///////////////////////////////////////////////////////////
@@ -237,7 +242,7 @@ void SModularRigModel::Construct(const FArguments& InArgs, TSharedRef<FControlRi
 	}
 }
 
-void SModularRigModel::OnEditorClose(const FRigVMEditor* InEditor, URigVMBlueprint* InBlueprint)
+void SModularRigModel::OnEditorClose(const FRigVMEditorBase* InEditor, URigVMBlueprint* InBlueprint)
 {
 	if (InEditor)
 	{
@@ -922,15 +927,19 @@ void SModularRigModel::HandleMirrorModules(const TArray<FString>& InPaths)
 	{
 		FRigVMMirrorSettings Settings;
 		TSharedPtr<FStructOnScope> StructToDisplay = MakeShareable(new FStructOnScope(FRigVMMirrorSettings::StaticStruct(), (uint8*)&Settings));
-		
-		TSharedRef<SKismetInspector> KismetInspector = SNew(SKismetInspector);
-		KismetInspector->ShowSingleStruct(StructToDisplay);
+
+#if WITH_RIGVMLEGACYEDITOR
+		TSharedRef<SKismetInspector> DetailsInspector = SNew(SKismetInspector);
+#else
+		TSharedRef<SRigVMDetailsInspector> DetailsInspector = SNew(SRigVMDetailsInspector);
+#endif
+		DetailsInspector->ShowSingleStruct(StructToDisplay);
 		
 		TSharedRef<SCustomDialog> MirrorDialog = SNew(SCustomDialog)
 			.Title(FText(LOCTEXT("ControlModularModelMirror", "Mirror Selected Modules")))
 			.Content()
 			[
-				KismetInspector
+				DetailsInspector
 			]
 			.Buttons({
 				SCustomDialog::FButton(LOCTEXT("OK", "OK")),
diff --git a/Engine/Plugins/Animation/ControlRig/Source/ControlRigEditor/Private/Editor/SModularRigModel.h b/Engine/Plugins/Animation/ControlRig/Source/ControlRigEditor/Private/Editor/SModularRigModel.h
--- a/Engine/Plugins/Animation/ControlRig/Source/ControlRigEditor/Private/Editor/SModularRigModel.h
+++ b/Engine/Plugins/Animation/ControlRig/Source/ControlRigEditor/Private/Editor/SModularRigModel.h
@@ -7,7 +7,7 @@
 #include "Editor/SModularRigTreeView.h"
 #include "ControlRigBlueprint.h"
 #include "DragAndDrop/GraphNodeDragDropOp.h"
-#include "Editor/RigVMEditor.h"
+#include "Editor/RigVMNewEditor.h"
 #include "ControlRigDragOps.h"
 
 class SModularRigModel;
@@ -43,7 +43,7 @@ public:
 
 private:
 
-	void OnEditorClose(const FRigVMEditor* InEditor, URigVMBlueprint* InBlueprint);
+	void OnEditorClose(const FRigVMEditorBase* InEditor, URigVMBlueprint* InBlueprint);
 
 	/** Bind commands that this widget handles */
 	void BindCommands();
diff --git a/Engine/Plugins/Animation/ControlRig/Source/ControlRigEditor/Private/Editor/SRigCurveContainer.cpp b/Engine/Plugins/Animation/ControlRig/Source/ControlRigEditor/Private/Editor/SRigCurveContainer.cpp
--- a/Engine/Plugins/Animation/ControlRig/Source/ControlRigEditor/Private/Editor/SRigCurveContainer.cpp
+++ b/Engine/Plugins/Animation/ControlRig/Source/ControlRigEditor/Private/Editor/SRigCurveContainer.cpp
@@ -310,7 +310,7 @@ TSharedPtr<SWidget> SRigCurveContainer::OnGetContextMenuContent() const
 	return MenuBuilder.MakeWidget();
 }
 
-void SRigCurveContainer::OnEditorClose(const FRigVMEditor* InEditor, URigVMBlueprint* InBlueprint)
+void SRigCurveContainer::OnEditorClose(const FRigVMEditorBase* InEditor, URigVMBlueprint* InBlueprint)
 {
 	if (UControlRigBlueprint* BP = Cast<UControlRigBlueprint>(InBlueprint))
 	{
diff --git a/Engine/Plugins/Animation/ControlRig/Source/ControlRigEditor/Private/Editor/SRigCurveContainer.h b/Engine/Plugins/Animation/ControlRig/Source/ControlRigEditor/Private/Editor/SRigCurveContainer.h
--- a/Engine/Plugins/Animation/ControlRig/Source/ControlRigEditor/Private/Editor/SRigCurveContainer.h
+++ b/Engine/Plugins/Animation/ControlRig/Source/ControlRigEditor/Private/Editor/SRigCurveContainer.h
@@ -18,7 +18,7 @@ class SInlineEditableTextBlock;
 class FControlRigEditor;
 struct FAssetData;
 class URigVMBlueprint;
-class FRigVMEditor;
+class FRigVMEditorBase;
 
 class FDisplayedRigCurveInfo
 {
@@ -213,7 +213,7 @@ private:
 	/** Handler for context menus */
 	TSharedPtr<SWidget> OnGetContextMenuContent() const;
 
-	void OnEditorClose(const FRigVMEditor* InEditor, URigVMBlueprint* InBlueprint);
+	void OnEditorClose(const FRigVMEditorBase* InEditor, URigVMBlueprint* InBlueprint);
 
 	/**
 	* Clears and rebuilds the table, according to an optional search string
diff --git a/Engine/Plugins/Animation/ControlRig/Source/ControlRigEditor/Private/Editor/SRigHierarchy.cpp b/Engine/Plugins/Animation/ControlRig/Source/ControlRigEditor/Private/Editor/SRigHierarchy.cpp
--- a/Engine/Plugins/Animation/ControlRig/Source/ControlRigEditor/Private/Editor/SRigHierarchy.cpp
+++ b/Engine/Plugins/Animation/ControlRig/Source/ControlRigEditor/Private/Editor/SRigHierarchy.cpp
@@ -34,7 +34,6 @@
 #include "HAL/PlatformTime.h"
 #include "Dialogs/Dialogs.h"
 #include "IPersonaToolkit.h"
-#include "SKismetInspector.h"
 #include "Types/WidgetActiveTimerDelegate.h"
 #include "Dialog/SCustomDialog.h"
 #include "EditMode/ControlRigEditMode.h"
@@ -55,6 +54,12 @@
 #include "Preferences/PersonaOptions.h"
 #include "Editor/SModularRigModel.h"
 
+#if WITH_RIGVMLEGACYEDITOR
+#include "SKismetInspector.h"
+#else
+#include "Editor/SRigVMDetailsInspector.h"
+#endif
+
 #include UE_INLINE_GENERATED_CPP_BY_NAME(SRigHierarchy)
 
 #define LOCTEXT_NAMESPACE "SRigHierarchy"
@@ -250,7 +255,7 @@ void SRigHierarchy::Construct(const FArguments& InArgs, TSharedRef<FControlRigEd
 	}));
 }
 
-void SRigHierarchy::OnEditorClose(const FRigVMEditor* InEditor, URigVMBlueprint* InBlueprint)
+void SRigHierarchy::OnEditorClose(const FRigVMEditorBase* InEditor, URigVMBlueprint* InBlueprint)
 {
 	if (InEditor)
 	{
@@ -533,8 +538,11 @@ EVisibility SRigHierarchy::IsSearchbarVisible() const
 FReply SRigHierarchy::OnImportSkeletonClicked()
 {
 	TSharedPtr<FStructOnScope> StructToDisplay = MakeShareable(new FStructOnScope(FRigHierarchyImportSettings::StaticStruct(), (uint8*)&ImportSettings));
-
+#if WITH_RIGVMLEGACYEDITOR
 	TSharedRef<SKismetInspector> KismetInspector = SNew(SKismetInspector);
+#else
+	TSharedRef<SRigVMDetailsInspector> KismetInspector = SNew(SRigVMDetailsInspector);
+#endif
 	KismetInspector->ShowSingleStruct(StructToDisplay);
 
 	SGenericDialogWidget::FArguments DialogArguments;
@@ -2363,8 +2371,11 @@ void SRigHierarchy::HandleMirrorItem()
 		
 		FRigVMMirrorSettings Settings;
 		TSharedPtr<FStructOnScope> StructToDisplay = MakeShareable(new FStructOnScope(FRigVMMirrorSettings::StaticStruct(), (uint8*)&Settings));
-
+#if WITH_RIGVMLEGACYEDITOR
 		TSharedRef<SKismetInspector> KismetInspector = SNew(SKismetInspector);
+#else
+		TSharedRef<SRigVMDetailsInspector> KismetInspector = SNew(SRigVMDetailsInspector);
+#endif
 		KismetInspector->ShowSingleStruct(StructToDisplay);
 
 		TSharedRef<SCustomDialog> MirrorDialog = SNew(SCustomDialog)
diff --git a/Engine/Plugins/Animation/ControlRig/Source/ControlRigEditor/Private/Editor/SRigHierarchy.h b/Engine/Plugins/Animation/ControlRig/Source/ControlRigEditor/Private/Editor/SRigHierarchy.h
--- a/Engine/Plugins/Animation/ControlRig/Source/ControlRigEditor/Private/Editor/SRigHierarchy.h
+++ b/Engine/Plugins/Animation/ControlRig/Source/ControlRigEditor/Private/Editor/SRigHierarchy.h
@@ -8,7 +8,7 @@
 #include "Editor/SRigHierarchyTreeView.h"
 #include "Units/RigUnitContext.h"
 #include "ControlRigBlueprint.h"
-#include "Editor/RigVMEditor.h"
+#include "Editor/RigVMNewEditor.h"
 #include "ControlRigSchematicModel.h"
 #include "ControlRigDragOps.h"
 #include "SRigHierarchy.generated.h"
@@ -59,7 +59,7 @@ public:
 
 private:
 
-	void OnEditorClose(const FRigVMEditor* InEditor, URigVMBlueprint* InBlueprint);
+	void OnEditorClose(const FRigVMEditorBase* InEditor, URigVMBlueprint* InBlueprint);
 
 	/** Bind commands that this widget handles */
 	void BindCommands();
diff --git a/Engine/Plugins/Animation/ControlRig/Source/ControlRigEditor/Public/EditMode/ControlRigEditMode.h b/Engine/Plugins/Animation/ControlRig/Source/ControlRigEditor/Public/EditMode/ControlRigEditMode.h
--- a/Engine/Plugins/Animation/ControlRig/Source/ControlRigEditor/Public/EditMode/ControlRigEditMode.h
+++ b/Engine/Plugins/Animation/ControlRig/Source/ControlRigEditor/Public/EditMode/ControlRigEditMode.h
@@ -22,7 +22,7 @@
 #include "ControlRigEditMode.generated.h"
 
 class FEditorViewportClient;
-class FRigVMEditor;
+class FRigVMEditorBase;
 class FViewport;
 class UActorFactory;
 struct FViewportClick;
diff --git a/Engine/Plugins/Animation/ControlRig/Source/ControlRigEditor/Public/IControlRigEditor.h b/Engine/Plugins/Animation/ControlRig/Source/ControlRigEditor/Public/IControlRigEditor.h
--- a/Engine/Plugins/Animation/ControlRig/Source/ControlRigEditor/Public/IControlRigEditor.h
+++ b/Engine/Plugins/Animation/ControlRig/Source/ControlRigEditor/Public/IControlRigEditor.h
@@ -3,11 +3,23 @@
 #pragma once
 
 #include "CoreMinimal.h"
-#include "Editor/RigVMEditor.h"
+
+#if WITH_RIGVMLEGACYEDITOR
+#include "Editor/RigVMLegacyEditor.h"
+#else
+#include "Editor/RigVMNewEditor.h"
+#endif
 
 class UControlRigBlueprint;
 class FBlueprintActionDatabaseRegistrar;
 
-class IControlRigEditor : public FRigVMEditor
+#if WITH_RIGVMLEGACYEDITOR
+class IControlRigEditor : public FRigVMLegacyEditor
+{
+};
+
+#else
+class IControlRigEditor : public FRigVMNewEditor
 {
 };
+#endif
\ No newline at end of file
diff --git a/Engine/Plugins/Experimental/AnimNext/Source/AnimNextEditor/Private/Common/AnimNextEdGraphNodeCustomization.cpp b/Engine/Plugins/Experimental/AnimNext/Source/AnimNextEditor/Private/Common/AnimNextEdGraphNodeCustomization.cpp
--- a/Engine/Plugins/Experimental/AnimNext/Source/AnimNextEditor/Private/Common/AnimNextEdGraphNodeCustomization.cpp
+++ b/Engine/Plugins/Experimental/AnimNext/Source/AnimNextEditor/Private/Common/AnimNextEdGraphNodeCustomization.cpp
@@ -108,7 +108,7 @@ void FAnimNextEdGraphNodeCustomization::CustomizeObjects(IDetailLayoutBuilder& D
 						continue;
 					}
 
-					TWeakPtr<FRigVMEditor> WeakEditor; // passing an empty RigVMEditor to FRigVMCollapseGraphLayout
+					TWeakPtr<FRigVMEditorBase> WeakEditor; // passing an empty RigVMEditor to FRigVMCollapseGraphLayout
 					IRigVMClientHost* RigVMClientHost = EdGraphNode->GetController()->GetClientHost();
 					check(RigVMClientHost);
 
diff --git a/Engine/Plugins/Experimental/AnimNext/Source/AnimNextEditor/Private/Graph/AnimNextEdGraphCustomization.cpp b/Engine/Plugins/Experimental/AnimNext/Source/AnimNextEditor/Private/Graph/AnimNextEdGraphCustomization.cpp
--- a/Engine/Plugins/Experimental/AnimNext/Source/AnimNextEditor/Private/Graph/AnimNextEdGraphCustomization.cpp
+++ b/Engine/Plugins/Experimental/AnimNext/Source/AnimNextEditor/Private/Graph/AnimNextEdGraphCustomization.cpp
@@ -29,7 +29,7 @@ void FAnimNextEdGraphCustomization::CustomizeDetails(IDetailLayoutBuilder& Detai
 	{
 		if(URigVMGraph* Model = EdGraph->GetModel())
 		{
-			TWeakPtr<FRigVMEditor> WeakEditor; // passing an empty RigVMEditor to FRigVMCollapseGraphLayout
+			TWeakPtr<FRigVMEditorBase> WeakEditor; // passing an empty RigVMEditor to FRigVMCollapseGraphLayout
 			IRigVMClientHost* RigVMClientHost = EdGraph->GetController()->GetClientHost();
 			check(RigVMClientHost);
 
diff --git a/Engine/Plugins/Runtime/RigVM/Source/RigVM/Private/RigVMSettings.cpp b/Engine/Plugins/Runtime/RigVM/Source/RigVM/Private/RigVMSettings.cpp
--- a/Engine/Plugins/Runtime/RigVM/Source/RigVM/Private/RigVMSettings.cpp
+++ b/Engine/Plugins/Runtime/RigVM/Source/RigVM/Private/RigVMSettings.cpp
@@ -8,6 +8,9 @@
 
 URigVMEditorSettings::URigVMEditorSettings(const FObjectInitializer& ObjectInitializer)
     : Super(ObjectInitializer)
+#if WITH_EDITORONLY_DATA
+	, bEnableContextMenuTimeSlicing(true)
+#endif
 {
 #if WITH_EDITORONLY_DATA
 	bAutoLinkMutableNodes = false;
diff --git a/Engine/Plugins/Runtime/RigVM/Source/RigVM/Public/RigVMSettings.h b/Engine/Plugins/Runtime/RigVM/Source/RigVM/Public/RigVMSettings.h
--- a/Engine/Plugins/Runtime/RigVM/Source/RigVM/Public/RigVMSettings.h
+++ b/Engine/Plugins/Runtime/RigVM/Source/RigVM/Public/RigVMSettings.h
@@ -33,6 +33,9 @@ public:
 	UPROPERTY(EditAnywhere, config, Category = Interaction)
 	bool bAutoLinkMutableNodes;
 
+	/** When the Blueprint graph context menu is invoked (e.g. by right-clicking in the graph or dragging off a pin), do not block the UI while populating the available actions list. */
+	UPROPERTY(EditAnywhere, config, Category = Workflow, meta = (DisplayName = "Enable Non-Blocking Context Menu"))
+	bool bEnableContextMenuTimeSlicing;
 #endif
 };
 
diff --git a/Engine/Plugins/Runtime/RigVM/Source/RigVMDeveloper/Private/EdGraph/RigVMEdGraphNode.cpp b/Engine/Plugins/Runtime/RigVM/Source/RigVMDeveloper/Private/EdGraph/RigVMEdGraphNode.cpp
--- a/Engine/Plugins/Runtime/RigVM/Source/RigVMDeveloper/Private/EdGraph/RigVMEdGraphNode.cpp
+++ b/Engine/Plugins/Runtime/RigVM/Source/RigVMDeveloper/Private/EdGraph/RigVMEdGraphNode.cpp
@@ -10,6 +10,12 @@
 #include "KismetCompiler.h"
 #include "BlueprintNodeSpawner.h"
 #include "BlueprintActionDatabaseRegistrar.h"
+#if WITH_RIGVMLEGACYEDITOR
+#include "FindInBlueprintManager.h"
+#else
+#include "Kismet/RigVMBlueprintCompilationManager.h"
+#include "Kismet/RigVMFindInBlueprintManager.h"
+#endif
 #include "RigVMEditorModule.h"
 #include "Textures/SlateIcon.h"
 #include "RigVMBlueprint.h"
@@ -2162,6 +2168,7 @@ void URigVMEdGraphNode::BeginDestroy()
 }
 
 #if WITH_EDITOR
+#if WITH_RIGVMLEGACYEDITOR
 void URigVMEdGraphNode::AddPinSearchMetaDataInfo(const UEdGraphPin* Pin, TArray<FSearchTagDataPair>& OutTaggedMetaData) const
 {
 	Super::AddPinSearchMetaDataInfo(Pin, OutTaggedMetaData);
@@ -2175,6 +2182,62 @@ void URigVMEdGraphNode::AddPinSearchMetaDataInfo(const UEdGraphPin* Pin, TArray<
 		}
 	}
 }
+#else
+void URigVMEdGraphNode::AddRigVMSearchMetaDataInfo(TArray<struct UBlueprintExtension::FSearchTagDataPair>& OutTaggedMetaData) const
+{
+	// Copied from EdGraphNode::AddSearchMetaDataInfo
+	if (GetSchema() == nullptr)
+	{
+		ensure(false);
+		return;
+	}
+
+	// Searchable - Primary label for the item in the search results
+	OutTaggedMetaData.Add(UBlueprintExtension::FSearchTagDataPair(FRigVMFindInBlueprintSearchTags::FiB_Name, GetNodeTitle(ENodeTitleType::ListView)));
+
+	// Searchable - As well as being searchable, this displays in the tooltip for the node
+	OutTaggedMetaData.Add(UBlueprintExtension::FSearchTagDataPair(FRigVMFindInBlueprintSearchTags::FiB_ClassName, FText::FromString(GetClass()->GetName())));
+
+	// Non-searchable - Used to lookup the node when attempting to jump to it
+	OutTaggedMetaData.Add(UBlueprintExtension::FSearchTagDataPair(FRigVMFindInBlueprintSearchTags::FiB_NodeGuid, FText::FromString(NodeGuid.ToString(EGuidFormats::Digits))));
+
+	// Non-searchable - Important for matching pin types with icons and colors, stored here so that each pin does not store it
+	OutTaggedMetaData.Add(UBlueprintExtension::FSearchTagDataPair(FRigVMFindInBlueprintSearchTags::FiB_SchemaName, FText::FromString(GetSchema()->GetClass()->GetName())));
+
+	// Non-Searchable - Used to display the icon and color for this node for better visual identification.
+	FLinearColor GlyphColor = FLinearColor::White;
+	FSlateIcon Icon = GetIconAndTint(GlyphColor);
+	OutTaggedMetaData.Add(UBlueprintExtension::FSearchTagDataPair(FRigVMFindInBlueprintSearchTags::FiB_Glyph, FText::FromName(Icon.GetStyleName())));
+	OutTaggedMetaData.Add(UBlueprintExtension::FSearchTagDataPair(FRigVMFindInBlueprintSearchTags::FiB_GlyphStyleSet, FText::FromName(Icon.GetStyleSetName())));
+	OutTaggedMetaData.Add(UBlueprintExtension::FSearchTagDataPair(FRigVMFindInBlueprintSearchTags::FiB_GlyphColor, FText::FromString(GlyphColor.ToString())));
+	OutTaggedMetaData.Add(UBlueprintExtension::FSearchTagDataPair(FRigVMFindInBlueprintSearchTags::FiB_Comment, FText::FromString(NodeComment)));
+}
+
+void URigVMEdGraphNode::AddRigVMPinSearchMetaDataInfo(const UEdGraphPin* Pin, TArray<UBlueprintExtension::FSearchTagDataPair>& OutTaggedMetaData) const
+{
+	// Copied from EdGraphNode::AddPinSearchMetaDataInfo
+	{
+		//Super::AddPinSearchMetaDataInfo(Pin, OutTaggedMetaData);
+		// Searchable - Primary label for the item in the search results
+		OutTaggedMetaData.Add(UBlueprintExtension::FSearchTagDataPair(FRigVMFindInBlueprintSearchTags::FiB_Name, Pin->GetSchema()->GetPinDisplayName(Pin)));
+
+		// Searchable - The pin's default value as a text string
+		OutTaggedMetaData.Add(UBlueprintExtension::FSearchTagDataPair(FRigVMFindInBlueprintSearchTags::FiB_DefaultValue, Pin->GetDefaultAsText()));
+	}
+	
+	if(const URigVMPin* ModelPin = FindModelPinFromGraphPin(Pin))
+	{
+		OutTaggedMetaData.Emplace(FText::FromString(TEXT("Data Type")), FText::FromString(ModelPin->GetCPPType()));
+		
+		if (ModelPin->IsBoundToVariable())
+		{
+			OutTaggedMetaData.Add(UBlueprintExtension::FSearchTagDataPair(FText::FromString(TEXT("Binding")), FText::FromString(ModelPin->GetBoundVariablePath())));
+		}
+	}
+}
+
+#endif
+
 #endif
 
 FEdGraphPinType URigVMEdGraphNode::GetPinTypeForModelPin(const URigVMPin* InModelPin)
diff --git a/Engine/Plugins/Runtime/RigVM/Source/RigVMDeveloper/Private/EdGraph/RigVMEdGraphSchema.cpp b/Engine/Plugins/Runtime/RigVM/Source/RigVMDeveloper/Private/EdGraph/RigVMEdGraphSchema.cpp
--- a/Engine/Plugins/Runtime/RigVM/Source/RigVMDeveloper/Private/EdGraph/RigVMEdGraphSchema.cpp
+++ b/Engine/Plugins/Runtime/RigVM/Source/RigVMDeveloper/Private/EdGraph/RigVMEdGraphSchema.cpp
@@ -28,7 +28,11 @@
 #include "RigVMFunctions/RigVMDispatch_Constant.h"
 #include "RigVMFunctions/RigVMDispatch_MakeStruct.h"
 #include "RigVMModel/Nodes/RigVMDispatchNode.h"
+#if WITH_RIGVMLEGACYEDITOR
 #include "BlueprintEditor.h"
+#else
+#include "Editor/RigVMNewEditor.h"
+#endif
 
 #include UE_INLINE_GENERATED_CPP_BY_NAME(RigVMEdGraphSchema)
 
@@ -973,6 +977,11 @@ TSharedPtr<INameValidatorInterface> URigVMEdGraphSchema::GetNameValidator(const
 	return MakeShareable(new FRigVMNameValidator(BlueprintObj, ValidationScope, OriginalName));		
 }
 
+bool URigVMEdGraphSchema::CanPromotePinToVariable(const UEdGraphPin& Pin, bool bCond) const
+{
+	return true;
+}
+
 bool URigVMEdGraphSchema::SupportsPinType(const UScriptStruct* ScriptStruct) const
 {
 	if(!ScriptStruct)
@@ -1468,8 +1477,11 @@ TSharedPtr<FEdGraphSchemaAction> URigVMEdGraphSchema::MakeActionFromVariableDesc
 		{
 			Category = FText::GetEmpty();
 		}
-
+#if WITH_RIGVMLEGACYEDITOR
 		TSharedPtr<FRigVMEdGraphSchemaAction_LocalVar> Action = MakeShareable(new FRigVMEdGraphSchemaAction_LocalVar(Category, FText::FromName(Variable.VarName), FText::GetEmpty(), 0, NodeSectionID::LOCAL_VARIABLE));
+#else
+		TSharedPtr<FRigVMEdGraphSchemaAction_LocalVar> Action = MakeShareable(new FRigVMEdGraphSchemaAction_LocalVar(Category, FText::FromName(Variable.VarName), FText::GetEmpty(), 0, RigVMNodeSectionID::LOCAL_VARIABLE));
+#endif
 		Action->SetVariableInfo(Variable.VarName, RigGraph, Variable.VarType.PinCategory == UEdGraphSchema_K2::PC_Boolean);
 		return Action;
 	}
diff --git a/Engine/Plugins/Runtime/RigVM/Source/RigVMDeveloper/Private/RigVMBlueprint.cpp b/Engine/Plugins/Runtime/RigVM/Source/RigVMDeveloper/Private/RigVMBlueprint.cpp
--- a/Engine/Plugins/Runtime/RigVM/Source/RigVMDeveloper/Private/RigVMBlueprint.cpp
+++ b/Engine/Plugins/Runtime/RigVM/Source/RigVMDeveloper/Private/RigVMBlueprint.cpp
@@ -38,6 +38,9 @@
 #include "CookOnTheSide/CookOnTheFlyServer.h"
 #include "RigVMEditorModule.h"
 #include "ScopedTransaction.h"
+#if !WITH_RIGVMLEGACYEDITOR
+#include "Kismet/RigVMBlueprintCompilationManager.h"
+#endif
 #endif//WITH_EDITOR
 
 #define LOCTEXT_NAMESPACE "RigVMBlueprint"
@@ -2709,8 +2712,13 @@ FName URigVMBlueprint::AddMemberVariable(const FName& InName, const FString& InC
 	FName Result = AddHostMemberVariableFromExternal(Variable, InDefaultValue);
 	if (!Result.IsNone())
 	{
+#if WITH_RIGVMLEGACYEDITOR
 		FBPCompileRequest Request(this, EBlueprintCompileOptions::None, nullptr);
 		FBlueprintCompilationManager::CompileSynchronously(Request);
+#else
+		FRigVMBPCompileRequest Request(this, EBlueprintCompileOptions::None, nullptr);
+		FRigVMBlueprintCompilationManager::CompileSynchronously(Request);
+#endif
 	}
 	return Result;
 }
diff --git a/Engine/Plugins/Runtime/RigVM/Source/RigVMDeveloper/Public/EdGraph/RigVMEdGraph.h b/Engine/Plugins/Runtime/RigVM/Source/RigVMDeveloper/Public/EdGraph/RigVMEdGraph.h
--- a/Engine/Plugins/Runtime/RigVM/Source/RigVMDeveloper/Public/EdGraph/RigVMEdGraph.h
+++ b/Engine/Plugins/Runtime/RigVM/Source/RigVMDeveloper/Public/EdGraph/RigVMEdGraph.h
@@ -208,7 +208,7 @@ private:
 #endif
 	friend class URigVMEdGraphNode;
 	friend class URigVMBlueprint;
-	friend class FRigVMEditor;
+	friend class FRigVMEditorBase;
 	friend class SRigVMGraphNode;
 	friend class URigVMBlueprint;
 
diff --git a/Engine/Plugins/Runtime/RigVM/Source/RigVMDeveloper/Public/EdGraph/RigVMEdGraphNode.h b/Engine/Plugins/Runtime/RigVM/Source/RigVMDeveloper/Public/EdGraph/RigVMEdGraphNode.h
--- a/Engine/Plugins/Runtime/RigVM/Source/RigVMDeveloper/Public/EdGraph/RigVMEdGraphNode.h
+++ b/Engine/Plugins/Runtime/RigVM/Source/RigVMDeveloper/Public/EdGraph/RigVMEdGraphNode.h
@@ -6,6 +6,7 @@
 #include "CoreMinimal.h"
 #include "UObject/WeakObjectPtrTemplates.h"
 #include "RigVMModel/RigVMGraph.h"
+#include "Blueprint/BlueprintExtension.h"
 #include "RigVMEdGraphNode.generated.h"
 
 class FBlueprintActionDatabaseRegistrar;
@@ -26,7 +27,7 @@ class RIGVMDEVELOPER_API URigVMEdGraphNode : public UEdGraphNode
 	friend class URigVMBlueprint;
 	friend class FRigVMEdGraphTraverser;
 	friend class FRigVMEdGraphPanelPinFactory;
-	friend class FRigVMEditor;
+	friend class FRigVMEditorBase;
 	friend struct FRigVMBlueprintUtils;
 	friend class SRigVMEdGraphPinCurveFloat;
 
@@ -125,7 +126,12 @@ public:
 	virtual bool ShouldDrawNodeAsControlPointOnly(int32& OutInputPinIndex, int32& OutOutputPinIndex) const override;
 	virtual void BeginDestroy() override;
 #if WITH_EDITOR
+#if WITH_RIGVMLEGACYEDITOR
 	virtual void AddPinSearchMetaDataInfo(const UEdGraphPin* Pin, TArray<FSearchTagDataPair>& OutTaggedMetaData) const override;
+#else
+	void AddRigVMSearchMetaDataInfo(TArray<struct UBlueprintExtension::FSearchTagDataPair>& OutTaggedMetaData) const;
+	void AddRigVMPinSearchMetaDataInfo(const UEdGraphPin* Pin, TArray<UBlueprintExtension::FSearchTagDataPair>& OutTaggedMetaData) const;
+#endif
 #endif
 	virtual FString GetPinMetaData(FName InPinName, FName InKey) override;
 
diff --git a/Engine/Plugins/Runtime/RigVM/Source/RigVMDeveloper/Public/EdGraph/RigVMEdGraphSchema.h b/Engine/Plugins/Runtime/RigVM/Source/RigVMDeveloper/Public/EdGraph/RigVMEdGraphSchema.h
--- a/Engine/Plugins/Runtime/RigVM/Source/RigVMDeveloper/Public/EdGraph/RigVMEdGraphSchema.h
+++ b/Engine/Plugins/Runtime/RigVM/Source/RigVMDeveloper/Public/EdGraph/RigVMEdGraphSchema.h
@@ -321,6 +321,7 @@ public:
 	virtual void SetPinBeingDroppedOnNode(UEdGraphPin* InSourcePin) const override { PinBeingDropped = InSourcePin; }
 	virtual void InsertAdditionalActions(TArray<UBlueprint*> InBlueprints, TArray<UEdGraph*> InGraphs, TArray<UEdGraphPin*> InPins, FGraphActionListBuilderBase& OutAllActions) const override;
 	virtual TSharedPtr<INameValidatorInterface> GetNameValidator(const UBlueprint* BlueprintObj, const FName& OriginalName, const UStruct* ValidationScope, const FName& ActionTypeId) const override;
+	bool CanPromotePinToVariable(const UEdGraphPin& Pin, bool bCond) const;
 
 	/** Returns true if the schema supports the script type */
 	bool SupportsPinType(const UScriptStruct* ScriptStruct) const;
diff --git a/Engine/Plugins/Runtime/RigVM/Source/RigVMDeveloper/Public/RigVMBlueprint.h b/Engine/Plugins/Runtime/RigVM/Source/RigVMDeveloper/Public/RigVMBlueprint.h
--- a/Engine/Plugins/Runtime/RigVM/Source/RigVMDeveloper/Public/RigVMBlueprint.h
+++ b/Engine/Plugins/Runtime/RigVM/Source/RigVMDeveloper/Public/RigVMBlueprint.h
@@ -849,7 +849,9 @@ private:
 #endif
 
 	friend class FRigVMBlueprintCompilerContext;
-	friend class FRigVMEditor;
+	friend class FRigVMLegacyEditor;
+	friend class FRigVMNewEditor;
+	friend class FRigVMEditorBase;
 	friend class FRigVMEditorModule;
 	friend class URigVMEdGraphSchema;
 	friend struct FRigVMEdGraphSchemaAction_PromoteToVariable;
diff --git a/Engine/Plugins/Runtime/RigVM/Source/RigVMDeveloper/RigVMDeveloper.Build.cs b/Engine/Plugins/Runtime/RigVM/Source/RigVMDeveloper/RigVMDeveloper.Build.cs
--- a/Engine/Plugins/Runtime/RigVM/Source/RigVMDeveloper/RigVMDeveloper.Build.cs
+++ b/Engine/Plugins/Runtime/RigVM/Source/RigVMDeveloper/RigVMDeveloper.Build.cs
@@ -30,9 +30,29 @@ public class RigVMDeveloper : ModuleRules
 					"MessageLog",
 					"BlueprintGraph",
 					"GraphEditor",
-					"Kismet",
                 }
 			);
+
+            bool bAddBlueprintEditorDependency = true;
+            PublicDefinitions.Add("WITH_RIGVMLEGACYEDITOR=" + (bAddBlueprintEditorDependency ? '1' : '0'));
+            if (bAddBlueprintEditorDependency)
+            {
+	            PrivateDependencyModuleNames.Add("Kismet");
+            }
+            else
+            {
+	            PrivateDependencyModuleNames.AddRange(new string[]
+	            {
+		            "ApplicationCore",
+		            "FieldNotification",
+		            "InputCore",
+		            "Json",
+		            "JsonUtilities",
+		            "ToolWidgets",
+		            "TraceLog",
+		            "WorkspaceMenuStructure"
+	            });
+            }
             
             PrivateIncludePathModuleNames.Add("RigVMEditor");
             DynamicallyLoadedModuleNames.Add("RigVMEditor");
