Transform Widgets: PreventThrottling argument added to disable throttling on underlying numeric widgets
- regression from 37628647

#rnx

diff --git a/Engine/Plugins/Animation/ControlRig/Source/ControlRigEditor/Private/ControlRigElementDetails.cpp b/Engine/Plugins/Animation/ControlRig/Source/ControlRigEditor/Private/ControlRigElementDetails.cpp
--- a/Engine/Plugins/Animation/ControlRig/Source/ControlRigEditor/Private/ControlRigElementDetails.cpp
+++ b/Engine/Plugins/Animation/ControlRig/Source/ControlRigEditor/Private/ControlRigElementDetails.cpp
@@ -1946,7 +1946,8 @@ void FRigTransformElementDetails::CustomizeTransform(IDetailLayoutBuilder& Detai
 	SAdvancedTransformInputBox<FEulerTransform>::FArguments TransformWidgetArgs = SAdvancedTransformInputBox<FEulerTransform>::FArguments()
 	.DisplayToggle(false)
 	.DisplayRelativeWorld(true)
-	.Font(IDetailLayoutBuilder::GetDetailFont());
+	.Font(IDetailLayoutBuilder::GetDetailFont())
+	.PreventThrottling(true);
 
 	for(int32 Index = 0; Index < ButtonLabels.Num(); Index++)
 	{
@@ -3442,7 +3443,8 @@ void FRigControlElementDetails::CustomizeValue(IDetailLayoutBuilder& DetailBuild
 				SAdvancedTransformInputBox<FEulerTransform>::FArguments TransformWidgetArgs = SAdvancedTransformInputBox<FEulerTransform>::FArguments()
 				.DisplayToggle(false)
 				.DisplayRelativeWorld(true)
-				.Visibility(VisibilityAttribute);
+				.Visibility(VisibilityAttribute)
+				.PreventThrottling(true);
 
 				CreateTransformComponentValueWidgetRow(
 					ControlType,
@@ -4030,7 +4032,8 @@ void FRigControlElementDetails::CustomizeAnimationChannels(IDetailLayoutBuilder&
 					SAdvancedTransformInputBox<FEulerTransform>::FArguments()
 				.DisplayToggle(false)
 				.DisplayRelativeWorld(false)
-				.Visibility(EVisibility::Visible);
+				.Visibility(EVisibility::Visible)
+				.PreventThrottling(true);
 
 				WidgetRow = &CreateTransformComponentValueWidgetRow(
 					AssignedAnimationChannel->Settings.ControlType,
@@ -4052,7 +4055,8 @@ void FRigControlElementDetails::CustomizeAnimationChannels(IDetailLayoutBuilder&
 				.DisplayToggle(false)
 				.DisplayRelativeWorld(false)
 				.Font(IDetailLayoutBuilder::GetDetailFont())
-				.Visibility(EVisibility::Visible);
+				.Visibility(EVisibility::Visible)
+				.PreventThrottling(true);
 
 				WidgetRow = &CreateEulerTransformValueWidgetRow(
 					ChildElementKeys,
@@ -4895,7 +4899,8 @@ void FRigControlElementDetails::CustomizeShape(IDetailLayoutBuilder& DetailBuild
 	.IsEnabled(bIsEnabled)
 	.DisplayToggle(false)
 	.DisplayRelativeWorld(false)
-	.Font(IDetailLayoutBuilder::GetDetailFont());
+	.Font(IDetailLayoutBuilder::GetDetailFont())
+	.PreventThrottling(true);
 
 	TArray<FRigElementKey> Keys = GetElementKeys();
 	Keys = PerElementInfos[0].GetHierarchy()->SortKeys(Keys);
@@ -6050,6 +6055,7 @@ FDetailWidgetRow& FRigControlElementDetails::CreateVector2DValueWidgetRow(
 		 .OnToggleXChanged(OnToggleXChanged)
 		 .OnToggleYChanged(OnToggleYChanged)
 		 .IsEnabled(bIsEnabled)
+		 .PreventThrottling(true)
 	]
 	.CopyAction(FUIAction(
 	FExecuteAction::CreateLambda([ValueType, Keys, Hierarchy]()
diff --git a/Engine/Plugins/Animation/ControlRig/Source/ControlRigEditor/Private/ControlRigElementDetails.h b/Engine/Plugins/Animation/ControlRig/Source/ControlRigEditor/Private/ControlRigElementDetails.h
--- a/Engine/Plugins/Animation/ControlRig/Source/ControlRigEditor/Private/ControlRigElementDetails.h
+++ b/Engine/Plugins/Animation/ControlRig/Source/ControlRigEditor/Private/ControlRigElementDetails.h
@@ -711,6 +711,7 @@ private:
 			 .OnToggleChanged(OnToggleChanged)
 			 .UndeterminedString(NSLOCTEXT("FRigControlElementDetails", "MultipleValues", "Multiple Values"))
 			 .IsEnabled(bIsEnabled)
+			 .PreventThrottling(true)
 		]
 		.CopyAction(FUIAction(
 		FExecuteAction::CreateLambda([ValueType, Keys, Hierarchy]()
diff --git a/Engine/Source/Developer/AnimationWidgets/Public/SAdvancedRotationInputBox.h b/Engine/Source/Developer/AnimationWidgets/Public/SAdvancedRotationInputBox.h
--- a/Engine/Source/Developer/AnimationWidgets/Public/SAdvancedRotationInputBox.h
+++ b/Engine/Source/Developer/AnimationWidgets/Public/SAdvancedRotationInputBox.h
@@ -100,6 +100,7 @@ public:
 		, _AllowSpin(true)
 		, _DisplayToggle( false )
 		, _TogglePadding( FMargin( 1.f,0.f,1.f,0.f ) )
+		, _PreventThrottling(false)
 		{}
 
 		/** Representation of the rotation */
@@ -159,6 +160,9 @@ public:
 		/** Padding around the toggle checkbox */
 		SLATE_ARGUMENT( FMargin, TogglePadding )
 
+		/** If refresh requests for the viewport should happen for all value changes **/
+		SLATE_ARGUMENT(bool, PreventThrottling)
+
 	SLATE_END_ARGS()
 
 	/**
@@ -230,6 +234,7 @@ public:
 				.OnRollEndSliderMovement(this, &SAdvancedRotationInputBox::HandleOnEndSliderMovement, ESlateTransformSubComponent::Roll)
 				.TogglePadding(InArgs._TogglePadding)
 				.Visibility(this, &SAdvancedRotationInputBox::IsRotatorInputBoxVisible)
+				.PreventThrottling(InArgs._PreventThrottling)
 			]
 
 			// a vector 3 variation to represent euler with rotation order
@@ -268,6 +273,7 @@ public:
 				.OnZEndSliderMovement(this, &SAdvancedRotationInputBox::HandleOnEndSliderMovement, ESlateTransformSubComponent::Z)
 				.TogglePadding(InArgs._TogglePadding)
 				.Visibility(this, &SAdvancedRotationInputBox::IsEulerInputBoxVisible)
+				.PreventThrottling(InArgs._PreventThrottling)
 			]
 			
 			// a vector 4d variation - to represent quaternion
@@ -313,6 +319,7 @@ public:
 				.OnWEndSliderMovement(this, &SAdvancedRotationInputBox::HandleOnEndSliderMovement, ESlateTransformSubComponent::W)
 				.TogglePadding(InArgs._TogglePadding)
 				.Visibility(this, &SAdvancedRotationInputBox::IsQuaternionInputBoxVisible)
+				.PreventThrottling(InArgs._PreventThrottling)
 			]
 
 			/*
@@ -366,6 +373,7 @@ public:
 				.OnZEndSliderMovement(this, &SAdvancedRotationInputBox::HandleOnEndSliderMovement, ESlateTransformSubComponent::Z)
 				.TogglePadding(InArgs._TogglePadding)
 				.Visibility(this, &SAdvancedRotationInputBox::IsAxisAndAngleInputBoxVisible)
+				.PreventThrottling(InArgs._PreventThrottling)
 			]
 
 			// a numeric input for the angle for the axis and angle representation
@@ -406,6 +414,7 @@ public:
 				.OnBeginSliderMovement(this, &SAdvancedRotationInputBox::HandleOnBeginSliderMovement, ESlateTransformSubComponent::Angle)
 				.OnEndSliderMovement(this, &SAdvancedRotationInputBox::HandleOnEndSliderMovement, ESlateTransformSubComponent::Angle)
 				.Visibility(this, &SAdvancedRotationInputBox::IsAxisAndAngleInputBoxVisible)
+				.PreventThrottling(InArgs._PreventThrottling)
 			]
 		];
 
diff --git a/Engine/Source/Developer/AnimationWidgets/Public/SAdvancedTransformInputBox.h b/Engine/Source/Developer/AnimationWidgets/Public/SAdvancedTransformInputBox.h
--- a/Engine/Source/Developer/AnimationWidgets/Public/SAdvancedTransformInputBox.h
+++ b/Engine/Source/Developer/AnimationWidgets/Public/SAdvancedTransformInputBox.h
@@ -118,6 +118,7 @@ public:
 		, _DisplayRelativeWorld(false)
 		, _DisplayToggle( false )
 		, _TogglePadding( FMargin( 1.f,0.f,1.f,0.f ) )
+		, _PreventThrottling(false)
 		{}
 
 		/** Whether or not to construct the location component */
@@ -237,6 +238,9 @@ public:
 		/** Delegate to determine if a row (or the whole transform) differs from its default */
 		SLATE_EVENT( FDiffersFromDefault, DiffersFromDefault )
 
+		/** If refresh requests for the viewport should happen for all value changes **/
+		SLATE_ARGUMENT(bool, PreventThrottling)
+	
 	SLATE_END_ARGS()
 
 	/**
@@ -630,6 +634,7 @@ public:
 					.ConstrainVector(ConstrainComponents)
 					.DisplayToggle(InArgs._DisplayToggle)
 					.TogglePadding(InArgs._TogglePadding)
+					.PreventThrottling(InArgs._PreventThrottling)
 					.OnXBeginSliderMovement(XBeginSlide)
 					.OnYBeginSliderMovement(YBeginSlide)
 					.OnZBeginSliderMovement(ZBeginSlide)
@@ -862,6 +867,7 @@ public:
 					.OnQuaternionCommitted(QuaternionCommitted)
 					.DisplayToggle(InArgs._DisplayToggle)
 					.TogglePadding(InArgs._TogglePadding)
+					.PreventThrottling(InArgs._PreventThrottling)
 					.OnBeginSliderMovement(RotationBeginSliderMovement)
 					.OnEndSliderMovement(RotationEndSliderMovement)
 					.OnGetToggleChecked(SAdvancedRotationInputBox<NumericType>::FOnGetToggleChecked::CreateLambda(
diff --git a/Engine/Source/Runtime/Slate/Public/Widgets/Input/SRotatorInputBox.h b/Engine/Source/Runtime/Slate/Public/Widgets/Input/SRotatorInputBox.h
--- a/Engine/Source/Runtime/Slate/Public/Widgets/Input/SRotatorInputBox.h
+++ b/Engine/Source/Runtime/Slate/Public/Widgets/Input/SRotatorInputBox.h
@@ -41,6 +41,7 @@ public:
 		, _ToggleYawChecked( ECheckBoxState::Checked )
 		, _ToggleRollChecked( ECheckBoxState::Checked )
 		, _TogglePadding( FMargin( 1.f,0.f,1.f,0.f ) )
+		, _PreventThrottling(false)
 		{}
 
 		/** Roll component of the rotator */
@@ -141,6 +142,9 @@ public:
 
 		/** Padding around the toggle checkbox */
 		SLATE_ARGUMENT( FMargin, TogglePadding )
+		
+		/** If refresh requests for the viewport should happen for all value changes **/
+		SLATE_ARGUMENT(bool, PreventThrottling)
 
 	SLATE_END_ARGS()
 
@@ -189,6 +193,7 @@ public:
 				.ToggleChecked(InArgs._ToggleRollChecked)
 				.OnToggleChanged(InArgs._OnToggleRollChanged)
 				.TogglePadding(InArgs._TogglePadding)
+				.PreventThrottling(InArgs._PreventThrottling)
 			]
 			+SHorizontalBox::Slot()
 			[
@@ -225,6 +230,7 @@ public:
 				.ToggleChecked(InArgs._TogglePitchChecked)
 				.OnToggleChanged(InArgs._OnTogglePitchChanged)
 				.TogglePadding(InArgs._TogglePadding)
+				.PreventThrottling(InArgs._PreventThrottling)
 			]
 			+SHorizontalBox::Slot()
 			[
@@ -261,6 +267,7 @@ public:
 				.ToggleChecked(InArgs._ToggleYawChecked)
 				.OnToggleChanged(InArgs._OnToggleYawChanged)
 				.TogglePadding(InArgs._TogglePadding)
+				.PreventThrottling(InArgs._PreventThrottling)
 			]
 		];
 
diff --git a/Engine/Source/Runtime/Slate/Public/Widgets/Input/SVectorInputBox.h b/Engine/Source/Runtime/Slate/Public/Widgets/Input/SVectorInputBox.h
--- a/Engine/Source/Runtime/Slate/Public/Widgets/Input/SVectorInputBox.h
+++ b/Engine/Source/Runtime/Slate/Public/Widgets/Input/SVectorInputBox.h
@@ -204,6 +204,7 @@ public:
 			, _bColorAxisLabels(false)
 			, _DisplayToggle(false)
 			, _TogglePadding(FMargin(1.f,0.f,1.f,0.f) )
+			, _PreventThrottling(false)
 		{}
 
 		/** Optional Value of the vector */
@@ -265,6 +266,9 @@ public:
 
 		/** Delegate to constrain the vector */
 		SLATE_ARGUMENT( FOnConstrainVector, ConstrainVector )
+
+		/** If refresh requests for the viewport should happen for all value changes **/
+		SLATE_ARGUMENT(bool, PreventThrottling)
 		
 	SLATE_END_ARGS()
 
@@ -461,6 +465,7 @@ private:
 			.TogglePadding(InArgs._TogglePadding)
 			.ToggleChecked(ToggleChecked)
 			.OnToggleChanged(OnToggleChanged)
+			.PreventThrottling(InArgs._PreventThrottling)
 		];
 	}
 
