Control rig added an option to refresh the hierarchy using only bones that exist on the mesh, removing bones that exist in the skeleton but have been removed from the mesh

diff --git a/Engine/Plugins/Animation/ControlRig/Source/ControlRig/Private/Rigs/RigHierarchyController.cpp b/Engine/Plugins/Animation/ControlRig/Source/ControlRig/Private/Rigs/RigHierarchyController.cpp
--- a/Engine/Plugins/Animation/ControlRig/Source/ControlRig/Private/Rigs/RigHierarchyController.cpp
+++ b/Engine/Plugins/Animation/ControlRig/Source/ControlRig/Private/Rigs/RigHierarchyController.cpp
@@ -1027,9 +1027,72 @@ bool URigHierarchyController::SetComponentContent(FRigComponentKey InComponent,
 	return true;
 }
 
+TArray<FRigElementKey> URigHierarchyController::ImportBones(
+	USkeletalMesh* InSkeletalMesh,
+	const FName& InNameSpace,
+	bool bReplaceExistingBones,
+	bool bRemoveObsoleteBones,
+	bool bSelectBones,
+	bool bSetupUndo)
+{
+	TArray<FRigElementKey> RigElementKeys;
+
+	if (InSkeletalMesh != nullptr)
+	{
+		const FReferenceSkeleton& RefSkeleton = InSkeletalMesh->GetRefSkeleton();
+		const TArray<FMeshBoneInfo>& RefBoneInfos = RefSkeleton.GetRefBoneInfo();
+		const TArray<FTransform>& RefBonePoses = RefSkeleton.GetRefBonePose();
+
+		USkeleton* MeshSkeleton = InSkeletalMesh->GetSkeleton();
+	
+		TArray<FMeshBoneInfo> BoneInfos;
+		BoneInfos.Reserve(RefBoneInfos.Num());
+		TArray<FTransform> BoneTransforms;
+		BoneTransforms.Reserve(RefBonePoses.Num());
+
+		const int32 NumSkeletonBones = RefBoneInfos.Num();
+		for (int32 i=0; i<NumSkeletonBones; ++i)
+		{
+			const FMeshBoneInfo& MeshBoneInfo = RefBoneInfos[i];
+			const int32 SkeletonBoneIndex = RefSkeleton.FindBoneIndex(MeshBoneInfo.Name);
+			if (MeshSkeleton->GetMeshBoneIndexFromSkeletonBoneIndex(InSkeletalMesh, SkeletonBoneIndex) == INDEX_NONE)
+			{
+				// if bone index is none, the bone does not exist in the mesh and it has been excluded
+				// add an empty bone info to keep the same indexes
+				BoneInfos.Add(FMeshBoneInfo());
+				BoneTransforms.Add(FTransform::Identity);
+			}
+			else
+			{
+				const FTransform& BoneTransform = RefBonePoses[i];
+				BoneInfos.Add(MeshBoneInfo);
+				BoneTransforms.Add(BoneTransform);
+			}
+		}
+
+		RigElementKeys = ImportBones(BoneInfos, BoneTransforms, InNameSpace, bReplaceExistingBones, bRemoveObsoleteBones, bSelectBones, bSetupUndo);
+	}
+	return RigElementKeys;
+}
+
 TArray<FRigElementKey> URigHierarchyController::ImportBones(const FReferenceSkeleton& InSkeleton,
-                                                            const FName& InNameSpace, bool bReplaceExistingBones, bool bRemoveObsoleteBones, bool bSelectBones,
-                                                            bool bSetupUndo)
+	const FName& InNameSpace, bool bReplaceExistingBones, bool bRemoveObsoleteBones, bool bSelectBones,
+	bool bSetupUndo)
+{
+	const TArray<FMeshBoneInfo>& BoneInfos = InSkeleton.GetRefBoneInfo();
+	const TArray<FTransform>& BoneTransforms = InSkeleton.GetRefBonePose();
+
+	return ImportBones(BoneInfos, BoneTransforms, InNameSpace, bReplaceExistingBones, bRemoveObsoleteBones, bSelectBones, bSetupUndo);
+}
+
+TArray<FRigElementKey> URigHierarchyController::ImportBones(
+	const TArray<FMeshBoneInfo> & BoneInfos,
+	const TArray<FTransform>& BoneTransforms,
+	const FName& InNameSpace,
+	bool bReplaceExistingBones,
+	bool bRemoveObsoleteBones,
+	bool bSelectBones,
+	bool bSetupUndo)
 {
 	TArray<FRigElementKey> AddedBones;
 
@@ -1045,9 +1108,6 @@ TArray<FRigElementKey> URigHierarchyController::ImportBones(const FReferenceSkel
 
 	Hierarchy->ResetPoseToInitial();
 
-	const TArray<FMeshBoneInfo>& BoneInfos = InSkeleton.GetRefBoneInfo();
-	const TArray<FTransform>& BonePoses = InSkeleton.GetRefBonePose();
-
 	struct Local
 	{
 		static FName DetermineBoneName(const FName& InBoneName, const FName& InLocalNameSpace)
@@ -1070,11 +1130,17 @@ TArray<FRigElementKey> URigHierarchyController::ImportBones(const FReferenceSkel
 
 		for (int32 Index = 0; Index < BoneInfos.Num(); ++Index)
 		{
-			const FRigElementKey ExistingBoneKey(BoneInfos[Index].Name, ERigElementType::Bone);
+			const FName& BoneName = BoneInfos[Index].Name;
+			if (BoneName == NAME_None)
+			{
+				continue;
+			}
+			const FRigElementKey ExistingBoneKey(BoneName, ERigElementType::Bone);
 			const int32 ExistingBoneIndex = Hierarchy->GetIndex(ExistingBoneKey);
 			
-			const FName DesiredBoneName = Local::DetermineBoneName(BoneInfos[Index].Name, InNameSpace);
-			FName ParentName = (BoneInfos[Index].ParentIndex != INDEX_NONE) ? BoneInfos[BoneInfos[Index].ParentIndex].Name : NAME_None;
+			const FName DesiredBoneName = Local::DetermineBoneName(BoneName, InNameSpace);
+			const int32 ParentBoneIndex = BoneInfos[Index].ParentIndex;
+			FName ParentName = (ParentBoneIndex != INDEX_NONE && BoneInfos.IsValidIndex(ParentBoneIndex)) ? BoneInfos[ParentBoneIndex].Name : NAME_None;
 			ParentName = Local::DetermineBoneName(ParentName, InNameSpace);
 
 			const FName* MappedParentNamePtr = BoneNameMap.Find(ParentName);
@@ -1096,14 +1162,14 @@ TArray<FRigElementKey> URigHierarchyController::ImportBones(const FReferenceSkel
 					SetParent(ExistingBoneKey, ParentKey, bSetupUndo);
 				}
 
-				Hierarchy->SetInitialLocalTransform(ExistingBoneIndex, BonePoses[Index], true, bSetupUndo);
-				Hierarchy->SetLocalTransform(ExistingBoneIndex, BonePoses[Index], true, bSetupUndo);
+				Hierarchy->SetInitialLocalTransform(ExistingBoneIndex, BoneTransforms[Index], true, bSetupUndo);
+				Hierarchy->SetLocalTransform(ExistingBoneIndex, BoneTransforms[Index], true, bSetupUndo);
 
 				BonesToSelect.Add(ExistingBoneKey);
 			}
 			else
 			{
-				const FRigElementKey AddedBoneKey = AddBone(DesiredBoneName, ParentKey, BonePoses[Index], false, ERigBoneType::Imported, bSetupUndo);
+				const FRigElementKey AddedBoneKey = AddBone(DesiredBoneName, ParentKey, BoneTransforms[Index], false, ERigBoneType::Imported, bSetupUndo);
 				BoneNameMap.Add(DesiredBoneName, AddedBoneKey.Name);
 				AddedBones.Add(AddedBoneKey);
 				BonesToSelect.Add(AddedBoneKey);
@@ -1115,7 +1181,12 @@ TArray<FRigElementKey> URigHierarchyController::ImportBones(const FReferenceSkel
 	{
 		for (int32 Index = 0; Index < BoneInfos.Num(); ++Index)
 		{
-			FName DesiredBoneName = Local::DetermineBoneName(BoneInfos[Index].Name, InNameSpace);
+			const FName& BoneName = BoneInfos[Index].Name;
+			if (BoneName == NAME_None)
+			{
+				continue;
+			}
+			FName DesiredBoneName = Local::DetermineBoneName(BoneName, InNameSpace);
 			FName ParentName = (BoneInfos[Index].ParentIndex != INDEX_NONE) ? BoneInfos[BoneInfos[Index].ParentIndex].Name : NAME_None;
 			ParentName = Local::DetermineBoneName(ParentName, InNameSpace);
 
@@ -1126,7 +1197,7 @@ TArray<FRigElementKey> URigHierarchyController::ImportBones(const FReferenceSkel
 			}
 
 			const FRigElementKey ParentKey(ParentName, ERigElementType::Bone);
-			const FRigElementKey AddedBoneKey = AddBone(DesiredBoneName, ParentKey, BonePoses[Index], false, ERigBoneType::Imported, bSetupUndo);
+			const FRigElementKey AddedBoneKey = AddBone(DesiredBoneName, ParentKey, BoneTransforms[Index], false, ERigBoneType::Imported, bSetupUndo);
 			BoneNameMap.Add(DesiredBoneName, AddedBoneKey.Name);
 			AddedBones.Add(AddedBoneKey);
 			BonesToSelect.Add(AddedBoneKey);
@@ -1138,7 +1209,12 @@ TArray<FRigElementKey> URigHierarchyController::ImportBones(const FReferenceSkel
 		TMap<FName, int32> BoneNameToIndexInSkeleton;
 		for (const FMeshBoneInfo& BoneInfo : BoneInfos)
 		{
-			FName DesiredBoneName = Local::DetermineBoneName(BoneInfo.Name, InNameSpace);
+			const FName& BoneName = BoneInfo.Name;
+			if (BoneName == NAME_None)
+			{
+				continue;
+			}
+			FName DesiredBoneName = Local::DetermineBoneName(BoneName, InNameSpace);
 			BoneNameToIndexInSkeleton.Add(DesiredBoneName, BoneNameToIndexInSkeleton.Num());
 		}
 		
@@ -1179,7 +1255,12 @@ TArray<FRigElementKey> URigHierarchyController::ImportBones(const FReferenceSkel
 		// update the sub index to match the bone index in the skeleton
 		for (int32 Index = 0; Index < BoneInfos.Num(); ++Index)
 		{
-			FName DesiredBoneName = Local::DetermineBoneName(BoneInfos[Index].Name, InNameSpace);
+			const FName& BoneName = BoneInfos[Index].Name;
+			if (BoneName == NAME_None)
+			{
+				continue;
+			}
+			FName DesiredBoneName = Local::DetermineBoneName(BoneName, InNameSpace);
 			const FRigElementKey Key(DesiredBoneName, ERigElementType::Bone);
 			if(FRigBoneElement* BoneElement = Hierarchy->Find<FRigBoneElement>(Key))
 			{
@@ -1232,6 +1313,36 @@ TArray<FRigElementKey> URigHierarchyController::ImportBones(USkeleton* InSkeleto
 	return BoneKeys;
 }
 
+TArray<FRigElementKey> URigHierarchyController::ImportBonesFromSkeletalMesh(
+	USkeletalMesh* InSkeletalMesh,
+	const FName& InNameSpace,
+	bool bReplaceExistingBones,
+	bool bRemoveObsoleteBones,
+	bool bSelectBones,
+	bool bSetupUndo,
+	bool bPrintPythonCommand)
+{
+	const TArray<FRigElementKey> BoneKeys = ImportBones(InSkeletalMesh, InNameSpace, bReplaceExistingBones, bRemoveObsoleteBones, bSelectBones, bSetupUndo);
+
+#if WITH_EDITOR
+	if (!BoneKeys.IsEmpty() && bPrintPythonCommand && !bSuspendPythonPrinting)
+	{
+		UBlueprint* Blueprint = GetTypedOuter<UBlueprint>();
+		if (Blueprint)
+		{
+			RigVMPythonUtils::Print(Blueprint->GetFName().ToString(),
+				FString::Printf(TEXT("hierarchy_controller.import_bones_from_asset('%s', '%s', %s, %s, %s)"),
+					(InSkeletalMesh != nullptr) ? *InSkeletalMesh->GetPathName() : TEXT(""),
+					*InNameSpace.ToString(),
+					(bReplaceExistingBones) ? TEXT("True") : TEXT("False"),
+					(bRemoveObsoleteBones) ? TEXT("True") : TEXT("False"),
+					(bSelectBones) ? TEXT("True") : TEXT("False")));
+		}
+	}
+#endif
+	return BoneKeys;
+}
+
 #if WITH_EDITOR
 
 TArray<FRigElementKey> URigHierarchyController::ImportBonesFromAsset(FString InAssetPath, FName InNameSpace,
@@ -4325,4 +4436,3 @@ void URigHierarchyController::ReportAndNotifyError(const FString& InMessage) con
 	}
 #endif
 }
-
diff --git a/Engine/Plugins/Animation/ControlRig/Source/ControlRig/Public/Rigs/RigHierarchyController.h b/Engine/Plugins/Animation/ControlRig/Source/ControlRig/Public/Rigs/RigHierarchyController.h
--- a/Engine/Plugins/Animation/ControlRig/Source/ControlRig/Public/Rigs/RigHierarchyController.h
+++ b/Engine/Plugins/Animation/ControlRig/Source/ControlRig/Public/Rigs/RigHierarchyController.h
@@ -314,6 +314,43 @@ public:
 		bool bSelectBones = false,
 		bool bSetupUndo = false);
 
+	/**
+	 * Imports an existing skeleton to the hierarchy, restricting the bone list to the ones that exist in the provided Skeletal Mesh
+	 * @param InSkeletalMesh The skeletal mesh asset to import
+	 * @param InNameSpace The namespace to prefix the bone names with
+	 * @param bReplaceExistingBones If true existing bones will be removed
+	 * @param bRemoveObsoleteBones If true bones non-existent in the skeleton will be removed from the hierarchy
+	 * @param bSelectBones If true the bones will be selected upon import
+	 * @param bSetupUndo If set to true the stack will record the change for undo / redo
+	 * @return The keys of the imported elements
+	 */
+	TArray<FRigElementKey> ImportBones(
+		USkeletalMesh* InSkeletalMesh,
+		const FName& InNameSpace,
+		bool bReplaceExistingBones = true,
+		bool bRemoveObsoleteBones = true,
+		bool bSelectBones = false,
+		bool bSetupUndo = false);
+
+	/**
+	 * Imports the provided bone list to the hierarchy
+	 * @param BoneInfos The array of MeshBoneInfos to use as source. Bones with name set to NAME_Mone are bones that have been excluded.
+	 * @param BoneTransforms The transforms to initialize the hierarchy elements
+	 * @param bReplaceExistingBones If true existing bones will be removed
+	 * @param bRemoveObsoleteBones If true bones non-existent in the skeleton will be removed from the hierarchy
+	 * @param bSelectBones If true the bones will be selected upon import
+	 * @param bSetupUndo If set to true the stack will record the change for undo / redo
+	 * @return The keys of the imported elements
+	 */	
+	 TArray<FRigElementKey> ImportBones(
+		const TArray<FMeshBoneInfo>& BoneInfos,
+		const TArray<FTransform>& BoneTransforms,
+		const FName& InNameSpace,
+		bool bReplaceExistingBones = true,
+		bool bRemoveObsoleteBones = true,
+		bool bSelectBones = false,
+		bool bSetupUndo = false);
+
 	/**
 	 * Imports an existing skeleton to the hierarchy
 	 * @param InSkeleton The skeleton to import
@@ -335,6 +372,26 @@ public:
 		bool bSetupUndo = false,
 		bool bPrintPythonCommand = false);
 
+	/**
+	 * Imports an existing skeleton to the hierarchy, restricting the bone list to the ones that exist in the provided Skeletal Mesh
+	 * @param InSkeletalMesh The skeletal mesh asset to import
+	 * @param InNameSpace The namespace to prefix the bone names with
+	 * @param bReplaceExistingBones If true existing bones will be removed
+	 * @param bRemoveObsoleteBones If true bones non-existent in the skeleton will be removed from the hierarchy
+	 * @param bSelectBones If true the bones will be selected upon import
+	 * @param bSetupUndo If set to true the stack will record the change for undo / redo
+	 * @return The keys of the imported elements
+	 */
+	 UFUNCTION(BlueprintCallable, Category = URigHierarchyController)
+	 TArray<FRigElementKey> ImportBonesFromSkeletalMesh(
+		USkeletalMesh* InSkeletalMesh,
+		const FName& InNameSpace,
+		bool bReplaceExistingBones = true,
+		bool bRemoveObsoleteBones = true,
+		bool bSelectBones = false,
+		bool bSetupUndo = false,
+		bool bPrintPythonCommand = false);
+
 #if WITH_EDITOR
 	/**
 	 * Imports an existing skeleton to the hierarchy
diff --git a/Engine/Plugins/Animation/ControlRig/Source/ControlRigEditor/Private/Editor/SRigHierarchy.cpp b/Engine/Plugins/Animation/ControlRig/Source/ControlRigEditor/Private/Editor/SRigHierarchy.cpp
--- a/Engine/Plugins/Animation/ControlRig/Source/ControlRigEditor/Private/Editor/SRigHierarchy.cpp
+++ b/Engine/Plugins/Animation/ControlRig/Source/ControlRigEditor/Private/Editor/SRigHierarchy.cpp
@@ -1547,7 +1547,33 @@ void SRigHierarchy::CreateRefreshMenu(FMenuBuilder& MenuBuilder)
 		[
 			SNew(SObjectPropertyEntryBox)
 			.AllowedClass(USkeletalMesh::StaticClass())
-			.OnObjectChanged(this, &SRigHierarchy::RefreshHierarchy, false)
+			.OnObjectChanged(this, &SRigHierarchy::RefreshHierarchy, false, bRestrictRefreshToMeshBones)
+		]
+		+ SVerticalBox::Slot()
+		.AutoHeight()
+		.Padding(3)
+		[
+			SNew(SHorizontalBox)
+			+ SHorizontalBox::Slot()
+			.AutoWidth()
+			.Padding(0.0f, 0.0f, 5.0f, 0.0f)
+			[
+				SNew(STextBlock)
+				.Text(LOCTEXT("RefreshMesh_RestrictToMeshBones", "Restrict to Mesh Bones"))
+				.ToolTipText(LOCTEXT("RefreshMesh_RestrictToMeshBones_Tooltip", "It will remove any bones that does not exist in the Mesh. It might break compatibility if the Rig is used across different meshes that share a skeleton with a slightly different set of bones on each mesh."))
+			]
+			+ SHorizontalBox::Slot()
+			[
+				SNew(SCheckBox)
+				.IsChecked_Lambda([this]() -> ECheckBoxState
+				{
+					return bRestrictRefreshToMeshBones ? ECheckBoxState::Checked : ECheckBoxState::Unchecked;
+				})
+				.OnCheckStateChanged_Lambda([this](ECheckBoxState CheckBoxState)
+				{
+					bRestrictRefreshToMeshBones = CheckBoxState == ECheckBoxState::Checked;
+				})
+			]
 		]
 		,
 		FText()
@@ -1574,7 +1600,7 @@ void SRigHierarchy::CreateResetCurvesMenu(FMenuBuilder& MenuBuilder)
 		[
 			SNew(SObjectPropertyEntryBox)
 			.AllowedClass(USkeletalMesh::StaticClass())
-			.OnObjectChanged(this, &SRigHierarchy::RefreshHierarchy, true)
+			.OnObjectChanged(this, &SRigHierarchy::RefreshHierarchy, true, false)
 		]
 		,
 		FText()
@@ -1616,7 +1642,7 @@ void SRigHierarchy::UpdateMesh(USkeletalMesh* InMesh, const bool bImport) const
 	EditorSharedPtr->Compile();
 }
 
-void SRigHierarchy::RefreshHierarchy(const FAssetData& InAssetData, bool bOnlyResetCurves)
+void SRigHierarchy::RefreshHierarchy(const FAssetData& InAssetData, bool bOnlyResetCurves, bool bRestrictToMeshBones)
 {
 	if (bIsChangingRigHierarchy)
 	{
@@ -1665,7 +1691,14 @@ void SRigHierarchy::RefreshHierarchy(const FAssetData& InAssetData, bool bOnlyRe
 		}
 		else
 		{
-			Controller->ImportBones(Mesh->GetSkeleton(), NAME_None, true, true, bSelectBones, true, true);
+			if (bRestrictRefreshToMeshBones)
+			{
+				Controller->ImportBonesFromSkeletalMesh(Mesh, NAME_None, true, true, bSelectBones, true, true);
+			}
+			else
+			{
+				Controller->ImportBones(Mesh->GetSkeleton(), NAME_None, true, true, bSelectBones, true, true);
+			}
 			Controller->ImportCurvesFromSkeletalMesh(Mesh, NAME_None, false, true, true);
 		}
 	}
diff --git a/Engine/Plugins/Animation/ControlRig/Source/ControlRigEditor/Private/Editor/SRigHierarchy.h b/Engine/Plugins/Animation/ControlRig/Source/ControlRigEditor/Private/Editor/SRigHierarchy.h
--- a/Engine/Plugins/Animation/ControlRig/Source/ControlRigEditor/Private/Editor/SRigHierarchy.h
+++ b/Engine/Plugins/Animation/ControlRig/Source/ControlRigEditor/Private/Editor/SRigHierarchy.h
@@ -192,7 +192,7 @@ private:
 	void CreateRefreshMenu(FMenuBuilder& MenuBuilder);
 	void CreateResetCurvesMenu(FMenuBuilder& MenuBuilder);
 	bool ShouldFilterOnImport(const FAssetData& AssetData) const;
-	void RefreshHierarchy(const FAssetData& InAssetData, bool bOnlyResetCurves);
+	void RefreshHierarchy(const FAssetData& InAssetData, bool bOnlyResetCurves, bool bRestrictToMeshBones);
 	void UpdateMesh(USkeletalMesh* InMesh, const bool bImport) const;
 
 	void HandleResetTransform(bool bSelectionOnly);
@@ -227,6 +227,7 @@ private:
 	void OnNavigateToFirstConnectorWarning();
 
 	bool bIsConstructionEventRunning;
+	bool bRestrictRefreshToMeshBones = false;
 	uint32 LastHierarchyHash;
 	TArray<FRigElementKey> SelectionBeforeConstruction;
 	TMap<FRigElementKey, FModularRigResolveResult> DragRigResolveResults;
