Control Rig Edit Mode: constraint space transform
- parent & explicit modes can now transform the selected objects in the constraint space if any
- added FConstraintsInteractionCache designed to minimize the number of requests made to the ConstraintManager and thus improve performance
- added ControlRigEditModeUtil to start cleaning the FControlRigEditMode class that becomes overloaded

TODO: clarify the current architecture to explicitely split animation related features and rigging one in the ed mode

#rnx

diff --git a/Engine/Plugins/Animation/ControlRig/Source/ControlRigEditor/Private/EditMode/ControlRigEditMode.cpp b/Engine/Plugins/Animation/ControlRig/Source/ControlRigEditor/Private/EditMode/ControlRigEditMode.cpp
--- a/Engine/Plugins/Animation/ControlRig/Source/ControlRigEditor/Private/EditMode/ControlRigEditMode.cpp
+++ b/Engine/Plugins/Animation/ControlRig/Source/ControlRigEditor/Private/EditMode/ControlRigEditMode.cpp
@@ -70,6 +70,7 @@
 #include "AnimationEditorViewportClient.h"
 #include "ControlShapeActorHelper.h"
 #include "EditorInteractiveGizmoManager.h"
+#include "Constraints/ControlRigTransformableHandle.h"
 #include "Editor/ControlRigViewportToolbarExtensions.h"
 #include "Editor/Sequencer/Private/SSequencer.h"
 #include "Slate/SceneViewport.h"
@@ -83,21 +84,19 @@ TAutoConsoleVariable<bool> CVarClickSelectThroughGizmo(TEXT("ControlRig.Sequence
 
 namespace UE::ControlRigEditMode::Private
 {
+	const TCHAR* FocusModeName = TEXT("AnimMode.PendingFocusMode");
 
-const TCHAR* FocusModeName = TEXT("AnimMode.PendingFocusMode");
+	static bool bFocusMode = false;
+	static FAutoConsoleVariableRef CVarSetFocusOnHover(
+	   FocusModeName,
+	   bFocusMode,
+	   TEXT("Force setting focus on the hovered viewport when entering a key.")
+	   );
 
-static bool bFocusMode = false;
-static FAutoConsoleVariableRef CVarSetFocusOnHover(
-   FocusModeName,
-   bFocusMode,
-   TEXT("Force setting focus on the hovered viewport when entering a key.")
-   );
-
-IConsoleVariable* GetFocusModeVariable()
-{
-	return IConsoleManager::Get().FindConsoleVariable(FocusModeName);
-}
-	
+	IConsoleVariable* GetFocusModeVariable()
+	{
+		return IConsoleManager::Get().FindConsoleVariable(FocusModeName);
+	}
 }
 
 using namespace UE::ControlRigEditMode;
@@ -504,6 +503,11 @@ void FControlRigEditMode::Enter()
 
 	// initialize the gizmo context
 	WeakGizmoContext = GetModeManager()->GetGizmoContext();
+
+	if (IsInLevelEditor())
+	{
+		ConstraintsCache.RegisterNotifications(WorldPtr);
+	}
 }
 
 //todo get working with Persona
@@ -526,6 +530,9 @@ static void ClearOutAnyActiveTools()
 
 void FControlRigEditMode::Exit()
 {
+	ConstraintsCache.UnregisterNotifications(WorldPtr);
+	ConstraintsCache.Reset();
+	
 	if (WeakGizmoContext.IsValid())
 	{
 		WeakGizmoContext->RotationContext = FRotationContext();
@@ -2571,7 +2578,9 @@ ETransformSpace FControlRigEditMode::GetTransformSpace() const
 	return ETransformSpace::Local;	
 };
 
-FTransform FControlRigEditMode::GetPivotOrientation(const FRigElementKey& InControlKey, URigHierarchy* InHierarchy, const ETransformSpace InSpace, const FTransform& InComponentTransform) const
+FTransform FControlRigEditMode::GetPivotOrientation(
+	const FRigElementKey& InControlKey, const UControlRig* InControlRig, URigHierarchy* InHierarchy,
+	const ETransformSpace InSpace, const FTransform& InComponentTransform) const
 {
 	FRigControlElement* Control = InHierarchy->Find<FRigControlElement>(InControlKey);
 	if (!Control)
@@ -2585,8 +2594,13 @@ FTransform FControlRigEditMode::GetPivotOrientation(const FRigElementKey& InCont
 	case ETransformSpace::Local:
 		return InHierarchy->GetTransform(Control, ERigTransformType::CurrentGlobal);
 	case ETransformSpace::Parent:
-		{
-			const int32 NumParents = InHierarchy->GetNumberOfParents(Control);			
+		{	
+			if (TOptional<FTransform> ConstraintSpace = GetConstraintParentTransform(InControlRig, InControlKey.Name))
+			{
+				return ConstraintSpace->GetRelativeTransform(InComponentTransform);
+			}
+			
+			const int32 NumParents = InHierarchy->GetNumberOfParents(Control);
 			return NumParents > 0 ? InHierarchy->GetParentTransform(Control, ERigTransformType::CurrentGlobal) : InHierarchy->GetTransform(Control, ERigTransformType::CurrentGlobal);
 		}
 	case ETransformSpace::Explicit:
@@ -2601,12 +2615,27 @@ FTransform FControlRigEditMode::GetPivotOrientation(const FRigElementKey& InCont
 
 			if (bRotating)
 			{
-				const FTransform Offset = InHierarchy->GetControlOffsetTransform(Control, ERigTransformType::CurrentGlobal);
-				RotationContext.Offset = Offset * InComponentTransform;
+				if (TOptional<FTransform> ConstraintSpace = GetConstraintParentTransform(InControlRig, InControlKey.Name))
+				{
+					RotationContext.Offset = *ConstraintSpace;
+				}
+				else
+				{
+					const FTransform Offset = InHierarchy->GetControlOffsetTransform(Control, ERigTransformType::CurrentGlobal);
+					RotationContext.Offset = Offset * InComponentTransform;					
+				}
+				
 				return InHierarchy->GetTransform(Control, ERigTransformType::CurrentGlobal);
 			}
 
 			RotationContext.Offset = FTransform::Identity;
+			if (IsInLevelEditor())
+			{
+				if (TOptional<FTransform> ConstraintSpace = GetConstraintParentTransform(InControlRig, InControlKey.Name))
+				{
+					return ConstraintSpace->GetRelativeTransform(InComponentTransform);
+				}
+			}
 			return InHierarchy->GetControlOffsetTransform(Control, ERigTransformType::CurrentGlobal);
 		}
 	default:
@@ -2645,7 +2674,7 @@ bool FControlRigEditMode::ComputePivotFromShapeActors(UControlRig* InControlRig,
 			{
 				const FRigElementKey ControlKey = ShapeActor->GetElementKey();
 				const FTransform ShapeTransform = ShapeActor->GetActorTransform().GetRelativeTransform(ComponentTransform);
-				LastTransform = GetPivotOrientation(ControlKey, Hierarchy, InSpace, ComponentTransform);
+				LastTransform = GetPivotOrientation(ControlKey, InControlRig, Hierarchy, InSpace, ComponentTransform);
 				PivotLocation += ShapeTransform.GetLocation();
 				
 				++NumSelectedControls;
@@ -5025,7 +5054,7 @@ void FControlRigEditMode::SetControlShapeTransform(
 	const FTransform& InToWorldTransform,
 	const FRigControlModifiedContext& InContext,
 	const bool bPrintPython,
-	const FRotator& InDeltaRot,
+	const FTransformContext& InTransformContext,
 	const bool bFixEulerFlips) const
 {
 	UControlRig* ControlRig = InShapeActor->ControlRig.Get();
@@ -5042,67 +5071,14 @@ void FControlRigEditMode::SetControlShapeTransform(
 	}
 
 	static constexpr bool bNotify = true, bUndo = true;
-	
-	const bool bIsExplicitRotation = GetTransformSpace() == ETransformSpace::Explicit && GetModeManager()->GetWidgetMode() == UE::Widget::EWidgetMode::WM_Rotate;
-	auto ApplyExplicitRotation = [bIsExplicitRotation, ControlRig, Hierarchy, ControlElement, InDeltaRot]()
-	{
-		if (!bIsExplicitRotation)
-		{
-			return false;
-		}
-		
-		if (InDeltaRot.IsZero())
-		{
-			return false;
-		}
-		
-		if (ControlRig->IsAdditive())
-		{
-			return false;
-		}
-
-		if (!Hierarchy->UsesPreferredEulerAngles())
-		{
-			return false;
-		}
-
-		const ERigControlType ControlType = ControlElement->Settings.ControlType;
-		return ControlType == ERigControlType::Rotator || ControlType == ERigControlType::EulerTransform;
-	};
-
-	auto SetControlRotation = [ControlRig, ControlElement, Hierarchy, InContext, bPrintPython](const FTransform& InGlobalTransform, const FRotator& InDeltaRot)
-	{
-		const FName ControlName = ControlElement->GetKey().Name;
-		const FVector DeltaEulerAngle(InDeltaRot.Roll, InDeltaRot.Pitch, InDeltaRot.Yaw);
-		FVector NewEulerAngle = Hierarchy->GetControlSpecifiedEulerAngle(ControlElement);
-		NewEulerAngle += DeltaEulerAngle;
-		
-		const ERigControlType ControlType = ControlElement->Settings.ControlType;
-		if (ControlType == ERigControlType::Rotator)
-		{
-			const FRotator Rotator(Hierarchy->GetControlQuaternion(ControlElement, NewEulerAngle));
-			
-			Hierarchy->SetControlSpecifiedEulerAngle(ControlElement, NewEulerAngle);
-			ControlRig->SetControlValue<FRotator>(ControlName, Rotator, /*bNotify*/true, InContext, /*bUndo*/true, bPrintPython);
-		}
-		else if (ControlType == ERigControlType::EulerTransform)
-		{
-			const FQuat Quat = Hierarchy->GetControlQuaternion(ControlElement, NewEulerAngle);
-			FRigControlValue NewValue = ControlRig->GetControlValueFromGlobalTransform(ControlName, InGlobalTransform, ERigTransformType::CurrentGlobal);
-			FEulerTransform EulerTransform = NewValue.Get<FRigControlValue::FEulerTransform_Float>().ToTransform();
-			EulerTransform.Rotation = FRotator(Quat);
-
-			Hierarchy->SetControlSpecifiedEulerAngle(ControlElement, NewEulerAngle);
-			ControlRig->SetControlValue<FRigControlValue::FEulerTransform_Float>(ControlName, EulerTransform, /*bNotify*/true, InContext, /*bUndo*/true, bPrintPython);
-			Hierarchy->SetControlSpecifiedEulerAngle(ControlElement, NewEulerAngle);			
-		}
-	};
+	const FExplicitRotationInteraction ExplicitRotation(InTransformContext, ControlRig, Hierarchy, ControlElement, InToWorldTransform);
+	const bool bApplyExplicitRotation = ExplicitRotation.IsValid();
 
 	if (AreEditingControlRigDirectly())
 	{
-		if (ApplyExplicitRotation())
+		if (bApplyExplicitRotation)
 		{
-			SetControlRotation(InGlobalTransform, InDeltaRot);
+			ExplicitRotation.Apply(InGlobalTransform, InContext, bPrintPython);
 		}
 		else
 		{
@@ -5127,19 +5103,19 @@ void FControlRigEditMode::SetControlShapeTransform(
 	};
 	
 	// find the last constraint in the stack (this could be cached on mouse press)
-	TArray< TWeakObjectPtr<UTickableConstraint> > Constraints;
-	FTransformConstraintUtils::GetParentConstraints(ControlRig->GetWorld(), InShapeActor, Constraints);
+	UWorld* World = ControlRig->GetWorld();
+	const uint32 ControlHash = UTransformableControlHandle::ComputeHash(ControlRig, InShapeActor->ControlName);
+	const bool bNeedsConstraintPostProcess = ConstraintsCache.HasAnyActiveConstraint(ControlHash, World);
+	static const TArray< TWeakObjectPtr<UTickableConstraint> > EmptyConstraints;
+	const TArray< TWeakObjectPtr<UTickableConstraint> >& Constraints = bNeedsConstraintPostProcess ? ConstraintsCache.Get(ControlHash, World) : EmptyConstraints;
 
-	const int32 LastActiveIndex = FTransformConstraintUtils::GetLastActiveConstraintIndex(Constraints);
-	const bool bNeedsConstraintPostProcess = Constraints.IsValidIndex(LastActiveIndex);
-	
 	// set the global space, assumes it's attached to actor
 	// no need to compensate for constraints here, this will be done after when setting the control in the constraint space
 	{
 		TGuardValue<bool> CompensateGuard(FMovieSceneConstraintChannelHelper::bDoNotCompensate, true);
-		if (ApplyExplicitRotation())
+		if (bApplyExplicitRotation)
 		{
-			SetControlRotation(InGlobalTransform, InDeltaRot);
+			ExplicitRotation.Apply(InGlobalTransform, InContext, bPrintPython, Constraints);
 		}
 		else
 		{
@@ -5149,31 +5125,30 @@ void FControlRigEditMode::SetControlShapeTransform(
 		}
 	}
 
-	if (!bNeedsConstraintPostProcess)
+	if (bNeedsConstraintPostProcess)
 	{
-		return;
-	}
-	
-	// switch to constraint space
-	const FTransform WorldTransform = InGlobalTransform * InToWorldTransform;
-	FTransform LocalTransform = ControlRig->GetControlLocalTransform(InShapeActor->ControlName);
+		if (!bApplyExplicitRotation)	
+		{
+			// switch to constraint space
+			const FTransform WorldTransform = InGlobalTransform * InToWorldTransform;
+			FTransform LocalTransform = ControlRig->GetControlLocalTransform(InShapeActor->ControlName);
 
-	const TOptional<FTransform> RelativeTransform =
-		FTransformConstraintUtils::GetConstraintsRelativeTransform(Constraints, LocalTransform, WorldTransform);
-	if (RelativeTransform)
-	{
-		LocalTransform = *RelativeTransform; 
-	}
+			const TOptional<FTransform> RelativeTransform =
+				FTransformConstraintUtils::GetConstraintsRelativeTransform(Constraints, LocalTransform, WorldTransform);
+			if (RelativeTransform)
+			{
+				LocalTransform = *RelativeTransform; 
+			}
 
-	FRigControlModifiedContext Context = InContext;
-	Context.bConstraintUpdate = false;
+			FRigControlModifiedContext Context = InContext;
+			Context.bConstraintUpdate = false;
 	
-	ControlRig->SetControlLocalTransform(InShapeActor->ControlName, LocalTransform, bNotify, Context, bUndo, bFixEulerFlips);
-	EvaluateRigIfAdditive();
+			ControlRig->SetControlLocalTransform(InShapeActor->ControlName, LocalTransform, bNotify, Context, bUndo, bFixEulerFlips);
+			EvaluateRigIfAdditive();
+		}
 
-	{
 		TGuardValue<bool> CompensateGuard(FMovieSceneConstraintChannelHelper::bDoNotCompensate, true);
-		const FConstraintsManagerController& Controller = FConstraintsManagerController::Get(ControlRig->GetWorld());
+		const FConstraintsManagerController& Controller = FConstraintsManagerController::Get(World);
 		Controller.EvaluateAllConstraints();
 	}
 }
@@ -5187,151 +5162,182 @@ FTransform FControlRigEditMode::GetControlShapeTransform(const AControlRigShapeA
 	return FTransform::Identity;
 }
 
-void FControlRigEditMode::MoveControlShape( AControlRigShapeActor* ShapeActor,
+bool FControlRigEditMode::MoveControlShapeLocally( AControlRigShapeActor* ShapeActor,
 	const FTransformContext& InTransformContext, const FTransform& ToWorldTransform,
-	const bool bUseLocal, const bool bCalcLocal, FTransform* InOutLocal,
-	TArray<TFunction<void()>>& OutTasks)
+	const FTransform& InLocal)
 {
 	if (!ensure(InTransformContext.CanTransform()))
 	{
-		return;
+		return false;
 	}
 	
 	UControlRig* ControlRig = ShapeActor ? ShapeActor->ControlRig.Get() : nullptr;
 	if (!ensure(ControlRig))
 	{
-		return;
+		return  false;
 	}
 
 	bool bTransformChanged = false;
+	
+	FTransform CurrentLocalTransform = ControlRig->GetControlLocalTransform(ShapeActor->ControlName);
 
-	// In case for some reason the shape actor was detached, make sure to attach it again to the scene component
-	if (!ShapeActor->GetAttachParentActor())
-	{
-		if (USceneComponent* SceneComponent = GetHostingSceneComponent(ControlRig))
-		{
-			if (AActor* OwnerActor = SceneComponent->GetOwner())
-			{
-				ShapeActor->AttachToActor(OwnerActor, FAttachmentTransformRules::KeepWorldTransform);
-			}
-		}
-	}
+	URigHierarchy* Hierarchy = ControlRig->GetHierarchy();
+	FRigControlElement* Control = Hierarchy ? Hierarchy->Find<FRigControlElement>(ShapeActor->GetElementKey()) : nullptr;
 	
-	//first case is where we do all controls by the local diff.
-	if (bUseLocal)
+	if (InTransformContext.bRotation)
 	{
-		FTransform CurrentLocalTransform = ControlRig->GetControlLocalTransform(ShapeActor->ControlName);
+		FQuat CurrentRotation = CurrentLocalTransform.GetRotation();
 
-		URigHierarchy* Hierarchy = ControlRig->GetHierarchy();
-		FRigControlElement* Control = Hierarchy ? Hierarchy->Find<FRigControlElement>(ShapeActor->GetElementKey()) : nullptr;
-		
-		if (InTransformContext.bRotation)
+		FQuat DeltaRot = InLocal.GetRotation();
+
+		if (ensure(Control))
 		{
-			FQuat CurrentRotation = CurrentLocalTransform.GetRotation();
-			if (InOutLocal)
+			switch(InTransformContext.Space)
 			{
-				FQuat DeltaRot = InOutLocal->GetRotation();
-
-				if (ensure(Control))
+			case ETransformSpace::World:
+			case ETransformSpace::Explicit:
+				ensure(false);
+				break;
+			case ETransformSpace::Local:
+				break;
+			case ETransformSpace::Parent:
 				{
-					switch(InTransformContext.Space)
-					{
-					case ETransformSpace::World:
-					case ETransformSpace::Explicit:
-						ensure(false);
-						break;
-					case ETransformSpace::Local:
-						break;
-					case ETransformSpace::Parent:
-						{
-							const int32 NumParents = Hierarchy->GetNumberOfParents(Control);
-							const FTransform Global = Hierarchy->GetTransform(Control, ERigTransformType::CurrentGlobal);
-							const FTransform Parent = NumParents > 0 ? Hierarchy->GetParentTransform(Control, ERigTransformType::CurrentGlobal) : Global;
-
-							FTransform RelativeToParent = Global.GetRelativeTransform(Parent);
-							RelativeToParent = RelativeToParent * DeltaRot;
-							const FTransform NewGlobal = RelativeToParent * Parent;
-							
-							DeltaRot = NewGlobal.GetRelativeTransform(Global).GetRotation();
-							break;
-						}
-					}
-				}
+					const int32 NumParents = Hierarchy->GetNumberOfParents(Control);
+					const FTransform Global = Hierarchy->GetTransform(Control, ERigTransformType::CurrentGlobal);
+					const FTransform Parent = NumParents > 0 ? Hierarchy->GetParentTransform(Control, ERigTransformType::CurrentGlobal) : Global;
 
-				CurrentRotation = (CurrentRotation * DeltaRot);
+					FTransform RelativeToParent = Global.GetRelativeTransform(Parent);
+					RelativeToParent = RelativeToParent * DeltaRot;
+					const FTransform NewGlobal = RelativeToParent * Parent;
+					
+					DeltaRot = NewGlobal.GetRelativeTransform(Global).GetRotation();
+					break;
+				}
 			}
-			CurrentLocalTransform.SetRotation(CurrentRotation);
-			bTransformChanged = true;
 		}
 
-		if (InTransformContext.bTranslation)
+		CurrentRotation = (CurrentRotation * DeltaRot);
+		
+		CurrentLocalTransform.SetRotation(CurrentRotation);
+		bTransformChanged = true;
+	}
+
+	if (InTransformContext.bTranslation)
+	{
+		FVector CurrentLocation = CurrentLocalTransform.GetLocation();
+
+		FVector Delta = InLocal.GetLocation();
+
+		if (ensure(Control))
 		{
-			FVector CurrentLocation = CurrentLocalTransform.GetLocation();
-			if (InOutLocal)
+			switch(InTransformContext.Space)
 			{
-				FVector Delta = InOutLocal->GetLocation();
-
-				if (ensure(Control))
+			case ETransformSpace::World:
+				ensure(false);
+				break;
+			case ETransformSpace::Local:
+			{
+				// in local mode, the incoming delta is expressed in the global space so it has to be switched back to the offset space
+				const FTransform Global = Hierarchy->GetTransform(Control, ERigTransformType::CurrentGlobal);
+				const FVector GlobalDelta = Global.TransformVector(Delta);
+				const FTransform Offset = Hierarchy->GetControlOffsetTransform(Control, ERigTransformType::CurrentGlobal);
+				Delta = Offset.InverseTransformVector(GlobalDelta);
+				break;
+			}
+			case ETransformSpace::Parent:
+			{
+				// in parent mode, the incoming delta is expressed in the parent space so it has to be switched back to the offset space
+				FTransform Parent = FTransform::Identity;  
+				if (TOptional<FTransform> ConstraintSpace = GetConstraintParentTransform(ControlRig, ShapeActor->ControlName))
 				{
-					switch(InTransformContext.Space)
-					{
-					case ETransformSpace::World:
-						ensure(false);
-						break;
-					case ETransformSpace::Local:
-					{
-						// in local mode, the incoming delta is expressed in the global space so it has to be switched back to the offset space
-						const FTransform Global = Hierarchy->GetTransform(Control, ERigTransformType::CurrentGlobal);
-						const FVector GlobalDelta = Global.TransformVector(Delta);
-						const FTransform Offset = Hierarchy->GetControlOffsetTransform(Control, ERigTransformType::CurrentGlobal);
-						Delta = Offset.InverseTransformVector(GlobalDelta);
-						break;
-					}
-					case ETransformSpace::Parent:
-					{
-						// in parent mode, the incoming delta is expressed in the parent space so it has to be switched back to the offset space
-						const int32 NumParents = Hierarchy->GetNumberOfParents(Control);
-						const FTransform Parent = NumParents > 0 ? Hierarchy->GetParentTransform(Control, ERigTransformType::CurrentGlobal) : Hierarchy->GetTransform(Control, ERigTransformType::CurrentGlobal);
-						const FVector GlobalDelta = Parent.TransformVector(Delta);
-						const FTransform Offset = Hierarchy->GetControlOffsetTransform(Control, ERigTransformType::CurrentGlobal);
-						Delta = Offset.InverseTransformVector(GlobalDelta);
-						break;
-					}
-					case ETransformSpace::Explicit:
-					{
-						// nothing to do as it has already been computed in the right space
-						break;
-					}
-					}
+					Parent = ConstraintSpace->GetRelativeTransform(ToWorldTransform);
 				}
-    
-				CurrentLocation = CurrentLocation + Delta;
+				else
+				{
+					const int32 NumParents = Hierarchy->GetNumberOfParents(Control);
+					Parent = NumParents > 0 ?
+						Hierarchy->GetParentTransform(Control, ERigTransformType::CurrentGlobal) :
+						Hierarchy->GetTransform(Control, ERigTransformType::CurrentGlobal);
+				}
+				const FVector GlobalDelta = Parent.TransformVector(Delta);
+				const FTransform Offset = Hierarchy->GetControlOffsetTransform(Control, ERigTransformType::CurrentGlobal);
+				Delta = Offset.InverseTransformVector(GlobalDelta);
+				break;
 			}
-			CurrentLocalTransform.SetLocation(CurrentLocation);
-			bTransformChanged = true;
+			case ETransformSpace::Explicit:
+			{
+				// nothing to do as it has already been computed in the right space
+				break;
+			}
+			}
+			
+			CurrentLocation = CurrentLocation + Delta;
 		}
+		CurrentLocalTransform.SetLocation(CurrentLocation);
+		bTransformChanged = true;
+	}
 
-		if (bTransformChanged)
-		{
-			ControlRig->InteractionType = InteractionType;
-			ControlRig->ElementsBeingInteracted.AddUnique(ShapeActor->GetElementKey());
-			
-			ControlRig->SetControlLocalTransform(ShapeActor->ControlName, CurrentLocalTransform,true, FRigControlModifiedContext(), true, /*fix eulers*/ true);
+	if (bTransformChanged)
+	{
+		ControlRig->InteractionType = InteractionType;
+		ControlRig->ElementsBeingInteracted.AddUnique(ShapeActor->GetElementKey());
+		
+		ControlRig->SetControlLocalTransform(ShapeActor->ControlName, CurrentLocalTransform,true, FRigControlModifiedContext(), true, /*fix eulers*/ true);
 
-			FTransform CurrentTransform  = ControlRig->GetControlGlobalTransform(ShapeActor->ControlName);			// assumes it's attached to actor
-			CurrentTransform = ToWorldTransform * CurrentTransform;
+		FTransform CurrentTransform  = ControlRig->GetControlGlobalTransform(ShapeActor->ControlName);			// assumes it's attached to actor
+		CurrentTransform = ToWorldTransform * CurrentTransform;
 
-			// make the transform relative to the offset transform again.
-			// first we'll make it relative to the offset used at the time of starting the drag
-			// and then we'll make it absolute again based on the current offset. these two can be
-			// different if we are interacting on a control on an animated character
-			CurrentTransform = CurrentTransform.GetRelativeTransform(ShapeActor->OffsetTransform);
-			CurrentTransform = CurrentTransform * ControlRig->GetHierarchy()->GetGlobalControlOffsetTransform(ShapeActor->GetElementKey(), false);
+		// make the transform relative to the offset transform again.
+		// first we'll make it relative to the offset used at the time of starting the drag
+		// and then we'll make it absolute again based on the current offset. these two can be
+		// different if we are interacting on a control on an animated character
+		CurrentTransform = CurrentTransform.GetRelativeTransform(ShapeActor->OffsetTransform);
+		CurrentTransform = CurrentTransform * ControlRig->GetHierarchy()->GetGlobalControlOffsetTransform(ShapeActor->GetElementKey(), false);
 
-			// Don't set the global transform to the shape actor to avoid drifting
-			//ShapeActor->SetGlobalTransform(CurrentTransform);
+		// Don't set the global transform to the shape actor to avoid drifting
+		//ShapeActor->SetGlobalTransform(CurrentTransform);
 
-			RigsToEvaluateDuringThisTick.AddUnique(ControlRig);
+		RigsToEvaluateDuringThisTick.AddUnique(ControlRig);
+	}
+
+	return bTransformChanged; 
+}
+
+void FControlRigEditMode::MoveControlShape( AControlRigShapeActor* ShapeActor,
+	const FTransformContext& InTransformContext, const FTransform& ToWorldTransform,
+	const bool bUseLocal, const bool bCalcLocal, FTransform* InOutLocal,
+	TArray<TFunction<void()>>& OutTasks)
+{
+	if (!ensure(InTransformContext.CanTransform()))
+	{
+		return;
+	}
+	
+	UControlRig* ControlRig = ShapeActor ? ShapeActor->ControlRig.Get() : nullptr;
+	if (!ensure(ControlRig))
+	{
+		return;
+	}
+
+	// In case for some reason the shape actor was detached, make sure to attach it again to the scene component
+	if (!ShapeActor->GetAttachParentActor())
+	{
+		if (USceneComponent* SceneComponent = GetHostingSceneComponent(ControlRig))
+		{
+			if (AActor* OwnerActor = SceneComponent->GetOwner())
+			{
+				ShapeActor->AttachToActor(OwnerActor, FAttachmentTransformRules::KeepWorldTransform);
+			}
+		}
+	}
+	
+	// first case is where we do all controls by the local diff.
+	bool bTransformChanged = false;
+	if (bUseLocal && InOutLocal)
+	{
+		bTransformChanged = MoveControlShapeLocally(ShapeActor, InTransformContext, ToWorldTransform, *InOutLocal);
+		if (bTransformChanged)
+		{
 			return;
 		}
 	}
@@ -5410,21 +5416,35 @@ void FControlRigEditMode::MoveControlShape( AControlRigShapeActor* ShapeActor,
 				break;
 			case ETransformSpace::Parent:
 				{
-					const int32 NumParents = Hierarchy->GetNumberOfParents(ControlElement);
-					if (NumParents > 0)
+					if (TOptional<FTransform> ConstraintSpace = GetConstraintParentTransform(ControlRig, ShapeActor->ControlName))
 					{
-						TransformSpace = Hierarchy->GetParentTransform(ControlElement, ERigTransformType::CurrentGlobal) * ToWorldTransform;
+						TransformSpace = *ConstraintSpace;
 					}
 					else
 					{
-						TransformSpace = GlobalTransform;
+						const int32 NumParents = Hierarchy->GetNumberOfParents(ControlElement);
+						if (NumParents > 0)
+						{
+							TransformSpace = Hierarchy->GetParentTransform(ControlElement, ERigTransformType::CurrentGlobal) * ToWorldTransform;
+						}
+						else
+						{
+							TransformSpace = GlobalTransform;
+						}
 					}
 				}
 				break;
 			case ETransformSpace::Explicit:
 				if (ensure(InTransformContext.bTranslation))
 				{
-					TransformSpace = Hierarchy->GetControlOffsetTransform(ControlElement, ERigTransformType::CurrentGlobal) * ToWorldTransform;
+					if (TOptional<FTransform> ConstraintSpace = GetConstraintParentTransform(ControlRig, ShapeActor->ControlName))
+					{
+						TransformSpace = *ConstraintSpace;
+					}
+					else
+					{
+						TransformSpace = Hierarchy->GetControlOffsetTransform(ControlElement, ERigTransformType::CurrentGlobal) * ToWorldTransform;	
+					}
 				}
 				break;
 			}
@@ -5462,11 +5482,11 @@ void FControlRigEditMode::MoveControlShape( AControlRigShapeActor* ShapeActor,
 		const FTransform ToWorldTransformCopy = ToWorldTransform;
 		
 		const TFunction<void()> SetControlShapeTask = [this, ShapeActor, ControlRig, ControlElement, NewTransform, ToWorldTransformCopy, Context,
-			bPrintPythonCommands, bIsTransientControl, bRotation = InTransformContext.bRotation, Rot = InTransformContext.Rot, bQueueTasks]()
+			bPrintPythonCommands, bIsTransientControl, InTransformContext, bQueueTasks]()
 		{
 			//fix flips and do rotation orders only if not additive
-			const bool bFixEulerFlips = !ControlRig->IsAdditive() && bRotation;
-			SetControlShapeTransform(ShapeActor, NewTransform, ToWorldTransformCopy, Context, bPrintPythonCommands, Rot, bFixEulerFlips);
+			const bool bFixEulerFlips = !ControlRig->IsAdditive() && InTransformContext.bRotation;
+			SetControlShapeTransform(ShapeActor, NewTransform, ToWorldTransformCopy, Context, bPrintPythonCommands, InTransformContext, bFixEulerFlips);
 			NotifyDrivenControls(ControlRig, ShapeActor->GetElementKey(),Context);
 
 			if (ControlElement && !bIsTransientControl)
@@ -6346,9 +6366,16 @@ void FControlRigEditMode::UpdateRotationContext()
 						RotationContext.RotationOrder = bUsePreferredRotationOrder ? Hierarchy->GetControlPreferredEulerRotationOrder(Control) : EEulerRotationOrder::XYZ;
 						RotationContext.Rotation = Hierarchy->GetControlPreferredRotator(Control);
 
-						const FTransform Offset = Hierarchy->GetControlOffsetTransform(Control, ERigTransformType::CurrentGlobal);
-						const FTransform ComponentTransform = GetHostingSceneComponentTransform(ControlRig);
-						RotationContext.Offset = Offset * ComponentTransform;
+						if (TOptional<FTransform> ConstraintSpace = GetConstraintParentTransform(ControlRig, ShapeActor->ControlName))
+						{
+							RotationContext.Offset = *ConstraintSpace;
+						}
+						else
+						{
+							const FTransform Offset = Hierarchy->GetControlOffsetTransform(Control, ERigTransformType::CurrentGlobal);
+							const FTransform ComponentTransform = GetHostingSceneComponentTransform(ControlRig);
+							RotationContext.Offset = Offset * ComponentTransform;
+						}
 
 						// only get the first rotation order
 						return;
@@ -6731,4 +6758,22 @@ FControlRigEditMode::FPendingControlRigEvaluator::~FPendingControlRigEvaluator()
 	}
 }
 
+TOptional<FTransform> FControlRigEditMode::GetConstraintParentTransform(const UControlRig* InControlRig, const FName& InControlName) const
+{
+	static const TOptional<FTransform> Dummy;
+
+	if (IsInLevelEditor())
+	{
+		if (!InControlRig || InControlName == NAME_None)
+		{
+			return Dummy;
+		}
+		
+		const uint32 ControlHash = UTransformableControlHandle::ComputeHash(InControlRig, InControlName);
+		return ConstraintsCache.GetParentTransform(ControlHash, InControlRig->GetWorld());
+	}
+
+	return Dummy;
+}
+
 #undef LOCTEXT_NAMESPACE
diff --git a/Engine/Plugins/Animation/ControlRig/Source/ControlRigEditor/Private/EditMode/ControlRigEditModeUtil.cpp b/Engine/Plugins/Animation/ControlRig/Source/ControlRigEditor/Private/EditMode/ControlRigEditModeUtil.cpp
new file mode 100644
--- /dev/null
+++ b/Engine/Plugins/Animation/ControlRig/Source/ControlRigEditor/Private/EditMode/ControlRigEditModeUtil.cpp
@@ -0,0 +1,113 @@
+﻿// Copyright Epic Games, Inc. All Rights Reserved.
+
+#include "EditMode/ControlRigEditModeUtil.h"
+
+#include "ControlRig.h"
+#include "TransformConstraint.h"
+
+#include "Rigs/RigHierarchy.h"
+
+namespace UE::ControlRigEditMode
+{
+
+FExplicitRotationInteraction::FExplicitRotationInteraction(const FTransformContext& InContext,
+	UControlRig* InControlRig,
+	URigHierarchy* InHierarchy,
+	FRigControlElement* InControlElement,
+	const FTransform& InComponentWorldTransform)
+		: TransformContext(InContext)
+		, ControlRig(InControlRig)
+		, Hierarchy(InHierarchy)
+		, ControlElement(InControlElement)
+		, ComponentWorldTransform(InComponentWorldTransform)
+{}
+	
+bool FExplicitRotationInteraction::IsValid() const
+{
+	const bool bIsExplicitRotation = TransformContext.bRotation && TransformContext.Space == ETransformSpace::Explicit && !TransformContext.Rot.IsZero();
+	if (!bIsExplicitRotation)
+	{
+		return false;
+	}
+
+	if (!ControlRig || !Hierarchy)
+	{
+		return false;
+	}
+
+	if (ControlRig->IsAdditive() || !Hierarchy->UsesPreferredEulerAngles())	
+	{
+		return false;
+	}
+
+	if (!ControlElement)
+	{
+		return false;
+	}
+	
+	const ERigControlType ControlType = ControlElement->Settings.ControlType;
+	return ControlType == ERigControlType::Rotator || ControlType == ERigControlType::EulerTransform;
+}
+
+void FExplicitRotationInteraction::Apply(
+	const FTransform& InGlobalTransform,
+	const FRigControlModifiedContext& InContext,
+	const bool bPrintPython,
+	const TArray< TWeakObjectPtr<UTickableConstraint> >& InConstraints) const
+{
+	static constexpr bool bNotify = true, bUndo = true;
+
+	const FName ControlName = ControlElement->GetKey().Name;
+	const ERigControlType ControlType = ControlElement->Settings.ControlType;
+	
+	const FVector DeltaEulerAngle(TransformContext.Rot.Roll, TransformContext.Rot.Pitch, TransformContext.Rot.Yaw);
+	FVector NewEulerAngle = Hierarchy->GetControlSpecifiedEulerAngle(ControlElement);
+	NewEulerAngle += DeltaEulerAngle;
+	
+	if (ControlType == ERigControlType::Rotator)
+	{
+		const FQuat Quat = Hierarchy->GetControlQuaternion(ControlElement, NewEulerAngle);
+		const FRotator Rotator(Quat);
+			
+		Hierarchy->SetControlSpecifiedEulerAngle(ControlElement, NewEulerAngle);
+		ControlRig->SetControlValue<FRotator>(ControlName, Rotator, bNotify, InContext, bUndo, bPrintPython);
+	}
+	else if (ControlType == ERigControlType::EulerTransform)
+	{
+		FRigControlModifiedContext Context = InContext;
+		
+		const FQuat Quat = Hierarchy->GetControlQuaternion(ControlElement, NewEulerAngle);
+
+		FEulerTransform EulerTransform;
+		if (InConstraints.IsEmpty())
+		{
+			FRigControlValue NewValue = ControlRig->GetControlValueFromGlobalTransform(ControlName, InGlobalTransform, ERigTransformType::CurrentGlobal);
+			EulerTransform = NewValue.Get<FRigControlValue::FEulerTransform_Float>().ToTransform();
+			EulerTransform.Rotation = FRotator(Quat);
+		}
+		else
+		{
+			Context.bConstraintUpdate = false;
+			
+			const FTransform WorldTransform = InGlobalTransform * ComponentWorldTransform;
+			FTransform LocalTransform = ControlRig->GetControlLocalTransform(ControlName);
+
+			const TOptional<FTransform> RelativeTransform =
+				FTransformConstraintUtils::GetConstraintsRelativeTransform(InConstraints, LocalTransform, WorldTransform);
+			if (ensure(RelativeTransform))
+			{
+				LocalTransform = *RelativeTransform; 
+			}
+
+			EulerTransform = FEulerTransform(LocalTransform);
+			EulerTransform.Rotation = FRotator(Quat);
+		}
+		
+		Hierarchy->SetControlSpecifiedEulerAngle(ControlElement, NewEulerAngle);
+		ControlRig->SetControlValue<FRigControlValue::FEulerTransform_Float>(ControlName, EulerTransform, bNotify, Context, bUndo, bPrintPython);
+		Hierarchy->SetControlSpecifiedEulerAngle(ControlElement, NewEulerAngle);
+	}
+}
+
+	
+}
\ No newline at end of file
diff --git a/Engine/Plugins/Animation/ControlRig/Source/ControlRigEditor/Private/Sequencer/ControlRigParameterTrackEditor.cpp b/Engine/Plugins/Animation/ControlRig/Source/ControlRigEditor/Private/Sequencer/ControlRigParameterTrackEditor.cpp
--- a/Engine/Plugins/Animation/ControlRig/Source/ControlRigEditor/Private/Sequencer/ControlRigParameterTrackEditor.cpp
+++ b/Engine/Plugins/Animation/ControlRig/Source/ControlRigEditor/Private/Sequencer/ControlRigParameterTrackEditor.cpp
@@ -3506,7 +3506,9 @@ void FControlRigParameterTrackEditor::HandleConstraintRemoved(IMovieSceneConstra
 						break;
 					case EConstraintsManagerNotifyType::ManagerUpdated:
 						InSection->OnConstraintsChanged();
-						break;		
+						break;
+					case EConstraintsManagerNotifyType::GraphUpdated:
+						break;
 				}
 			});
 			ConstraintHandlesToClear.Add(InSection->OnConstraintRemovedHandle);
diff --git a/Engine/Plugins/Animation/ControlRig/Source/ControlRigEditor/Public/EditMode/ControlRigEditMode.h b/Engine/Plugins/Animation/ControlRig/Source/ControlRigEditor/Public/EditMode/ControlRigEditMode.h
--- a/Engine/Plugins/Animation/ControlRig/Source/ControlRigEditor/Public/EditMode/ControlRigEditMode.h
+++ b/Engine/Plugins/Animation/ControlRig/Source/ControlRigEditor/Public/EditMode/ControlRigEditMode.h
@@ -19,8 +19,12 @@
 #include "Materials/MaterialInstanceDynamic.h"
 #include "IDetailKeyframeHandler.h"
 #include "WidgetFocusUtils.h"
+#include "Transform/TransformConstraintUtil.h"
+#include "ControlRigEditModeUtil.h"
+
 #include "ControlRigEditMode.generated.h"
 
+class UTickableConstraint;
 class FEditorViewportClient;
 class FRigVMEditor;
 class FViewport;
@@ -124,59 +128,6 @@ private:
 
 };
 
-namespace UE::ControlRigEditMode
-{
-
-enum class ETransformSpace
-{
-	World,
-	Local,
-	Parent,
-	Explicit
-};
-	
-struct FTransformContext
-{
-	bool CanTransform() const
-	{
-		return bTranslation || bRotation || bScale;
-	}
-	
-	bool bTranslation = false;
-	FVector Drag = FVector::ZeroVector;
-	
-	bool bRotation = false;
-	FRotator Rot = FRotator::ZeroRotator;
-	
-	bool bScale = false;
-	FVector Scale = FVector::OneVector;
-	
-	ETransformSpace Space = ETransformSpace::World;
-};
-
-/**
- * FInteractionDependencyCache provides a minimal "dependency graph" between the selected controls during interaction.
- * it simply stores information on who is a child and who is a parent in whatever is being manipulated.
- */
-struct FInteractionDependencyCache
-{
-	bool HasDownwardDependencies(const FRigElementKey& InKey) const
-	{
-		return Parents.Contains(InKey);
-	}
-	
-	bool HasUpwardDependencies(const FRigElementKey& InKey) const
-	{
-		return Children.Contains(InKey);
-	}
-	
-	TSet<FRigElementKey> Parents;
-	TSet<FRigElementKey> Children;
-};
-	
-}
-
-
 class CONTROLRIGEDITOR_API FControlRigEditMode : public IPersonaEditMode
 {
 public:
@@ -305,7 +256,7 @@ protected:
 	bool ComputePivotFromEditedShape(UControlRig* InControlRig, FTransform& OutTransform) const;
 	bool ComputePivotFromShapeActors(UControlRig* InControlRig, const bool bEachLocalSpace, const UE::ControlRigEditMode::ETransformSpace InSpace, FTransform& OutTransform) const;
 	bool ComputePivotFromElements(UControlRig* InControlRig, FTransform& OutTransform) const;
-	FTransform GetPivotOrientation(const FRigElementKey& InControlKey, URigHierarchy* InHierarchy, const UE::ControlRigEditMode::ETransformSpace InSpace, const FTransform& InComponentTransform) const;
+	FTransform GetPivotOrientation(const FRigElementKey& InControlKey, const UControlRig* InControlRig, URigHierarchy* InHierarchy, const UE::ControlRigEditMode::ETransformSpace InSpace, const FTransform& InComponentTransform) const;
 	
 	/** Get the current coordinate system space */
 	ECoordSystem GetCoordSystemSpace() const;
@@ -598,7 +549,7 @@ private:
 	bool CheckMovieSceneSig();
 	void SetControlShapeTransform( const AControlRigShapeActor* InShapeActor, const FTransform& InGlobalTransform,
 		const FTransform& InToWorldTransform, const FRigControlModifiedContext& InContext, const bool bPrintPython,
-		const FRotator& InRot, const bool bFixEulerFlips = true) const;
+		const UE::ControlRigEditMode::FTransformContext& InTransformContext, const bool bFixEulerFlips = true) const;
 	static FTransform GetControlShapeTransform(const AControlRigShapeActor* ShapeActor);
 
 	static void ChangeControlShapeTransform(AControlRigShapeActor* ShapeActor, const UE::ControlRigEditMode::FTransformContext& InContext, const FTransform& ToWorldTransform);
@@ -611,6 +562,10 @@ public:
 
 protected:
 
+	bool MoveControlShapeLocally( AControlRigShapeActor* ShapeActor,
+		const UE::ControlRigEditMode::FTransformContext& InTransformContext, const FTransform& ToWorldTransform,
+		const FTransform& InLocal);
+	
 	void MoveControlShape(AControlRigShapeActor* ShapeActor,
 		const UE::ControlRigEditMode::FTransformContext& InContext, const FTransform& ToWorldTransform,
 		const bool bUseLocal, const bool bCalcLocal, FTransform* InOutLocal, TArray<TFunction<void()>>& OutTasks);
@@ -725,6 +680,9 @@ private:
 	FRotationContext& GetRotationContext() const;
 	void UpdateRotationContext();
 
+	TOptional<FTransform> GetConstraintParentTransform(const UControlRig* InControlRig, const FName& InControlName) const;
+	mutable UE::TransformConstraintUtil::FConstraintsInteractionCache ConstraintsCache;
+
 	friend class FControlRigEditorModule;
 	friend class FControlRigEditor;
 	friend class FControlRigEditModeGenericDetails;
diff --git a/Engine/Plugins/Animation/ControlRig/Source/ControlRigEditor/Public/EditMode/ControlRigEditModeUtil.h b/Engine/Plugins/Animation/ControlRig/Source/ControlRigEditor/Public/EditMode/ControlRigEditModeUtil.h
new file mode 100644
--- /dev/null
+++ b/Engine/Plugins/Animation/ControlRig/Source/ControlRigEditor/Public/EditMode/ControlRigEditModeUtil.h
@@ -0,0 +1,102 @@
+﻿// Copyright Epic Games, Inc. All Rights Reserved.
+
+#pragma once
+
+#include "Math/MathFwd.h"
+#include "Math/Vector.h"
+#include "Math/Rotator.h"
+
+#include "Rigs/RigHierarchyDefines.h"
+
+#include "UObject/WeakObjectPtr.h"
+
+class UControlRig;
+class URigHierarchy;
+class UTickableConstraint;
+
+struct FRigControlModifiedContext;
+struct FRigControlElement;
+
+namespace UE::ControlRigEditMode
+{
+
+enum class ETransformSpace
+{
+	World,
+	Local,
+	Parent,
+	Explicit
+};
+
+/**
+ * FTransformContext provides a way of passing the various transform parameters to functions that need to know what the transform context is.
+ * Extend it if necessary, particularly to avoid overloading some functions signatures.
+ */
+
+struct FTransformContext
+{
+	bool CanTransform() const
+	{
+		return bTranslation || bRotation || bScale;
+	}
+
+	bool bTranslation = false;
+	FVector Drag = FVector::ZeroVector;
+
+	bool bRotation = false;
+	FRotator Rot = FRotator::ZeroRotator;
+
+	bool bScale = false;
+	FVector Scale = FVector::OneVector;
+
+	ETransformSpace Space = ETransformSpace::World;
+};
+
+/**
+ * FInteractionDependencyCache provides a minimal "dependency graph" between the selected controls during interaction.
+ * it simply stores information on who is a child and who is a parent in whatever is being manipulated.
+ */
+
+struct FInteractionDependencyCache
+{
+	bool HasDownwardDependencies(const FRigElementKey& InKey) const
+	{
+		return Parents.Contains(InKey);
+	}
+
+	bool HasUpwardDependencies(const FRigElementKey& InKey) const
+	{
+		return Children.Contains(InKey);
+	}
+
+	TSet<FRigElementKey> Parents;
+	TSet<FRigElementKey> Children;
+};
+
+/**
+* FExplicitRotationInteraction is a wrapper struct to apply euler angle deltas to controls.
+*/
+	
+struct FExplicitRotationInteraction
+{
+	FExplicitRotationInteraction(const FTransformContext& InContext,
+		UControlRig* InControlRig,
+		URigHierarchy* InHierarchy,
+		FRigControlElement* InControlElement,
+		const FTransform& InComponentWorldTransform);
+	
+	bool IsValid() const;
+
+	void Apply(const FTransform& InGlobalTransform, const FRigControlModifiedContext& InContext, const bool bPrintPython,
+		const TArray< TWeakObjectPtr<UTickableConstraint> >& InConstraints = TArray< TWeakObjectPtr<UTickableConstraint> >()) const;
+
+private:
+	const FTransformContext& TransformContext;
+	UControlRig* ControlRig = nullptr;
+	URigHierarchy* Hierarchy = nullptr;
+	FRigControlElement* ControlElement = nullptr;
+	const FTransform& ComponentWorldTransform;
+};
+	
+}
+
diff --git a/Engine/Source/Editor/MovieSceneTools/Private/TrackEditors/TransformTrackEditor.cpp b/Engine/Source/Editor/MovieSceneTools/Private/TrackEditors/TransformTrackEditor.cpp
--- a/Engine/Source/Editor/MovieSceneTools/Private/TrackEditors/TransformTrackEditor.cpp
+++ b/Engine/Source/Editor/MovieSceneTools/Private/TrackEditors/TransformTrackEditor.cpp
@@ -1883,7 +1883,9 @@ void F3DTransformTrackEditor::HandleConstraintRemoved(IMovieSceneConstrainedSect
 							break;
 						case EConstraintsManagerNotifyType::ManagerUpdated:
 							InSection->OnConstraintsChanged();
-							break;		
+							break;
+						case EConstraintsManagerNotifyType::GraphUpdated:
+							break;
 					}
 				});
 		ConstraintHandlesToClear.Add(InSection->OnConstraintRemovedHandle);
diff --git a/Engine/Source/Runtime/Experimental/Animation/Constraints/Private/ConstraintSubsystem.cpp b/Engine/Source/Runtime/Experimental/Animation/Constraints/Private/ConstraintSubsystem.cpp
--- a/Engine/Source/Runtime/Experimental/Animation/Constraints/Private/ConstraintSubsystem.cpp
+++ b/Engine/Source/Runtime/Experimental/Animation/Constraints/Private/ConstraintSubsystem.cpp
@@ -364,4 +364,7 @@ void FConstraintsInWorld::InvalidateGraph()
 	{
 		EvaluationGraph.Reset();
 	}
+
+	const FConstraintsManagerController& Controller = FConstraintsManagerController::Get(World.Get());
+	Controller.Notify(EConstraintsManagerNotifyType::GraphUpdated, nullptr);
 }
\ No newline at end of file
diff --git a/Engine/Source/Runtime/Experimental/Animation/Constraints/Private/ConstraintsManager.cpp b/Engine/Source/Runtime/Experimental/Animation/Constraints/Private/ConstraintsManager.cpp
--- a/Engine/Source/Runtime/Experimental/Animation/Constraints/Private/ConstraintsManager.cpp
+++ b/Engine/Source/Runtime/Experimental/Animation/Constraints/Private/ConstraintsManager.cpp
@@ -879,6 +879,10 @@ void FConstraintsManagerController::Notify(EConstraintsManagerNotifyType InNotif
 		case EConstraintsManagerNotifyType::ManagerUpdated:
 			checkSlow(Cast<UConstraintSubsystem>(InObject) != nullptr);
 			break;
+		
+		case EConstraintsManagerNotifyType::GraphUpdated:
+			break;
+		
 		default:
 			checkfSlow(false, TEXT("Unchecked EConstraintsManagerNotifyType!"));
 			break;
@@ -904,7 +908,9 @@ void FConstraintsManagerController::InvalidateEvaluationGraph() const
 	{
 		return;
 	}
+	
 	SubSystem->GetEvaluationGraph(World).InvalidateData();
+	Notify(EConstraintsManagerNotifyType::GraphUpdated, nullptr);
 }
 
 void FConstraintsManagerController::FlushEvaluationGraph() const
diff --git a/Engine/Source/Runtime/Experimental/Animation/Constraints/Private/Transform/TransformConstraintUtil.cpp b/Engine/Source/Runtime/Experimental/Animation/Constraints/Private/Transform/TransformConstraintUtil.cpp
new file mode 100644
--- /dev/null
+++ b/Engine/Source/Runtime/Experimental/Animation/Constraints/Private/Transform/TransformConstraintUtil.cpp
@@ -0,0 +1,145 @@
+﻿// Copyright Epic Games, Inc. All Rights Reserved.
+
+#include "Transform/TransformConstraintUtil.h"
+
+#include "ConstraintsManager.h"
+#include "TransformableHandle.h"
+#include "TransformableRegistry.h"
+#include "TransformConstraint.h"
+
+#include "UObject/Object.h"
+
+namespace UE::TransformConstraintUtil
+{
+
+namespace Private
+{
+
+uint32 GetHandleHash(const UObject* InObject, const FName& InAttachmentName)
+{
+	if (InObject)
+	{
+		const FTransformableRegistry& Registry = FTransformableRegistry::Get();
+		if (const FTransformableRegistry::GetHashFuncT HashFunction = Registry.GetHashFunction(InObject->GetClass()))
+		{
+			return HashFunction(InObject, InAttachmentName);
+		}
+	}
+
+	return 0;
+}
+
+}
+
+const TArray< TWeakObjectPtr<UTickableConstraint> >& FConstraintsInteractionCache::Get(const UObject* InObject, const FName& InAttachmentName)
+{
+	const uint32 HandleHash = Private::GetHandleHash(InObject, InAttachmentName);
+	if (HandleHash == 0)
+	{
+		static const TArray< TWeakObjectPtr<UTickableConstraint> > Empty;
+		return Empty;
+	}
+
+	return Get(HandleHash, InObject->GetWorld());
+}
+
+bool FConstraintsInteractionCache::HasAnyActiveConstraint(const UObject* InObject, const FName& InAttachmentName)
+{
+	const uint32 HandleHash = Private::GetHandleHash(InObject, InAttachmentName);
+	if (HandleHash == 0)
+	{
+		return false;
+	}
+	
+	return HasAnyActiveConstraint(HandleHash, InObject->GetWorld());
+}
+	
+TOptional<FTransform> FConstraintsInteractionCache::GetParentTransform(const UObject* InObject, const FName& InAttachmentName)
+{
+	const uint32 HandleHash = Private::GetHandleHash(InObject, InAttachmentName);
+	if (HandleHash == 0)
+	{
+		static const TOptional<FTransform> Invalid;
+		return Invalid;
+	}
+
+	return GetParentTransform(HandleHash, InObject->GetWorld());
+}
+
+const TArray< TWeakObjectPtr<UTickableConstraint> >& FConstraintsInteractionCache::Get(const uint32 InHandleHash, UWorld* InWorld)
+{
+	static const TArray< TWeakObjectPtr<UTickableConstraint> > Empty;
+	if (!ensure(InHandleHash != 0 && InWorld))
+	{
+		return Empty;
+	}
+
+	if (TArray< TWeakObjectPtr<UTickableConstraint> >* FoundConstraints = PerHandleConstraints.Find(InHandleHash))
+	{
+		return *FoundConstraints;
+	}
+	
+	static constexpr bool bSorted = true;
+	const FConstraintsManagerController& Controller = FConstraintsManagerController::Get(InWorld);
+	const TArray< TWeakObjectPtr<UTickableConstraint> > Constraints = Controller.GetParentConstraints(InHandleHash, bSorted);
+
+	TArray< TWeakObjectPtr<UTickableConstraint> >& TransformConstraints = PerHandleConstraints.Add(InHandleHash);
+	TransformConstraints = Constraints.FilterByPredicate([](const TWeakObjectPtr<UTickableConstraint>& WeakConstraint)
+	{
+		return Cast<UTickableTransformConstraint>(WeakConstraint.Get()) != nullptr;
+	});
+	return TransformConstraints;
+}
+
+bool FConstraintsInteractionCache::HasAnyActiveConstraint(const uint32 InHandleHash, UWorld* InWorld)
+{
+	const TArray< TWeakObjectPtr<UTickableConstraint> >& TransformConstraints = Get(InHandleHash, InWorld);
+	return FTransformConstraintUtils::GetLastActiveConstraintIndex(TransformConstraints) != INDEX_NONE;
+}
+
+TOptional<FTransform> FConstraintsInteractionCache::GetParentTransform(const uint32 InHandleHash, UWorld* InWorld)
+{
+	const TArray< TWeakObjectPtr<UTickableConstraint> >& TransformConstraints = Get(InHandleHash, InWorld);
+	const int32 ConstraintsIndex = FTransformConstraintUtils::GetLastActiveConstraintIndex(TransformConstraints);
+	if (ConstraintsIndex != INDEX_NONE)
+	{
+		UTickableTransformConstraint* Constraint = Cast<UTickableTransformConstraint>(TransformConstraints[ConstraintsIndex].Get());
+		if (Constraint && Constraint->ParentTRSHandle)
+		{
+			return Constraint->ParentTRSHandle->GetGlobalTransform();	
+		}
+	}
+
+	static const TOptional<FTransform> Invalid;
+	return Invalid;
+}
+	
+void FConstraintsInteractionCache::Reset()
+{
+	PerHandleConstraints.Reset();
+}
+
+void FConstraintsInteractionCache::RegisterNotifications(UWorld* InWorld)
+{
+	FConstraintsManagerController& Controller = FConstraintsManagerController::Get(InWorld);
+	ConstraintsNotificationHandle =
+		Controller.GetNotifyDelegate().AddLambda([this](EConstraintsManagerNotifyType InNotifyType, UObject* InObject)
+	{
+		if (InNotifyType == EConstraintsManagerNotifyType::GraphUpdated)
+		{
+			Reset();
+		}
+	});
+}
+
+void FConstraintsInteractionCache::UnregisterNotifications(UWorld* InWorld)
+{
+	if (ConstraintsNotificationHandle.IsValid())
+	{
+		FConstraintsManagerController& Controller = FConstraintsManagerController::Get(InWorld);
+		Controller.GetNotifyDelegate().Remove(ConstraintsNotificationHandle);
+		ConstraintsNotificationHandle.Reset();
+	}
+}
+	
+}
diff --git a/Engine/Source/Runtime/Experimental/Animation/Constraints/Public/ConstraintsManager.h b/Engine/Source/Runtime/Experimental/Animation/Constraints/Public/ConstraintsManager.h
--- a/Engine/Source/Runtime/Experimental/Animation/Constraints/Public/ConstraintsManager.h
+++ b/Engine/Source/Runtime/Experimental/Animation/Constraints/Public/ConstraintsManager.h
@@ -230,6 +230,7 @@ enum class EConstraintsManagerNotifyType
 	ConstraintRemoved,					/** A constraint has been removed. */
 	ConstraintRemovedWithCompensation,	/** A constraint has been removed and needs compensation. */
 	ManagerUpdated,						/** The manager has been updated/reset. */
+	GraphUpdated						/** The evaluation graph has been updated/reset. */
 };
 
 DECLARE_MULTICAST_DELEGATE_TwoParams(FConstraintsManagerNotifyDelegate, EConstraintsManagerNotifyType /* type */, UObject* /* subject */);
diff --git a/Engine/Source/Runtime/Experimental/Animation/Constraints/Public/Transform/TransformConstraintUtil.h b/Engine/Source/Runtime/Experimental/Animation/Constraints/Public/Transform/TransformConstraintUtil.h
new file mode 100644
--- /dev/null
+++ b/Engine/Source/Runtime/Experimental/Animation/Constraints/Public/Transform/TransformConstraintUtil.h
@@ -0,0 +1,46 @@
+﻿// Copyright Epic Games, Inc. All Rights Reserved.
+
+#pragma once
+
+#include "Containers/Array.h"
+#include "Containers/Map.h"
+
+#include "Delegates/IDelegateInstance.h"
+
+#include "UObject/WeakObjectPtr.h"
+
+class UWorld;
+class UTickableConstraint;
+
+namespace UE::TransformConstraintUtil
+{
+
+/**
+ * FConstraintsInteractionCache is designed to minimize the number of requests made to the ConstraintManager and thus improve performance.
+ * it is particularly useful for interface components that need to be updated frequently.
+ * note that PerHandleConstraints is reset when the constraint graph is updated, to keep it synchronized.
+ */
+	
+struct CONSTRAINTS_API FConstraintsInteractionCache
+{
+public:
+	void Reset();
+
+	const TArray< TWeakObjectPtr<UTickableConstraint> >& Get(const UObject* InObject, const FName& InAttachmentName);
+	bool HasAnyActiveConstraint(const UObject* InObject, const FName& InAttachmentName);
+	TOptional<FTransform> GetParentTransform(const UObject* InObject, const FName& InAttachmentName);
+	
+	const TArray< TWeakObjectPtr<UTickableConstraint> >& Get(const uint32 InHandleHash, UWorld* InWorld);
+	bool HasAnyActiveConstraint(const uint32 InHandleHash, UWorld* InWorld);
+	TOptional<FTransform> GetParentTransform(const uint32 InHandleHash, UWorld* InWorld);
+
+	void RegisterNotifications(UWorld* InWorld);
+	void UnregisterNotifications(UWorld* InWorld);
+
+private:
+	
+	TMap<uint32, TArray< TWeakObjectPtr<UTickableConstraint> >> PerHandleConstraints;
+	FDelegateHandle ConstraintsNotificationHandle;
+};
+	
+}
