Control Rig: Improve hierarchy instantiation cost

#rnx

diff --git a/Engine/Plugins/Animation/ControlRig/Source/ControlRig/Private/Rigs/RigHierarchy.cpp b/Engine/Plugins/Animation/ControlRig/Source/ControlRig/Private/Rigs/RigHierarchy.cpp
--- a/Engine/Plugins/Animation/ControlRig/Source/ControlRig/Private/Rigs/RigHierarchy.cpp
+++ b/Engine/Plugins/Animation/ControlRig/Source/ControlRig/Private/Rigs/RigHierarchy.cpp
@@ -434,6 +434,8 @@ void URigHierarchy::Load(FArchive& Ar)
 	int32 NumTransforms = 0;
 	int32 NumDirtyStates = 0;
 	int32 NumCurves = 0;
+
+	ElementIndexLookup.Reserve(ElementCount);
 	
 	const bool bAllocateStoragePerElement = ArchiveForElements->CustomVer(FControlRigObjectVersion::GUID) < FControlRigObjectVersion::RigHierarchyIndirectElementStorage;
 	for(int32 ElementIndex = 0; ElementIndex < ElementCount; ElementIndex++)
@@ -663,12 +665,37 @@ void URigHierarchy::Reset_Impl(bool bResetElements)
 	if(bResetElements)
 	{
 		UE::TScopeLock Lock(ElementsLock);
-		
+
+		// remove the metadata first. elements hold an integer pointing to the
+		// element metadata for each element - so if this is empty it won't cause
+		// any further deallocation / removal.
+		ElementMetadata.Reset([](int32 InIndex, FMetadataStorage* MetadataStorage)
+		{
+			for (TTuple<FName, FRigBaseMetadata*>& Item: MetadataStorage->MetadataMap)
+			{
+				FRigBaseMetadata::DestroyMetadata(&Item.Value);
+			} 
+		});
+
+		// clear these first - the indices pointing to these
+		// won't be valid - which in turn won't cause any deallocation
+		ElementTransforms.Reset();
+		ElementDirtyStates.Reset();
+		ElementCurves.Reset();
+
+		// clearing the components first will avoid cost of looping over
+		// each component per elements and destroying them one by onw
+		ElementComponents.Reset();
+
 		// walk in reverse since certain elements might not have been allocated themselves
 		for(int32 ElementIndex = Elements.Num() - 1; ElementIndex >= 0; ElementIndex--)
 		{
-			DestroyElement(Elements[ElementIndex]);
+			DestroyElement(Elements[ElementIndex],
+				false /* don't clear components */,
+				false /* don't clear element storage */,
+				false /* don't destroy metadata */);
 		}
+		
 		Elements.Reset();
 		ElementsPerType.Reset();
 		const UEnum* ElementTypeEnum = StaticEnum<ERigElementType>();
@@ -682,18 +709,6 @@ void URigHierarchy::Reset_Impl(bool bResetElements)
 		}
 		ElementIndexLookup.Reset();
 		ComponentIndexLookup.Reset();
-
-		ElementTransforms.Reset();
-		ElementDirtyStates.Reset();
-		ElementCurves.Reset();
-
-		ElementMetadata.Reset([](int32 InIndex, FMetadataStorage* MetadataStorage)
-		{
-			for (TTuple<FName, FRigBaseMetadata*>& Item: MetadataStorage->MetadataMap)
-			{
-				FRigBaseMetadata::DestroyMetadata(&Item.Value);
-			} 
-		});
 	}
 
 	ResetPoseHash = INDEX_NONE;
@@ -845,7 +860,7 @@ void URigHierarchy::CopyHierarchy(URigHierarchy* InHierarchy)
 			}
 
 			Elements.Reserve(InHierarchy->Elements.Num());
-			ElementIndexLookup.Reserve(InHierarchy->ElementIndexLookup.Num());
+			ElementIndexLookup = InHierarchy->ElementIndexLookup;
 
 			for(int32 Index = 0; Index < InHierarchy->Num(); Index++)
 			{
@@ -871,7 +886,6 @@ void URigHierarchy::CopyHierarchy(URigHierarchy* InHierarchy)
 				Target->ComponentIndices.Reset();
 
 				ElementsPerType[ElementTypeIndex].Add(Target);
-				ElementIndexLookup.Add(Key, Target->Index);
 			
 				IncrementPoseVersion(Index);
 
@@ -3732,6 +3746,26 @@ FRigPose URigHierarchy::GetPose(bool bInitial, ERigElementType InElementType,
 	Pose.HierarchyTopologyVersion = GetTopologyVersion();
 	Pose.PoseHash = Pose.HierarchyTopologyVersion;
 
+	static constexpr int32 FirstElementType = static_cast<int32>(ERigElementType::First); 
+	static constexpr int32 LastElementType = static_cast<int32>(ERigElementType::Last);
+	const int32 InElementTypeAsInt32 = static_cast<int32>(InElementType);
+	
+	int32 ExpectedNumElements = 0;
+	for(int32 ElementType = FirstElementType; ElementType <= LastElementType; ElementType++)
+	{
+		if ((InElementTypeAsInt32 & ElementType) != 0)
+		{
+			ExpectedNumElements += Num(static_cast<ERigElementType>(ElementType));
+		}
+	}
+	if(ExpectedNumElements == 0)
+	{
+		ExpectedNumElements = Elements.Num();
+	}
+	Pose.Elements.Reserve(ExpectedNumElements);
+
+	const uint32 TopologyVersionHash = GetTopologyVersionHash();
+
 	for (int32 ElementIndex = 0; ElementIndex < Elements.Num(); ElementIndex++)
 	{
 		FRigBaseElement* Element = Elements[ElementIndex];
@@ -3752,7 +3786,7 @@ FRigPose URigHierarchy::GetPose(bool bInitial, ERigElementType InElementType,
 		}
 		
 		FRigPoseElement PoseElement;
-		PoseElement.Index.UpdateCache(Element->GetKey(), this);
+		PoseElement.Index.Set(Element, TopologyVersionHash);
 		
 		if(FRigTransformElement* TransformElement = Cast<FRigTransformElement>(Element))
 		{
@@ -5983,7 +6017,7 @@ FRigBaseElement* URigHierarchy::MakeElement(ERigElementType InElementType, int32
 	return Element;
 }
 
-void URigHierarchy::DestroyElement(FRigBaseElement*& InElement)
+void URigHierarchy::DestroyElement(FRigBaseElement*& InElement, bool bDestroyComponents, bool bDestroyElementStorage, bool bDestroyMetadata)
 {
 	LLM_SCOPE_BYNAME(TEXT("Animation/ControlRig"));
 	check(InElement != nullptr);
@@ -6002,8 +6036,9 @@ void URigHierarchy::DestroyElement(FRigBaseElement*& InElement)
 			for(int32 Index=0;Index<Count;Index++)
 			{
 				TGuardValue<const FRigBaseElement*> DestroyGuard(ElementBeingDestroyed, &ExistingElements[Index]);
-				DestroyComponents(&ExistingElements[Index]);
-				DeallocateElementStorage(&ExistingElements[Index]);
+				if(bDestroyComponents) DestroyComponents(&ExistingElements[Index]);
+				if(bDestroyElementStorage) DeallocateElementStorage(&ExistingElements[Index]);
+				if(!bDestroyMetadata) ExistingElements[Index].Owner = nullptr;
 				ExistingElements[Index].~FRigBoneElement(); 
 			}
 			break;
@@ -6014,8 +6049,9 @@ void URigHierarchy::DestroyElement(FRigBaseElement*& InElement)
 			for(int32 Index=0;Index<Count;Index++)
 			{
 				TGuardValue<const FRigBaseElement*> DestroyGuard(ElementBeingDestroyed, &ExistingElements[Index]);
-				DestroyComponents(&ExistingElements[Index]);
-				DeallocateElementStorage(&ExistingElements[Index]);
+				if(bDestroyComponents) DestroyComponents(&ExistingElements[Index]);
+				if(bDestroyElementStorage) DeallocateElementStorage(&ExistingElements[Index]);
+				if(!bDestroyMetadata) ExistingElements[Index].Owner = nullptr;
 				ExistingElements[Index].~FRigNullElement(); 
 			}
 			break;
@@ -6026,8 +6062,9 @@ void URigHierarchy::DestroyElement(FRigBaseElement*& InElement)
 			for(int32 Index=0;Index<Count;Index++)
 			{
 				TGuardValue<const FRigBaseElement*> DestroyGuard(ElementBeingDestroyed, &ExistingElements[Index]);
-				DestroyComponents(&ExistingElements[Index]);
-				DeallocateElementStorage(&ExistingElements[Index]);
+				if(bDestroyComponents) DestroyComponents(&ExistingElements[Index]);
+				if(bDestroyElementStorage) DeallocateElementStorage(&ExistingElements[Index]);
+				if(!bDestroyMetadata) ExistingElements[Index].Owner = nullptr;
 				ExistingElements[Index].~FRigControlElement(); 
 			}
 			break;
@@ -6038,8 +6075,9 @@ void URigHierarchy::DestroyElement(FRigBaseElement*& InElement)
 			for(int32 Index=0;Index<Count;Index++)
 			{
 				TGuardValue<const FRigBaseElement*> DestroyGuard(ElementBeingDestroyed, &ExistingElements[Index]);
-				DestroyComponents(&ExistingElements[Index]);
-				DeallocateElementStorage(&ExistingElements[Index]);
+				if(bDestroyComponents) DestroyComponents(&ExistingElements[Index]);
+				if(bDestroyElementStorage) DeallocateElementStorage(&ExistingElements[Index]);
+				if(!bDestroyMetadata) ExistingElements[Index].Owner = nullptr;
 				ExistingElements[Index].~FRigCurveElement(); 
 			}
 			break;
@@ -6050,8 +6088,9 @@ void URigHierarchy::DestroyElement(FRigBaseElement*& InElement)
 			for(int32 Index=0;Index<Count;Index++)
 			{
 				TGuardValue<const FRigBaseElement*> DestroyGuard(ElementBeingDestroyed, &ExistingElements[Index]);
-				DestroyComponents(&ExistingElements[Index]);
-				DeallocateElementStorage(&ExistingElements[Index]);
+				if(bDestroyComponents) DestroyComponents(&ExistingElements[Index]);
+				if(bDestroyElementStorage) DeallocateElementStorage(&ExistingElements[Index]);
+				if(!bDestroyMetadata) ExistingElements[Index].Owner = nullptr;
 				ExistingElements[Index].~FRigReferenceElement(); 
 			}
 			break;
@@ -6062,8 +6101,9 @@ void URigHierarchy::DestroyElement(FRigBaseElement*& InElement)
 			for(int32 Index=0;Index<Count;Index++)
 			{
 				TGuardValue<const FRigBaseElement*> DestroyGuard(ElementBeingDestroyed, &ExistingElements[Index]);
-				DestroyComponents(&ExistingElements[Index]);
-				DeallocateElementStorage(&ExistingElements[Index]);
+				if(bDestroyComponents) DestroyComponents(&ExistingElements[Index]);
+				if(bDestroyElementStorage) DeallocateElementStorage(&ExistingElements[Index]);
+				if(!bDestroyMetadata) ExistingElements[Index].Owner = nullptr;
 				ExistingElements[Index].~FRigConnectorElement(); 
 			}
 			break;
@@ -6074,8 +6114,9 @@ void URigHierarchy::DestroyElement(FRigBaseElement*& InElement)
 			for(int32 Index=0;Index<Count;Index++)
 			{
 				TGuardValue<const FRigBaseElement*> DestroyGuard(ElementBeingDestroyed, &ExistingElements[Index]);
-				DestroyComponents(&ExistingElements[Index]);
-				DeallocateElementStorage(&ExistingElements[Index]);
+				if(bDestroyComponents) DestroyComponents(&ExistingElements[Index]);
+				if(bDestroyElementStorage) DeallocateElementStorage(&ExistingElements[Index]);
+				if(!bDestroyMetadata) ExistingElements[Index].Owner = nullptr;
 				ExistingElements[Index].~FRigSocketElement(); 
 			}
 			break;
@@ -6685,195 +6726,164 @@ bool URigHierarchy::SortElementStorage()
 	{
 		return false;
 	}
-	
-	TMap<int32, FRigComputedTransform*> TransformIndexToComputedTransform;
-	TMap<int32, FRigTransformDirtyState*> TransformIndexToDirtyState;
 
-	struct FTypeInfo
-	{
-		FTypeInfo()
-		: Index(INDEX_NONE)
-		, Element(ERigElementType::Bone)
-		, Transform(ERigTransformType::InitialLocal)
-		, Storage(ERigTransformStorageType::Pose)
-		{
-		}
+	static constexpr int32 NumTransformTypes = ERigTransformType::NumTransformTypes;
 
-		FTypeInfo(int32 InIndex, ERigElementType InElement, ERigTransformType::Type InTransform, ERigTransformStorageType::Type InStorage)
-		: Index(InIndex)
-		, Element(InElement)
-		, Transform(InTransform)
-		, Storage(InStorage)
-		{
-		}
+	// figure out which types of elements actually exist on this hierarchy
+	static const int32 LastElementType = RigElementTypeToFlatIndex(ERigElementType::Last);
+	static const int32 ControlElementType = RigElementTypeToFlatIndex(ERigElementType::Control);
 
-		uint32 GetSortKey(uint32 InNumElements) const
-		{
-			const uint32 StorageMultiplier = InNumElements; 
-			const uint32 ElementMultiplier = StorageMultiplier * (uint32)ERigTransformStorageType::NumStorageTypes;
-			const uint32 TransformMultiplier = ElementMultiplier * uint32(RigElementTypeToFlatIndex(ERigElementType::Last) + 1);
-			// we sort elements by laying them out first by transform type, then element type, then storage type
-			// and finally in the order of their original index in the hierarchy
-			return Index +
-				int32(Transform) * TransformMultiplier +
-				int32(Element) * ElementMultiplier +
-				int32(Storage) * StorageMultiplier;
-		}
+	TArray<int32> PerElementStride, NumElementsPerType, OffsetPerElementType;
+	PerElementStride.Reserve(LastElementType+1);
+	NumElementsPerType.Reserve(LastElementType+1);
+	OffsetPerElementType.AddZeroed(LastElementType+1);
 
-		int32 Index;
-		ERigElementType Element;
-		ERigTransformType::Type Transform;
-		ERigTransformStorageType::Type Storage;
-	};
+	int32 NumElementTypes = 0;
+	int32 TotalNumberOfTransformsPerTransformType = 0;
 
-	struct FTransformInfo
+	for(int32 ElementType = 0; ElementType <= LastElementType; ElementType++)
 	{
-		FTransformInfo()
-		: ComputedTransform(nullptr)
-		, DirtyState(nullptr)
-		{
-		}
+		NumElementsPerType.Add(ElementsPerType[ElementType].Num());
+		PerElementStride.Add(ElementType == ControlElementType ? static_cast<int32>(ERigTransformStorageType::NumStorageTypes) : 1);
 
-		FTransformInfo(const FTypeInfo& InTypeInfo, FRigComputedTransform* InComputedTransform, FRigTransformDirtyState* InDirtyState)
-		: TypeInfo(InTypeInfo)
-		, ComputedTransform(InComputedTransform)
-		, DirtyState(InDirtyState)
+		if(ElementType > 0)
 		{
+			OffsetPerElementType[ElementType] = OffsetPerElementType[ElementType-1] + NumElementsPerType[ElementType-1] * PerElementStride[ElementType-1]; 
 		}
 		
-		FTypeInfo TypeInfo;
-		FRigComputedTransform* ComputedTransform;
-		FRigTransformDirtyState* DirtyState;
-	};
+		NumElementTypes += ElementsPerType[ElementType].IsEmpty() ? 0 : 1;
+		TotalNumberOfTransformsPerTransformType += NumElementsPerType.Last() * PerElementStride.Last();
+	}
 
-	TArray<FTransformInfo> TransformInfos;
-	TransformInfos.Reserve(Elements.Num() * 4);
+	const int32 TotalNumberOfTransforms = TotalNumberOfTransformsPerTransformType * NumTransformTypes;
+	TArray<int32> TransformIndexToAllocatedIndex;
+	TransformIndexToAllocatedIndex.AddUninitialized(TotalNumberOfTransforms);
 
-	ForEachTransformElementStorage([&TransformInfos](
+	// we sort elements by laying them out first by transform type, then element type, then storage type.
+	auto GetElementIndexForSort = [OffsetPerElementType, PerElementStride, TotalNumberOfTransformsPerTransformType](
+		const FRigBaseElement* InElement,
+		ERigTransformType::Type InTransformType,
+		ERigTransformStorageType::Type InStorageType
+	) -> int32 {
+		const int32 FlatElementIndex = RigElementTypeToFlatIndex(InElement->GetType());
+		return
+			static_cast<int32>(InTransformType) * TotalNumberOfTransformsPerTransformType +
+			(OffsetPerElementType[FlatElementIndex] + InElement->GetSubIndex() * PerElementStride[FlatElementIndex]) +
+			static_cast<int32>(InStorageType);
+	};
+
+	ForEachTransformElementStorage([&TransformIndexToAllocatedIndex, GetElementIndexForSort](
 		FRigTransformElement* InTransformElement,
 		ERigTransformType::Type InTransformType,
 		ERigTransformStorageType::Type InStorageType,
 		FRigComputedTransform* InComputedTransform,
 		FRigTransformDirtyState* InDirtyState)
 	{
-		TransformInfos.Emplace(
-			FTypeInfo(InTransformElement->GetIndex(), InTransformElement->GetType(), InTransformType, InStorageType),
-			InComputedTransform,
-			InDirtyState);
-	});
-
-	uint32 NumElements = (uint32)Elements.Num();
-	Algo::SortBy(TransformInfos, [NumElements](const FTransformInfo& InInfo) -> uint32
-	{
-		return InInfo.TypeInfo.GetSortKey(NumElements);
-	});
-
-	bool bHasMatchingRanges = true;
-	TMap<ERigTransformType::Type, TTuple<int32, int32>> NewRanges;
-	
-	// check if the sorting is at all required
-	bool bRequiresSort = false;
-	for(int32 Index = 0; Index < TransformInfos.Num(); Index++)
-	{
-		const int32 CurrentTransformIndex = TransformInfos[Index].ComputedTransform->StorageIndex;
-		if(CurrentTransformIndex != INDEX_NONE && CurrentTransformIndex != Index)
+		const int32 Index = GetElementIndexForSort(InTransformElement, InTransformType, InStorageType);
+		if(InComputedTransform->GetStorageIndex() == INDEX_NONE)
 		{
-			bRequiresSort = true;
+			TransformIndexToAllocatedIndex[Index] = INDEX_NONE;
 		}
-
-		const int32 CurrentDirtyStateIndex = TransformInfos[Index].DirtyState->StorageIndex;
-		if(CurrentDirtyStateIndex != INDEX_NONE && CurrentDirtyStateIndex != Index)
-		{
-			bRequiresSort = true;
-		}
-
-		if(CurrentTransformIndex == INDEX_NONE && CurrentDirtyStateIndex == INDEX_NONE)
+		else
 		{
-			// skip over entries with external storage
-			continue;
+			TransformIndexToAllocatedIndex[Index] = 0;
 		}
+	});
 
-		if(CurrentTransformIndex != CurrentDirtyStateIndex)
-		{
-			bHasMatchingRanges = false;
-		}
-		if(bHasMatchingRanges)
+	// accumulate the allocated indices
+	int32 NumAllocatedTransforms = 0;
+	for(int32 Index = 0; Index < TransformIndexToAllocatedIndex.Num(); Index++)
+	{
+		if(TransformIndexToAllocatedIndex[Index] != INDEX_NONE)
 		{
-			TTuple<int32,int32>& Range = NewRanges.FindOrAdd(TransformInfos[Index].TypeInfo.Transform, {INT32_MAX, INDEX_NONE});
-			Range.Get<0>() = FMath::Min(CurrentTransformIndex, Range.Get<0>());
-			Range.Get<1>() = FMath::Max(CurrentTransformIndex, Range.Get<1>());
+			TransformIndexToAllocatedIndex[Index] = NumAllocatedTransforms++;
 		}
 	}
 
-	if(!bRequiresSort)
+	TArray<TTuple<int32, int32>> SortedRanges;
+	SortedRanges.Reserve(NumTransformTypes);
+	for(int32 TransformType = 0; TransformType < NumTransformTypes; TransformType++)
 	{
-		if(bHasMatchingRanges)
+		SortedRanges.Emplace(INT32_MAX, INDEX_NONE);
+	}
+
+	bool bRequiresSort = false;
+	ForEachTransformElementStorage([
+		this,
+		GetElementIndexForSort,
+		TransformIndexToAllocatedIndex,
+		&SortedRanges,
+		&bRequiresSort
+	](
+		const FRigTransformElement* InTransformElement,
+		ERigTransformType::Type InTransformType,
+		ERigTransformStorageType::Type InStorageType,
+		const FRigComputedTransform* InComputedTransform,
+		const FRigTransformDirtyState* InDirtyState
+	) {
+		
+		if(InComputedTransform->GetStorageIndex() == INDEX_NONE)
 		{
-			ElementTransformRanges = NewRanges;
+			return;
 		}
-		return false;
-	}
+		
+		const int32 Index = GetElementIndexForSort(InTransformElement, InTransformType, InStorageType);
+		const int32 AllocatedIndex = TransformIndexToAllocatedIndex[Index];
 
-	FRigReusableElementStorage<FTransform> NewElementTransforms;
-	NewElementTransforms.Reserve(ElementTransforms.Num() - ElementTransforms.GetFreeList().Num());
-	
-	FRigReusableElementStorage<bool> NewElementDirtyStates;
-	NewElementDirtyStates.Reserve(ElementDirtyStates.Num() - ElementDirtyStates.GetFreeList().Num());
+		TTuple<int32, int32>& Range = SortedRanges[static_cast<int32>(InTransformType)]; 
+		Range.Get<0>() = FMath::Min(Range.Get<0>(), AllocatedIndex); 
+		Range.Get<1>() = FMath::Max(Range.Get<1>(), AllocatedIndex); 
 
-	bHasMatchingRanges = true;
-	NewRanges.Reset();
-	for(int32 Index = 0; Index < TransformInfos.Num(); Index++)
-	{
-		const int32 CurrentTransformIndex = TransformInfos[Index].ComputedTransform->StorageIndex;
-		const int32 CurrentDirtyStateIndex = TransformInfos[Index].DirtyState->StorageIndex;
-
-		int32 NewTransformIndex = INDEX_NONE;
-		int32 NewDirtyStateIndex = INDEX_NONE;
+		bRequiresSort = bRequiresSort || (InComputedTransform->GetStorageIndex() != AllocatedIndex && InDirtyState->GetStorageIndex() != AllocatedIndex);
+	});
 
-		if(CurrentTransformIndex != INDEX_NONE)
-		{
-			NewTransformIndex = NewElementTransforms.Add(ElementTransforms[CurrentTransformIndex]);
-			TransformInfos[Index].ComputedTransform->StorageIndex = NewTransformIndex;
-		}
+	if(!bRequiresSort)
+	{
+		ElementTransformRanges = MoveTemp(SortedRanges);
+		return false;
+	}
 
-		if(CurrentDirtyStateIndex != INDEX_NONE)
-		{
-			NewDirtyStateIndex = NewElementDirtyStates.Add(ElementDirtyStates[CurrentDirtyStateIndex]);
-			TransformInfos[Index].DirtyState->StorageIndex = NewDirtyStateIndex;
-		}
+	FRigReusableElementStorage<FTransform> SortedElementTransforms;
+	FRigReusableElementStorage<bool> SortedElementDirtyStates;
+	SortedElementTransforms.AddUninitialized(NumAllocatedTransforms);
+	SortedElementDirtyStates.AddUninitialized(NumAllocatedTransforms);
 
-		if(CurrentTransformIndex == INDEX_NONE && CurrentDirtyStateIndex)
+	ForEachTransformElementStorage([
+		this,
+		&SortedElementTransforms,
+		&SortedElementDirtyStates,
+		GetElementIndexForSort,
+		TransformIndexToAllocatedIndex
+	](
+		const FRigTransformElement* InTransformElement,
+		ERigTransformType::Type InTransformType,
+		ERigTransformStorageType::Type InStorageType,
+		FRigComputedTransform* InOutComputedTransform,
+		FRigTransformDirtyState* InOutDirtyState
+	) {
+		if(InOutComputedTransform->GetStorageIndex() == INDEX_NONE)
 		{
-			// ignore entries with external storage
-			continue;
+			return;
 		}
 
-		if(NewTransformIndex == NewDirtyStateIndex)
-		{
-			TTuple<int32,int32>& Range = NewRanges.FindOrAdd(TransformInfos[Index].TypeInfo.Transform, {INT32_MAX, INDEX_NONE});
-			Range.Get<0>() = FMath::Min(NewTransformIndex, Range.Get<0>());
-			Range.Get<1>() = FMath::Max(NewTransformIndex, Range.Get<1>());
-		}
-		else
-		{
-			bHasMatchingRanges = false;
-		}
-	}
+		const int32 Index = GetElementIndexForSort(InTransformElement, InTransformType, InStorageType);
+		const int32 AllocatedIndex = TransformIndexToAllocatedIndex[Index];
 
-	ElementTransforms = NewElementTransforms;
-	ElementDirtyStates = NewElementDirtyStates;
-	
-	UpdateElementStorage();
+		SortedElementTransforms[AllocatedIndex] = InOutComputedTransform->Get();
+		SortedElementDirtyStates[AllocatedIndex] = InOutDirtyState->Get();
 
-	if(bHasMatchingRanges)
-	{
-		ElementTransformRanges = NewRanges;
-	}
-	else
-	{
-		ElementTransformRanges.Reset();
-	}
+		InOutComputedTransform->StorageIndex = AllocatedIndex;
+		InOutComputedTransform->Storage = &SortedElementTransforms[AllocatedIndex];
+		
+		InOutDirtyState->StorageIndex = AllocatedIndex;
+		InOutDirtyState->Storage = &SortedElementDirtyStates[AllocatedIndex];
+	});
 
+	ElementTransforms.FreeList.Empty();
+	ElementTransforms.Storage = MoveTemp(SortedElementTransforms.Storage);
+	ElementDirtyStates.FreeList.Empty();
+	ElementDirtyStates.Storage = MoveTemp(SortedElementDirtyStates.Storage);
+	ElementTransformRanges = MoveTemp(SortedRanges);
 	return true;
 }
 
@@ -7012,9 +7022,10 @@ TTuple<FRigComputedTransform*, FRigTransformDirtyState*> URigHierarchy::GetEleme
 
 TOptional<TTuple<int32, int32>> URigHierarchy::GetElementStorageRange(ERigTransformType::Type InTransformType) const
 {
-	if(const TTuple<int32,int32>* Range = ElementTransformRanges.Find(InTransformType))
+	const int32 Index = static_cast<int32>(InTransformType);
+	if(ElementTransformRanges.IsValidIndex(Index))
 	{
-		return *Range;
+		return ElementTransformRanges[Index];
 	}
 	return TOptional<TTuple<int32, int32>>();
 }
diff --git a/Engine/Plugins/Animation/ControlRig/Source/ControlRig/Private/Rigs/RigHierarchyCache.cpp b/Engine/Plugins/Animation/ControlRig/Source/ControlRig/Private/Rigs/RigHierarchyCache.cpp
--- a/Engine/Plugins/Animation/ControlRig/Source/ControlRig/Private/Rigs/RigHierarchyCache.cpp
+++ b/Engine/Plugins/Animation/ControlRig/Source/ControlRig/Private/Rigs/RigHierarchyCache.cpp
@@ -85,6 +85,25 @@ bool FCachedRigElement::IsIdentical(const FRigElementKey& InKey, const URigHiera
 	return InKey == Key && InHierarchy->GetTopologyVersionHash() == ContainerVersion;
 }
 
+void FCachedRigElement::Set(const FRigBaseElement* InElement, int32 InTopologyHashVersion)
+{
+	check(InElement);
+	Element = InElement;
+	Key = Element->GetKey();
+	Index = Element->GetIndex();
+	if(InTopologyHashVersion == INDEX_NONE)
+	{
+		if(const URigHierarchy* Hierarchy = Element->GetOwner())
+		{
+			ContainerVersion = Hierarchy->GetTopologyVersionHash();
+		}
+	}
+	else
+	{
+		ContainerVersion = InTopologyHashVersion;
+	}
+}
+
 ////////////////////////////////////////////////////////////////////////////////
 // FCachedRigComponent
 ////////////////////////////////////////////////////////////////////////////////
diff --git a/Engine/Plugins/Animation/ControlRig/Source/ControlRig/Public/Rigs/RigHierarchy.h b/Engine/Plugins/Animation/ControlRig/Source/ControlRig/Public/Rigs/RigHierarchy.h
--- a/Engine/Plugins/Animation/ControlRig/Source/ControlRig/Public/Rigs/RigHierarchy.h
+++ b/Engine/Plugins/Animation/ControlRig/Source/ControlRig/Public/Rigs/RigHierarchy.h
@@ -4355,7 +4355,7 @@ private:
 	/*
 	* Helper function to destroy an element
 	*/
-	void DestroyElement(FRigBaseElement*& InElement);
+	void DestroyElement(FRigBaseElement*& InElement, bool bDestroyComponents = true, bool bDestroyElementStorage = true, bool bDestroyMetadata = true);
 
 	/*
 	 * Templated helper function to create an element
@@ -4520,7 +4520,7 @@ private:
 	FRigReusableElementStorage<float> ElementCurves;
 
 	// A list of ranges which can be used to copy all poses from initial to current, for example 
-	TMap<ERigTransformType::Type, TTuple<int32, int32>> ElementTransformRanges;
+	TArray<TTuple<int32, int32>> ElementTransformRanges;
 
 	// Allocates the default element storage for an element
 	void AllocateDefaultElementStorage(FRigBaseElement* InElement, bool bUpdateAllElements);
diff --git a/Engine/Plugins/Animation/ControlRig/Source/ControlRig/Public/Rigs/RigHierarchyCache.h b/Engine/Plugins/Animation/ControlRig/Source/ControlRig/Public/Rigs/RigHierarchyCache.h
--- a/Engine/Plugins/Animation/ControlRig/Source/ControlRig/Public/Rigs/RigHierarchyCache.h
+++ b/Engine/Plugins/Animation/ControlRig/Source/ControlRig/Public/Rigs/RigHierarchyCache.h
@@ -164,6 +164,8 @@ public:
 
 private:
 
+	void Set(const FRigBaseElement* InElement, int32 InTopologyHashVersion = INDEX_NONE);
+
 	UPROPERTY()
 	FRigElementKey Key;
 
@@ -175,6 +177,7 @@ private:
 
 	const FRigBaseElement* Element;
 	friend struct FCachedRigComponent;
+	friend class URigHierarchy;
 };
 
 USTRUCT(BlueprintType)
diff --git a/Engine/Plugins/Animation/ControlRig/Source/ControlRig/Public/Rigs/RigReusableElementStorage.h b/Engine/Plugins/Animation/ControlRig/Source/ControlRig/Public/Rigs/RigReusableElementStorage.h
--- a/Engine/Plugins/Animation/ControlRig/Source/ControlRig/Public/Rigs/RigReusableElementStorage.h
+++ b/Engine/Plugins/Animation/ControlRig/Source/ControlRig/Public/Rigs/RigReusableElementStorage.h
@@ -116,6 +116,18 @@ public:
 		return Indices[0];
 	}
 
+	int32 AddZeroed(int32 InNum)
+	{
+		check(FreeList.IsEmpty());
+		return Storage.AddZeroed(InNum);
+	}
+
+	int32 AddUninitialized(int32 InNum)
+	{
+		check(FreeList.IsEmpty());
+		return Storage.AddUninitialized(InNum);
+	}
+
 	TArray<int32, TInlineAllocator<4>> Allocate(int32 InCount, const T& InDefault, bool bContiguous = false)
 	{
 		TArray<int32, TInlineAllocator<4>> Indices;
@@ -323,6 +335,7 @@ public:
 		
 		if(!FreeList.IsEmpty())
 		{
+			
 			TArray<bool> ToRemove;
 
 			if(FreeList.Num() != Storage.Num() || OnDestroyCallback != nullptr)
@@ -340,16 +353,23 @@ public:
 
 			if(FreeList.Num() != Storage.Num())
 			{
+				const int32 NewNum = Storage.Num() - FreeList.Num();
+				OldToNew.Reserve(NewNum);
+				
 				TArray<T> NewStorage;
-				NewStorage.Reserve(FMath::Max(Storage.Num() - FreeList.Num(), 0));
+				NewStorage.Reserve(FMath::Max(NewNum, 0));
 				for(int32 OldIndex = 0; OldIndex < Storage.Num(); OldIndex++)
 				{
 					if(!ToRemove[OldIndex])
 					{
-						OldToNew.Add(OldIndex, NewStorage.Add(Storage[OldIndex]));
+						const int32 NewIndex = NewStorage.Add(Storage[OldIndex]);
+						if(OldIndex != NewIndex)
+						{
+							OldToNew.Add(OldIndex, NewIndex);
+						}
 					}
 				}
-				Storage = NewStorage;
+				Storage = MoveTemp(NewStorage);
 			}
 			else
 			{
@@ -364,4 +384,6 @@ public:
 
 		return OldToNew;
 	}
+
+	friend class URigHierarchy;
 };
