Control Rig: Replace old tweening slider UI in Control Rig with new tweening UI used by CurveEditor.
- FControlRigTweenModel and FTweenLogic now contain all tweening logic for ControlRig
- Deprecates all symbols in ControlRigTweener.h. These symbols are now superceded the more generic tweening system in TweeningUtilsEditor.
- Tweening widget no longer manages tweening logic. As result simplified several code paths in the edit mode and toolkit since they no longer need to need to inform the tweening widget.

Tasks left to do and that will be done in a separate changes (these have been marked with TODO UE-xxxxx in code comments):
- UE-234007: Synchronize blend function on ControlRig and CurveEditor
- UE-240484: Move preferred Tween function to TweeningUtilsEditor
- UE-240486 : Remove forced TSharedPtr use in FFunctionBasedTweenModel and ITweenFunctionContainer

diff --git a/Engine/Plugins/Animation/ControlRig/Source/ControlRigEditor/ControlRigEditor.Build.cs b/Engine/Plugins/Animation/ControlRig/Source/ControlRigEditor/ControlRigEditor.Build.cs
--- a/Engine/Plugins/Animation/ControlRig/Source/ControlRigEditor/ControlRigEditor.Build.cs
+++ b/Engine/Plugins/Animation/ControlRig/Source/ControlRigEditor/ControlRigEditor.Build.cs
@@ -13,7 +13,8 @@ namespace UnrealBuildTool.Rules
 					"AppFramework",
                     "RigVMEditor",
 					"CurveEditor",
-                    "InteractiveToolsFramework"
+                    "InteractiveToolsFramework", 
+                    "TweeningUtilsEditor"
 				}
 			);
 
@@ -84,8 +85,7 @@ namespace UnrealBuildTool.Rules
                     "ActorPickerMode",
                     "Constraints",
                     "AnimationEditMode",
-					"SequencerWidgets",
-					"TweeningUtilsEditor"
+					"SequencerWidgets"
 				}
             );
 
diff --git a/Engine/Plugins/Animation/ControlRig/Source/ControlRigEditor/Private/ControlRigSequencerEditorLibrary.cpp b/Engine/Plugins/Animation/ControlRig/Source/ControlRigEditor/Private/ControlRigSequencerEditorLibrary.cpp
--- a/Engine/Plugins/Animation/ControlRig/Source/ControlRigEditor/Private/ControlRigSequencerEditorLibrary.cpp
+++ b/Engine/Plugins/Animation/ControlRig/Source/ControlRigEditor/Private/ControlRigSequencerEditorLibrary.cpp
@@ -26,7 +26,6 @@
 #include "Engine/SCS_Node.h"
 #include "ILevelSequenceEditorToolkit.h"
 #include "ISequencer.h"
-#include "Tools/ControlRigTweener.h"
 #include "LevelSequencePlayer.h"
 #include "LevelSequenceActor.h"
 #include "Logging/LogMacros.h"
@@ -50,6 +49,9 @@
 #include "BakingAnimationKeySettings.h"
 #include "EditMode/ControlRigEditMode.h"
 #include "SequencerTools.h"
+#include "EditMode/Tween/ControlRigTweenModel.h"
+#include "Math/Abstraction/BuiltInTweenFunctionContainer.h"
+#include "Math/Abstraction/KeyBlendingAbstraction.h"
 #include "Sequencer/AnimLayers/AnimLayers.h"
 
 
@@ -400,23 +402,48 @@ TArray<UMovieSceneTrack*> UControlRigSequencerEditorLibrary::FindOrCreateControl
 
 bool UControlRigSequencerEditorLibrary::TweenControlRig(ULevelSequence* LevelSequence, UControlRig* ControlRig, float TweenValue)
 {
+	using namespace UE::TweeningUtilsEditor;
+	using namespace UE::ControlRigEditor;
+	
 	TWeakPtr<ISequencer> WeakSequencer = GetSequencerFromAsset();
 	if (WeakSequencer.IsValid() && WeakSequencer.Pin()->GetFocusedMovieSceneSequence() == LevelSequence
 		&& ControlRig && LevelSequence->GetMovieScene())
 	{
-		FControlsToTween ControlsToTween;
 		LevelSequence->GetMovieScene()->Modify();
-		TArray<UControlRig*> SelectedControlRigs;
-		SelectedControlRigs.Add(ControlRig);
-		ControlsToTween.Setup(SelectedControlRigs, WeakSequencer);
-		ControlsToTween.Blend(WeakSequencer, TweenValue);
+		
+		// TODO UE-240486: This has unnecessary dynamic memory allocation overhead until UE-240486 is addressed.
+		const FBuiltInTweenFunctionContainer TweenFunctionContainer; 
+		FControlRigTweenModel BlendTool(WeakSequencer, TweenFunctionContainer.GetBlendFunction(EBlendFunction::ControlsToTween).ToSharedRef());
+		BlendTool.StartBlendOperationWithControlRigs({ ControlRig });
+		BlendTool.BlendValues(TweenValue);
+		BlendTool.StopBlendOperation();
 		return true;
 	}
 	return false;
 }
 
+namespace UE::ControlRigEditor::SequencerLibraryDetail
+{
+	static TweeningUtilsEditor::EBlendFunction ConvertBlendFunction(EAnimToolBlendOperation Operation)
+	{
+		switch (Operation)
+		{
+		case EAnimToolBlendOperation::Tween: return TweeningUtilsEditor::EBlendFunction::ControlsToTween;
+		case EAnimToolBlendOperation::BlendToNeighbor: return TweeningUtilsEditor::EBlendFunction::BlendNeighbor;
+		case EAnimToolBlendOperation::PushPull: return TweeningUtilsEditor::EBlendFunction::PushPull;
+		case EAnimToolBlendOperation::BlendRelative: return TweeningUtilsEditor::EBlendFunction::BlendRelative;
+		case EAnimToolBlendOperation::BlendToEase: return TweeningUtilsEditor::EBlendFunction::BlendEase;
+		case EAnimToolBlendOperation::SmoothRough: return TweeningUtilsEditor::EBlendFunction::SmoothRough;
+		default: checkNoEntry(); return TweeningUtilsEditor::EBlendFunction::ControlsToTween;;
+		}
+	}
+}
+
 bool UControlRigSequencerEditorLibrary::BlendValuesOnSelected(ULevelSequence* LevelSequence, EAnimToolBlendOperation BlendOperation, float BlendValue)
 {
+	using namespace UE::TweeningUtilsEditor;
+	using namespace UE::ControlRigEditor;
+	
 	TWeakPtr<ISequencer> WeakSequencer = GetSequencerFromAsset();
 	if (WeakSequencer.IsValid() && WeakSequencer.Pin()->GetFocusedMovieSceneSequence() == LevelSequence
 		&& LevelSequence->GetMovieScene())
@@ -424,54 +451,15 @@ bool UControlRigSequencerEditorLibrary::BlendValuesOnSelected(ULevelSequence* Le
 		if (FControlRigEditMode* EditMode = static_cast<FControlRigEditMode*>(GLevelEditorModeTools().GetActiveMode(FControlRigEditMode::ModeName)))
 		{
 			TWeakPtr<FControlRigEditMode> WeakMode = StaticCastSharedRef<FControlRigEditMode, FEdMode>(EditMode->AsShared()).ToWeakPtr();
-			
-			FControlsToTween ControlsToTween;
 			LevelSequence->GetMovieScene()->Modify();
-			switch(BlendOperation)
-			{ 
-			case EAnimToolBlendOperation::Tween:
-				{
-					FControlsToTween BlendTool;
-					BlendTool.Setup(WeakSequencer, WeakMode);
-					BlendTool.Blend(WeakSequencer, BlendValue);
-					return true;
-				}
-			case EAnimToolBlendOperation::BlendToNeighbor:
-				{
-					FBlendNeighborSlider BlendTool;
-					BlendTool.Setup(WeakSequencer, WeakMode);
-					BlendTool.Blend(WeakSequencer, BlendValue);
-					return true;
-				}
-			case EAnimToolBlendOperation::PushPull:
-				{
-					FPushPullSlider BlendTool;
-					BlendTool.Setup(WeakSequencer, WeakMode);
-					BlendTool.Blend(WeakSequencer, BlendValue);
-					return true;
-				}
-			case EAnimToolBlendOperation::BlendRelative:
-				{
-					FBlendRelativeSlider BlendTool;
-					BlendTool.Setup(WeakSequencer, WeakMode);
-					BlendTool.Blend(WeakSequencer, BlendValue);
-					return true;
-				}
-			case EAnimToolBlendOperation::BlendToEase:
-				{
-					FBlendToEaseSlider BlendTool;
-					BlendTool.Setup(WeakSequencer, WeakMode);
-					BlendTool.Blend(WeakSequencer, BlendValue);
-					return true;
-				}
-			case EAnimToolBlendOperation::SmoothRough:
-				{
-					FSmoothRoughSlider BlendTool;
-					BlendTool.Setup(WeakSequencer, WeakMode);
-					BlendTool.Blend(WeakSequencer, BlendValue);
-					return true;
-				}
-			}
+			
+			// TODO UE-240486: This has unnecessary dynamic memory allocation overhead until UE-240486 is addressed.
+			const FBuiltInTweenFunctionContainer TweenFunctionContainer;
+			const EBlendFunction BlendFunction = SequencerLibraryDetail::ConvertBlendFunction(BlendOperation);
+			FControlRigTweenModel BlendTool(WeakSequencer, WeakMode,
+				TweenFunctionContainer.GetBlendFunction(BlendFunction).ToSharedRef()
+				);
+			return BlendTool.BlendSingle(BlendValue);
 		}
 	}
 	return false;
diff --git a/Engine/Plugins/Animation/ControlRig/Source/ControlRigEditor/Private/EditMode/ControlRigEditMode.cpp b/Engine/Plugins/Animation/ControlRig/Source/ControlRigEditor/Private/EditMode/ControlRigEditMode.cpp
--- a/Engine/Plugins/Animation/ControlRig/Source/ControlRigEditor/Private/EditMode/ControlRigEditMode.cpp
+++ b/Engine/Plugins/Animation/ControlRig/Source/ControlRigEditor/Private/EditMode/ControlRigEditMode.cpp
@@ -1066,69 +1066,17 @@ bool FControlRigEditMode::InputKey(FEditorViewportClient* InViewportClient, FVie
 		{
 			return true;
 		}
-		if (IsDragAnimSliderToolPressed(InViewport)) //this is needed to make sure we get all of the processed mouse events, for some reason the above may not return true
-		{
-			return true;
-		}
 	}
 
 	return FEdMode::InputKey(InViewportClient, InViewport, InKey, InEvent);
 }
 
-bool FControlRigEditMode::ProcessCapturedMouseMoves(FEditorViewportClient* InViewportClient, FViewport* InViewport, const TArrayView<FIntPoint>& CapturedMouseMoves)
-{
-	const bool bChange = IsDragAnimSliderToolPressed(InViewport);
-	if (CapturedMouseMoves.Num() > 0)
-	{
-		if (bChange)
-		{
-			for (int32 Index = 0; Index < CapturedMouseMoves.Num(); ++Index)
-			{
-				int32 X = CapturedMouseMoves[Index].X;
-				if (StartXValue.IsSet() == false)
-				{
-					StartAnimSliderTool(X);
-				}
-				else
-				{
-					int32 Diff = X - StartXValue.GetValue();
-					if (Diff != 0)
-					{
-						FIntPoint Origin, Size;
-						InViewportClient->GetViewportDimensions(Origin, Size);
-						const double ViewPortSize = (double)Size.X * 0.125; // 1/8 screen drag should do full blend, todo perhaps expose this as a sensitivity
-						const double ViewDiff = (double)(Diff) / ViewPortSize;
-						DragAnimSliderTool(ViewDiff);
-					}
-				}
-			}
-		}
-		else if (bisTrackingAnimToolDrag && bChange == false)
-		{
-			ResetAnimSlider();
-		}
-		return bisTrackingAnimToolDrag;
-	}
-	else if(bisTrackingAnimToolDrag && bChange == false)
-	{
-		ResetAnimSlider();
-	}
-	return false;
-}
-
 bool FControlRigEditMode::StartTracking(FEditorViewportClient* InViewportClient, FViewport* InViewport)
 {
 	if(RuntimeControlRigs.IsEmpty())
 	{
 		return false;
 	}
-	
-	//break right out if doing anim slider scrub
-
-	if (IsDragAnimSliderToolPressed(InViewport))
-	{
-		return true;
-	}
 
 	if (IsMovingCamera(InViewport))
 	{
@@ -1146,20 +1094,11 @@ bool FControlRigEditMode::StartTracking(FEditorViewportClient* InViewportClient,
 
 bool FControlRigEditMode::EndTracking(FEditorViewportClient* InViewportClient, FViewport* InViewport)
 {
-	if(RuntimeControlRigs.IsEmpty())
+	if (RuntimeControlRigs.IsEmpty())
 	{
 		return false;
 	}
 
-	if (bisTrackingAnimToolDrag)
-	{
-		ResetAnimSlider();
-	}
-	if (IsDragAnimSliderToolPressed(InViewport))
-	{
-		return true;
-	}
-
 	if (IsMovingCamera(InViewport))
 	{
 		return true;
@@ -2057,12 +1996,6 @@ bool FControlRigEditMode::IsDoingDrag(const FViewport* InViewport) const
 
 bool FControlRigEditMode::InputDelta(FEditorViewportClient* InViewportClient, FViewport* InViewport, FVector& InDrag, FRotator& InRot, FVector& InScale)
 {
-	if (IsDragAnimSliderToolPressed(InViewport)) //this is needed to make sure we get all of the processed mouse events, for some reason the above may not return true
-	{
-		//handled by processed mouse clicks
-		return true;
-	}
-
 	if (IsDoingDrag(InViewport))
 	{
 		return DragToolHandler.InputDelta(InViewportClient, InViewport, InDrag, InRot, InScale);
@@ -2952,12 +2885,6 @@ void FControlRigEditMode::BindCommands()
 		Commands.ClearSelection,
 		FExecuteAction::CreateRaw(this, &FControlRigEditMode::ClearSelection));
 
-	CommandBindings->MapAction(
-		Commands.ChangeAnimSliderTool,
-		FExecuteAction::CreateRaw(this, &FControlRigEditMode::ChangeAnimSliderTool),
-		FCanExecuteAction::CreateRaw(this, &FControlRigEditMode::CanChangeAnimSliderTool)
-	);
-
 	CommandBindings->MapAction(
 		Commands.FrameSelection,
 		FExecuteAction::CreateRaw(this, &FControlRigEditMode::FrameSelection),
@@ -3224,76 +3151,6 @@ void FControlRigEditMode::GetAllSelectedControls(TMap<UControlRig*, TArray<FRigE
 	}
 }
 
-/** If Anim Slider is open, got to the next tool*/
-void FControlRigEditMode::ChangeAnimSliderTool()
-{
-	if (Toolkit.IsValid())
-	{
-		StaticCastSharedPtr<FControlRigEditModeToolkit>(Toolkit)->GetToNextActiveSlider();
-	}
-}
-
-/** If Anim Slider is open, then can drag*/
-bool FControlRigEditMode::CanChangeAnimSliderTool() const
-{
-	if (Toolkit.IsValid())
-	{
-		return (StaticCastSharedPtr<FControlRigEditModeToolkit>(Toolkit)->CanChangeAnimSliderTool());
-	}
-	return false;
-}
-
-/** If Anim Slider is open, drag the tool*/
-void FControlRigEditMode::DragAnimSliderTool(double IncrementVal)
-{
-	if (Toolkit.IsValid())
-	{
-		StaticCastSharedPtr<FControlRigEditModeToolkit>(Toolkit)->DragAnimSliderTool(IncrementVal);
-	}
-}
-
-void FControlRigEditMode::StartAnimSliderTool(int32 InX)
-{
-	if (Toolkit.IsValid())
-	{
-		bisTrackingAnimToolDrag = true;
-		GEditor->BeginTransaction(LOCTEXT("AnimSliderBlend", "AnimSlider Blend"));
-		StartXValue = InX;
-		StaticCastSharedPtr<FControlRigEditModeToolkit>(Toolkit)->StartAnimSliderTool();
-	}
-}
-
-/** Reset and stop user the anim slider tool*/
-void FControlRigEditMode::ResetAnimSlider()
-{
-	if (Toolkit.IsValid())
-	{
-		GEditor->EndTransaction();
-		bisTrackingAnimToolDrag = false;
-		StartXValue.Reset();
-		StaticCastSharedPtr<FControlRigEditModeToolkit>(Toolkit)->ResetAnimSlider();
-	}
-}
-
-/** If the Drag Anim Slider Tool is pressed*/
-bool FControlRigEditMode::IsDragAnimSliderToolPressed(FViewport* InViewport)
-{
-	if (IsInLevelEditor() && Toolkit.IsValid() &&  CanChangeAnimSliderTool())
-	{
-		bool bIsMovingSlider = false;
-		const FControlRigEditModeCommands& Commands = FControlRigEditModeCommands::Get();
-		// Need to iterate through primary and secondary to make sure they are all pressed.
-		for (uint32 i = 0; i < static_cast<uint8>(EMultipleKeyBindingIndex::NumChords); ++i)
-		{
-			EMultipleKeyBindingIndex ChordIndex = static_cast<EMultipleKeyBindingIndex>(i);
-			const FInputChord& Chord = *Commands.DragAnimSliderTool->GetActiveChord(ChordIndex);
-			bIsMovingSlider |= Chord.IsValidChord() && InViewport->KeyState(Chord.Key);
-		}
-		return bIsMovingSlider;
-	}
-	return false;
-}
-
 void FControlRigEditMode::SetAnimLayerPassthroughKey()
 {
 	ISequencer* Sequencer = WeakSequencer.Pin().Get();
diff --git a/Engine/Plugins/Animation/ControlRig/Source/ControlRigEditor/Private/EditMode/ControlRigEditModeCommands.cpp b/Engine/Plugins/Animation/ControlRig/Source/ControlRigEditor/Private/EditMode/ControlRigEditModeCommands.cpp
--- a/Engine/Plugins/Animation/ControlRig/Source/ControlRigEditor/Private/EditMode/ControlRigEditModeCommands.cpp
+++ b/Engine/Plugins/Animation/ControlRig/Source/ControlRigEditor/Private/EditMode/ControlRigEditModeCommands.cpp
@@ -19,9 +19,6 @@ void FControlRigEditModeCommands::RegisterCommands()
 	UI_COMMAND(DecreaseControlShapeSize, "Decrease Shape Size", "Decrease Size of the Shapes In The Viewport", EUserInterfaceActionType::Button, FInputChord(EKeys::Hyphen, EModifierKey::Shift));
 	UI_COMMAND(ResetControlShapeSize, "Reset Shape Size", "Resize Shape Size To Default", EUserInterfaceActionType::Button, FInputChord(EKeys::Equals));
 
-	UI_COMMAND(DragAnimSliderTool, "Drag Anim Slider Tool", "Drag existing anim slider", EUserInterfaceActionType::Button, FInputChord(EKeys::U));
-	UI_COMMAND(ChangeAnimSliderTool, "Change Anim Slider Tool", "Go to the next anim slider", EUserInterfaceActionType::Button, FInputChord(EKeys::U, EModifierKey::Shift));
-
 	UI_COMMAND(SetAnimLayerPassthroughKey,"Set Anim Layer Passthrough Key","Set previous layer value on current selected anim layer", EUserInterfaceActionType::Button, FInputChord());
 
 	UI_COMMAND(SelectMirroredControls, "Select Mirrored Controls", " Select mirrored Controls on current selection", EUserInterfaceActionType::Button, FInputChord(EKeys::M, EModifierKey::Control));
diff --git a/Engine/Plugins/Animation/ControlRig/Source/ControlRigEditor/Private/EditMode/ControlRigEditModeCommands.h b/Engine/Plugins/Animation/ControlRig/Source/ControlRigEditor/Private/EditMode/ControlRigEditModeCommands.h
--- a/Engine/Plugins/Animation/ControlRig/Source/ControlRigEditor/Private/EditMode/ControlRigEditModeCommands.h
+++ b/Engine/Plugins/Animation/ControlRig/Source/ControlRigEditor/Private/EditMode/ControlRigEditModeCommands.h
@@ -26,12 +26,6 @@ public:
 	/** Toggles hiding all manipulators on all control rigs in the viewport */
 	TSharedPtr< FUICommandInfo > ToggleAllManipulators;
 
-	/** Drag Anim Slider Tool*/
-	TSharedPtr< FUICommandInfo > DragAnimSliderTool;
-
-	/**  Change current Anim Slider Tool */
-	TSharedPtr< FUICommandInfo > ChangeAnimSliderTool;
-
 	/** Invert Transforms to Rest Pose */
 	TSharedPtr< FUICommandInfo > InvertTransforms;
 
diff --git a/Engine/Plugins/Animation/ControlRig/Source/ControlRigEditor/Private/EditMode/ControlRigEditModeToolkit.cpp b/Engine/Plugins/Animation/ControlRig/Source/ControlRigEditor/Private/EditMode/ControlRigEditModeToolkit.cpp
--- a/Engine/Plugins/Animation/ControlRig/Source/ControlRigEditor/Private/EditMode/ControlRigEditModeToolkit.cpp
+++ b/Engine/Plugins/Animation/ControlRig/Source/ControlRigEditor/Private/EditMode/ControlRigEditModeToolkit.cpp
@@ -4,8 +4,8 @@
 * Control Rig Edit Mode Toolkit
 */
 #include "EditMode/ControlRigEditModeToolkit.h"
+
 #include "Widgets/DeclarativeSyntaxSupport.h"
-#include "Widgets/SWidget.h"
 #include "EditorModes.h"
 #include "Toolkits/BaseToolkit.h"
 #include "EditorModeManager.h"
@@ -13,7 +13,6 @@
 #include "EditMode/ControlRigEditMode.h"
 #include "Modules/ModuleManager.h"
 #include "EditMode/SControlRigBaseListWidget.h"
-#include "EditMode/SControlRigTweenWidget.h"
 #include "EditMode/SControlRigSnapper.h"
 #include "Tools/SMotionTrailOptions.h"
 #include "Toolkits/AssetEditorModeUILayer.h"
@@ -26,6 +25,7 @@
 #include "SLevelViewport.h"
 #include "Sequencer/AnimLayers/SAnimLayers.h"
 #include "Sequencer/AnimLayers/AnimLayers.h"
+#include "Tween/SControlRigTweenWidget.h"
 
 #define LOCTEXT_NAMESPACE "FControlRigEditModeToolkit"
 
@@ -260,7 +260,7 @@ void FControlRigEditModeToolkit::CreateAndShowTweenOverlay()
 		.HAlign(HAlign_Left)
 		.Padding(TAttribute<FMargin>(this, &FControlRigEditModeToolkit::GetTweenWidgetPadding))
 		[
-			SAssignNew(TweenWidget, SControlRigTweenWidget)
+			SAssignNew(TweenWidget, UE::ControlRigEditor::SControlRigTweenWidget)
 			.InOwningToolkit(SharedThis(this))
 			.InOwningEditMode(SharedThis(&EditMode))
 		];
@@ -268,40 +268,6 @@ void FControlRigEditModeToolkit::CreateAndShowTweenOverlay()
 	TryShowTweenOverlay();
 }
 
-void FControlRigEditModeToolkit::GetToNextActiveSlider()
-{
-	if (TweenWidgetParent && TweenWidget)
-	{
-		TweenWidget->GetToNextActiveSlider();
-	}
-}
-bool FControlRigEditModeToolkit::CanChangeAnimSliderTool() const
-{
-	return (TweenWidgetParent && TweenWidget);
-}
-
-void FControlRigEditModeToolkit::DragAnimSliderTool(double Val)
-{
-	if (TweenWidgetParent && TweenWidget)
-	{
-		TweenWidget->DragAnimSliderTool(Val);
-	}
-}
-void FControlRigEditModeToolkit::ResetAnimSlider()
-{
-	if (TweenWidgetParent && TweenWidget)
-	{
-		TweenWidget->ResetAnimSlider();
-	}
-}
-
-void FControlRigEditModeToolkit::StartAnimSliderTool()
-{
-	if (TweenWidgetParent && TweenWidget)
-	{
-		TweenWidget->StartAnimSliderTool();
-	}
-}
 void FControlRigEditModeToolkit::TryShowTweenOverlay()
 {
 	if (TweenWidgetParent)
diff --git a/Engine/Plugins/Animation/ControlRig/Source/ControlRigEditor/Private/EditMode/ControlRigEditModeToolkit.h b/Engine/Plugins/Animation/ControlRig/Source/ControlRigEditor/Private/EditMode/ControlRigEditModeToolkit.h
--- a/Engine/Plugins/Animation/ControlRig/Source/ControlRigEditor/Private/EditMode/ControlRigEditModeToolkit.h
+++ b/Engine/Plugins/Animation/ControlRig/Source/ControlRigEditor/Private/EditMode/ControlRigEditModeToolkit.h
@@ -10,15 +10,16 @@
 #include "EditorModeManager.h"
 #include "EditMode/SControlRigEditModeTools.h"
 #include "EditMode/ControlRigEditMode.h"
+#include "Tween/ControlRigTweenModel.h"
 
-class SControlRigTweenWidget;
 class SControlRigDetails;
 class SControlRigOutliner;
+namespace UE::ControlRigEditor { class SControlRigTweenWidget; }
 
 class FControlRigEditModeToolkit : public FModeToolkit
 {
 public:
-	friend class SControlRigTweenWidget;
+	friend UE::ControlRigEditor::SControlRigTweenWidget;
 
 	FControlRigEditModeToolkit(FControlRigEditMode& InEditMode)
 		: EditMode(InEditMode)
@@ -56,13 +57,8 @@ public:
 	void TryInvokeToolkitUI(const FName InName);
 	bool IsToolkitUIActive(const FName InName) const;
 
-	/** For updating tween values via hot keys*/
-	void GetToNextActiveSlider();
-	bool CanChangeAnimSliderTool() const;
-	void DragAnimSliderTool(double Val);
-	void ResetAnimSlider();
-	void StartAnimSliderTool();
 public:
+	
 	static const FName PoseTabName;
 	static const FName MotionTrailTabName;
 	static const FName AnimLayerTabName;
@@ -100,13 +96,13 @@ protected:
 	static bool bSnapperTabOpen;
 	static bool bTweenOpen;
 private:
+	
 	/** The edit mode we are bound to */
 	FControlRigEditMode& EditMode;
 	TSharedPtr<SWidget> TweenWidgetParent;
-	TSharedPtr<SControlRigTweenWidget> TweenWidget;
+	TSharedPtr<UE::ControlRigEditor::SControlRigTweenWidget> TweenWidget;
 
 	FVector2D InViewportTweenWidgetLocation;
 	/** The tools widget */
 	TSharedPtr<SControlRigEditModeTools> ModeTools;
-
 };
diff --git a/Engine/Plugins/Animation/ControlRig/Source/ControlRigEditor/Private/EditMode/SControlRigTweenWidget.cpp b/Engine/Plugins/Animation/ControlRig/Source/ControlRigEditor/Private/EditMode/SControlRigTweenWidget.cpp
deleted file mode 100644
diff --git a/Engine/Plugins/Animation/ControlRig/Source/ControlRigEditor/Private/EditMode/SControlRigTweenWidget.h b/Engine/Plugins/Animation/ControlRig/Source/ControlRigEditor/Private/EditMode/SControlRigTweenWidget.h
deleted file mode 100644
diff --git a/Engine/Plugins/Animation/ControlRig/Source/ControlRigEditor/Private/EditMode/Tween/ControlRigObjectSelection.cpp b/Engine/Plugins/Animation/ControlRig/Source/ControlRigEditor/Private/EditMode/Tween/ControlRigObjectSelection.cpp
new file mode 100644
--- /dev/null
+++ b/Engine/Plugins/Animation/ControlRig/Source/ControlRigEditor/Private/EditMode/Tween/ControlRigObjectSelection.cpp
@@ -0,0 +1,301 @@
+﻿// Copyright Epic Games, Inc. All Rights Reserved.
+
+#include "ControlRigObjectSelection.h"
+
+#include "CurveDataAbstraction.h"
+#include "Channels/MovieSceneDoubleChannel.h"
+#include "Channels/MovieSceneFloatChannel.h"
+#include "EditMode/ControlRigEditMode.h"
+#include "ISequencer.h"
+#include "MovieScene.h"
+#include "Sequencer/MovieSceneControlRigParameterSection.h"
+#include "Sequencer/MovieSceneControlRigParameterTrack.h"
+
+namespace UE::ControlRigEditor
+{
+bool FControlRigObjectSelection::Setup(const TWeakPtr<ISequencer>& InSequencer, const TWeakPtr<FControlRigEditMode>& InEditMode)
+{
+	ChannelsArray.Reset();
+	const TArray<UControlRig*> ControlRigs = GetControlRigs(InEditMode);
+	return Setup(ControlRigs, InSequencer);
+}
+	
+bool FControlRigObjectSelection::Setup(const TArray<UControlRig*>& SelectedControlRigs, const TWeakPtr<ISequencer>& InSequencer)
+{
+	ChannelsArray.Reset();
+	if (InSequencer.IsValid() == false)
+	{
+		return false;
+	}
+	ISequencer* Sequencer = InSequencer.Pin().Get();
+
+	const FQualifiedFrameTime CurrentTime = Sequencer->GetLocalTime();
+	const UMovieScene* MovieScene = Sequencer->GetFocusedMovieSceneSequence()->GetMovieScene();
+
+	//get selected controls and objects
+	TArray<FGuid> SelectedGuids;
+	Sequencer->GetSelectedObjects(SelectedGuids);
+
+	TArray<UMovieSceneControlRigParameterSection*> HandledSections;
+	auto SetupControlRigTrackChannels = [this, &CurrentTime, &HandledSections](const UMovieSceneControlRigParameterTrack* Track)
+	{
+		check(Track);		
+		for (UMovieSceneSection* MovieSection : Track->GetAllSections())
+		{
+			UMovieSceneControlRigParameterSection* Section = Cast<UMovieSceneControlRigParameterSection>(MovieSection);
+			if (Section && Section->IsActive() && Section->GetRange().Contains(CurrentTime.Time.GetFrame()) && !HandledSections.Contains(Section))
+			{
+				HandledSections.Add(Section);				
+				UControlRig* ControlRig = Track->GetControlRig();
+				TArray<FRigControlElement*> CurrentControls;
+				ControlRig->GetControlsInOrder(CurrentControls);
+
+				FMovieSceneChannelProxy& ChannelProxy = Section->GetChannelProxy();
+				TArrayView<FMovieSceneFloatChannel*> Channels = ChannelProxy.GetChannels<FMovieSceneFloatChannel>();
+				//reuse these arrays
+				TArray<FFrameNumber> KeyTimes;
+				TArray<FKeyHandle> Handles;
+				for (const FRigControlElement* ControlElement : CurrentControls)
+				{
+					if (ControlRig->GetHierarchy()->IsAnimatable(ControlElement) &&  ControlRig->IsControlSelected(ControlElement->GetFName()))
+					{
+						FControlRigObjectSelection::FObjectChannels ObjectChannels;
+						ObjectChannels.Section = Section;
+						KeyTimes.SetNum(0);
+						Handles.SetNum(0);
+						const FChannelMapInfo* pChannelIndex = Section->ControlChannelMap.Find(ControlElement->GetFName());
+						if (pChannelIndex == nullptr)
+						{
+							continue;
+						}
+
+						const int32 NumChannels = [&ControlElement]()
+						{
+							switch (ControlElement->Settings.ControlType)
+							{
+								case ERigControlType::Float:
+								case ERigControlType::ScaleFloat:
+								{
+									return 1;
+								}
+								case ERigControlType::Vector2D:
+								{
+									return 2;
+								}
+								case ERigControlType::Position:
+								case ERigControlType::Scale:
+								case ERigControlType::Rotator:
+								{
+									return 3;
+								}
+								case ERigControlType::TransformNoScale:
+								{
+									return 6;
+								}
+								case ERigControlType::Transform:
+								case ERigControlType::EulerTransform:
+								{
+									return 9;
+								}
+							default:
+									return 0;
+							}
+						}();
+					
+						int32 BoundIndex = 0;
+						int32 NumValidChannels = 0;
+						ObjectChannels.KeyBounds.SetNum(NumChannels);
+						for (int32 ChannelIdx = pChannelIndex->ChannelIndex; ChannelIdx < (pChannelIndex->ChannelIndex + NumChannels); ++ChannelIdx)
+						{
+							FMovieSceneFloatChannel* Channel = Channels[ChannelIdx];
+							SetupChannel(CurrentTime.Time.GetFrame(), KeyTimes, Handles, Channel, nullptr,
+							ObjectChannels.KeyBounds[BoundIndex]);
+							if (ObjectChannels.KeyBounds[BoundIndex].bValid)
+							{
+								++NumValidChannels;
+							}
+							++BoundIndex;
+						}
+						if (NumValidChannels > 0)
+						{
+							ChannelsArray.Add(ObjectChannels);
+						}
+					}
+				}
+			}
+		}
+	};
+	
+	// Handle MovieScene bindings
+	const TArray<FMovieSceneBinding>& Bindings = MovieScene->GetBindings();
+	for (const FMovieSceneBinding& Binding : Bindings)
+	{
+		if (SelectedGuids.Num() > 0 && SelectedGuids.Contains(Binding.GetObjectGuid()))
+		{
+			const TArray<UMovieSceneTrack*>& Tracks = Binding.GetTracks();
+			for (const UMovieSceneTrack* Track : Tracks)
+			{
+				if (Track && Track->IsA<UMovieSceneControlRigParameterTrack>() == false)
+				{
+					const TArray<UMovieSceneSection*>& Sections = Track->GetAllSections();
+					for (UMovieSceneSection* Section : Sections)
+					{
+						if (Section && Section->IsActive())
+						{
+							FMovieSceneChannelProxy& ChannelProxy = Section->GetChannelProxy();
+							TArrayView<FMovieSceneFloatChannel*> FloatChannels = ChannelProxy.GetChannels<FMovieSceneFloatChannel>();
+							TArrayView<FMovieSceneDoubleChannel*> DoubleChannels = ChannelProxy.GetChannels<FMovieSceneDoubleChannel>();
+
+							//reuse these arrays
+							TArray<FFrameNumber> KeyTimes;
+							TArray<FKeyHandle> Handles;
+							FControlRigObjectSelection::FObjectChannels ObjectChannels;
+							ObjectChannels.Section = Section;
+							KeyTimes.SetNum(0);
+							Handles.SetNum(0);
+							const int32 NumFloatChannels = FloatChannels.Num();
+							const int32 NumDoubleChannels = DoubleChannels.Num();
+
+							ObjectChannels.KeyBounds.SetNum(NumFloatChannels + NumDoubleChannels);
+							int32 NumValidChannels = 0;
+							for (int32 ChannelIdx = 0;ChannelIdx < NumFloatChannels; ++ChannelIdx)
+							{
+								FMovieSceneFloatChannel* FloatChannel = FloatChannels[ChannelIdx];
+								SetupChannel(CurrentTime.Time.GetFrame(), KeyTimes, Handles, FloatChannel, nullptr,
+									ObjectChannels.KeyBounds[ChannelIdx]);
+								if (ObjectChannels.KeyBounds[ChannelIdx].bValid)
+								{
+									++NumValidChannels;
+								}
+							}
+							for (int32 ChannelIdx = 0; ChannelIdx < NumDoubleChannels; ++ChannelIdx)
+							{
+								FMovieSceneDoubleChannel* DoubleChannel = DoubleChannels[ChannelIdx];
+								SetupChannel(CurrentTime.Time.GetFrame(), KeyTimes, Handles, nullptr, DoubleChannel,
+									ObjectChannels.KeyBounds[ChannelIdx]);
+								if (ObjectChannels.KeyBounds[ChannelIdx].bValid)
+								{
+									++NumValidChannels;
+								}
+							}
+							if (NumValidChannels > 0)
+							{
+								ChannelsArray.Add(ObjectChannels);
+							}
+						}
+					}
+				}
+			}
+		}
+		const UMovieSceneControlRigParameterTrack* Track = Cast<UMovieSceneControlRigParameterTrack>(MovieScene->FindTrack(UMovieSceneControlRigParameterTrack::StaticClass(), Binding.GetObjectGuid(), NAME_None));
+		if (Track && Track->GetControlRig() && SelectedControlRigs.Contains(Track->GetControlRig()))
+		{
+			SetupControlRigTrackChannels(Track);
+		}
+	}
+	
+	// Handle movie tracks in general (for non-binding, USkeleton, ControlRig tracks)
+	for (const UMovieSceneTrack* MovieTrack : MovieScene->GetTracks())
+	{
+		if (const UMovieSceneControlRigParameterTrack* Track = Cast<UMovieSceneControlRigParameterTrack>(MovieTrack))
+		{
+			if (Track && Track->GetControlRig() && SelectedControlRigs.Contains(Track->GetControlRig()))
+			{
+				SetupControlRigTrackChannels(Track);
+			}
+		}
+	}
+	
+	return ChannelsArray.Num() > 0;
+}
+
+void FControlRigObjectSelection::SetupChannel(FFrameNumber CurrentFrame, TArray<FFrameNumber>& KeyTimes, TArray<FKeyHandle>& Handles, 
+	FMovieSceneFloatChannel* FloatChannel, FMovieSceneDoubleChannel* DoubleChannel,
+	FChannelKeyBounds& KeyBounds)
+{
+	KeyBounds.FloatChannel = FloatChannel;
+	KeyBounds.DoubleChannel = DoubleChannel;
+	KeyBounds.CurrentFrame = CurrentFrame;
+	KeyBounds.PreviousIndex = INDEX_NONE;
+	KeyBounds.NextIndex = INDEX_NONE;
+	KeyTimes.SetNum(0);
+	Handles.SetNum(0);
+	if (FloatChannel)
+	{
+		FloatChannel->GetKeys(TRange<FFrameNumber>(), &KeyTimes, &Handles);
+		float OutValue;
+		FloatChannel->Evaluate(FFrameTime(CurrentFrame), OutValue);
+		KeyBounds.CurrentValue = OutValue;
+	}
+	else if (DoubleChannel)
+	{
+		DoubleChannel->GetKeys(TRange<FFrameNumber>(), &KeyTimes, &Handles);
+		double OutValue;
+		DoubleChannel->Evaluate(FFrameTime(CurrentFrame), OutValue);
+		KeyBounds.CurrentValue = OutValue;
+	}
+
+	if (KeyTimes.Num() > 0)
+	{
+		TArrayView<const FMovieSceneFloatValue> FloatValues;
+		TArrayView<const FMovieSceneDoubleValue> DoubleValues;
+		if (FloatChannel)
+		{
+			FloatValues = FloatChannel->GetValues();
+		}
+		else if (DoubleChannel)
+		{
+			DoubleValues = DoubleChannel->GetValues();
+
+		}
+		for (int32 Index = 0; Index < KeyTimes.Num(); Index++)
+		{
+			const FFrameNumber FrameNumber = KeyTimes[Index];
+			if (FrameNumber < CurrentFrame || (FrameNumber == CurrentFrame && KeyBounds.PreviousIndex == INDEX_NONE))
+			{
+				KeyBounds.PreviousIndex = Index;
+				KeyBounds.PreviousFrame = FrameNumber;
+				if (FloatChannel)
+				{
+					KeyBounds.PreviousValue = FloatValues[Index].Value;
+				}
+				else if (DoubleChannel)
+				{
+					KeyBounds.PreviousValue = DoubleValues[Index].Value;
+				}
+			}
+			else if (FrameNumber > CurrentFrame || (FrameNumber == CurrentFrame && Index == KeyTimes.Num() -1))
+			{
+				KeyBounds.NextIndex = Index;
+				KeyBounds.NextFrame = FrameNumber;
+				if (FloatChannel)
+				{
+					KeyBounds.NextValue = FloatValues[Index].Value;
+				}
+				else if (DoubleChannel)
+				{
+					KeyBounds.NextValue = DoubleValues[Index].Value;
+				}
+				break;
+			}
+		}
+	}
+	KeyBounds.bValid = (KeyBounds.PreviousIndex != INDEX_NONE && KeyBounds.NextIndex != INDEX_NONE
+		&& KeyBounds.PreviousIndex != KeyBounds.NextIndex) ? true : false;
+}
+
+TArray<UControlRig*> FControlRigObjectSelection::GetControlRigs(const TWeakPtr<FControlRigEditMode>& InEditMode)
+{
+	TArray<UControlRig*> ControlRigs;
+	if (const FControlRigEditMode* EditMode = InEditMode.Pin().Get())
+	{
+		TMap<UControlRig*, TArray<FRigElementKey>> SelectedControls;
+		EditMode->GetAllSelectedControls(SelectedControls);
+		for (TPair<UControlRig*, TArray<FRigElementKey>>& Selected : SelectedControls)
+		{
+			ControlRigs.Add(Selected.Key);
+		}
+	}
+	return ControlRigs;
+}
+}
diff --git a/Engine/Plugins/Animation/ControlRig/Source/ControlRigEditor/Private/EditMode/Tween/ControlRigObjectSelection.h b/Engine/Plugins/Animation/ControlRig/Source/ControlRigEditor/Private/EditMode/Tween/ControlRigObjectSelection.h
new file mode 100644
--- /dev/null
+++ b/Engine/Plugins/Animation/ControlRig/Source/ControlRigEditor/Private/EditMode/Tween/ControlRigObjectSelection.h
@@ -0,0 +1,68 @@
+﻿// Copyright Epic Games, Inc. All Rights Reserved.
+
+#pragma once
+
+#include "Containers/Array.h"
+#include "HAL/Platform.h"
+#include "Misc/FrameNumber.h"
+#include "Templates/SharedPointerFwd.h"
+
+class FControlRigEditMode;
+class ISequencer;
+class UControlRig;
+class UMovieSceneSection;
+struct FKeyHandle;
+struct FMovieSceneDoubleChannel;
+struct FMovieSceneFloatChannel;
+
+namespace UE::ControlRigEditor
+{
+/** Key bounds and their values used to blend with Controls and Actor float channels */
+struct FChannelKeyBounds
+{
+	FChannelKeyBounds()
+		: bValid(false), FloatChannel(nullptr), DoubleChannel(nullptr)
+		, PreviousIndex(INDEX_NONE), NextIndex(INDEX_NONE)
+		, PreviousFrame(0), NextFrame(0), CurrentFrame(0), PreviousValue(0.0), NextValue(0.0), CurrentValue(0.0)
+	{}
+	
+	bool bValid;
+	FMovieSceneFloatChannel* FloatChannel;
+	FMovieSceneDoubleChannel* DoubleChannel;
+	int32 PreviousIndex;
+	int32 NextIndex;
+	FFrameNumber PreviousFrame;
+	FFrameNumber NextFrame;
+	FFrameNumber CurrentFrame;
+	double PreviousValue;
+	double NextValue;
+	double CurrentValue;
+};
+
+/** Contains the selection state for a set of Control Rig Controls to blend with the anim slider */
+struct FControlRigObjectSelection
+{
+	//set of possible float channels
+	struct FObjectChannels
+	{
+		FObjectChannels() : Section(nullptr) {};
+		TArray<FChannelKeyBounds>  KeyBounds;
+		UMovieSceneSection* Section;
+	};
+
+	bool Setup(const TWeakPtr<ISequencer>& InSequencer, const TWeakPtr<FControlRigEditMode>& InEditMode);
+	bool Setup(const TArray<UControlRig*>& SelectedControlRigs, const TWeakPtr<ISequencer>& InSequencer);
+
+	TArray<FObjectChannels>  ChannelsArray;
+
+private:
+
+	// Either float or double will be non-null
+	void SetupChannel(
+		FFrameNumber CurrentFrame, TArray<FFrameNumber>& KeyTimes, TArray<FKeyHandle>& Handles,
+		FMovieSceneFloatChannel* FloatChannel, FMovieSceneDoubleChannel* DoubleChannel, FChannelKeyBounds& KeyBounds
+		);
+
+	TArray<UControlRig*> GetControlRigs(const TWeakPtr<FControlRigEditMode>& InEditMode);
+};
+}
diff --git a/Engine/Plugins/Animation/ControlRig/Source/ControlRigEditor/Private/EditMode/Tween/ControlRigTweenModel.cpp b/Engine/Plugins/Animation/ControlRig/Source/ControlRigEditor/Private/EditMode/Tween/ControlRigTweenModel.cpp
new file mode 100644
--- /dev/null
+++ b/Engine/Plugins/Animation/ControlRig/Source/ControlRigEditor/Private/EditMode/Tween/ControlRigTweenModel.cpp
@@ -0,0 +1,128 @@
+﻿// Copyright Epic Games, Inc. All Rights Reserved.
+
+#include "ControlRigTweenModel.h"
+
+#include "ISequencer.h"
+#include "Channels/MovieSceneFloatChannel.h"
+#include "Channels/MovieSceneDoubleChannel.h"
+#include "Math/Abstraction/ITweenFunction.h"
+#include "Math/CurveBlending.h"
+#include "Misc/SequencerUtils.h"
+
+namespace UE::ControlRigEditor
+{
+bool FControlRigTweenModel::BlendSingle(float InNormalizedValue)
+{
+	StartBlendOperation();
+	const bool bBlended = DoBlend(InNormalizedValue);
+	StopBlendOperation(); // Technically unneeded but it's expected by the API so we do it in case base version changes in the future
+	return bBlended;
+}
+
+bool FControlRigTweenModel::BlendSingleWithControlRigs(float InNormalizedValue, const TArray<UControlRig*>& InControlRigs)
+{
+	StartBlendOperationWithControlRigs(InControlRigs);
+	const bool bBlended = DoBlend(InNormalizedValue);
+	StopBlendOperation(); // Technically unneeded but it's expected by the API so we do it in case base version changes in the future
+	return bBlended;
+}
+
+void FControlRigTweenModel::StartBlendOperationWithControlRigs(const TArray<UControlRig*>& InControlRigs)
+{
+	const TWeakPtr<ISequencer> Sequencer = SequencerAttr.Get();
+	if (const TSharedPtr<ISequencer> SequencerPin = Sequencer.Pin())
+	{
+		ObjectSelection.Setup(InControlRigs, Sequencer);
+	}
+}
+
+void FControlRigTweenModel::StartBlendOperation()
+{
+	KeySelection.KeyMap.Empty(KeySelection.KeyMap.Num());
+	ObjectSelection.ChannelsArray.Empty(ObjectSelection.ChannelsArray.Num());
+	
+	const TWeakPtr<ISequencer> Sequencer = SequencerAttr.Get();
+	const TSharedPtr<ISequencer> SequencerPin = Sequencer.Pin();
+	const TSharedPtr<FCurveEditor> CurveEditor = GetCurveEditorFromSequencer(SequencerPin);
+	if (!CurveEditor)
+	{
+		return;
+	}
+
+	KeySelection = TweeningUtilsEditor::FContiguousKeyMapping(*CurveEditor);
+	if (KeySelection.KeyMap.IsEmpty())
+	{
+		ObjectSelection.Setup(Sequencer, ControlRigMode);
+	}
+}
+
+bool FControlRigTweenModel::DoBlend(float InNormalizedValue)
+{
+	const TWeakPtr<ISequencer> Sequencer = SequencerAttr.Get();
+	const TSharedPtr<ISequencer> SequencerPin = Sequencer.Pin();
+	const TSharedPtr<FCurveEditor> CurveEditor = GetCurveEditorFromSequencer(SequencerPin);
+	if (!CurveEditor)
+	{
+		return false;
+	}
+
+	const float ScaledBlendValue = ScaleBlendValue(InNormalizedValue); 
+	if (!KeySelection.KeyMap.IsEmpty())
+	{
+		return TweeningUtilsEditor::BlendCurves(ScaledBlendValue, *CurveEditor, KeySelection, [this](const TweeningUtilsEditor::FKeyBlendInfo& InBlendInfo)
+		{
+			return GetTweenFunction()->Blend(InBlendInfo);
+		});
+	}
+
+	if (!ObjectSelection.ChannelsArray.IsEmpty())
+	{
+		const FFrameTime FrameTime = SequencerPin->GetLocalTime().Time;
+		const FFrameRate TickResolution = SequencerPin->GetFocusedTickResolution();
+		
+		const TArray<FKeyPosition> AllKeyPositions;
+		const TArray <int32> Indices;
+		const int32 CurrentIndex = -1;
+		TweeningUtilsEditor::FKeyBlendInfo BlendStruct(AllKeyPositions, Indices);
+
+		bool bDidBlend = false;
+		for (const FControlRigObjectSelection::FObjectChannels& ObjectChannels : ObjectSelection.ChannelsArray)
+		{
+			if (ObjectChannels.Section)
+			{
+				ObjectChannels.Section->Modify();
+			}
+			for (int32 Index = 0; Index < ObjectChannels.KeyBounds.Num(); ++Index)
+			{
+				if (ObjectChannels.KeyBounds[Index].bValid)
+				{
+					const double PreviousValue = ObjectChannels.KeyBounds[Index].PreviousValue;
+					const double PreviousTime = TickResolution.AsSeconds(FFrameTime(ObjectChannels.KeyBounds[Index].PreviousFrame));
+					const double NextValue = ObjectChannels.KeyBounds[Index].NextValue;
+					const double NextTime = TickResolution.AsSeconds(FFrameTime(ObjectChannels.KeyBounds[Index].NextFrame));
+					const double CurrentValue = ObjectChannels.KeyBounds[Index].CurrentValue;
+					const double CurrentTime = TickResolution.AsSeconds(FFrameTime(ObjectChannels.KeyBounds[Index].CurrentFrame));
+					BlendStruct.SetValues(PreviousTime, PreviousValue, CurrentTime, CurrentValue,
+						NextTime, NextValue, ScaledBlendValue, CurrentValue, CurrentValue, CurrentIndex);
+					const double NewValue = GetTweenFunction()->Blend(BlendStruct);
+					using namespace UE::MovieScene;
+					if (ObjectChannels.KeyBounds[Index].FloatChannel)
+					{
+						const EMovieSceneKeyInterpolation KeyInterpolation = GetInterpolationMode(ObjectChannels.KeyBounds[Index].FloatChannel, FrameTime.GetFrame(), SequencerPin->GetKeyInterpolation());
+						AddKeyToChannel(ObjectChannels.KeyBounds[Index].FloatChannel, FrameTime.GetFrame(), (float)NewValue, KeyInterpolation);
+					}
+					else if (ObjectChannels.KeyBounds[Index].DoubleChannel)
+					{
+						const EMovieSceneKeyInterpolation KeyInterpolation = GetInterpolationMode(ObjectChannels.KeyBounds[Index].DoubleChannel, FrameTime.GetFrame(), SequencerPin->GetKeyInterpolation());
+						AddKeyToChannel(ObjectChannels.KeyBounds[Index].DoubleChannel, FrameTime.GetFrame(), NewValue, KeyInterpolation);
+					}
+					bDidBlend = true;
+				}
+			}
+		}
+		return bDidBlend;
+	}
+
+	return false;
+}
+}
diff --git a/Engine/Plugins/Animation/ControlRig/Source/ControlRigEditor/Private/EditMode/Tween/ControlRigTweenModel.h b/Engine/Plugins/Animation/ControlRig/Source/ControlRigEditor/Private/EditMode/Tween/ControlRigTweenModel.h
new file mode 100644
--- /dev/null
+++ b/Engine/Plugins/Animation/ControlRig/Source/ControlRigEditor/Private/EditMode/Tween/ControlRigTweenModel.h
@@ -0,0 +1,73 @@
+﻿// Copyright Epic Games, Inc. All Rights Reserved.
+
+#pragma once
+
+#include "ControlRigObjectSelection.h"
+#include "Misc/Attribute.h"
+#include "Templates/SharedPointer.h"
+#include "Widgets/MVC/Models/FunctionBasedTweenModel.h"
+
+class ISequencer;
+
+namespace UE::ControlRigEditor
+{
+/**
+ * Implements tweening in ControlRig
+ *
+ * Priority of tweening:
+ *	1. If keys are selected in the curve editor, tween those keys. 
+ *	2. Otherwise, if a control rig element is selected, insert a key at the current scrub position into the curves corresponding to the control rig
+ *	and then interpolate those keys.
+ */
+class FControlRigTweenModel : public TweeningUtilsEditor::FFunctionBasedTweenModel
+{
+public:
+	
+	explicit FControlRigTweenModel(
+		TAttribute<TWeakPtr<ISequencer>> InSequencerAttr, TWeakPtr<FControlRigEditMode> InControlRigMode,
+		const TSharedRef<TweeningUtilsEditor::ITweenFunction>& InInitialFunction
+		)
+		: FFunctionBasedTweenModel(InInitialFunction)
+		, SequencerAttr(MoveTemp(InSequencerAttr))
+		, ControlRigMode(MoveTemp(InControlRigMode))
+	{}
+	explicit FControlRigTweenModel(
+		TAttribute<TWeakPtr<ISequencer>> InSequencerAttr, const TSharedRef<TweeningUtilsEditor::ITweenFunction>& InInitialFunction
+		)
+		: FFunctionBasedTweenModel(InInitialFunction)
+		, SequencerAttr(MoveTemp(InSequencerAttr))
+	{}
+
+	/** Does a blend and returns whether any values were actually blended. */
+	bool BlendSingle(float InNormalizedValue);
+	/** Does a blend on specified control rigs and returns whether any values were actually blended. */
+	bool BlendSingleWithControlRigs(float InNormalizedValue, const TArray<UControlRig*>& InControlRigs);
+
+	/** Inits this model with control rigs. */
+	void StartBlendOperationWithControlRigs(const TArray<UControlRig*>& InControlRigs);
+
+	//~ Begin FTweenModel Interface
+	virtual void StartBlendOperation() override;
+	virtual void BlendValues(float InNormalizedValue) override { DoBlend(InNormalizedValue); }
+	//~ End FTweenModel Interface
+
+private:
+
+	/** Gets the sequencer used by Control Rig. Can return nullptr (e.g. when user has closed the sequencer), in which case blending does nothing. */
+	const TAttribute<TWeakPtr<ISequencer>> SequencerAttr;
+	/** The control rig mode that owns this tweener. Can be nullptr, in which case only the curve editor key selection is blended. */
+	const TWeakPtr<FControlRigEditMode> ControlRigMode;
+
+	/** The key selection to blend. Created in StartBlendOperation and used for the entirety of the blend operation. */
+	TweeningUtilsEditor::FContiguousKeyMapping KeySelection;
+	/**
+	 * The control rig selection to blend if ContiguousKeySelection is empty.
+	 * Created in StartBlendOperation and used for the entirety of the blend operation.
+	 */
+	FControlRigObjectSelection ObjectSelection; // TODO: Extract this to UE::ControlRigEditor namespace and deprecated the global one.
+	
+	/** @return Whether any blended occured. */
+	bool DoBlend(float InNormalizedValue);
+};
+}
+
diff --git a/Engine/Plugins/Animation/ControlRig/Source/ControlRigEditor/Private/EditMode/Tween/SControlRigTweenWidget.cpp b/Engine/Plugins/Animation/ControlRig/Source/ControlRigEditor/Private/EditMode/Tween/SControlRigTweenWidget.cpp
new file mode 100644
--- /dev/null
+++ b/Engine/Plugins/Animation/ControlRig/Source/ControlRigEditor/Private/EditMode/Tween/SControlRigTweenWidget.cpp
@@ -0,0 +1,70 @@
+// Copyright Epic Games, Inc. All Rights Reserved.
+
+#include "SControlRigTweenWidget.h"
+
+#include "EditMode/ControlRigEditModeToolkit.h"
+#include "ILevelSequenceEditorToolkit.h"
+#include "ISequencer.h"
+#include "LevelSequence.h"
+#include "LevelSequenceEditorBlueprintLibrary.h"
+#include "Tools/ControlRigPose.h"
+#include "TweenLogic.h"
+#include "Viewports/InViewportUIDragOperation.h"
+
+#define LOCTEXT_NAMESPACE "ControlRigTweenWidget"
+
+namespace UE::ControlRigEditor
+{
+void SControlRigTweenWidget::Construct(const FArguments& InArgs)
+{
+	OwningToolkit = InArgs._InOwningToolkit;
+	TweenControllers = MakeUnique<FTweenLogic>(
+		TAttribute<TWeakPtr<ISequencer>>::CreateLambda([]
+		{
+			ULevelSequence* LevelSequence = ULevelSequenceEditorBlueprintLibrary::GetCurrentLevelSequence();
+			IAssetEditorInstance* AssetEditor = GEditor->GetEditorSubsystem<UAssetEditorSubsystem>()->FindEditorForAsset(LevelSequence, false);
+			ILevelSequenceEditorToolkit* LevelSequenceEditor = static_cast<ILevelSequenceEditorToolkit*>(AssetEditor);
+			return LevelSequenceEditor ? LevelSequenceEditor->GetSequencer() : nullptr;
+		}),
+		InArgs._InOwningEditMode.ToSharedRef()
+	);
+	
+
+	ChildSlot
+	[
+		TweenControllers->MakeWidget()
+	];
+}
+
+FReply SControlRigTweenWidget::OnDragDetected(const FGeometry& MyGeometry, const FPointerEvent& MouseEvent)
+{
+	// Need to remember where within a tab we grabbed
+	const FVector2D TabGrabScreenSpaceOffset = MouseEvent.GetScreenSpacePosition() - MyGeometry.GetAbsolutePosition();
+
+	FOnInViewportUIDropped OnUIDropped = FOnInViewportUIDropped::CreateSP(this, &SControlRigTweenWidget::FinishDraggingWidget);
+	// Start dragging.
+	TSharedRef<FInViewportUIDragOperation> DragDropOperation =
+		FInViewportUIDragOperation::New(
+			SharedThis(this),
+			TabGrabScreenSpaceOffset,
+			GetDesiredSize(),
+			OnUIDropped
+		);
+
+	if (OwningToolkit.IsValid())
+	{
+		OwningToolkit.Pin()->TryRemoveTweenOverlay();
+	}
+	return FReply::Handled().BeginDragDrop(DragDropOperation);
+}
+
+void SControlRigTweenWidget::FinishDraggingWidget(const FVector2D InLocation) const
+{
+	if (OwningToolkit.IsValid())
+	{
+		OwningToolkit.Pin()->UpdateTweenWidgetLocation(InLocation);
+		OwningToolkit.Pin()->TryShowTweenOverlay();
+	}
+}
+}
+#undef LOCTEXT_NAMESPACE
diff --git a/Engine/Plugins/Animation/ControlRig/Source/ControlRigEditor/Private/EditMode/Tween/SControlRigTweenWidget.h b/Engine/Plugins/Animation/ControlRig/Source/ControlRigEditor/Private/EditMode/Tween/SControlRigTweenWidget.h
new file mode 100644
--- /dev/null
+++ b/Engine/Plugins/Animation/ControlRig/Source/ControlRigEditor/Private/EditMode/Tween/SControlRigTweenWidget.h
@@ -0,0 +1,44 @@
+// Copyright Epic Games, Inc. All Rights Reserved.
+
+#pragma once
+
+#include "TweenLogic.h"
+#include "Widgets/DeclarativeSyntaxSupport.h"
+#include "Widgets/SCompoundWidget.h"
+
+class UControlRig;
+class ISequencer;
+class FControlRigEditModeToolkit;
+
+namespace UE::ControlRigEditor
+{
+/** Manages the tween widget in the viewport, which the user can reposition by dragging. */
+class SControlRigTweenWidget : public SCompoundWidget
+{
+	SLATE_BEGIN_ARGS(SControlRigTweenWidget) {}
+		SLATE_ARGUMENT(TSharedPtr<FControlRigEditModeToolkit>, InOwningToolkit)
+		SLATE_ARGUMENT(TSharedPtr<FControlRigEditMode>, InOwningEditMode)
+	SLATE_END_ARGS()
+
+	void Construct(const FArguments& InArgs);
+
+private:
+
+	/** Used to update tell the toolkit when this widget has finished being dragged somewhere else. */
+	TWeakPtr<FControlRigEditModeToolkit> OwningToolkit;
+
+	/** Manages logic for the tweening widget. */
+	TUniquePtr<FTweenLogic> TweenControllers;
+
+	//~ Begin SWidget Interface
+	virtual FReply OnMouseButtonDown(const FGeometry& MyGeometry, const FPointerEvent& MouseEvent) override
+	{
+		return FReply::Handled().DetectDrag(SharedThis(this), EKeys::LeftMouseButton);
+	};
+	virtual FReply OnDragDetected(const FGeometry& MyGeometry, const FPointerEvent& MouseEvent) override;
+	//~ End SWidget Interface
+	
+	void FinishDraggingWidget(const FVector2D InLocation) const;
+};
+}
+
diff --git a/Engine/Plugins/Animation/ControlRig/Source/ControlRigEditor/Private/EditMode/Tween/TweenLogic.cpp b/Engine/Plugins/Animation/ControlRig/Source/ControlRigEditor/Private/EditMode/Tween/TweenLogic.cpp
new file mode 100644
--- /dev/null
+++ b/Engine/Plugins/Animation/ControlRig/Source/ControlRigEditor/Private/EditMode/Tween/TweenLogic.cpp
@@ -0,0 +1,57 @@
+﻿// Copyright Epic Games, Inc. All Rights Reserved.
+
+#include "TweenLogic.h"
+
+#include "ControlRigTweenModel.h"
+#include "EditMode/ControlRigEditMode.h"
+#include "Framework/Commands/UICommandList.h"
+#include "Math/Abstraction/BuiltInTweenFunctionContainer.h"
+#include "Misc/SequencerUtils.h"
+#include "Templates/SharedPointer.h"
+#include "Widgets/MVC/Models/EditorTweenModel.h"
+
+namespace UE::ControlRigEditor
+{
+namespace TweenLogicDetail
+{
+static TSharedRef<FControlRigTweenModel> MakeTweenModel(
+	TAttribute<TWeakPtr<ISequencer>> InSequencerAttr, const TSharedRef<FControlRigEditMode>& InOwningEditMode,
+	const TweeningUtilsEditor::ITweenFunctionContainer& InFunctions
+	)
+{
+	using FModel = TweeningUtilsEditor::TEditorTweenModel<FControlRigTweenModel>;
+	TAttribute<TWeakPtr<FCurveEditor>> CurveEditorAttr = TAttribute<TWeakPtr<FCurveEditor>>::CreateLambda([SequencerAttr = InSequencerAttr]
+	{
+		const TSharedPtr<ISequencer> SequencerPin = SequencerAttr.Get().Pin();
+		return GetCurveEditorFromSequencer(SequencerPin);
+	});
+	// TODO UE-240484: Read preferred function from config
+	return MakeShared<FModel>(
+		MoveTemp(CurveEditorAttr),	// TTangentFlatteningTweenProxy arg
+		MoveTemp(InSequencerAttr), InOwningEditMode, InFunctions.GetFunction(0).ToSharedRef()
+		);
+}
+}
+	
+FTweenLogic::FTweenLogic(TAttribute<TWeakPtr<ISequencer>> InSequencerAttr, const TSharedRef<FControlRigEditMode>& InOwningEditMode)
+	: CommandList(InOwningEditMode->GetCommandBindings().ToSharedRef())
+	, TweenFunctions(MakeShared<TweeningUtilsEditor::FBuiltInTweenFunctionContainer>())
+	, TweenModel(TweenLogicDetail::MakeTweenModel(MoveTemp(InSequencerAttr), InOwningEditMode, *TweenFunctions)) 
+	, Controllers(TweenModel, InOwningEditMode->GetCommandBindings().ToSharedRef(), TweenFunctions)
+{}
+
+TSharedRef<SWidget> FTweenLogic::MakeWidget() const
+{
+	// Visually move the slider widget when the user uses the U+LMB button command to indirectly move the mouse
+	TAttribute<TOptional<float>> OverrideSliderPosition =
+		TAttribute<TOptional<float>>::CreateLambda([this]()
+		{
+			return Controllers.MouseSlidingController.GetCurrentSliderPosition();
+		});
+	
+	FSlimHorizontalToolBarBuilder ToolBarBuilder(CommandList, FMultiBoxCustomization::None);
+	ToolBarBuilder.SetStyle(&FAppStyle::Get(), TEXT("ControlRigTweenToolbar"));
+	Controllers.ToolbarController.AddToToolbar(ToolBarBuilder, { MoveTemp(OverrideSliderPosition) });
+	return ToolBarBuilder.MakeWidget();
+}
+}
diff --git a/Engine/Plugins/Animation/ControlRig/Source/ControlRigEditor/Private/EditMode/Tween/TweenLogic.h b/Engine/Plugins/Animation/ControlRig/Source/ControlRigEditor/Private/EditMode/Tween/TweenLogic.h
new file mode 100644
--- /dev/null
+++ b/Engine/Plugins/Animation/ControlRig/Source/ControlRigEditor/Private/EditMode/Tween/TweenLogic.h
@@ -0,0 +1,41 @@
+﻿// Copyright Epic Games, Inc. All Rights Reserved.
+
+#pragma once
+
+#include "Misc/Attribute.h"
+#include "Templates/SharedPointerFwd.h"
+#include "Widgets/MVC/TweenControllers.h"
+
+class FControlRigEditMode;
+class FUICommandList;
+class ISequencer;
+namespace UE::TweeningUtilsEditor { class ITweenFunctionContainer; }
+
+namespace UE::ControlRigEditor
+{
+class FControlRigTweenModel;
+
+/** Creates the tweening model specific to ControlRig and provides an interface for interacting with it. */
+class FTweenLogic
+{
+public:
+
+	explicit FTweenLogic(TAttribute<TWeakPtr<ISequencer>> InSequencerAttr, const TSharedRef<FControlRigEditMode>& InOwningEditMode);
+
+	/** Constructs the content for the viewport tweening widget. */
+	TSharedRef<SWidget> MakeWidget() const;
+	
+private:
+
+	/** Owning control rig's command list. */
+	const TSharedRef<FUICommandList> CommandList;
+
+	/** Holds the used tweening functions and info about how to display them in UI. */
+	const TSharedRef<TweeningUtilsEditor::ITweenFunctionContainer> TweenFunctions;
+	/** The control rig specific tween model. */
+	const TSharedRef<FControlRigTweenModel> TweenModel;
+
+	/** Common functionality that should be shared consistent with other editor modules, such as Curve Editor. */
+	const TweeningUtilsEditor::FTweenControllers Controllers;
+};
+}
diff --git a/Engine/Plugins/Animation/ControlRig/Source/ControlRigEditor/Private/Misc/SequencerUtils.cpp b/Engine/Plugins/Animation/ControlRig/Source/ControlRigEditor/Private/Misc/SequencerUtils.cpp
new file mode 100644
--- /dev/null
+++ b/Engine/Plugins/Animation/ControlRig/Source/ControlRigEditor/Private/Misc/SequencerUtils.cpp
@@ -0,0 +1,21 @@
+﻿// Copyright Epic Games, Inc. All Rights Reserved.
+
+#include "SequencerUtils.h"
+
+#include "MVVM/CurveEditorExtension.h"
+#include "MVVM/ViewModels/SequencerEditorViewModel.h"
+
+namespace UE::ControlRigEditor
+{
+TSharedPtr<FCurveEditor> GetCurveEditorFromSequencer(const TSharedPtr<ISequencer>& InSequencer)
+{
+	if (!InSequencer)
+	{
+		return nullptr;
+	}
+
+	const TSharedPtr<Sequencer::FSequencerEditorViewModel> SequencerViewModel = InSequencer->GetViewModel();
+	const Sequencer::FCurveEditorExtension* CurveEditorExtension = SequencerViewModel->CastDynamic<Sequencer::FCurveEditorExtension>();
+	return ensure(CurveEditorExtension) ? CurveEditorExtension->GetCurveEditor() : nullptr;
+}
+}
diff --git a/Engine/Plugins/Animation/ControlRig/Source/ControlRigEditor/Private/Misc/SequencerUtils.h b/Engine/Plugins/Animation/ControlRig/Source/ControlRigEditor/Private/Misc/SequencerUtils.h
new file mode 100644
--- /dev/null
+++ b/Engine/Plugins/Animation/ControlRig/Source/ControlRigEditor/Private/Misc/SequencerUtils.h
@@ -0,0 +1,14 @@
+﻿// Copyright Epic Games, Inc. All Rights Reserved.
+
+#pragma once
+
+#include "Templates/SharedPointerFwd.h"
+
+class ISequencer;
+class FCurveEditor;
+
+namespace UE::ControlRigEditor
+{
+/** @return The curve editor used by sequencer */
+TSharedPtr<FCurveEditor> GetCurveEditorFromSequencer(const TSharedPtr<ISequencer>& InSequencer);
+}
diff --git a/Engine/Plugins/Animation/ControlRig/Source/ControlRigEditor/Private/Tools/ControlRigSnapper.cpp b/Engine/Plugins/Animation/ControlRig/Source/ControlRigEditor/Private/Tools/ControlRigSnapper.cpp
--- a/Engine/Plugins/Animation/ControlRig/Source/ControlRigEditor/Private/Tools/ControlRigSnapper.cpp
+++ b/Engine/Plugins/Animation/ControlRig/Source/ControlRigEditor/Private/Tools/ControlRigSnapper.cpp
@@ -1,7 +1,6 @@
 // Copyright Epic Games, Inc. All Rights Reserved.
 
 #include "Tools/ControlRigSnapper.h"
-#include "Tools/ControlRigTweener.h" //remove
 
 #include "Channels/MovieSceneFloatChannel.h"
 #include "ControlRig.h"
diff --git a/Engine/Plugins/Animation/ControlRig/Source/ControlRigEditor/Private/Tools/ControlRigTweener.cpp b/Engine/Plugins/Animation/ControlRig/Source/ControlRigEditor/Private/Tools/ControlRigTweener.cpp
deleted file mode 100644
diff --git a/Engine/Plugins/Animation/ControlRig/Source/ControlRigEditor/Public/EditMode/ControlRigEditMode.h b/Engine/Plugins/Animation/ControlRig/Source/ControlRigEditor/Public/EditMode/ControlRigEditMode.h
--- a/Engine/Plugins/Animation/ControlRig/Source/ControlRigEditor/Public/EditMode/ControlRigEditMode.h
+++ b/Engine/Plugins/Animation/ControlRig/Source/ControlRigEditor/Public/EditMode/ControlRigEditMode.h
@@ -166,7 +166,6 @@ public:
 	virtual bool StartTracking(FEditorViewportClient* InViewportClient, FViewport* InViewport) override;
 	virtual bool BeginTransform(const FGizmoState& InState) override;
 	virtual bool EndTransform(const FGizmoState& InState) override;
-	virtual bool ProcessCapturedMouseMoves(FEditorViewportClient* InViewportClient, FViewport* InViewport, const TArrayView<FIntPoint>& CapturedMouseMoves) override;
 	virtual bool HandleClick(FEditorViewportClient* InViewportClient, HHitProxy *HitProxy, const FViewportClick &Click) override;
 	virtual bool BoxSelect(FBox& InBox, bool InSelect = true) override;
 	virtual bool FrustumSelect(const FConvexVolume& InFrustum, FEditorViewportClient* InViewportClient, bool InSelect = true) override;
@@ -225,7 +224,7 @@ public:
 	- this is to ensure preview scene doesn't remove shape actors */
 	bool CanRemoveFromPreviewScene(const USceneComponent* InComponent);
 
-	FUICommandList* GetCommandBindings() const { return CommandBindings.Get(); }
+	TSharedPtr<FUICommandList> GetCommandBindings() const { return CommandBindings; }
 
 	/** Requests to recreate the shape actors in the next tick. Will recreate only the ones for the specified
 	Control Rig, otherwise will recreate all of them*/
@@ -276,21 +275,6 @@ protected:
 	/** Returns true if all control rig shapes are visible in the viewport */
 	bool AreControlsVisible() const;
 
-	/** If Anim Slider is open, got to the next tool*/
-	void ChangeAnimSliderTool();
-
-	/** If Anim Slider is open, then can drag*/
-	bool CanChangeAnimSliderTool() const;
-	/** Start Slider Tool*/
-	void StartAnimSliderTool(int32 InX);
-	/** If Anim Slider is open, drag the tool*/
-	void DragAnimSliderTool(double Val);
-	/** Reset and stop user the anim slider tool*/
-	void ResetAnimSlider();
-
-	/** If the Drag Anim Slider Tool is pressed*/
-	bool IsDragAnimSliderToolPressed(FViewport* InViewport);
-
 	bool HandleBeginTransform(const FEditorViewportClient* InViewportClient);
 	bool HandleEndTransform(FEditorViewportClient* InViewportClient);
 
@@ -456,11 +440,6 @@ protected:
 	/** Set initial transform handlers */
 	void OpenContextMenu(FEditorViewportClient* InViewportClient);
 
-	/* params to handle mouse move for changing anim tools sliders*/
-	bool bisTrackingAnimToolDrag = false;
-	/** Starting X Value*/
-	TOptional<int32> StartXValue;
-
 	/** previous Gizmo(Widget) scale before we enter this mode, used to set it back*/
 	float PreviousGizmoScale = 1.0f;
 
diff --git a/Engine/Plugins/Animation/ControlRig/Source/ControlRigEditor/Public/Tools/ControlRigTweener.h b/Engine/Plugins/Animation/ControlRig/Source/ControlRigEditor/Public/Tools/ControlRigTweener.h
deleted file mode 100644
diff --git a/Engine/Plugins/Animation/TweeningUtils/Source/TweeningUtilsEditor/Private/Math/Abstraction/BuiltInTweenFunctionContainer.cpp b/Engine/Plugins/Animation/TweeningUtils/Source/TweeningUtilsEditor/Private/Math/Abstraction/BuiltInTweenFunctionContainer.cpp
--- a/Engine/Plugins/Animation/TweeningUtils/Source/TweeningUtilsEditor/Private/Math/Abstraction/BuiltInTweenFunctionContainer.cpp
+++ b/Engine/Plugins/Animation/TweeningUtils/Source/TweeningUtilsEditor/Private/Math/Abstraction/BuiltInTweenFunctionContainer.cpp
@@ -27,6 +27,18 @@ TOptional<EBlendFunction> FBuiltInTweenFunctionContainer::GetBlendFunctionType(c
 		: TOptional<EBlendFunction>();
 }
 
+TSharedPtr<ITweenFunction> FBuiltInTweenFunctionContainer::GetBlendFunction(EBlendFunction InBlendFunction) const
+{
+	for (const TSharedRef<FBuiltInTweenFunction>& Function : Functions)
+	{
+		if (Function->GetBlendFunction() == InBlendFunction)
+		{
+			return Function;
+		}
+	}
+	return nullptr;
+}
+
 void FBuiltInTweenFunctionContainer::ForEachFunction(const TFunctionRef<void(const TSharedRef<ITweenFunction>&)> InConsumer) const
 {
 	for (const TSharedRef<FBuiltInTweenFunction>& Function : Functions)
diff --git a/Engine/Plugins/Animation/TweeningUtils/Source/TweeningUtilsEditor/Private/Widgets/MVC/TweenControllers.cpp b/Engine/Plugins/Animation/TweeningUtils/Source/TweeningUtilsEditor/Private/Widgets/MVC/TweenControllers.cpp
new file mode 100644
--- /dev/null
+++ b/Engine/Plugins/Animation/TweeningUtils/Source/TweeningUtilsEditor/Private/Widgets/MVC/TweenControllers.cpp
@@ -0,0 +1,18 @@
+﻿// Copyright Epic Games, Inc. All Rights Reserved.
+
+#include "Widgets/MVC/TweenControllers.h"
+
+#include "Widgets/MVC/Models/FunctionBasedTweenModel.h"
+
+namespace UE::TweeningUtilsEditor
+{
+FTweenControllers::FTweenControllers(
+	const TSharedRef<FFunctionBasedTweenModel>& InTweenModel,
+	const TSharedRef<FUICommandList>& InCommandList,
+	const TSharedRef<ITweenFunctionContainer>& InTweenFunctions
+	)
+	: ToolbarController(InTweenModel, InCommandList, InTweenFunctions)
+	, CycleFunctionController(InTweenModel, InTweenFunctions, InCommandList)
+	, MouseSlidingController(InTweenModel, InCommandList)
+	{}
+}
diff --git a/Engine/Plugins/Animation/TweeningUtils/Source/TweeningUtilsEditor/Private/Widgets/MVC/TweenToolbarController.cpp b/Engine/Plugins/Animation/TweeningUtils/Source/TweeningUtilsEditor/Private/Widgets/MVC/TweenToolbarController.cpp
--- a/Engine/Plugins/Animation/TweeningUtils/Source/TweeningUtilsEditor/Private/Widgets/MVC/TweenToolbarController.cpp
+++ b/Engine/Plugins/Animation/TweeningUtils/Source/TweeningUtilsEditor/Private/Widgets/MVC/TweenToolbarController.cpp
@@ -96,7 +96,7 @@ FTweenToolbarController::~FTweenToolbarController()
 	CommandList->UnmapAction(FTweeningUtilsCommands::Get().ToggleOvershootMode);
 }
 
-void FTweenToolbarController::AddToToolbar(FToolBarBuilder& ToolBarBuilder, TAttribute<TOptional<float>> InOverrideSliderPosition) const
+void FTweenToolbarController::AddToToolbar(FToolBarBuilder& ToolBarBuilder, FMakeWidgetArgs InArgs) const
 {
 	ToolBarBuilder.SetLabelVisibility(EVisibility::Visible);
 	ToolBarBuilder.AddComboButton(
@@ -112,7 +112,7 @@ void FTweenToolbarController::AddToToolbar(FToolBarBuilder& ToolBarBuilder, TAtt
 		SNew(STweenView, TweenModel)
 		.SliderIcon(TAttribute<const FSlateBrush*>::CreateRaw(this, &FTweenToolbarController::GetIconForSlider))
 		.SliderColor(TAttribute<FLinearColor>::CreateRaw(this, &FTweenToolbarController::GetColorForSlider))
-		.OverrideSliderPosition(MoveTemp(InOverrideSliderPosition))
+		.OverrideSliderPosition(MoveTemp(InArgs.OverrideSliderPositionAttr))
 		);
 	
 	ToolBarBuilder.AddToolBarButton(FTweeningUtilsCommands::Get().ToggleOvershootMode,
diff --git a/Engine/Plugins/Animation/TweeningUtils/Source/TweeningUtilsEditor/Public/Math/Abstraction/BuiltInTweenFunctionContainer.h b/Engine/Plugins/Animation/TweeningUtils/Source/TweeningUtilsEditor/Public/Math/Abstraction/BuiltInTweenFunctionContainer.h
--- a/Engine/Plugins/Animation/TweeningUtils/Source/TweeningUtilsEditor/Public/Math/Abstraction/BuiltInTweenFunctionContainer.h
+++ b/Engine/Plugins/Animation/TweeningUtils/Source/TweeningUtilsEditor/Public/Math/Abstraction/BuiltInTweenFunctionContainer.h
@@ -21,6 +21,9 @@ public:
 	/** @return Gets the type of blend function the tween function corresponds to, or unset if you pass in a function not owned by this container. */
 	TOptional<EBlendFunction> GetBlendFunctionType(const TSharedRef<ITweenFunction>& InFunctionToFind) const;
 
+	/** @return The blend function identified by the enum. */
+	TSharedPtr<ITweenFunction> GetBlendFunction(EBlendFunction InBlendFunction) const;
+
 	//~ Begin ITweenFunctionContainer Interface
 	virtual void ForEachFunction(const TFunctionRef<void(const TSharedRef<ITweenFunction>&)> InConsumer) const override;
 	virtual TSharedPtr<ITweenFunction> GetFunction(int32 InIndex) const override;
diff --git a/Engine/Plugins/Animation/TweeningUtils/Source/TweeningUtilsEditor/Public/Widgets/MVC/Models/EditorTweenModel.h b/Engine/Plugins/Animation/TweeningUtils/Source/TweeningUtilsEditor/Public/Widgets/MVC/Models/EditorTweenModel.h
new file mode 100644
--- /dev/null
+++ b/Engine/Plugins/Animation/TweeningUtils/Source/TweeningUtilsEditor/Public/Widgets/MVC/Models/EditorTweenModel.h
@@ -0,0 +1,27 @@
+﻿// Copyright Epic Games, Inc. All Rights Reserved.
+
+#pragma once
+
+#include "TangentFlatteningTweenProxy.h"
+#include "TransactionalTweenModelProxy.h"
+
+namespace UE::TweeningUtilsEditor
+{
+/**
+ * Extends TUnderlyingModel with commonly expected shared functionality:
+ *	- Undo / redo
+ *	- As the tween function squishes the keys vertically, user-specified tangents are flattened.
+ *	- TODO UE-234007: Synchronize tween function selection from different editor tools across all instances of the tweener.
+ *	- TODO UE-240484: Move preferred function editor UCurveEditorToolsUserSettings to this module.
+ *	- In the future this can be extended to remove vertically stacked keys (same keys per frame). For now, no implemented tween function stacks keys.
+ *
+ * Construct like this:
+ *  using FModel = TEditorTweenModel<FFooTweenModel> // FFooTweenModel should inherit from FTweenModel.
+ *	const TSharedRef<FModel> TweenModel = MakeShared<FModel>(
+ *		WeakCurveEditor,					// TTangentFlatteningTweenProxy arg
+ *		Bar									// Any arguments that FFooTweenModel requires for construction
+ *	);
+ */
+template<typename TUnderlyingModel>
+using TEditorTweenModel = TTransactionalTweenModelProxy<TTangentFlatteningTweenProxy<TUnderlyingModel>>;
+}
diff --git a/Engine/Plugins/Animation/TweeningUtils/Source/TweeningUtilsEditor/Public/Widgets/MVC/Models/TangentFlatteningTweenProxy.h b/Engine/Plugins/Animation/TweeningUtils/Source/TweeningUtilsEditor/Public/Widgets/MVC/Models/TangentFlatteningTweenProxy.h
new file mode 100644
--- /dev/null
+++ b/Engine/Plugins/Animation/TweeningUtils/Source/TweeningUtilsEditor/Public/Widgets/MVC/Models/TangentFlatteningTweenProxy.h
@@ -0,0 +1,58 @@
+﻿// Copyright Epic Games, Inc. All Rights Reserved.
+
+#pragma once
+
+#include "CurveEditor.h"
+#include "Misc/Attribute.h"
+#include "Misc/Mirror/TangentSelectionFlattener.h"
+#include "Widgets/MVC/Models/TweenModel.h"
+#include <type_traits>
+
+namespace UE::TweeningUtilsEditor
+{
+/**
+ * This class squishes the curves based on how much the tween function squishes the keys vertically.
+ * The squishing is achieved by interpolating the tangents to 0.
+ */
+template<typename TBase> requires std::is_base_of_v<FTweenModel, TBase>
+class TTangentFlatteningTweenProxy : public TBase
+{
+public:
+
+	template<typename... TArg>
+	explicit TTangentFlatteningTweenProxy(TAttribute<TWeakPtr<FCurveEditor>> WeakCurveEditorAttr, TArg&&... Arg)
+		: TBase(Forward<TArg>(Arg)...)
+		, WeakCurveEditor(MoveTemp(WeakCurveEditorAttr))
+	{
+		check(WeakCurveEditorAttr.IsBound() || WeakCurveEditorAttr.IsSet());
+	}
+
+	//~ Begin FTweenModel Interface
+	virtual void StartBlendOperation() override
+	{
+		TBase::StartBlendOperation();
+		
+		if (const TSharedPtr<FCurveEditor> CurveEditorPin = WeakCurveEditor.Get().Pin())
+		{
+			TangentTweener.ResetFromSelection(*CurveEditorPin);
+		}
+	}
+	virtual void BlendValues(float InNormalizedValue) override
+	{
+		TBase::BlendValues(InNormalizedValue);
+		if (const TSharedPtr<FCurveEditor> CurveEditorPin = WeakCurveEditor.Get().Pin())
+		{
+			TangentTweener.ComputeMirroringParallel(*CurveEditorPin);
+		}
+	}
+	//~ Begin FTweenModel Interface
+
+private:
+
+	/** Needed as arg for TangentTweener. */
+	const TAttribute<TWeakPtr<FCurveEditor>> WeakCurveEditor;
+	/** Implements the logic for flattening the tangents. */
+	CurveEditor::FTangentSelectionFlattener TangentTweener;
+};
+}
+
diff --git a/Engine/Plugins/Animation/TweeningUtils/Source/TweeningUtilsEditor/Public/Widgets/MVC/TweenControllers.h b/Engine/Plugins/Animation/TweeningUtils/Source/TweeningUtilsEditor/Public/Widgets/MVC/TweenControllers.h
new file mode 100644
--- /dev/null
+++ b/Engine/Plugins/Animation/TweeningUtils/Source/TweeningUtilsEditor/Public/Widgets/MVC/TweenControllers.h
@@ -0,0 +1,36 @@
+﻿// Copyright Epic Games, Inc. All Rights Reserved.
+
+#pragma once
+
+#include "CycleFunctionController.h"
+#include "Templates/SharedPointerFwd.h"
+#include "TweenMouseSlidingController.h"
+#include "TweenToolbarController.h"
+
+class FUICommandList;
+
+namespace UE::TweeningUtilsEditor
+{
+class FFunctionBasedTweenModel;
+class ITweenFunctionContainer;
+
+/** Holds functionality that may be reused across modules to add tweening widgets to a toolbar. */
+struct TWEENINGUTILSEDITOR_API FTweenControllers
+{
+	/** Manages the toolbar widget */
+	FTweenToolbarController ToolbarController;
+
+	/** Cycles functions (Shift + U). */
+	FCycleFunctionController CycleFunctionController;
+
+	/** Allows indirect movement of the slider by using u + Move Mouse */
+	FTweenMouseSlidingController MouseSlidingController;
+	
+	explicit FTweenControllers(
+		const TSharedRef<FFunctionBasedTweenModel>& InTweenModel,
+		const TSharedRef<FUICommandList>& InCommandList,
+		const TSharedRef<ITweenFunctionContainer>& InTweenFunctions
+		);
+};
+}
+
diff --git a/Engine/Plugins/Animation/TweeningUtils/Source/TweeningUtilsEditor/Public/Widgets/MVC/TweenToolbarController.h b/Engine/Plugins/Animation/TweeningUtils/Source/TweeningUtilsEditor/Public/Widgets/MVC/TweenToolbarController.h
--- a/Engine/Plugins/Animation/TweeningUtils/Source/TweeningUtilsEditor/Public/Widgets/MVC/TweenToolbarController.h
+++ b/Engine/Plugins/Animation/TweeningUtils/Source/TweeningUtilsEditor/Public/Widgets/MVC/TweenToolbarController.h
@@ -47,14 +47,21 @@ public:
 		);
 	~FTweenToolbarController();
 
+	struct FMakeWidgetArgs
+	{
+		/**
+		 * Optional override for the where the slider is visually placed (does not affect functionality). Range [-1,1].
+		 * Useful e.g. in conjunction with FTweenMouseSlidingController (while pressing U+LMB, give the user a visual association). 
+		 */
+		TAttribute<TOptional<float>> OverrideSliderPositionAttr;
+	};
 	/**
 	 * Extends the toolbar. Safe to call multiple times (although typically you'll call this exactly once).
 	 * 
 	 * @param ToolbarBuilder The toolbar to add the widgets to
-	 * @param InOverrideSliderPosition Optional override for the where the slider is visually placed (does not affect functionality). Range [-1,1].
-	 *	Useful e.g. in conjunction with FTweenMouseSlidingController (while pressing U+LMB, give the user a visual association).
+	 * @param InArgs Optional additional args to control how the widget should be constructed.
 	 */
-	void AddToToolbar(FToolBarBuilder& ToolbarBuilder, TAttribute<TOptional<float>> InOverrideSliderPosition = {}) const;
+	void AddToToolbar(FToolBarBuilder& ToolbarBuilder, FMakeWidgetArgs InArgs = {}) const;
 
 	DECLARE_MULTICAST_DELEGATE_OneParam(FOnTweenFunctionChanged, const TSharedRef<ITweenFunction>&);
 	/** Invoked when the tween function is changed by this controller. */
diff --git a/Engine/Plugins/Editor/CurveEditorTools/Source/CurveEditorTools/Private/Extensions/Tweening/TangentFlatteningTweenProxy.h b/Engine/Plugins/Editor/CurveEditorTools/Source/CurveEditorTools/Private/Extensions/Tweening/TangentFlatteningTweenProxy.h
deleted file mode 100644
diff --git a/Engine/Plugins/Editor/CurveEditorTools/Source/CurveEditorTools/Private/Extensions/Tweening/TweenEditorExtension.cpp b/Engine/Plugins/Editor/CurveEditorTools/Source/CurveEditorTools/Private/Extensions/Tweening/TweenEditorExtension.cpp
--- a/Engine/Plugins/Editor/CurveEditorTools/Source/CurveEditorTools/Private/Extensions/Tweening/TweenEditorExtension.cpp
+++ b/Engine/Plugins/Editor/CurveEditorTools/Source/CurveEditorTools/Private/Extensions/Tweening/TweenEditorExtension.cpp
@@ -9,8 +9,7 @@
 #include "Framework/MultiBox/MultiBoxExtender.h"
 #include "Math/Abstraction/BuiltInTweenFunctionContainer.h"
 #include "Misc/Optional.h"
-#include "TangentFlatteningTweenProxy.h"
-#include "Widgets/MVC/Models/TransactionalTweenModelProxy.h"
+#include "Widgets/MVC/Models/EditorTweenModel.h"
 #include "Widgets/MVC/TweenToolbarController.h"
 
 namespace UE::CurveEditorTools
@@ -45,7 +44,7 @@ TSharedPtr<FExtender> FTweenEditorExtension::MakeToolbarExtender(const TSharedRe
 				
 				ToolbarBuilder.BeginStyleOverride(TEXT("CurveEditorTweenToolbar")); // Forces min sizes for combo button
 				ToolbarBuilder.BeginSection("Tween");
-				TweenControllers->ToolbarController.AddToToolbar(ToolbarBuilder, MoveTemp(OverrideSliderPosition));
+				TweenControllers->ToolbarController.AddToToolbar(ToolbarBuilder, { MoveTemp(OverrideSliderPosition) });
 				ToolbarBuilder.EndSection();
 				ToolbarBuilder.EndStyleOverride();
 			}
@@ -53,16 +52,6 @@ TSharedPtr<FExtender> FTweenEditorExtension::MakeToolbarExtender(const TSharedRe
 	
 	return Extender;
 }
-	
-FTweenEditorExtension::FTweenControllers::FTweenControllers(
-	const TSharedRef<TweeningUtilsEditor::FFunctionBasedTweenModel>& InTweenModel,
-	const TSharedRef<FUICommandList>& InCommandList,
-	const TSharedRef<TweeningUtilsEditor::ITweenFunctionContainer>& InTweenFunctions
-	)
-	: ToolbarController(InTweenModel, InCommandList, InTweenFunctions)
-	, CycleFunctionController(InTweenModel, InTweenFunctions, InCommandList)
-	, MouseSlidingController(InTweenModel, InCommandList)
-{}
 
 bool FTweenEditorExtension::InitControllers(TSharedRef<FUICommandList> InCommandList)
 {
@@ -80,8 +69,9 @@ bool FTweenEditorExtension::InitControllers(TSharedRef<FUICommandList> InCommand
 			);
 		const TSharedRef<ITweenFunction> InitialFunction = PreferredFunction
 			? PreferredFunction.ToSharedRef() : BlendFunctionContainer->GetFunction(0).ToSharedRef();
+
 		
-		using FModel = TTransactionalTweenModelProxy<TTangentFlatteningTweenProxy<FContextAwareTweeningModel>>;
+		using FModel = TEditorTweenModel<FContextAwareTweeningModel>;
 		const TSharedRef<FModel> TweenModel = MakeShared<FModel>(
 			WeakCurveEditor,					// TTangentFlatteningTweenProxy arg
 			WeakCurveEditor, InitialFunction	// FContextAwareTweeningModel args
diff --git a/Engine/Plugins/Editor/CurveEditorTools/Source/CurveEditorTools/Private/Extensions/Tweening/TweenEditorExtension.h b/Engine/Plugins/Editor/CurveEditorTools/Source/CurveEditorTools/Private/Extensions/Tweening/TweenEditorExtension.h
--- a/Engine/Plugins/Editor/CurveEditorTools/Source/CurveEditorTools/Private/Extensions/Tweening/TweenEditorExtension.h
+++ b/Engine/Plugins/Editor/CurveEditorTools/Source/CurveEditorTools/Private/Extensions/Tweening/TweenEditorExtension.h
@@ -5,9 +5,7 @@
 #include "ICurveEditorExtension.h"
 #include "Misc/Optional.h"
 #include "Templates/SharedPointer.h"
-#include "Widgets/MVC/CycleFunctionController.h"
-#include "Widgets/MVC/TweenMouseSlidingController.h"
-#include "Widgets/MVC/TweenToolbarController.h"
+#include "Widgets/MVC/TweenControllers.h"
 
 class FCurveEditor;
 class FExtender;
@@ -38,25 +36,8 @@ private:
 	/** Holds the functions that ToolbarController shows: just the default, built-in ones.*/
 	const TSharedRef<TweeningUtilsEditor::FBuiltInTweenFunctionContainer> BlendFunctionContainer;
 
-	struct FTweenControllers
-	{
-		/** Manages the toolbar widget */
-		TweeningUtilsEditor::FTweenToolbarController ToolbarController;
-
-		/** Cycles functions (Shift + U). */
-		TweeningUtilsEditor::FCycleFunctionController CycleFunctionController;
-
-		/** Allows indirect movement of the slider by using u + Move Mouse */
-		TweeningUtilsEditor::FTweenMouseSlidingController MouseSlidingController;
-		
-		explicit FTweenControllers(
-			const TSharedRef<TweeningUtilsEditor::FFunctionBasedTweenModel>& InTweenModel,
-			const TSharedRef<FUICommandList>& InCommandList,
-			const TSharedRef<TweeningUtilsEditor::ITweenFunctionContainer>& InTweenFunctions
-			);
-	};
 	/** Created by InitControllers. */
-	TOptional<FTweenControllers> TweenControllers;
+	TOptional<TweeningUtilsEditor::FTweenControllers> TweenControllers;
 
 	/** Makes sure that ToolbarController is created. */
 	bool InitControllers(TSharedRef<FUICommandList> InCommandList);
diff --git a/Engine/Source/Editor/EditorStyle/Private/StarshipStyle.cpp b/Engine/Source/Editor/EditorStyle/Private/StarshipStyle.cpp
--- a/Engine/Source/Editor/EditorStyle/Private/StarshipStyle.cpp
+++ b/Engine/Source/Editor/EditorStyle/Private/StarshipStyle.cpp
@@ -7504,12 +7504,21 @@ void FStarshipEditorStyle::FStyle::SetupToolkitStyles()
 		// Some min sizes for toolbar combo buttons. We don't want the combo buttons to resize depending on the size of their contained label.
 		CurveEditorToolbar.SetComboContentHorizontalAlignment(HAlign_Center);
 		{
+			constexpr float LongestTweeName = 90.f;
+			
 			CurveEditorToolbar.SetComboLabelMinWidth(68.f);
 			Set("CurveEditorToolBar.ToolsCombo", CurveEditorToolbar);
+			
 			CurveEditorToolbar.SetComboLabelMinWidth(90.f);
 			Set("CurveEditorToolBar.TangentMode", CurveEditorToolbar);
-			CurveEditorToolbar.SetComboLabelMinWidth(90.f);
+			
+			CurveEditorToolbar.SetComboLabelMinWidth(LongestTweeName);
 			Set("CurveEditorTweenToolbar", CurveEditorToolbar);
+			
+			CurveEditorToolbar.SetComboLabelMinWidth(LongestTweeName);
+			CurveEditorToolbar.SetAllowWrapButton(false);
+			CurveEditorToolbar.SetBackground(*GetBrush("EditorViewport.OverlayBrush"));
+			Set("ControlRigTweenToolbar", CurveEditorToolbar);
 		}
 
 		Set("CurveEditor.KeyDetailWidth", 180.f);
