Control Rig: Follow up fix for performance improvements

#rnx

diff --git a/Engine/Plugins/Animation/ControlRig/Source/ControlRig/Private/Rigs/RigHierarchy.cpp b/Engine/Plugins/Animation/ControlRig/Source/ControlRig/Private/Rigs/RigHierarchy.cpp
--- a/Engine/Plugins/Animation/ControlRig/Source/ControlRig/Private/Rigs/RigHierarchy.cpp
+++ b/Engine/Plugins/Animation/ControlRig/Source/ControlRig/Private/Rigs/RigHierarchy.cpp
@@ -1464,7 +1464,12 @@ void URigHierarchy::ResetChangedCurveIndices()
 
 int32 URigHierarchy::Num(ERigElementType InElementType) const
 {
-	return ElementsPerType[RigElementTypeToFlatIndex(InElementType)].Num();
+	const int32 FlatIndex = RigElementTypeToFlatIndex(InElementType);
+	if(ElementsPerType.IsValidIndex(FlatIndex))
+	{
+		return ElementsPerType[FlatIndex].Num();
+	}
+	return 0;
 }
 
 bool URigHierarchy::IsProcedural(const FRigElementKey& InKey) const
