Fix a crash that could occur in URigHierarchy::GetPose.

The elements are now represented as bits, so <<=1 is the proper way to loop over the elements.

diff --git a/Engine/Plugins/Animation/ControlRig/Source/ControlRig/Private/Rigs/RigHierarchy.cpp b/Engine/Plugins/Animation/ControlRig/Source/ControlRig/Private/Rigs/RigHierarchy.cpp
--- a/Engine/Plugins/Animation/ControlRig/Source/ControlRig/Private/Rigs/RigHierarchy.cpp
+++ b/Engine/Plugins/Animation/ControlRig/Source/ControlRig/Private/Rigs/RigHierarchy.cpp
@@ -3756,7 +3756,7 @@ FRigPose URigHierarchy::GetPose(bool bInitial, ERigElementType InElementType,
 	const int32 InElementTypeAsInt32 = static_cast<int32>(InElementType);
 	
 	int32 ExpectedNumElements = 0;
-	for(int32 ElementType = FirstElementType; ElementType <= LastElementType; ElementType++)
+	for(int32 ElementType = FirstElementType; ElementType <= LastElementType; ElementType <<= 1)
 	{
 		if ((InElementTypeAsInt32 & ElementType) != 0)
 		{
