Control Rig: Refactor aggregation code in URigHierarchy::GetPose

Follow up to 39725167

#rnx

diff --git a/Engine/Plugins/Animation/ControlRig/Source/ControlRig/Private/Rigs/RigHierarchy.cpp b/Engine/Plugins/Animation/ControlRig/Source/ControlRig/Private/Rigs/RigHierarchy.cpp
--- a/Engine/Plugins/Animation/ControlRig/Source/ControlRig/Private/Rigs/RigHierarchy.cpp
+++ b/Engine/Plugins/Animation/ControlRig/Source/ControlRig/Private/Rigs/RigHierarchy.cpp
@@ -3751,16 +3751,15 @@ FRigPose URigHierarchy::GetPose(bool bInitial, ERigElementType InElementType,
 	Pose.HierarchyTopologyVersion = GetTopologyVersion();
 	Pose.PoseHash = Pose.HierarchyTopologyVersion;
 
-	static constexpr int32 FirstElementType = static_cast<int32>(ERigElementType::First); 
-	static constexpr int32 LastElementType = static_cast<int32>(ERigElementType::Last);
-	const int32 InElementTypeAsInt32 = static_cast<int32>(InElementType);
+	const uint8 InElementTypeAsUInt8 = static_cast<uint8>(InElementType);
 	
 	int32 ExpectedNumElements = 0;
-	for(int32 ElementType = FirstElementType; ElementType <= LastElementType; ElementType <<= 1)
+	for(int32 ElementTypeIndex = 0; ElementTypeIndex < ElementsPerType.Num(); ElementTypeIndex++)
 	{
-		if ((InElementTypeAsInt32 & ElementType) != 0)
+		const ERigElementType ElementType = FlatIndexToRigElementType(ElementTypeIndex);  
+		if ((InElementTypeAsUInt8 & static_cast<uint8>(ElementType)) != 0)
 		{
-			ExpectedNumElements += Num(static_cast<ERigElementType>(ElementType));
+			ExpectedNumElements += ElementsPerType[ElementTypeIndex].Num();
 		}
 	}
 	if(ExpectedNumElements == 0)
