Control Rig: Unify use of connector target widget

#rnx

diff --git a/Engine/Plugins/Animation/ControlRig/Source/ControlRigDeveloper/Public/ControlRigBlueprint.h b/Engine/Plugins/Animation/ControlRig/Source/ControlRigDeveloper/Public/ControlRigBlueprint.h
--- a/Engine/Plugins/Animation/ControlRig/Source/ControlRigDeveloper/Public/ControlRigBlueprint.h
+++ b/Engine/Plugins/Animation/ControlRig/Source/ControlRigDeveloper/Public/ControlRigBlueprint.h
@@ -451,4 +451,5 @@ private:
 	friend struct FControlRigGraphSchemaAction_PromoteToVariable;
 	friend class UControlRigGraphSchema;
 	friend class FControlRigBlueprintDetails;
+	friend class FRigConnectorElementDetails;
 };
diff --git a/Engine/Plugins/Animation/ControlRig/Source/ControlRigEditor/Private/ControlRigElementDetails.cpp b/Engine/Plugins/Animation/ControlRig/Source/ControlRigEditor/Private/ControlRigElementDetails.cpp
--- a/Engine/Plugins/Animation/ControlRig/Source/ControlRigEditor/Private/ControlRigElementDetails.cpp
+++ b/Engine/Plugins/Animation/ControlRig/Source/ControlRigEditor/Private/ControlRigElementDetails.cpp
@@ -28,6 +28,8 @@
 #include "Widgets/SRigVMGraphPinEnumPicker.h"
 #include "IStructureDataProvider.h"
 #include "Editor/ControlRigEditor.h"
+#include "Editor/SRigConnectorTargetWidget.h"
+#include "ModularRigRuleManager.h"
 
 #define LOCTEXT_NAMESPACE "ControlRigElementDetails"
 
@@ -6134,9 +6136,9 @@ void FRigConnectorElementDetails::CustomizeDetails(IDetailLayoutBuilder& DetailB
 {
 	FRigTransformElementDetails::CustomizeDetails(DetailBuilder);
 	CustomizeSettings(DetailBuilder);
-	//CustomizeTransform(DetailBuilder);
 	CustomizeComponents(DetailBuilder);
-	//CustomizeMetadata(DetailBuilder);
+	CustomizeConnectorTargets(DetailBuilder);
+	CustomizeMetadata(DetailBuilder);
 }
 
 void FRigConnectorElementDetails::CustomizeSettings(IDetailLayoutBuilder& DetailBuilder)
@@ -6210,6 +6212,76 @@ void FRigConnectorElementDetails::CustomizeSettings(IDetailLayoutBuilder& Detail
 	}
 }
 
+void FRigConnectorElementDetails::CustomizeConnectorTargets(IDetailLayoutBuilder& DetailBuilder)
+{
+	if(PerElementInfos.Num() != 1)
+	{
+		return;
+	}
+
+	const bool bArrayEnabled = IsArrayEnabled();
+
+	URigHierarchy* Hierarchy = PerElementInfos[0].GetHierarchy();
+	check(Hierarchy);
+
+	FRigTreeDelegates RigTreeDelegates;
+	RigTreeDelegates.OnGetHierarchy.BindLambda([Hierarchy]()
+	{
+		return Hierarchy;
+	});
+
+	TArray<FRigElementKey> CurrentTargets;
+	if(const UControlRig* ControlRig = Hierarchy->GetTypedOuter<UControlRig>())
+	{
+		if(const FRigConnectorElement* Connector = PerElementInfos[0].Element.Get<FRigConnectorElement>())
+		{
+			const FRigElementKeyRedirector Redirector = ControlRig->GetElementKeyRedirector();
+			const FRigElementKeyRedirector::FCachedKeyArray* Cache = Redirector.Find(Connector->GetKey());
+			if(Cache && !Cache->IsEmpty())
+			{
+				for(int32 Index = 0; Index < Cache->Num(); Index++)
+				{
+					CurrentTargets.Add((*Cache)[Index].GetKey());
+				}
+			}
+		}
+	}
+
+	IDetailCategoryBuilder& TargetsCategory = DetailBuilder.EditCategory(TEXT("Targets"), LOCTEXT("Targets", "Targets"));
+	FDetailWidgetRow& TargetsRow = TargetsCategory.AddCustomRow(LOCTEXT("Targets", "Targets"));
+	
+	TSharedRef<SRigConnectorTargetWidget> ConnectorTargetWidget = SNew(SRigConnectorTargetWidget)
+	.Outer(PerElementInfos[0].WrapperObject.Get())
+	.ConnectorKey(PerElementInfos[0].Element.GetKey())
+	.IsArray(bArrayEnabled)
+	.ExpandArrayByDefault(true)
+	.Targets(CurrentTargets)
+	.OnSetTargetArray(FRigConnectorTargetWidget_SetTargetArray::CreateSP(this, &FRigConnectorElementDetails::OnTargetsChanged))
+	.RigTreeDelegates(RigTreeDelegates);
+
+	if(bArrayEnabled)
+	{
+		TargetsRow.WholeRowContent()
+		[
+			ConnectorTargetWidget
+		];
+	}
+	else
+	{
+		TargetsRow
+		.NameContent()
+		[
+			SNew(STextBlock)
+			.Text(LOCTEXT("Target", "Target"))
+			.Font(IDetailLayoutBuilder::GetDetailFont())
+		]
+		.ValueContent()
+		[
+			ConnectorTargetWidget
+		];
+	}
+}
+
 TOptional<EConnectorType> FRigConnectorElementDetails::GetConnectorType() const
 {
 	TOptional<EConnectorType> Result;
@@ -6260,6 +6332,18 @@ bool FRigConnectorElementDetails::IsArrayEnabled() const
 	return false;
 }
 
+bool FRigConnectorElementDetails::OnTargetsChanged(TArray<FRigElementKey> InTargets)
+{
+	if(PerElementInfos.Num() == 1)
+	{
+		if(UControlRigBlueprint* ControlRigBlueprint = PerElementInfos[0].GetBlueprint())
+		{
+			return ControlRigBlueprint->ResolveConnectorToArray(PerElementInfos[0].Element.GetKey(), InTargets, true);
+		}
+	}
+	return false;
+}
+
 void FRigSocketElementDetails::CustomizeDetails(IDetailLayoutBuilder& DetailBuilder)
 {
 	FRigTransformElementDetails::CustomizeDetails(DetailBuilder);
diff --git a/Engine/Plugins/Animation/ControlRig/Source/ControlRigEditor/Private/ControlRigElementDetails.h b/Engine/Plugins/Animation/ControlRig/Source/ControlRigEditor/Private/ControlRigElementDetails.h
--- a/Engine/Plugins/Animation/ControlRig/Source/ControlRigEditor/Private/ControlRigElementDetails.h
+++ b/Engine/Plugins/Animation/ControlRig/Source/ControlRigEditor/Private/ControlRigElementDetails.h
@@ -868,12 +868,14 @@ public:
 	virtual void CustomizeDetails(IDetailLayoutBuilder& DetailBuilder) override;
 
 	void CustomizeSettings(IDetailLayoutBuilder& DetailBuilder);
+	void CustomizeConnectorTargets(IDetailLayoutBuilder& DetailBuilder);
 
 private:
 
 	TOptional<EConnectorType> GetConnectorType() const;
 	TOptional<bool> GetIsConnectorArray() const;
 	bool IsArrayEnabled() const;
+	bool OnTargetsChanged(TArray<FRigElementKey> InTargets);
 
 	TSharedPtr<IPropertyHandle> ConnectorTypeHandle;
 	TSharedPtr<IPropertyHandle> IsArrayHandle;
diff --git a/Engine/Plugins/Animation/ControlRig/Source/ControlRigEditor/Private/ControlRigModuleDetails.cpp b/Engine/Plugins/Animation/ControlRig/Source/ControlRigEditor/Private/ControlRigModuleDetails.cpp
--- a/Engine/Plugins/Animation/ControlRig/Source/ControlRigEditor/Private/ControlRigModuleDetails.cpp
+++ b/Engine/Plugins/Animation/ControlRig/Source/ControlRigEditor/Private/ControlRigModuleDetails.cpp
@@ -29,6 +29,7 @@
 #include "Editor/SRigHierarchyTreeView.h"
 #include "Widgets/SRigVMVariantTagWidget.h"
 #include "Algo/Sort.h"
+#include "Editor/SRigConnectorTargetWidget.h"
 
 #define LOCTEXT_NAMESPACE "ControlRigModuleDetails"
 
@@ -378,56 +379,12 @@ void FRigModuleInstanceDetails::CustomizeDetails(IDetailLayoutBuilder& DetailBui
 				const FText Label = FText::FromString(Connector.Name);
 				bool bIsArrayConnector = IsArrayConnector[ConnectorIndex].Get(false);
 
-				TArray<FRigElementResolveResult> Matches;
-				for (int32 ModuleIndex=0; ModuleIndex<PerModuleInfos.Num(); ++ModuleIndex)
-				{
-					const FPerModuleInfo& Info = PerModuleInfos[ModuleIndex];
-					if (const FRigModuleInstance* Module = Info.GetModule())
-					{
-						if (UModularRig* ModularRig = Info.GetModularRig())
-						{
-							if (URigHierarchy* Hierarchy = ModularRig->GetHierarchy())
-							{
-								FString ConnectorPath = FString::Printf(TEXT("%s:%s"), *Info.Path, *Connector.Name);
-								FRigElementKey ConnectorKey(*ConnectorPath, ERigElementType::Connector);
-								if (FRigConnectorElement* ConnectorElement = Cast<FRigConnectorElement>(Hierarchy->Find(ConnectorKey)))
-								{
-									const UModularRigRuleManager* RuleManager = ModularRig->GetHierarchy()->GetRuleManager();
-									if (ModuleIndex == 0)
-									{
-										Matches = RuleManager->FindMatches(ConnectorElement, Module, ModularRig->GetElementKeyRedirector()).GetMatches();
-									}
-									else
-									{
-										const FModularRigResolveResult& ConnectorMatches = RuleManager->FindMatches(ConnectorElement, Module, ModularRig->GetElementKeyRedirector());
-										Matches.RemoveAll([&ConnectorMatches](const FRigElementResolveResult& Match)
-										{
-											return !ConnectorMatches.ContainsMatch(Match.GetKey());
-										});
-									}
-								}
-							}
-						}
-					}
-				}
-
-				FRigTreeDelegates TreeDelegates;
-				TreeDelegates.OnGetHierarchy = FOnGetRigTreeHierarchy::CreateLambda([this]()
+				FRigTreeDelegates RigTreeDelegates;
+				RigTreeDelegates.OnGetHierarchy = FOnGetRigTreeHierarchy::CreateLambda([this]()
 				{
 					return PerModuleInfos[0].GetModularRig()->GetHierarchy();
 				});
 
-				TArray<FRigElementKey> MatchingKeys;
-				MatchingKeys.Reserve(Matches.Num());
-				for(const FRigElementResolveResult& SingleMatch : Matches)
-				{
-					MatchingKeys.Add(SingleMatch.GetKey());
-				}
-				TreeDelegates.OnRigTreeIsItemVisible = FOnRigTreeIsItemVisible::CreateLambda([MatchingKeys](const FRigElementKey& InTarget)
-				{
-					return MatchingKeys.Contains(InTarget);
-				});
-
 				static const FSlateBrush* PrimaryBrush = FControlRigEditorStyle::Get().GetBrush("ControlRig.ConnectorPrimary");
 				static const FSlateBrush* SecondaryBrush = FControlRigEditorStyle::Get().GetBrush("ControlRig.ConnectorSecondary");
 				static const FSlateBrush* OptionalBrush = FControlRigEditorStyle::Get().GetBrush("ControlRig.ConnectorOptional");
@@ -436,10 +393,9 @@ void FRigModuleInstanceDetails::CustomizeDetails(IDetailLayoutBuilder& DetailBui
 
 				TSharedRef<SVerticalBox> ValueContentVerticalBox = SNew(SVerticalBox);
 
-				TArray<FRigElementKey> CurrentTargetKeys;
+				TOptional<TArray<FRigElementKey>> TargetKeys;
 				if (PerModuleInfos.Num() >= 1)
 				{
-					TOptional<TArray<FRigElementKey>> TargetKeys;
 					for (int32 ModuleIndex=0; ModuleIndex<PerModuleInfos.Num(); ++ModuleIndex)
 					{
 						FString ConnectorPath = FString::Printf(TEXT("%s:%s"), *PerModuleInfos[ModuleIndex].Path, *Connector.Name);
@@ -474,326 +430,17 @@ void FRigModuleInstanceDetails::CustomizeDetails(IDetailLayoutBuilder& DetailBui
 									break;
 								}
 							}
-							break;
 						}
 					}
-
-					if(TargetKeys.IsSet())
-					{
-						CurrentTargetKeys = TargetKeys.GetValue();
-					}
-				}
-
-				if(bIsArrayConnector)
-				{
-					ValueContentVerticalBox->AddSlot()
-					.AutoHeight()
-					.Padding(0)
-					[
-						SNew(SHorizontalBox)
-
-						+SHorizontalBox::Slot()
-						.Padding(2, 4, 0, 4)
-						.AutoWidth()
-						.HAlign(HAlign_Left)
-						.VAlign(VAlign_Center)
-						[
-							SNew(STextBlock)
-							.Font(DetailBuilder.GetDetailFont())
-							.Text(FText::FromString(FString::Printf(TEXT("%d Connected Targets"), CurrentTargetKeys.Num())))
-						]
-
-						+SHorizontalBox::Slot()
-						.Padding(4, 0, 0, 0)
-						.AutoWidth()
-						.HAlign(HAlign_Left)
-						.VAlign(VAlign_Center)
-						[
-							SNew(SButton)
-							.IsEnabled(true)
-							.ButtonStyle(FAppStyle::Get(), "SimpleButton")
-							.ContentPadding(FMargin(1, 0))
-							.OnClicked(this, &FRigModuleInstanceDetails::OnAddTargetToArrayConnector, ConnectorName, DetailBuilder.GetPropertyUtilities())
-							.HAlign(HAlign_Right)
-							.ToolTipText(LOCTEXT("AddTargetToArrayConnectors", "Adds a new target for this array connector"))
-							.VAlign(VAlign_Center)
-							[
-								SNew(SImage)
-								.Image(FAppStyle::Get().GetBrush("Icons.PlusCircle"))
-								.ColorAndOpacity(FSlateColor::UseForeground())
-							]
-						]
-
-						+SHorizontalBox::Slot()
-						.Padding(4, 0, 0, 0)
-						.AutoWidth()
-						.HAlign(HAlign_Left)
-						.VAlign(VAlign_Center)
-						[
-							SNew(SButton)
-							.IsEnabled(true)
-							.ButtonStyle(FAppStyle::Get(), "SimpleButton")
-							.ContentPadding(FMargin(1, 0))
-							.OnClicked(this, &FRigModuleInstanceDetails::OnClearTargetsForArrayConnector, ConnectorName, DetailBuilder.GetPropertyUtilities())
-							.HAlign(HAlign_Right)
-							.ToolTipText(LOCTEXT("ClearTargetsForArrayConnector", "Removes all targets for this array connector"))
-							.VAlign(VAlign_Center)
-							[
-								SNew(SImage)
-								.Image(FAppStyle::Get().GetBrush("Icons.Delete"))
-								.ColorAndOpacity(FSlateColor::UseForeground())
-							]
-						]
-					];
 				}
 
-				for(int32 TargetIndex = 0; TargetIndex < CurrentTargetKeys.Num(); TargetIndex++)
-				{
-					const FString SearchableTreeViewKey = FString::Printf(TEXT("%s_%d"), *Connector.Name, TargetIndex);
-					
-					const FRigElementKey& TargetKey = CurrentTargetKeys[TargetIndex];
-
-					TreeDelegates.OnGetSelection.BindLambda([this, ConnectorName, TargetIndex]() -> TArray<FRigElementKey>
-					{
-						TOptional<FRigElementKey> Target;
-						FRigElementKeyRedirector Redirector = PerModuleInfos[0].GetModularRig()->GetElementKeyRedirector();
-						for (int32 ModuleIndex=0; ModuleIndex<PerModuleInfos.Num(); ++ModuleIndex)
-						{
-							FString ConnectorPath = FString::Printf(TEXT("%s:%s"), *PerModuleInfos[ModuleIndex].Path, *ConnectorName);
-							FRigElementKey ConnectorKey(*ConnectorPath, ERigElementType::Connector);
-							if(const UControlRigBlueprint* Blueprint = PerModuleInfos[ModuleIndex].GetBlueprint())
-							{
-								const TArray<FRigElementKey> CurrentTargets = Blueprint->ModularRigModel.Connections.FindTargetsFromConnector(ConnectorKey);
-								if(CurrentTargets.IsValidIndex(TargetIndex))
-								{
-									if (!Target.IsSet())
-									{
-										Target = CurrentTargets[TargetIndex];
-									}
-									else if (Target.GetValue() != CurrentTargets[TargetIndex])
-									{
-										Target.GetValue().Name = *ControlRigModuleDetailsMultipleValues.ToString();
-										return {Target.GetValue()};
-									}
-								}
-								else if(Target.IsSet())
-								{
-									Target.GetValue().Name = *ControlRigModuleDetailsMultipleValues.ToString();
-									return {Target.GetValue()};
-								}
-							}
-							else
-							{
-								return {FRigElementKey(*ControlRigModuleDetailsMultipleValues.ToString(), ERigElementType::None)};
-							}
-						}
-						return {};
-					});
-								
-					TreeDelegates.OnSelectionChanged.BindSP(this, &FRigModuleInstanceDetails::OnConnectorTargetChanged, Connector, TargetIndex);
-
-					TSharedPtr<SSearchableRigHierarchyTreeView>& SearchableTreeView = ConnectionListBox.FindOrAdd(SearchableTreeViewKey);
-					SearchableTreeView = SNew(SSearchableRigHierarchyTreeView)
-						.RigTreeDelegates(TreeDelegates);
-					
-					TSharedRef<SHorizontalBox> ValueContentHorizontalBox = SNew(SHorizontalBox);
-					ValueContentVerticalBox->AddSlot()
-					.AutoHeight()
-					.Padding(0)
-					[
-						ValueContentHorizontalBox
-					];
-
-					TSharedPtr<SVerticalBox> ButtonBox;
-
-					// Combo button
-					ValueContentHorizontalBox->AddSlot()
-					.AutoWidth()
-					.Padding(0.f, 0.f, 0.f, 0.f)
-					[
-						SNew( SComboButton )
-						.ContentPadding(3)
-						.MenuPlacement(MenuPlacement_BelowAnchor)
-						.OnComboBoxOpened(this, &FRigModuleInstanceDetails::PopulateConnectorTargetList, SearchableTreeViewKey)
-						.ButtonContent()
-						[
-							// Wrap in configurable box to restrain height/width of menu
-							SNew(SBox)
-							.MinDesiredWidth(150.0f)
-							[
-								SAssignNew(ButtonBox, SVerticalBox)
-							]
-						]
-						.MenuContent()
-						[
-							SNew(SBorder)
-							.Visibility(EVisibility::Visible)
-							.BorderImage(FAppStyle::GetBrush("Menu.Background"))
-							[
-								SearchableTreeView.ToSharedRef()
-							]
-						]
-					];
-
-					ValueContentHorizontalBox->AddSlot()
-					.AutoWidth()
-					.VAlign(VAlign_Center)
-					.Padding(0.f, 0.f, 0.f, 0.f)
-					[
-						SNew(SVerticalBox)
-						+SVerticalBox::Slot()
-						.VAlign(VAlign_Center)
-						.AutoHeight()
-						[
-							// Reset button
-							SNew(SHorizontalBox)
-							+SHorizontalBox::Slot()
-							.AutoWidth()
-							.Padding(0.f, 0.f, 0.f, 0.f)
-							[
-								SAssignNew(ResetConnectorButton.FindOrAdd(Connector.Name), SButton)
-								.ButtonStyle( FAppStyle::Get(), "NoBorder" )
-								.ButtonColorAndOpacity_Lambda([this, Connector]()
-								{
-									const TSharedPtr<SButton>& Button = ResetConnectorButton.FindRef(Connector.Name);
-									return Button.IsValid() && Button->IsHovered()
-										? FSlateColor(FLinearColor(1.f, 1.f, 1.f, 0.8))
-										: FSlateColor(FLinearColor(1.f, 1.f, 1.f, 0.4));
-								})
-								.OnClicked_Lambda([this, Connector]()
-								{
-									for (FPerModuleInfo& Info : PerModuleInfos)
-									{
-										FString ConnectorPath = FString::Printf(TEXT("%s:%s"), *Info.Path, *Connector.Name);
-										FRigElementKey ConnectorKey(*ConnectorPath, ERigElementType::Connector);
-										Info.GetBlueprint()->GetModularRigController()->DisconnectConnector(ConnectorKey);
-									}
-									return FReply::Handled();
-								})
-								.ContentPadding(1.f)
-								.ToolTipText(NSLOCTEXT("ControlRigModuleDetails", "Reset_Connector", "Reset Connector"))
-								[
-									SNew(SImage)
-									.ColorAndOpacity_Lambda( [this, Connector]()
-									{
-										const TSharedPtr<SButton>& Button = ResetConnectorButton.FindRef(Connector.Name);
-										return Button.IsValid() && Button->IsHovered()
-										? FSlateColor(FLinearColor(1.f, 1.f, 1.f, 0.8))
-										: FSlateColor(FLinearColor(1.f, 1.f, 1.f, 0.4));
-									})
-									.Image(FSlateIcon(FAppStyle::Get().GetStyleSetName(), "PropertyWindow.DiffersFromDefault").GetIcon())
-								]
-							]
-
-							// Use button
-							+SHorizontalBox::Slot()
-							.AutoWidth()
-							.Padding(0.f, 0.f, 0.f, 0.f)
-							[
-								SAssignNew(UseSelectedButton.FindOrAdd(Connector.Name), SButton)
-								.ButtonStyle( FAppStyle::Get(), "NoBorder" )
-								.ButtonColorAndOpacity_Lambda([this, Connector]()
-								{
-									const TSharedPtr<SButton>& Button = UseSelectedButton.FindRef(Connector.Name);
-									return Button.IsValid() && Button->IsHovered()
-										? FSlateColor(FLinearColor(1.f, 1.f, 1.f, 0.8))
-										: FSlateColor(FLinearColor(1.f, 1.f, 1.f, 0.4));
-								})
-								.OnClicked_Lambda([this, Connector, TargetIndex]()
-								{
-									if (UModularRig* ModularRig = PerModuleInfos[0].GetModularRig())
-									{
-										const TArray<FRigElementKey>& Selected = ModularRig->GetHierarchy()->GetSelectedKeys();
-										if (Selected.Num() > 0)
-										{
-											for (FPerModuleInfo& Info : PerModuleInfos)
-											{
-												FString ConnectorPath = FString::Printf(TEXT("%s:%s"), *Info.Path, *Connector.Name);
-												FRigElementKey ConnectorKey(*ConnectorPath, ERigElementType::Connector);
-
-												TArray<FRigElementKey> Targets = Info.GetBlueprint()->ModularRigModel.Connections.FindTargetsFromConnector(ConnectorKey);
-												if(Targets.IsValidIndex(TargetIndex))
-												{
-													Targets[TargetIndex] = Selected[0];
-												}
-												else
-												{
-													Targets.Add(Selected[0]);
-												}
-												
-												Info.GetBlueprint()->GetModularRigController()->ConnectConnectorToElements(ConnectorKey, Targets, true, ModularRig->GetModularRigSettings().bAutoResolve);
-											}
-										}
-									}
-									return FReply::Handled();
-								})
-								.ContentPadding(1.f)
-								.ToolTipText(NSLOCTEXT("ControlRigModuleDetails", "Use_Selected", "Use Selected"))
-								[
-									SNew(SImage)
-									.ColorAndOpacity_Lambda( [this, Connector]()
-									{
-										const TSharedPtr<SButton>& Button = UseSelectedButton.FindRef(Connector.Name);
-										return Button.IsValid() && Button->IsHovered()
-										? FSlateColor(FLinearColor(1.f, 1.f, 1.f, 0.8))
-										: FSlateColor(FLinearColor(1.f, 1.f, 1.f, 0.4));
-									})
-									.Image(FAppStyle::GetBrush("Icons.CircleArrowLeft"))
-								]
-							]
-
-							// Select in hierarchy button
-							+SHorizontalBox::Slot()
-							.AutoWidth()
-							.Padding(0.f, 0.f, 0.f, 0.f)
-							[
-								SAssignNew(SelectElementButton.FindOrAdd(Connector.Name), SButton)
-								.ButtonStyle( FAppStyle::Get(), "NoBorder" )
-								.ButtonColorAndOpacity_Lambda([this, Connector]()
-								{
-									const TSharedPtr<SButton>& Button = SelectElementButton.FindRef(Connector.Name);
-									return Button.IsValid() && Button->IsHovered()
-										? FSlateColor(FLinearColor(1.f, 1.f, 1.f, 0.8))
-										: FSlateColor(FLinearColor(1.f, 1.f, 1.f, 0.4));
-								})
-								.OnClicked_Lambda([this, Connector]()
-								{
-									if (UModularRig* ModularRig = PerModuleInfos[0].GetModularRig())
-									{
-										FString ConnectorPath = FString::Printf(TEXT("%s:%s"), *PerModuleInfos[0].Path, *Connector.Name);
-										FRigElementKey ConnectorKey(*ConnectorPath, ERigElementType::Connector);
-										if (const FRigElementKeyRedirector::FKeyArray* TargetKeys = ModularRig->GetElementKeyRedirector().FindExternalKey(ConnectorKey))
-										{
-											TArray<FRigElementKey> KeysToSelect;
-											KeysToSelect.Append(*TargetKeys);
-											ModularRig->GetHierarchy()->GetController()->SetSelection(KeysToSelect, true);
-										}
-									}
-									return FReply::Handled();
-								})
-								.ContentPadding(1.f)
-								.ToolTipText(NSLOCTEXT("ControlRigModuleDetails", "Select_Element", "Select Element"))
-								[
-									SNew(SImage)
-									.ColorAndOpacity_Lambda( [this, Connector]()
-									{
-										const TSharedPtr<SButton>& Button = SelectElementButton.FindRef(Connector.Name);
-										return Button.IsValid() && Button->IsHovered()
-										? FSlateColor(FLinearColor(1.f, 1.f, 1.f, 0.8))
-										: FSlateColor(FLinearColor(1.f, 1.f, 1.f, 0.4));
-									})
-									.Image(FAppStyle::GetBrush("Icons.Search"))
-								]
-							]
-						]
-					];
-
-					TPair<const FSlateBrush*, FSlateColor> IconAndColor = SRigHierarchyItem::GetBrushForElementType(PerModuleInfos[0].GetModularRig()->GetHierarchy(), TargetKey);
-					PopulateConnectorCurrentTarget(ButtonBox, IconAndColor.Key, IconAndColor.Value, FText::FromName(TargetKey.Name));
-				}
+				FString ConnectorPath = FString::Printf(TEXT("%s:%s"), *PerModuleInfos[0].Path, *Connector.Name);
+				FRigElementKey ConnectorKey(*ConnectorPath, ERigElementType::Connector);
 
-				ConnectionsCategory.AddCustomRow(Label)
+				FDetailWidgetRow& ConnectorRow = ConnectionsCategory.AddCustomRow(Label)
 				.NameContent()
+				.HAlign(HAlign_Left)
+				.VAlign(VAlign_Top)
 				[
 					SNew(SHorizontalBox)
 
@@ -801,7 +448,7 @@ void FRigModuleInstanceDetails::CustomizeDetails(IDetailLayoutBuilder& DetailBui
 					.AutoWidth()
 					.Padding(0.f, 0.f, 4.f, 0.f)
 					.HAlign(HAlign_Left)
-					.VAlign(VAlign_Center)
+					.VAlign(VAlign_Top)
 					[
 						SNew(SImage)
 						.Image(IconBrush)
@@ -813,18 +460,45 @@ void FRigModuleInstanceDetails::CustomizeDetails(IDetailLayoutBuilder& DetailBui
 					.AutoWidth()
 					.Padding(0.f, 0.f, 0.f, 0.f)
 					.HAlign(HAlign_Left)
-					.VAlign(VAlign_Center)
+					.VAlign(VAlign_Top)
 					[
 						SNew(STextBlock)
 						.Text(Label)
 						.Font(IDetailLayoutBuilder::GetDetailFont())
 						.IsEnabled(true)
 					]
-				]
-				.ValueContent()
-				[
-					ValueContentVerticalBox
 				];
+
+
+				if(TargetKeys.IsSet())
+				{
+					ConnectorRow
+					.ValueContent()
+					.HAlign(HAlign_Left)
+					.VAlign(VAlign_Center)
+					[
+						SNew(SRigConnectorTargetWidget)
+						.Outer(PerModuleInfos[0].GetBlueprint())
+						.ConnectorKey(ConnectorKey)
+						.IsArray(bIsArrayConnector)
+						.ExpandArrayByDefault(true)
+						.Targets(TargetKeys.GetValue())
+						.OnSetTargetArray(FRigConnectorTargetWidget_SetTargetArray::CreateSP(this, &FRigModuleInstanceDetails::OnConnectorTargetChanged, Connector))
+						.RigTreeDelegates(RigTreeDelegates)
+					];
+				}
+				else
+				{
+					ConnectorRow
+					.ValueContent()
+					.HAlign(HAlign_Left)
+					.VAlign(VAlign_Center)
+					[
+						SNew(STextBlock)
+						.Text(ControlRigModuleDetailsMultipleValues)
+						.Font(IDetailLayoutBuilder::GetDetailFont())
+					];
+				}
 			}
 		}
 	}
@@ -1205,73 +879,6 @@ FRigElementKeyRedirector FRigModuleInstanceDetails::GetConnections() const
 	return FRigElementKeyRedirector();
 }
 
-void FRigModuleInstanceDetails::PopulateConnectorTargetList(const FString InConnectorName)
-{
-	ConnectionListBox.FindRef(InConnectorName)->GetTreeView()->RefreshTreeView(true);
-}
-
-void FRigModuleInstanceDetails::PopulateConnectorCurrentTarget(TSharedPtr<SVerticalBox> InListBox, const FSlateBrush* InBrush, const FSlateColor& InColor, const FText& InTitle)
-{
-	static const FSlateBrush* RoundedBoxBrush = FControlRigEditorStyle::Get().GetBrush(TEXT("ControlRig.SpacePicker.RoundedRect"));
-	
-	TSharedPtr<SHorizontalBox> RowBox, ButtonBox;
-	InListBox->AddSlot()
-	.AutoHeight()
-	.VAlign(VAlign_Top)
-	.HAlign(HAlign_Fill)
-	.Padding(4.0, 0.0, 4.0, 0.0)
-	[
-		SNew( SButton )
-		.ButtonStyle(FAppStyle::Get(), "SimpleButton")
-		.ContentPadding(FMargin(0.0))
-		[
-			SAssignNew(RowBox, SHorizontalBox)
-			+ SHorizontalBox::Slot()
-			.AutoWidth()
-			.VAlign(VAlign_Center)
-			.HAlign(HAlign_Fill)
-			.Padding(0)
-			[
-				SNew(SBorder)
-				.Padding(FMargin(5.0, 2.0, 5.0, 2.0))
-				.BorderImage(RoundedBoxBrush)
-				.BorderBackgroundColor(FSlateColor(FLinearColor::Transparent))
-				.Content()
-				[
-					SAssignNew(ButtonBox, SHorizontalBox)
-					+ SHorizontalBox::Slot()
-					.AutoWidth()
-					.VAlign(VAlign_Center)
-					.HAlign(HAlign_Left)
-					.Padding(FMargin(0.f, 0.f, 3.f, 0.f))
-					[
-						SNew(SImage)
-						.Image(InBrush)
-						.ColorAndOpacity(InColor)
-					]
-
-					+ SHorizontalBox::Slot()
-					.AutoWidth()
-					.VAlign(VAlign_Center)
-					.HAlign(HAlign_Left)
-					.Padding(0)
-					[
-						SNew( STextBlock )
-						.Text( InTitle )
-						.Font( IDetailLayoutBuilder::GetDetailFont() )
-					]
-
-					+ SHorizontalBox::Slot()
-					.FillWidth(1.f)
-					[
-						SNew(SSpacer)
-					]
-				]
-			]
-		]
-	];
-}
-
 void FRigModuleInstanceDetails::OnConfigValueChanged(const FName InVariableName)
 {
 	TMap<FString, FString> ModuleValues;
@@ -1297,13 +904,9 @@ void FRigModuleInstanceDetails::OnConfigValueChanged(const FName InVariableName)
 	}
 }
 
-void FRigModuleInstanceDetails::OnConnectorTargetChanged(TSharedPtr<FRigTreeElement> Selection, ESelectInfo::Type SelectInfo, const FRigModuleConnector InConnector, int32 InTargetIndex)
+bool FRigModuleInstanceDetails::OnConnectorTargetChanged(TArray<FRigElementKey> InTargets, FRigModuleConnector InConnector)
 {
-	if (SelectInfo == ESelectInfo::OnNavigation)
-	{
-		return;
-	}
-	
+	bool bResult = true;
 	FScopedTransaction Transaction(LOCTEXT("ModuleHierarchyResolveConnector", "Resolve Connector"));
 	for (FPerModuleInfo& Info : PerModuleInfos)
 	{
@@ -1311,39 +914,25 @@ void FRigModuleInstanceDetails::OnConnectorTargetChanged(TSharedPtr<FRigTreeElem
 		{
 			FString ConnectorPath = FString::Printf(TEXT("%s:%s"), *Info.Path, *InConnector.Name);
 			FRigElementKey ConnectorKey(*ConnectorPath, ERigElementType::Connector);
-
-			FRigElementKey ChosenTargetKey;
-			if(Selection.IsValid())
-			{
-				ChosenTargetKey = Selection->Key;
-			}
-			
-			TArray<FRigElementKey> Targets = Info.GetBlueprint()->ModularRigModel.Connections.FindTargetsFromConnector(ConnectorKey);
-			if(Targets.IsValidIndex(InTargetIndex))
-			{
-				Targets[InTargetIndex] = ChosenTargetKey;
-			}
-			else
-			{
-				Targets.Add(ChosenTargetKey);
-			}
-
-			Targets.RemoveAll([](const FRigElementKey& Key) -> bool
-			{
-				return !Key.IsValid();
-			});
 			
-			if (!Targets.IsEmpty())
+			if (!InTargets.IsEmpty())
 			{
 				const FModularRigSettings& Settings = Info.GetModularRig()->GetModularRigSettings();
-				Controller->ConnectConnectorToElements(ConnectorKey, Targets, true, Settings.bAutoResolve);
+				if(!Controller->ConnectConnectorToElements(ConnectorKey, InTargets, true, Settings.bAutoResolve))
+				{
+					bResult = false;
+				}
 			}
 			else
 			{
-				Controller->DisconnectConnector(ConnectorKey);
+				if(!Controller->DisconnectConnector(ConnectorKey))
+				{
+					bResult = false;
+				}
 			}
 		}
 	}
+	return bResult;
 }
 
 const FRigModuleInstanceDetails::FPerModuleInfo& FRigModuleInstanceDetails::FindModule(const FString& InPath) const
diff --git a/Engine/Plugins/Animation/ControlRig/Source/ControlRigEditor/Private/ControlRigModuleDetails.h b/Engine/Plugins/Animation/ControlRig/Source/ControlRigEditor/Private/ControlRigModuleDetails.h
--- a/Engine/Plugins/Animation/ControlRig/Source/ControlRigEditor/Private/ControlRigModuleDetails.h
+++ b/Engine/Plugins/Animation/ControlRig/Source/ControlRigEditor/Private/ControlRigModuleDetails.h
@@ -50,15 +50,9 @@ public:
 	FText GetRigClassPath() const;
 	TArray<FRigModuleConnector> GetConnectors() const;
 	FRigElementKeyRedirector GetConnections() const;
-	void PopulateConnectorTargetList(const FString InConnectorKey);
-	void PopulateConnectorCurrentTarget(
-		TSharedPtr<SVerticalBox> InListBox,
-		const FSlateBrush* InBrush,
-		const FSlateColor& InColor,
-		const FText& InTitle);
 
 	void OnConfigValueChanged(const FName InVariableName);
-	void OnConnectorTargetChanged(TSharedPtr<FRigTreeElement> Selection, ESelectInfo::Type SelectInfo, const FRigModuleConnector InConnector, int32 InTargetIndex);
+	bool OnConnectorTargetChanged(TArray<FRigElementKey> InTargets, FRigModuleConnector InConnector);
 	
 	struct FPerModuleInfo
 	{
@@ -146,7 +140,6 @@ protected:
 	FReply OnClearTargetsForArrayConnector(const FString InConnectorName, const TSharedRef<IPropertyUtilities> PropertyUtilities);
 
 	TArray<FPerModuleInfo> PerModuleInfos;
-	TMap<FString, TSharedPtr<SSearchableRigHierarchyTreeView>> ConnectionListBox;
 
 	/** Helper buttons. */
 	TMap<FString, TSharedPtr<SButton>> UseSelectedButton;
diff --git a/Engine/Plugins/Animation/ControlRig/Source/ControlRigEditor/Private/Editor/SModularRigModel.cpp b/Engine/Plugins/Animation/ControlRig/Source/ControlRigEditor/Private/Editor/SModularRigModel.cpp
--- a/Engine/Plugins/Animation/ControlRig/Source/ControlRigEditor/Private/Editor/SModularRigModel.cpp
+++ b/Engine/Plugins/Animation/ControlRig/Source/ControlRigEditor/Private/Editor/SModularRigModel.cpp
@@ -131,7 +131,7 @@ void SModularRigModel::Construct(const FArguments& InArgs, TSharedRef<FControlRi
 		.DefaultLabel(FText::FromName(SModularRigTreeView::Column_Module))
 		.HAlignCell(HAlign_Left)
 		.HAlignHeader(HAlign_Left)
-		.VAlignCell(VAlign_Fill)
+		.VAlignCell(VAlign_Top)
 	);
 	HeaderRowWidget->AddColumn(
 		SHeaderRow::Column(SModularRigTreeView::Column_Tags)
@@ -139,14 +139,14 @@ void SModularRigModel::Construct(const FArguments& InArgs, TSharedRef<FControlRi
 		.HAlignCell(HAlign_Fill)
 		.HAlignHeader(HAlign_Fill)
 		.FixedWidth(16.f)
-		.VAlignCell(VAlign_Center)
+		.VAlignCell(VAlign_Top)
 	);
 	HeaderRowWidget->AddColumn(
 		SHeaderRow::Column(SModularRigTreeView::Column_Connector)
 		.DefaultLabel(FText::FromName(SModularRigTreeView::Column_Connector))
 		.HAlignCell(HAlign_Left)
 		.HAlignHeader(HAlign_Left)
-		.VAlignCell(VAlign_Center)
+		.VAlignCell(VAlign_Top)
 	);
 	HeaderRowWidget->AddColumn(
 		SHeaderRow::Column(SModularRigTreeView::Column_Buttons)
@@ -154,7 +154,7 @@ void SModularRigModel::Construct(const FArguments& InArgs, TSharedRef<FControlRi
 		.ManualWidth(60)
 		.HAlignCell(HAlign_Left)
 		.HAlignHeader(HAlign_Left)
-		.VAlignCell(VAlign_Center)
+		.VAlignCell(VAlign_Top)
 	);
 	
 	ChildSlot
@@ -1194,7 +1194,7 @@ void SModularRigModel::HandleSwapClassForModules(const TArray<FString>& InPaths)
 	SwapModulesDialog->ShowNormal();
 }
 
-void SModularRigModel::HandleConnectorResolved(const FRigElementKey& InConnector, const FRigElementKey& InTarget)
+void SModularRigModel::HandleConnectorResolved(const FRigElementKey& InConnector, const TArray<FRigElementKey>& InTargets)
 {
 	if (ControlRigBlueprint.IsValid())
 	{
@@ -1205,7 +1205,7 @@ void SModularRigModel::HandleConnectorResolved(const FRigElementKey& InConnector
 
 		if (const UModularRig* ModularRig = GetModularRig())
 		{
-			Controller->ConnectConnectorToElement(InConnector, InTarget, true, ModularRig->GetModularRigSettings().bAutoResolve);
+			Controller->ConnectConnectorToElements(InConnector, InTargets, true, ModularRig->GetModularRigSettings().bAutoResolve);
 		}
 	}
 }
diff --git a/Engine/Plugins/Animation/ControlRig/Source/ControlRigEditor/Private/Editor/SModularRigModel.h b/Engine/Plugins/Animation/ControlRig/Source/ControlRigEditor/Private/Editor/SModularRigModel.h
--- a/Engine/Plugins/Animation/ControlRig/Source/ControlRigEditor/Private/Editor/SModularRigModel.h
+++ b/Engine/Plugins/Animation/ControlRig/Source/ControlRigEditor/Private/Editor/SModularRigModel.h
@@ -92,7 +92,7 @@ private:
 	void HandleSwapClassForModules(const TArray<FString>& InPaths);
 
 	/** Resolve connector */
-	void HandleConnectorResolved(const FRigElementKey& InConnector, const FRigElementKey& InTarget);
+	void HandleConnectorResolved(const FRigElementKey& InConnector, const TArray<FRigElementKey>& InTargets);
 
 	/** UnResolve connector */
 	void HandleConnectorDisconnect(const FRigElementKey& InConnector);
diff --git a/Engine/Plugins/Animation/ControlRig/Source/ControlRigEditor/Private/Editor/SModularRigTreeView.cpp b/Engine/Plugins/Animation/ControlRig/Source/ControlRigEditor/Private/Editor/SModularRigTreeView.cpp
--- a/Engine/Plugins/Animation/ControlRig/Source/ControlRigEditor/Private/Editor/SModularRigTreeView.cpp
+++ b/Engine/Plugins/Animation/ControlRig/Source/ControlRigEditor/Private/Editor/SModularRigTreeView.cpp
@@ -31,6 +31,7 @@
 #include "ScopedTransaction.h"
 #include "DetailLayoutBuilder.h"
 #include "Widgets/SRigVMVariantTagWidget.h"
+#include "Editor/SRigConnectorTargetWidget.h"
 
 #define LOCTEXT_NAMESPACE "SModularRigTreeView"
 
@@ -129,147 +130,11 @@ void SModularRigModelItem::Construct(const FArguments& InArgs, const TSharedRef<
 		.ShowWires(true), OwnerTable);
 }
 
-
-void SModularRigModelItem::PopulateConnectorTargetList(const FRigElementKey InConnectorKey)
-{
-	if (!WeakRigTreeElement.IsValid())
-	{
-		return;
-	}
-
-	if (const UModularRig* ModularRig = Delegates.GetModularRig())
-	{
-		if (URigHierarchy* Hierarchy = ModularRig->GetHierarchy())
-		{
-			if (Cast<FRigConnectorElement>(Hierarchy->Find(InConnectorKey)))
-			{
-				ConnectorComboBox->GetTreeView()->RefreshTreeView(true);
-
-				// set the focus to the search box so you can start typing right away
-				RegisterActiveTimer(0.f, FWidgetActiveTimerDelegate::CreateLambda([this](double,float)
-				{
-					FSlateApplication::Get().ForEachUser([this](FSlateUser& User)
-					{
-						User.SetFocus(ConnectorComboBox->GetSearchBox());
-					});
-					return EActiveTimerReturnType::Stop;
-				}));
-			}
-		}
-	}
-}
-
-void SModularRigModelItem::PopulateConnectorCurrentTarget(TSharedPtr<SVerticalBox> InListBox, const FRigElementKey& InConnectorKey, const TArray<FRigElementKey>& InTargetKeys, const FSlateBrush* InBrush, const FSlateColor& InColor, const FText& InTitle)
+bool SModularRigModelItem::OnConnectorTargetChanged(TArray<FRigElementKey> InTargets, const FRigElementKey InConnectorKey)
 {
-	static const FSlateBrush* RoundedBoxBrush = FControlRigEditorStyle::Get().GetBrush(TEXT("ControlRig.SpacePicker.RoundedRect"));
-
-	TAttribute<FSlateColor> TextColor;
-	const FRigElementKey FirstTargetKey = InTargetKeys.Num() == 1 ? InTargetKeys[0] : FRigElementKey();
-
-	if(InTargetKeys.Num() == 1)
-	{
-		TextColor = TAttribute<FSlateColor>::CreateLambda([this, InConnectorKey, FirstTargetKey]() -> FSlateColor
-		{
-			if(const UModularRig* ModularRig = Delegates.GetModularRig())
-			{
-				if(const URigHierarchy* Hierarchy = ModularRig->GetHierarchy())
-				{
-					if(!Hierarchy->Contains(FirstTargetKey))
-					{
-						if(!FirstTargetKey.IsValid())
-						{
-							if(const FRigConnectorElement* Connector = Hierarchy->Find<FRigConnectorElement>(InConnectorKey))
-							{
-								if(Connector->IsOptional())
-								{
-									return FSlateColor::UseForeground();
-								}
-							}
-						}
-						return FLinearColor::Red;
-					}
-				}
-			}
-			return FSlateColor::UseForeground();
-		});
-	}
-	
-	TSharedPtr<SHorizontalBox> RowBox, ButtonBox;
-	InListBox->AddSlot()
-	.AutoHeight()
-	.VAlign(VAlign_Top)
-	.HAlign(HAlign_Left)
-	.Padding(0.0, 0.0, 4.0, 0.0)
-	[
-		SNew(SHorizontalBox)
-		+ SHorizontalBox::Slot()
-		.AutoWidth()
-		.HAlign(HAlign_Left)
-		.Padding(0, 0, 0, 0)
-		[
-			SNew( SButton )
-			.ButtonStyle(FAppStyle::Get(), "SimpleButton")
-			.ContentPadding(FMargin(0.0))
-			.OnClicked_Lambda([this, InConnectorKey]()
-			{
-				PopulateConnectorTargetList(InConnectorKey);
-				return FReply::Handled();
-			})
-			.IsEnabled(InTargetKeys.Num() == 1)
-			[
-				SAssignNew(RowBox, SHorizontalBox)
-				+ SHorizontalBox::Slot()
-				.AutoWidth()
-				.VAlign(VAlign_Center)
-				.HAlign(HAlign_Fill)
-				.Padding(0)
-				[
-					SNew(SBorder)
-					.Padding(FMargin(2.0, 2.0, 5.0, 2.0))
-					.BorderImage(RoundedBoxBrush)
-					.BorderBackgroundColor(FSlateColor(FLinearColor::Transparent))
-					.Content()
-					[
-						SAssignNew(ButtonBox, SHorizontalBox)
-						+ SHorizontalBox::Slot()
-						.AutoWidth()
-						.VAlign(VAlign_Center)
-						.HAlign(HAlign_Left)
-						.Padding(FMargin(0.f, 0.f, 3.f, 0.f))
-						[
-							SNew(SImage)
-							.Image(InBrush)
-							.ColorAndOpacity(InColor)
-						]
-
-						+ SHorizontalBox::Slot()
-						.AutoWidth()
-						.VAlign(VAlign_Center)
-						.HAlign(HAlign_Left)
-						.Padding(0)
-						[
-							SNew( STextBlock )
-							.Text( InTitle )
-							.Font( IDetailLayoutBuilder::GetDetailFont() )
-							.ColorAndOpacity(TextColor)
-						]
-					]
-				]
-			]
-		]
-	];
-}
-
-void SModularRigModelItem::OnConnectorTargetChanged(TSharedPtr<FRigTreeElement> Selection, ESelectInfo::Type SelectInfo, const FRigElementKey InConnectorKey)
-{
-	if (!Selection.IsValid() || SelectInfo == ESelectInfo::OnNavigation)
-	{
-		return;
-	}
-	
 	FScopedTransaction Transaction(LOCTEXT("ModuleHierarchyResolveConnector", "Resolve Connector"));
-	TSharedPtr<FRigTreeElement> NewSelection = Selection;
-	Delegates.HandleResolveConnector(InConnectorKey, NewSelection->Key);
+	Delegates.HandleResolveConnector(InConnectorKey, InTargets);
+	return false;
 }
 
 void SModularRigModelItem::OnNameCommitted(const FText& InText, ETextCommit::Type InCommitType) const
@@ -296,13 +161,14 @@ TSharedRef<SWidget> SModularRigModelItem::GenerateWidgetForColumn(const FName& C
 {
 	if(ColumnName == SModularRigTreeView::Column_Module)
 	{
+		static constexpr float TopPadding = 2.f;
+		
 		TSharedPtr< SInlineEditableTextBlock > InlineWidget;
-
 		TSharedRef<SWidget> Widget = SNew(SHorizontalBox)
 
 		+SHorizontalBox::Slot()
 		.AutoWidth()
-		.Padding(6, 0, 0, 0)
+		.Padding(6, TopPadding, 0, 0)
 		.VAlign(VAlign_Fill)
 		[
 			SNew(SExpanderArrow, SharedThis(this))
@@ -314,7 +180,7 @@ TSharedRef<SWidget> SModularRigModelItem::GenerateWidgetForColumn(const FName& C
 		.MaxWidth(25)
 		.HAlign(HAlign_Left)
 		.VAlign(VAlign_Center)
-		.Padding(FMargin(0.f, 0.f, 3.f, 0.f))
+		.Padding(FMargin(0.f, TopPadding, 3.f, 0.f))
 		[
 			SNew(SVerticalBox)
 			+SVerticalBox::Slot()
@@ -344,6 +210,7 @@ TSharedRef<SWidget> SModularRigModelItem::GenerateWidgetForColumn(const FName& C
 		+ SHorizontalBox::Slot()
 		.AutoWidth()
 		.VAlign(VAlign_Center)
+		.Padding(0, TopPadding, 0, 0)
 		[
 			SAssignNew(InlineWidget, SInlineEditableTextBlock)
 			.Text(this, &SModularRigModelItem::GetName, true)
@@ -436,139 +303,33 @@ TSharedRef<SWidget> SModularRigModelItem::GenerateWidgetForColumn(const FName& C
 	}
 	if(ColumnName == SModularRigTreeView::Column_Connector)
 	{
-		TSharedPtr<SVerticalBox> ComboButtonBox;
-
+		bool bIsArrayConnector = false;
 		if (const UModularRig* ModularRig = Delegates.GetModularRig())
 		{
-			if (const FRigModuleInstance* Module = ModularRig->FindModule(WeakRigTreeElement.Pin()->ModulePath))
+			if (URigHierarchy* Hierarchy = ModularRig->GetHierarchy())
 			{
-				if (URigHierarchy* Hierarchy = ModularRig->GetHierarchy())
+				if (const FRigConnectorElement* ConnectorElement = Cast<FRigConnectorElement>(Hierarchy->Find(ConnectorKey)))
 				{
-					if (const FRigConnectorElement* ConnectorElement = Cast<FRigConnectorElement>(Hierarchy->Find(ConnectorKey)))
-					{
-						if(ConnectorElement->IsArrayConnector())
-						{
-							return SNew( STextBlock )
-							.Text( LOCTEXT("ArrayConnectorPleaseRelyOnDetailsPanel", "WIP. Array Connector. Use Details Panel.") )
-							.Font( IDetailLayoutBuilder::GetDetailFont() );
-						}
-					}
+					bIsArrayConnector = ConnectorElement->IsArrayConnector();
 				}
 			}
 		}
+
+		TSharedPtr<SVerticalBox> ComboButtonBox;
 		
-		FRigTreeDelegates TreeDelegates;
-		TreeDelegates.OnGetHierarchy = FOnGetRigTreeHierarchy::CreateLambda([this]()
+		FRigTreeDelegates RigTreeDelegates;
+		RigTreeDelegates.OnGetHierarchy = FOnGetRigTreeHierarchy::CreateLambda([this]()
 		{
 			return Delegates.GetModularRig()->GetHierarchy();
 		});
-		TreeDelegates.OnRigTreeIsItemVisible = FOnRigTreeIsItemVisible::CreateLambda([this](const FRigElementKey& InTarget)
-		{
-			if(!ConnectorMatches.IsSet() && WeakRigTreeElement.IsValid())
-			{
-				if (const UModularRig* ModularRig = Delegates.GetModularRig())
-				{
-					if (const FRigModuleInstance* Module = ModularRig->FindModule(WeakRigTreeElement.Pin()->ModulePath))
-					{
-						if (URigHierarchy* Hierarchy = ModularRig->GetHierarchy())
-						{
-							if (const FRigConnectorElement* ConnectorElement = Cast<FRigConnectorElement>(Hierarchy->Find(ConnectorKey)))
-							{
-								const UModularRigRuleManager* RuleManager = ModularRig->GetHierarchy()->GetRuleManager();
-								ConnectorMatches = RuleManager->FindMatches(ConnectorElement, Module, ModularRig->GetElementKeyRedirector());
-							}
-						}
-					}
-				}
-			}
-
-			if(ConnectorMatches.IsSet())
-			{
-				return ConnectorMatches.GetValue().ContainsMatch(InTarget); 
-			}
-			return true;
-		});
-		TreeDelegates.OnGetSelection.BindLambda([this]() -> TArray<FRigElementKey>
-		{
-			const FRigElementKeyRedirector Redirector = Delegates.GetModularRig()->GetElementKeyRedirector();
-			if (const FRigElementKeyRedirector::FKeyArray* Cache = Redirector.FindExternalKey(ConnectorKey))
-			{
-				TArray<FRigElementKey> Keys;
-				Keys.Append(*Cache);
-				return Keys;
-			}
-			return {};
-		});
-		TreeDelegates.OnSelectionChanged.BindSP(this, &SModularRigModelItem::OnConnectorTargetChanged, ConnectorKey);
-
-		TSharedRef<SWidget> Widget = SNew(SHorizontalBox)
-		+ SHorizontalBox::Slot()
-		.AutoWidth()
-		.HAlign(HAlign_Right)
-		.VAlign(VAlign_Center)
-		[
-			SNew( SComboButton )
-			.ContentPadding(3)
-			.MenuPlacement(MenuPlacement_BelowAnchor)
-			.OnComboBoxOpened(this, &SModularRigModelItem::PopulateConnectorTargetList, ConnectorKey)
-			.ButtonContent()
-			[
-				// Wrap in configurable box to restrain height/width of menu
-				SNew(SBox)
-				.MinDesiredWidth(200.0f)
-				.Padding(0, 0, 0, 0)
-				.HAlign(HAlign_Left)
-				[
-					SAssignNew(ComboButtonBox, SVerticalBox)
-				]
-			]
-			.MenuContent()
-			[
-				SNew(SBorder)
-				.Visibility(EVisibility::Visible)
-				.BorderImage(FAppStyle::GetBrush("Menu.Background"))
-				[
-					SAssignNew(ConnectorComboBox, SSearchableRigHierarchyTreeView)
-						.RigTreeDelegates(TreeDelegates)
-						.MaxHeight(300)
-				]
-			]
-		];
 
-		const UModularRig* ModularRig = Delegates.GetModularRig();
-		const FRigElementKeyRedirector& Redirector = ModularRig->GetElementKeyRedirector();
-		TArray<FRigElementKey> CurrentTargetKeys;
-		if (const FRigElementKeyRedirector::FKeyArray* Keys = Redirector.FindExternalKey(ConnectorKey))
-		{
-			CurrentTargetKeys.Append(*Keys);
-		}
-		if(CurrentTargetKeys.IsEmpty())
-		{
-			CurrentTargetKeys.Add(FRigElementKey());
-		}
-
-		TPair<const FSlateBrush*, FSlateColor> IconAndColor;
-		FText CurrentTargetShortName;
-		
-		if(CurrentTargetKeys.Num() == 1)
-		{
-			const FRigElementKey& FirstTargetKey = CurrentTargetKeys[0];
-			IconAndColor = SRigHierarchyItem::GetBrushForElementType(ModularRig->GetHierarchy(), FirstTargetKey);
-
-			CurrentTargetShortName = FText::FromString(ModularRig->GetShortestDisplayPathForElement(FirstTargetKey, false));
-			if(CurrentTargetShortName.IsEmpty())
-			{
-				CurrentTargetShortName = FText::FromName(FirstTargetKey.Name);
-			}
-		}
-		else
-		{
-			CurrentTargetShortName = LOCTEXT("MultipleValues", "Multiple Values");
-			ConnectorComboBox->SetEnabled(false);
-		}
-		PopulateConnectorCurrentTarget(ComboButtonBox, ConnectorKey, CurrentTargetKeys, IconAndColor.Key, IconAndColor.Value, CurrentTargetShortName);
-
-		return Widget;
+		return SNew(SRigConnectorTargetWidget)
+			.Outer(const_cast<UModularRig*>(Delegates.GetModularRig()))
+			.ConnectorKey(ConnectorKey)
+			.IsArray(bIsArrayConnector)
+			.Targets(GetTargetKeys())
+			.OnSetTargetArray(FRigConnectorTargetWidget_SetTargetArray::CreateSP(this, &SModularRigModelItem::OnConnectorTargetChanged, ConnectorKey))
+			.RigTreeDelegates(RigTreeDelegates);
 	}
 	if(ColumnName == SModularRigTreeView::Column_Buttons)
 	{
@@ -626,7 +387,7 @@ TSharedRef<SWidget> SModularRigModelItem::GenerateWidgetForColumn(const FName& C
 					const TArray<FRigElementKey>& Selected = ModularRig->GetHierarchy()->GetSelectedKeys();
 					if (Selected.Num() > 0)
 					{
-						Delegates.HandleResolveConnector(ConnectorKey, Selected[0]);
+						Delegates.HandleResolveConnector(ConnectorKey, Selected);
 					}
 				}
 				return FReply::Handled();
@@ -665,9 +426,11 @@ TSharedRef<SWidget> SModularRigModelItem::GenerateWidgetForColumn(const FName& C
 					const FRigElementKeyRedirector& Redirector = ModularRig->GetElementKeyRedirector();
 					if (const FRigElementKeyRedirector::FKeyArray* TargetKeys = Redirector.FindExternalKey(ConnectorKey))
 					{
+						bool bClearSelection = true;
 						for(const FRigElementKey& TargetKey : *TargetKeys)
 						{
-							ModularRig->GetHierarchy()->GetController()->SelectElement(TargetKey, true, true);
+							ModularRig->GetHierarchy()->GetController()->SelectElement(TargetKey, true, bClearSelection);
+							bClearSelection = false;
 						}
 					}
 				}
@@ -710,6 +473,24 @@ FText SModularRigModelItem::GetItemTooltip() const
 	return FullName;
 }
 
+TArray<FRigElementKey> SModularRigModelItem::GetTargetKeys() const
+{
+	TArray<FRigElementKey> Result;
+	if (const UModularRig* ModularRig = Delegates.GetModularRig())
+	{
+		const FRigElementKeyRedirector ElementKeyRedirector = ModularRig->GetElementKeyRedirector();
+		if(const FRigElementKeyRedirector::FCachedKeyArray* Cache = ElementKeyRedirector.Find(ConnectorKey))
+		{
+			Result.Reserve(Cache->Num());
+			for(int32 Index = 0; Index < Cache->Num(); Index++)
+			{
+				Result.Add(Cache->operator[](Index).GetKey());
+			}
+		}
+	}
+	return Result;
+}
+
 //////////////////////////////////////////////////////////////
 /// SModularRigTreeView
 ///////////////////////////////////////////////////////////
diff --git a/Engine/Plugins/Animation/ControlRig/Source/ControlRigEditor/Private/Editor/SModularRigTreeView.h b/Engine/Plugins/Animation/ControlRig/Source/ControlRigEditor/Private/Editor/SModularRigTreeView.h
--- a/Engine/Plugins/Animation/ControlRig/Source/ControlRigEditor/Private/Editor/SModularRigTreeView.h
+++ b/Engine/Plugins/Animation/ControlRig/Source/ControlRigEditor/Private/Editor/SModularRigTreeView.h
@@ -16,7 +16,7 @@ class FModularRigTreeElement;
 DECLARE_DELEGATE_RetVal(const UModularRig*, FOnGetModularRigTreeRig);
 DECLARE_DELEGATE_OneParam(FOnModularRigTreeRequestDetailsInspection, const FString&);
 DECLARE_DELEGATE_RetVal_TwoParams(FName, FOnModularRigTreeRenameElement, const FString& /*OldPath*/, const FName& /*NewName*/);
-DECLARE_DELEGATE_TwoParams(FOnModularRigTreeResolveConnector, const FRigElementKey& /*Connector*/, const FRigElementKey& /*Target*/);
+DECLARE_DELEGATE_TwoParams(FOnModularRigTreeResolveConnector, const FRigElementKey& /*Connector*/, const TArray<FRigElementKey>& /*Targets*/);
 DECLARE_DELEGATE_OneParam(FOnModularRigTreeDisconnectConnector, const FRigElementKey& /*Connector*/);
 DECLARE_DELEGATE_RetVal_ThreeParams(bool, FOnModularRigTreeVerifyElementNameChanged, const FString& /*OldPath*/, const FName& /*NewName*/, FText& /*OutErrorMessage*/);
 
@@ -75,11 +75,11 @@ public:
 		return false;
 	}
 
-	bool HandleResolveConnector(const FRigElementKey& InConnector, const FRigElementKey& InTarget)
+	bool HandleResolveConnector(const FRigElementKey& InConnector, const TArray<FRigElementKey>& InTargets)
 	{
 		if(OnResolveConnector.IsBound())
 		{
-			OnResolveConnector.Execute(InConnector, InTarget);
+			OnResolveConnector.Execute(InConnector, InTargets);
 			return true;
 		}
 		return false;
@@ -157,15 +157,7 @@ class SModularRigModelItem : public SMultiColumnTableRow<TSharedPtr<FModularRigT
 public:
 	
 	void Construct(const FArguments& InArgs, const TSharedRef<STableViewBase>& OwnerTable, TSharedRef<FModularRigTreeElement> InRigTreeElement, TSharedPtr<SModularRigTreeView> InTreeView, bool bPinned);
-	void PopulateConnectorTargetList(FRigElementKey InConnectorKey);
-	void PopulateConnectorCurrentTarget(
-		TSharedPtr<SVerticalBox> InListBox, 
-		const FRigElementKey& InConnectorKey,
-		const TArray<FRigElementKey>& InTargetKeys,
-		const FSlateBrush* InBrush,
-		const FSlateColor& InColor,
-		const FText& InTitle);
-	void OnConnectorTargetChanged(TSharedPtr<FRigTreeElement> Selection, ESelectInfo::Type SelectInfo, const FRigElementKey InConnectorKey);
+	bool OnConnectorTargetChanged(TArray<FRigElementKey> InTargets, const FRigElementKey InConnectorKey);
 
 	void OnNameCommitted(const FText& InText, ETextCommit::Type InCommitType) const;
 	bool OnVerifyNameChanged(const FText& InText, FText& OutErrorMessage);
@@ -184,6 +176,7 @@ private:
 
 	FText GetName(bool bUseShortName) const;
 	FText GetItemTooltip() const;
+	TArray<FRigElementKey> GetTargetKeys() const;
 
 	friend class SModularRigTreeView; 
 };
diff --git a/Engine/Plugins/Animation/ControlRig/Source/ControlRigEditor/Private/Editor/SRigConnectorTargetWidget.cpp b/Engine/Plugins/Animation/ControlRig/Source/ControlRigEditor/Private/Editor/SRigConnectorTargetWidget.cpp
new file mode 100644
--- /dev/null
+++ b/Engine/Plugins/Animation/ControlRig/Source/ControlRigEditor/Private/Editor/SRigConnectorTargetWidget.cpp
@@ -0,0 +1,419 @@
+// Copyright Epic Games, Inc. All Rights Reserved.
+
+#include "SRigConnectorTargetWidget.h"
+
+#include "DetailWidgetRow.h"
+#include "DetailLayoutBuilder.h"
+#include "IDetailChildrenBuilder.h"
+#include "ISinglePropertyView.h"
+#include "Modules/ModuleManager.h"
+#include "PropertyEditorModule.h"
+#include "PropertyHandle.h"
+#include "ControlRigEditorStyle.h"
+#include "ModularRig.h"
+#include "ModularRigRuleManager.h"
+#include "Widgets/Input/SSearchBox.h"
+#include "Framework/Application/SlateApplication.h"
+#include "Framework/Application/SlateUser.h"
+
+#define LOCTEXT_NAMESPACE "SRigConnectorTargetWidget"
+
+SRigConnectorTargetComboButton::~SRigConnectorTargetComboButton()
+{
+}
+
+void SRigConnectorTargetComboButton::Construct(const FArguments& InArgs)
+{
+	ConnectorKey = InArgs._ConnectorKey;
+	TargetKey = InArgs._TargetKey;
+	ArrayIndex = InArgs._ArrayIndex;
+	RigTreeDelegates = InArgs._RigTreeDelegates;
+	OnSetTarget = InArgs._OnSetTarget;
+
+	RigTreeDelegates.OnGetSelection.BindLambda([this]() -> TArray<FRigElementKey>
+	{
+		return {TargetKey.Get()};
+	});
+	RigTreeDelegates.OnSelectionChanged.BindSP(this, &SRigConnectorTargetComboButton::OnConnectorTargetPicked);
+
+	SearchableTreeView = SNew(SSearchableRigHierarchyTreeView)
+		.RigTreeDelegates(RigTreeDelegates);
+
+	SComboButton::FArguments ComboButtonArgs;
+	ComboButtonArgs.ContentPadding(InArgs._ContentPadding);
+	ComboButtonArgs.MenuPlacement(InArgs._MenuPlacement);
+	ComboButtonArgs.OnComboBoxOpened(this, &SRigConnectorTargetComboButton::OnComboBoxOpened);
+	ComboButtonArgs.ButtonContent()
+	[
+		// Wrap in configurable box to restrain height/width of menu
+		SNew(SBox)
+		.MinDesiredWidth(InArgs._ButtonMinDesiredWith)
+		[
+			SAssignNew(VerticalButtonBox, SVerticalBox)
+		]
+	];
+	ComboButtonArgs.MenuContent()
+	[
+		SNew(SBorder)
+		.Visibility(EVisibility::Visible)
+		.BorderImage(FAppStyle::GetBrush("Menu.Background"))
+		[
+			SearchableTreeView.ToSharedRef()
+		]
+	];
+
+	SComboButton::Construct(ComboButtonArgs);
+	PopulateButtonBox();
+}
+
+void SRigConnectorTargetComboButton::OnComboBoxOpened()
+{
+	SearchableTreeView->GetTreeView()->RefreshTreeView();
+
+	// set the focus to the search box so you can start typing right away
+	RegisterActiveTimer(0.f, FWidgetActiveTimerDelegate::CreateLambda([this](double,float)
+	{
+		FSlateApplication::Get().ForEachUser([this](FSlateUser& User)
+		{
+			User.SetFocus(SearchableTreeView->GetSearchBox());
+		});
+		return EActiveTimerReturnType::Stop;
+	}));
+}
+
+void SRigConnectorTargetComboButton::PopulateButtonBox()
+{
+	const URigHierarchy* Hierarchy = RigTreeDelegates.GetHierarchy();
+	static const FSlateBrush* RoundedBoxBrush = FControlRigEditorStyle::Get().GetBrush(TEXT("ControlRig.SpacePicker.RoundedRect"));
+
+	auto GetIconAndColor = [this, Hierarchy]
+	{
+		FRigElementKey CurrentTargetKey(NAME_None, ERigElementType::Bone);
+		if(TargetKey.IsSet() || TargetKey.IsBound())
+		{
+			CurrentTargetKey = TargetKey.Get();
+		}
+		return SRigHierarchyItem::GetBrushForElementType(Hierarchy, CurrentTargetKey);
+	};
+	
+	TSharedPtr<SHorizontalBox> RowBox, ButtonBox;
+	VerticalButtonBox->AddSlot()
+	.AutoHeight()
+	.VAlign(VAlign_Top)
+	.HAlign(HAlign_Fill)
+	.Padding(4.0, 0.0, 4.0, 0.0)
+	[
+		SNew( SButton )
+		.ButtonStyle(FAppStyle::Get(), "SimpleButton")
+		.ContentPadding(FMargin(0.0))
+		[
+			SAssignNew(RowBox, SHorizontalBox)
+			+ SHorizontalBox::Slot()
+			.AutoWidth()
+			.VAlign(VAlign_Center)
+			.HAlign(HAlign_Fill)
+			.Padding(0)
+			[
+				SNew(SBorder)
+				.Padding(FMargin(5.0, 2.0, 5.0, 2.0))
+				.BorderImage(RoundedBoxBrush)
+				.BorderBackgroundColor(FSlateColor(FLinearColor::Transparent))
+				.Content()
+				[
+					SAssignNew(ButtonBox, SHorizontalBox)
+					+ SHorizontalBox::Slot()
+					.AutoWidth()
+					.VAlign(VAlign_Center)
+					.HAlign(HAlign_Left)
+					.Padding(FMargin(0.f, 0.f, 3.f, 0.f))
+					[
+						SNew(SImage)
+						.Image_Lambda([GetIconAndColor]()
+						{
+							return GetIconAndColor().Key;
+						})
+						.ColorAndOpacity_Lambda([GetIconAndColor]()
+						{
+							return GetIconAndColor().Value;
+						})
+					]
+
+					+ SHorizontalBox::Slot()
+					.AutoWidth()
+					.VAlign(VAlign_Center)
+					.HAlign(HAlign_Left)
+					.Padding(0)
+					[
+						SNew( STextBlock )
+						.Text_Lambda([this]()
+						{
+							return FText::FromName(TargetKey.Get().Name);
+						})
+						.Font( IDetailLayoutBuilder::GetDetailFont() )
+					]
+
+					+ SHorizontalBox::Slot()
+					.FillWidth(1.f)
+					[
+						SNew(SSpacer)
+					]
+				]
+			]
+		]
+	];
+}
+
+void SRigConnectorTargetComboButton::OnConnectorTargetPicked(TSharedPtr<FRigTreeElement> Selection, ESelectInfo::Type SelectInfo)
+{
+	if(SelectInfo == ESelectInfo::OnNavigation)
+	{
+		return;
+	}
+	if(OnSetTarget.IsBound() && Selection.IsValid())
+	{
+		if(Selection->Key != TargetKey.Get())
+		{
+			(void)OnSetTarget.Execute(Selection->Key);
+		}
+	}
+}
+
+SRigConnectorTargetWidget::~SRigConnectorTargetWidget()
+{
+}
+
+void SRigConnectorTargetWidget::Construct(const FArguments& InArgs)
+{
+	check(InArgs._Outer);
+
+	SetPadding(InArgs._Padding);
+	bIsArray = InArgs._IsArray;
+	Connector = InArgs._ConnectorKey;
+	RigTreeDelegates = InArgs._RigTreeDelegates;
+
+	if(!RigTreeDelegates.OnRigTreeIsItemVisible.IsBound())
+	{
+		TArray<FRigElementKey> PotentialTargets;
+		if(URigHierarchy* Hierarchy = RigTreeDelegates.GetHierarchy())
+		{
+			if(const UControlRig* ControlRig = Hierarchy->GetTypedOuter<UControlRig>())
+			{
+				if(const FRigConnectorElement* ConnectorElement = Hierarchy->Find<FRigConnectorElement>(Connector))
+				{
+					if(const UModularRig* ModularRig = Cast<UModularRig>(ControlRig))
+					{
+						if(const FRigModuleInstance* Module = ModularRig->FindModule(ConnectorElement))
+						{
+							const UModularRigRuleManager* RuleManager = Hierarchy->GetRuleManager();
+							TArray<FRigElementResolveResult> Matches = RuleManager->FindMatches(ConnectorElement, Module, ModularRig->GetElementKeyRedirector()).GetMatches();
+							PotentialTargets.Reserve(Matches.Num());
+							for(const FRigElementResolveResult& SingleMatch : Matches)
+							{
+								PotentialTargets.Add(SingleMatch.GetKey());
+							}
+						}
+					}
+				}
+			}
+		}
+
+		RigTreeDelegates.OnRigTreeIsItemVisible.BindLambda([PotentialTargets](const FRigElementKey& InItem)
+		{
+			return InItem.IsValid() && (PotentialTargets.IsEmpty() || PotentialTargets.Contains(InItem)); 
+		});
+	}
+
+	OnSetTargetArray = InArgs._OnSetTargetArray;
+
+	if(bIsArray)
+	{
+		TargetsDetailWrapper = TStrongObjectPtr<URigConnectorTargetsDetailWrapper>(NewObject<URigConnectorTargetsDetailWrapper>(InArgs._Outer, NAME_None, RF_Public | RF_Transient | RF_TextExportTransient | RF_DuplicateTransient));
+		TargetsDetailWrapper->Connector = Connector;
+		TargetsDetailWrapper->RigTreeDelegates = &RigTreeDelegates;
+	}
+	HandleTargetsChangedInClient(InArgs._Targets);
+
+	if(bIsArray)
+	{
+		FPropertyEditorModule& PropertyEditorModule = FModuleManager::GetModuleChecked<FPropertyEditorModule>("PropertyEditor");
+
+		FDetailsViewArgs DetailsViewArgs;
+		DetailsViewArgs.NameAreaSettings = FDetailsViewArgs::HideNameArea;
+		DetailsViewArgs.bAllowSearch = false;
+		DetailsViewArgs.bShowOptions = false;
+		DetailsViewArgs.bShowSectionSelector = false;
+		DetailsViewArgs.bShowObjectLabel = false;
+		DetailsViewArgs.bHideSelectionTip = true;
+		DetailsViewArgs.bShowScrollBar = false;
+		DetailsViewArgs.bAllowFavoriteSystem = false;
+		DetailsViewArgs.bLockable = false;
+		DetailsViewArgs.bUpdatesFromSelection = false;
+
+		const TSharedRef<IDetailsView> DetailsView = PropertyEditorModule.CreateDetailView(DetailsViewArgs);
+
+		ChildSlot
+		[
+			SNew(SHorizontalBox)
+			+ SHorizontalBox::Slot()
+			.Padding(0)
+			.FillWidth(1)
+			[
+				SNew(SVerticalBox)
+				+SVerticalBox::Slot()
+				.Padding(0)
+				.FillHeight(1)
+				[
+					DetailsView
+				]
+			]
+		];
+
+		DetailsView->OnFinishedChangingProperties().AddSP(this, &SRigConnectorTargetWidget::OnFinishedChangingProperties);
+
+		DetailsView->RegisterInstancedCustomPropertyTypeLayout(
+			FRigElementKey::StaticStruct()->GetFName(),
+			FOnGetPropertyTypeCustomizationInstance::CreateSP(this, &SRigConnectorTargetWidget::GetRigElementKeyCustomization)
+		);
+
+		TArray<UObject*> Objects = {TargetsDetailWrapper.Get()}; 
+		DetailsView->SetObjects(Objects, true);
+	}
+	else
+	{
+		ChildSlot
+		[
+			SNew(SHorizontalBox)
+			+ SHorizontalBox::Slot()
+			.Padding(0)
+			.AutoWidth()
+			[
+				SNew(SVerticalBox)
+				+SVerticalBox::Slot()
+				.Padding(0)
+				.FillHeight(1)
+				[
+					SNew(SRigConnectorTargetComboButton)
+					.ConnectorKey(Connector)
+					.TargetKey(this, &SRigConnectorTargetWidget::GetSingleTargetKey)
+					.OnSetTarget_Lambda([this](FRigElementKey InTarget) -> bool
+					{
+						SingleTarget = InTarget;
+						OnPropertyChanged();
+						return true;
+					})
+					.RigTreeDelegates(RigTreeDelegates)
+				]
+			]
+		];
+	}
+}
+
+void SRigConnectorTargetWidget::HandleTargetsChangedInClient(TArray<FRigElementKey> InTargets)
+{
+	SingleTarget = FRigElementKey(NAME_None, ERigElementType::Bone);
+	if(bIsArray)
+	{
+		TargetsDetailWrapper->TargetArray = InTargets;
+	}
+	else
+	{
+		if(InTargets.Num() == 1)
+		{
+			SingleTarget = InTargets[0];
+		}
+	}
+}
+
+void SRigConnectorTargetWidget::OnFinishedChangingProperties(const FPropertyChangedEvent& PropertyChangedEvent)
+{
+	OnPropertyChanged();
+}
+
+void SRigConnectorTargetWidget::OnPropertyChanged()
+{
+	if(!OnSetTargetArray.IsBound())
+	{
+		return;
+	}
+	if(bIsArray)
+	{
+		(void)OnSetTargetArray.Execute(TargetsDetailWrapper->TargetArray);
+	}
+	else
+	{
+		(void)OnSetTargetArray.Execute({SingleTarget});
+	}
+}
+
+TSharedRef<IPropertyTypeCustomization> SRigConnectorTargetWidget::GetRigElementKeyCustomization() const
+{
+	return MakeShareable(new FRigConnectorTargetWidgetCustomization);
+}
+
+FRigElementKey SRigConnectorTargetWidget::GetSingleTargetKey() const
+{
+	return SingleTarget;
+}
+
+void FRigConnectorTargetWidgetCustomization::CustomizeHeader(TSharedRef<IPropertyHandle> InStructPropertyHandle, FDetailWidgetRow& HeaderRow,
+                                                             IPropertyTypeCustomizationUtils& StructCustomizationUtils)
+{
+	StructPropertyHandle = InStructPropertyHandle.ToSharedPtr();
+
+	TArray<UObject*> Objects;
+	InStructPropertyHandle->GetOuterObjects(Objects);
+	for (UObject* Object : Objects)
+	{
+		if (Object->IsA<URigConnectorTargetsDetailWrapper>())
+		{
+			TargetsDetailWrapper = Cast<URigConnectorTargetsDetailWrapper>(Object);
+			break;
+		}
+	}
+
+	HeaderRow.NameContent()
+	.MaxDesiredWidth(30)
+	[
+		StructPropertyHandle->CreatePropertyNameWidget()
+	];
+
+	if(TargetsDetailWrapper.IsValid())
+	{
+		const int32 ArrayIndex = InStructPropertyHandle->GetArrayIndex();
+		TSharedPtr<SRigConnectorTargetComboButton> ComboButton;
+		if(TargetsDetailWrapper->TargetArray.IsValidIndex(ArrayIndex))
+		{
+			HeaderRow.ValueContent()
+			[
+				SAssignNew(ComboButton, SRigConnectorTargetComboButton)
+				.ConnectorKey(TargetsDetailWrapper->Connector)
+				.TargetKey(this, &FRigConnectorTargetWidgetCustomization::GetElementKey)
+				.OnSetTarget_Lambda([this, ArrayIndex](FRigElementKey InTarget) -> bool
+				{
+					TargetsDetailWrapper->TargetArray[ArrayIndex] = InTarget;
+					StructPropertyHandle->NotifyFinishedChangingProperties();
+					return true;
+				})
+				.RigTreeDelegates(*TargetsDetailWrapper->GetRigTreeDelegates())
+			];
+		}
+	}
+}
+
+void FRigConnectorTargetWidgetCustomization::CustomizeChildren(TSharedRef<IPropertyHandle> InStructPropertyHandle,
+	IDetailChildrenBuilder& StructBuilder, IPropertyTypeCustomizationUtils& StructCustomizationUtils)
+{
+	// nothing to do here
+}
+
+FRigElementKey FRigConnectorTargetWidgetCustomization::GetElementKey() const
+{
+	void* Data = nullptr;
+	if(FPropertyAccess::Success == StructPropertyHandle->GetValueData(Data))
+	{
+		return *reinterpret_cast<FRigElementKey*>(Data);
+	}
+	return FRigElementKey();
+}
+
+#undef LOCTEXT_NAMESPACE
diff --git a/Engine/Plugins/Animation/ControlRig/Source/ControlRigEditor/Private/Editor/SRigConnectorTargetWidget.h b/Engine/Plugins/Animation/ControlRig/Source/ControlRigEditor/Private/Editor/SRigConnectorTargetWidget.h
new file mode 100644
--- /dev/null
+++ b/Engine/Plugins/Animation/ControlRig/Source/ControlRigEditor/Private/Editor/SRigConnectorTargetWidget.h
@@ -0,0 +1,155 @@
+// Copyright Epic Games, Inc. All Rights Reserved.
+
+#pragma once
+
+#include "CoreMinimal.h"
+#include "Rigs/RigHierarchy.h"
+#include "Widgets/Layout/SBox.h"
+#include "Widgets/Input/SButton.h"
+#include "Widgets/Input/SComboButton.h"
+#include "PropertyEditorDelegates.h"
+#include "UObject/StrongObjectPtrTemplates.h"
+#include "IPropertyTypeCustomization.h"
+#include "SRigHierarchyTreeView.h"
+#include "SRigConnectorTargetWidget.generated.h"
+
+DECLARE_DELEGATE_RetVal_OneParam(bool, FRigConnectorTargetWidget_SetTarget, FRigElementKey);
+DECLARE_DELEGATE_RetVal_OneParam(bool, FRigConnectorTargetWidget_SetTargetArray, TArray<FRigElementKey>);
+DECLARE_DELEGATE_OneParam(FRigConnectorTargetWidget_HandleTargetsChangedInClient, TArray<FRigElementKey>);
+
+UCLASS()
+class URigConnectorTargetsDetailWrapper : public UObject
+{
+	GENERATED_BODY()
+
+public:
+	
+	URigConnectorTargetsDetailWrapper()
+	: Connector(NAME_None, ERigElementType::Connector)
+	{
+	}
+
+	UPROPERTY()
+	FRigElementKey Connector;
+
+	UPROPERTY(EditAnywhere, Category=NoCategory, meta=(FullyExpand=false))
+	TArray<FRigElementKey> TargetArray;
+
+	FRigTreeDelegates* GetRigTreeDelegates() const
+	{
+		return RigTreeDelegates;
+	}
+
+private:
+	
+	FRigTreeDelegates* RigTreeDelegates = nullptr;
+
+	friend class SRigConnectorTargetWidget;
+	friend class FRigConnectorTargetWidgetCustomization;
+};
+
+/** Widget allowing editing the targets of a connector */
+class SRigConnectorTargetComboButton : public SComboButton
+{
+public:
+	SLATE_BEGIN_ARGS(SRigConnectorTargetComboButton)
+		: _Padding(0)
+		, _ContentPadding(3)
+		, _MenuPlacement(EMenuPlacement::MenuPlacement_BelowAnchor)
+		, _ConnectorKey(NAME_None, ERigElementType::Connector)
+		, _ArrayIndex(INDEX_NONE)
+		, _ButtonMinDesiredWith(150.f)
+	{}
+	SLATE_ATTRIBUTE(FMargin, Padding)
+	SLATE_ATTRIBUTE( FMargin, ContentPadding )
+	SLATE_ATTRIBUTE( EMenuPlacement, MenuPlacement )
+	SLATE_ARGUMENT(FRigElementKey, ConnectorKey)
+	SLATE_ATTRIBUTE(FRigElementKey, TargetKey)
+	SLATE_ATTRIBUTE(int32, ArrayIndex)
+	SLATE_ARGUMENT(FRigTreeDelegates, RigTreeDelegates)
+	SLATE_ARGUMENT(float, ButtonMinDesiredWith)
+	SLATE_EVENT(FRigConnectorTargetWidget_SetTarget, OnSetTarget)
+	SLATE_END_ARGS()
+
+	virtual ~SRigConnectorTargetComboButton() override;
+
+	void Construct(const FArguments& InArgs);
+
+private:
+
+	void OnComboBoxOpened();
+	void PopulateButtonBox();
+	void OnConnectorTargetPicked(TSharedPtr<FRigTreeElement> Selection, ESelectInfo::Type SelectInfo);
+	
+	FRigElementKey ConnectorKey;
+	TAttribute<FRigElementKey> TargetKey;
+	TAttribute<int32> ArrayIndex;
+	FRigConnectorTargetWidget_SetTarget OnSetTarget;
+	TSharedPtr<SVerticalBox> VerticalButtonBox;
+	FRigTreeDelegates RigTreeDelegates;
+	TSharedPtr<SSearchableRigHierarchyTreeView> SearchableTreeView;
+};
+
+/** Widget allowing editing the targets of a connector */
+class SRigConnectorTargetWidget : public SBox
+{
+public:
+	SLATE_BEGIN_ARGS(SRigConnectorTargetWidget)
+		: _Outer(nullptr)
+		, _ConnectorKey(NAME_None, ERigElementType::Connector)
+		, _IsArray(false)
+		, _ExpandArrayByDefault(false)
+		, _RespectConnectorRules(true)
+		, _HandleTargetsChangedInClient(nullptr)
+		, _Padding(FMargin(0))
+	{}
+	SLATE_ARGUMENT(UObject*, Outer)
+	SLATE_ARGUMENT(FRigElementKey, ConnectorKey)
+	SLATE_ARGUMENT(TArray<FRigElementKey>, Targets)
+	SLATE_ARGUMENT(bool, IsArray)
+	SLATE_ARGUMENT(bool, ExpandArrayByDefault)
+	SLATE_ARGUMENT(bool, RespectConnectorRules)
+	SLATE_EVENT(FRigConnectorTargetWidget_SetTargetArray, OnSetTargetArray)
+	SLATE_ARGUMENT(FRigConnectorTargetWidget_HandleTargetsChangedInClient*, HandleTargetsChangedInClient)
+	SLATE_ATTRIBUTE(FMargin, Padding)
+	SLATE_ARGUMENT(FRigTreeDelegates, RigTreeDelegates)
+	SLATE_END_ARGS()
+
+	virtual ~SRigConnectorTargetWidget() override;
+
+	void Construct(const FArguments& InArgs);
+
+private:
+
+	void HandleTargetsChangedInClient(TArray<FRigElementKey> InTargets);
+	void OnFinishedChangingProperties(const FPropertyChangedEvent& PropertyChangedEvent);
+	void OnPropertyChanged();
+	TSharedRef<IPropertyTypeCustomization> GetRigElementKeyCustomization() const;
+
+	FRigElementKey GetSingleTargetKey() const;
+
+	FRigElementKey Connector;
+	FRigElementKey SingleTarget;
+	TStrongObjectPtr<URigConnectorTargetsDetailWrapper> TargetsDetailWrapper;
+	FRigConnectorTargetWidget_SetTargetArray OnSetTargetArray;
+	bool bIsArray = false;
+	FRigTreeDelegates RigTreeDelegates;
+};
+
+class FRigConnectorTargetWidgetCustomization : public IPropertyTypeCustomization
+{
+public:
+
+	/** IPropertyTypeCustomization interface */
+	virtual void CustomizeHeader(TSharedRef<class IPropertyHandle> InStructPropertyHandle, class FDetailWidgetRow& HeaderRow, IPropertyTypeCustomizationUtils& StructCustomizationUtils) override;
+	virtual void CustomizeChildren(TSharedRef<class IPropertyHandle> InStructPropertyHandle, class IDetailChildrenBuilder& StructBuilder, IPropertyTypeCustomizationUtils& StructCustomizationUtils) override;
+
+protected:
+
+	FRigElementKey GetElementKey() const;
+
+	mutable TSharedPtr<IPropertyHandle> StructPropertyHandle;
+	TWeakObjectPtr<URigConnectorTargetsDetailWrapper> TargetsDetailWrapper;
+
+	friend class SRigConnectorTargetWidget;
+};
diff --git a/Engine/Plugins/Animation/ControlRig/Source/ControlRigEditor/Private/Editor/SRigHierarchyTreeView.h b/Engine/Plugins/Animation/ControlRig/Source/ControlRigEditor/Private/Editor/SRigHierarchyTreeView.h
--- a/Engine/Plugins/Animation/ControlRig/Source/ControlRigEditor/Private/Editor/SRigHierarchyTreeView.h
+++ b/Engine/Plugins/Animation/ControlRig/Source/ControlRigEditor/Private/Editor/SRigHierarchyTreeView.h
@@ -116,6 +116,15 @@ struct CONTROLRIGEDITOR_API FRigTreeDelegates
 		bIsChangingRigHierarchy = false;
 	}
 
+	URigHierarchy* GetHierarchy()
+	{
+		if(OnGetHierarchy.IsBound())
+		{
+			return const_cast<URigHierarchy*>(OnGetHierarchy.Execute());
+		}
+		return nullptr;
+	}
+
 	const URigHierarchy* GetHierarchy() const
 	{
 		if(OnGetHierarchy.IsBound())
