Sequencer: Channel Names and Colors for 3D Transform Sections, Control Rig Sesctions, Subsequence Sections now respect AxisDisplayInfo

Not included:
- Custom Primitive Sections
- Parameter Sections
- 2D Transform Sections
- Vector Sections
- Composure Post Move Sections

diff --git a/Engine/Plugins/Animation/ControlRig/Source/ControlRig/Private/Sequencer/MovieSceneControlRigParameterSection.cpp b/Engine/Plugins/Animation/ControlRig/Source/ControlRig/Private/Sequencer/MovieSceneControlRigParameterSection.cpp
--- a/Engine/Plugins/Animation/ControlRig/Source/ControlRig/Private/Sequencer/MovieSceneControlRigParameterSection.cpp
+++ b/Engine/Plugins/Animation/ControlRig/Source/ControlRig/Private/Sequencer/MovieSceneControlRigParameterSection.cpp
@@ -14,6 +14,7 @@
 #include "Rigs/FKControlRig.h"
 #include "Animation/AnimSequence.h"
 #include "Units/Execution/RigUnit_InverseExecution.h"
+#include "Misc/AxisDisplayInfo.h"
 #include "Misc/ScopedSlowTask.h"
 #include "MovieSceneTimeHelpers.h"
 #include "Animation/AnimSequenceHelpers.h"
@@ -400,16 +401,19 @@ struct FParameterTransformChannelEditorData
 		FString TotalName = NameAsString;
 		FText TransformGroup = FText::Format(NSLOCTEXT("MovieSceneControlParameterRigSection", "MovieSceneControlParameterRigSectionGroupName", "{0}"), GroupName);
 
+		FText LocationDisplayName = NSLOCTEXT("MovieSceneControlParameterRigSection", "Location", "Location");
+		FText ScaleDisplayName = NSLOCTEXT("MovieSceneControlParameterRigSection", "Scale", "Scale");
+
 		{
 			//MetaData[0].SetIdentifiers("Location.X", FCommonChannelData::ChannelX, LocationGroup);
 			TotalName += ".Location.X";
-			MetaData[0].SetIdentifiers(FName(*TotalName), NSLOCTEXT("MovieSceneControlParameterRigSection", "Location.X", "Location.X"), TransformGroup);
+			MetaData[0].SetIdentifiers(FName(*TotalName), FText::Join(FText::FromString(TEXT(".")), LocationDisplayName, AxisDisplayInfo::GetAxisDisplayName(EAxisList::X)), TransformGroup);
 			MetaData[0].IntentName = NSLOCTEXT("MovieSceneControlParameterRigSection", "Location.X", "Location.X");
 			MetaData[0].LongIntentNameFormat = LongIntentFormatStr;
 			TotalName = NameAsString;
 
 			MetaData[0].bEnabled = bEnabledOverride && EnumHasAllFlags(Mask, EMovieSceneTransformChannel::TranslationX);
-			MetaData[0].Color = FCommonChannelData::RedChannelColor;
+			MetaData[0].Color = AxisDisplayInfo::GetAxisColor(EAxisList::X);
 			MetaData[0].SortOrder = SortStartIndex++;
 			MetaData[0].bCanCollapseToTrack = true;
 			for (const FName& PropertyMetaDataKey : PropertyMetaDataKeys)
@@ -419,13 +423,13 @@ struct FParameterTransformChannelEditorData
 
 			//MetaData[1].SetIdentifiers("Location.Y", FCommonChannelData::ChannelY, LocationGroup);
 			TotalName += ".Location.Y";
-			MetaData[1].SetIdentifiers(FName(*TotalName), NSLOCTEXT("MovieSceneControlParameterRigSection", "Location.Y", "Location.Y"), TransformGroup);
+			MetaData[1].SetIdentifiers(FName(*TotalName), FText::Join(FText::FromString(TEXT(".")), LocationDisplayName, AxisDisplayInfo::GetAxisDisplayName(EAxisList::Y)), TransformGroup);
 			MetaData[1].IntentName = NSLOCTEXT("MovieSceneControlParameterRigSection", "Location.Y", "Location.Y");
 			MetaData[1].LongIntentNameFormat = LongIntentFormatStr;
 			TotalName = NameAsString;
 
 			MetaData[1].bEnabled = bEnabledOverride && EnumHasAllFlags(Mask, EMovieSceneTransformChannel::TranslationY);
-			MetaData[1].Color = FCommonChannelData::GreenChannelColor;
+			MetaData[1].Color = AxisDisplayInfo::GetAxisColor(EAxisList::Y);
 			MetaData[1].SortOrder = SortStartIndex++;
 			MetaData[1].bCanCollapseToTrack = true;
 			for (const FName& PropertyMetaDataKey : PropertyMetaDataKeys)
@@ -435,13 +439,13 @@ struct FParameterTransformChannelEditorData
 
 			//MetaData[2].SetIdentifiers("Location.Z", FCommonChannelData::ChannelZ, LocationGroup);
 			TotalName += ".Location.Z";
-			MetaData[2].SetIdentifiers(FName(*TotalName), NSLOCTEXT("MovieSceneControlParameterRigSection", "Location.Z", "Location.Z"), TransformGroup);
+			MetaData[2].SetIdentifiers(FName(*TotalName), FText::Join(FText::FromString(TEXT(".")), LocationDisplayName, AxisDisplayInfo::GetAxisDisplayName(EAxisList::Z)), TransformGroup);
 			MetaData[2].IntentName = NSLOCTEXT("MovieSceneControlParameterRigSection", "Location.Z", "Location.Z");
 			MetaData[2].LongIntentNameFormat = LongIntentFormatStr;
 			TotalName = NameAsString;
 
 			MetaData[2].bEnabled = bEnabledOverride && EnumHasAllFlags(Mask, EMovieSceneTransformChannel::TranslationZ);
-			MetaData[2].Color = FCommonChannelData::BlueChannelColor;
+			MetaData[2].Color = AxisDisplayInfo::GetAxisColor(EAxisList::Z);
 			MetaData[2].SortOrder = SortStartIndex++;
 			MetaData[2].bCanCollapseToTrack = true;
 			for (const FName& PropertyMetaDataKey : PropertyMetaDataKeys)
@@ -458,7 +462,7 @@ struct FParameterTransformChannelEditorData
 			TotalName = NameAsString;
 
 			MetaData[3].bEnabled = bEnabledOverride && EnumHasAllFlags(Mask, EMovieSceneTransformChannel::RotationX);
-			MetaData[3].Color = FCommonChannelData::RedChannelColor;
+			MetaData[3].Color = AxisDisplayInfo::GetAxisColor(EAxisList::X);
 			MetaData[3].SortOrder = SortStartIndex++;
 			MetaData[3].bCanCollapseToTrack = true;
 			for (const FName& PropertyMetaDataKey : PropertyMetaDataKeys)
@@ -474,7 +478,7 @@ struct FParameterTransformChannelEditorData
 			TotalName = NameAsString;
 
 			MetaData[4].bEnabled = bEnabledOverride && EnumHasAllFlags(Mask, EMovieSceneTransformChannel::RotationY);
-			MetaData[4].Color = FCommonChannelData::GreenChannelColor;
+			MetaData[4].Color = AxisDisplayInfo::GetAxisColor(EAxisList::Y);
 			MetaData[4].SortOrder = SortStartIndex++;
 			MetaData[4].bCanCollapseToTrack = true;
 			for (const FName& PropertyMetaDataKey : PropertyMetaDataKeys)
@@ -490,7 +494,7 @@ struct FParameterTransformChannelEditorData
 			TotalName = NameAsString;
 
 			MetaData[5].bEnabled = bEnabledOverride && EnumHasAllFlags(Mask, EMovieSceneTransformChannel::RotationZ);
-			MetaData[5].Color = FCommonChannelData::BlueChannelColor;
+			MetaData[5].Color = AxisDisplayInfo::GetAxisColor(EAxisList::Z);
 			MetaData[5].SortOrder = SortStartIndex++;
 			MetaData[5].bCanCollapseToTrack = true;
 			for (const FName& PropertyMetaDataKey : PropertyMetaDataKeys)
@@ -501,13 +505,13 @@ struct FParameterTransformChannelEditorData
 		{
 			//MetaData[6].SetIdentifiers("Scale.X", FCommonChannelData::ChannelX, ScaleGroup);
 			TotalName += ".Scale.X";
-			MetaData[6].SetIdentifiers(FName(*TotalName), NSLOCTEXT("MovieSceneControlParameterRigSection", "Scale.X", "Scale.X"), TransformGroup);
+			MetaData[6].SetIdentifiers(FName(*TotalName), FText::Join(FText::FromString(TEXT(".")), ScaleDisplayName, AxisDisplayInfo::GetAxisDisplayName(EAxisList::X)), TransformGroup);
 			MetaData[6].IntentName = NSLOCTEXT("MovieSceneControlParameterRigSection", "Scale.X", "Scale.X");
 			MetaData[6].LongIntentNameFormat = LongIntentFormatStr;
 			TotalName = NameAsString;
 
 			MetaData[6].bEnabled = bEnabledOverride && EnumHasAllFlags(Mask, EMovieSceneTransformChannel::ScaleX);
-			MetaData[6].Color = FCommonChannelData::RedChannelColor;
+			MetaData[6].Color = AxisDisplayInfo::GetAxisColor(EAxisList::X);
 			MetaData[6].SortOrder = SortStartIndex++;
 			MetaData[6].bCanCollapseToTrack = true;
 			for (const FName& PropertyMetaDataKey : PropertyMetaDataKeys)
@@ -517,13 +521,13 @@ struct FParameterTransformChannelEditorData
 
 			//MetaData[7].SetIdentifiers("Scale.Y", FCommonChannelData::ChannelY, ScaleGroup);
 			TotalName += ".Scale.Y";
-			MetaData[7].SetIdentifiers(FName(*TotalName), NSLOCTEXT("MovieSceneControlParameterRigSection", "Scale.Y", "Scale.Y"), TransformGroup);
+			MetaData[7].SetIdentifiers(FName(*TotalName), FText::Join(FText::FromString(TEXT(".")), ScaleDisplayName, AxisDisplayInfo::GetAxisDisplayName(EAxisList::Y)), TransformGroup);
 			MetaData[7].IntentName = NSLOCTEXT("MovieSceneControlParameterRigSection", "Scale.Y", "Scale.Y");
 			MetaData[7].LongIntentNameFormat = LongIntentFormatStr;
 			TotalName = NameAsString;
 
 			MetaData[7].bEnabled = bEnabledOverride && EnumHasAllFlags(Mask, EMovieSceneTransformChannel::ScaleY);
-			MetaData[7].Color = FCommonChannelData::GreenChannelColor;
+			MetaData[7].Color = AxisDisplayInfo::GetAxisColor(EAxisList::Y);
 			MetaData[7].SortOrder = SortStartIndex++;
 			MetaData[7].bCanCollapseToTrack = true;
 			for (const FName& PropertyMetaDataKey : PropertyMetaDataKeys)
@@ -533,13 +537,13 @@ struct FParameterTransformChannelEditorData
 
 			//MetaData[8].SetIdentifiers("Scale.Z", FCommonChannelData::ChannelZ, ScaleGroup);
 			TotalName += ".Scale.Z";
-			MetaData[8].SetIdentifiers(FName(*TotalName), NSLOCTEXT("MovieSceneControlParameterRigSection", "Scale.Z", "Scale.Z"), TransformGroup);
+			MetaData[8].SetIdentifiers(FName(*TotalName), FText::Join(FText::FromString(TEXT(".")), ScaleDisplayName, AxisDisplayInfo::GetAxisDisplayName(EAxisList::Z)), TransformGroup);
 			MetaData[8].IntentName = NSLOCTEXT("MovieSceneControlParameterRigSection", "Scale.Z", "Scale.Z");
 			MetaData[8].LongIntentNameFormat = LongIntentFormatStr;
 			TotalName = NameAsString;
 
 			MetaData[8].bEnabled = bEnabledOverride && EnumHasAllFlags(Mask, EMovieSceneTransformChannel::ScaleZ);
-			MetaData[8].Color = FCommonChannelData::BlueChannelColor;
+			MetaData[8].Color = AxisDisplayInfo::GetAxisColor(EAxisList::Z);
 			MetaData[8].SortOrder = SortStartIndex++;
 			MetaData[8].bCanCollapseToTrack = true;
 			for (const FName& PropertyMetaDataKey : PropertyMetaDataKeys)
diff --git a/Engine/Plugins/Animation/ControlRig/Source/ControlRigEditor/Private/Sequencer/ControlRigParameterTrackEditor.cpp b/Engine/Plugins/Animation/ControlRig/Source/ControlRigEditor/Private/Sequencer/ControlRigParameterTrackEditor.cpp
--- a/Engine/Plugins/Animation/ControlRig/Source/ControlRigEditor/Private/Sequencer/ControlRigParameterTrackEditor.cpp
+++ b/Engine/Plugins/Animation/ControlRig/Source/ControlRigEditor/Private/Sequencer/ControlRigParameterTrackEditor.cpp
@@ -47,6 +47,7 @@
 #include "MovieSceneTimeHelpers.h"
 #include "Fonts/FontMeasure.h"
 #include "AnimationEditorUtils.h"
+#include "Misc/AxisDisplayInfo.h"
 #include "Misc/MessageDialog.h"
 #include "Framework/Notifications/NotificationManager.h"
 #include "Widgets/Notifications/SNotificationList.h"
@@ -5931,13 +5932,16 @@ void FControlRigParameterSection::BuildSectionContextMenu(FMenuBuilder& MenuBuil
 			LOCTEXT("AllTranslation", "Translation"), LOCTEXT("AllTranslation_ToolTip", "Causes this section to affect the translation of rig control transforms"),
 			FNewMenuDelegate::CreateLambda([=](FMenuBuilder& SubMenuBuilder) {
 				SubMenuBuilder.AddMenuEntry(
-					LOCTEXT("TranslationX", "X"), LOCTEXT("TranslationX_ToolTip", "Causes this section to affect the X channel of the transform's translation"),
+					AxisDisplayInfo::GetAxisDisplayName(EAxisList::X),
+					FText::Format(LOCTEXT("ActivateTranslationChannel_Tooltip", "Causes this section to affect the {0} channel of the transform's translation"), AxisDisplayInfo::GetAxisDisplayName(EAxisList::X)),
 					FSlateIcon(), MakeUIAction(EMovieSceneTransformChannel::TranslationX), NAME_None, EUserInterfaceActionType::ToggleButton);
 				SubMenuBuilder.AddMenuEntry(
-					LOCTEXT("TranslationY", "Y"), LOCTEXT("TranslationY_ToolTip", "Causes this section to affect the Y channel of the transform's translation"),
+					AxisDisplayInfo::GetAxisDisplayName(EAxisList::Y), 
+					FText::Format(LOCTEXT("ActivateTranslationChannel_Tooltip", "Causes this section to affect the {0} channel of the transform's translation"), AxisDisplayInfo::GetAxisDisplayName(EAxisList::Y)),
 					FSlateIcon(), MakeUIAction(EMovieSceneTransformChannel::TranslationY), NAME_None, EUserInterfaceActionType::ToggleButton);
 				SubMenuBuilder.AddMenuEntry(
-					LOCTEXT("TranslationZ", "Z"), LOCTEXT("TranslationZ_ToolTip", "Causes this section to affect the Z channel of the transform's translation"),
+					AxisDisplayInfo::GetAxisDisplayName(EAxisList::Z), 
+					FText::Format(LOCTEXT("ActivateTranslationChannel_Tooltip", "Causes this section to affect the {0} channel of the transform's translation"), AxisDisplayInfo::GetAxisDisplayName(EAxisList::Z)),
 					FSlateIcon(), MakeUIAction(EMovieSceneTransformChannel::TranslationZ), NAME_None, EUserInterfaceActionType::ToggleButton);
 				}),
 			MakeUIAction(EMovieSceneTransformChannel::Translation),
@@ -5948,13 +5952,13 @@ void FControlRigParameterSection::BuildSectionContextMenu(FMenuBuilder& MenuBuil
 			LOCTEXT("AllRotation", "Rotation"), LOCTEXT("AllRotation_ToolTip", "Causes this section to affect the rotation of the rig control transform"),
 			FNewMenuDelegate::CreateLambda([=](FMenuBuilder& SubMenuBuilder) {
 				SubMenuBuilder.AddMenuEntry(
-					LOCTEXT("RotationX", "Roll (X)"), LOCTEXT("RotationX_ToolTip", "Causes this section to affect the roll (X) channel the transform's rotation"),
+					LOCTEXT("RotationX", "Roll"), LOCTEXT("RotationX_ToolTip", "Causes this section to affect the roll channel the transform's rotation"),
 					FSlateIcon(), MakeUIAction(EMovieSceneTransformChannel::RotationX), NAME_None, EUserInterfaceActionType::ToggleButton);
 				SubMenuBuilder.AddMenuEntry(
-					LOCTEXT("RotationY", "Pitch (Y)"), LOCTEXT("RotationY_ToolTip", "Causes this section to affect the pitch (Y) channel the transform's rotation"),
+					LOCTEXT("RotationY", "Pitch"), LOCTEXT("RotationY_ToolTip", "Causes this section to affect the pitch channel the transform's rotation"),
 					FSlateIcon(), MakeUIAction(EMovieSceneTransformChannel::RotationY), NAME_None, EUserInterfaceActionType::ToggleButton);
 				SubMenuBuilder.AddMenuEntry(
-					LOCTEXT("RotationZ", "Yaw (Z)"), LOCTEXT("RotationZ_ToolTip", "Causes this section to affect the yaw (Z) channel the transform's rotation"),
+					LOCTEXT("RotationZ", "Yaw"), LOCTEXT("RotationZ_ToolTip", "Causes this section to affect the yaw channel the transform's rotation"),
 					FSlateIcon(), MakeUIAction(EMovieSceneTransformChannel::RotationZ), NAME_None, EUserInterfaceActionType::ToggleButton);
 				}),
 			MakeUIAction(EMovieSceneTransformChannel::Rotation),
@@ -5965,13 +5969,16 @@ void FControlRigParameterSection::BuildSectionContextMenu(FMenuBuilder& MenuBuil
 			LOCTEXT("AllScale", "Scale"), LOCTEXT("AllScale_ToolTip", "Causes this section to affect the scale of the rig control transform"),
 			FNewMenuDelegate::CreateLambda([=](FMenuBuilder& SubMenuBuilder) {
 				SubMenuBuilder.AddMenuEntry(
-					LOCTEXT("ScaleX", "X"), LOCTEXT("ScaleX_ToolTip", "Causes this section to affect the X channel of the transform's scale"),
+					AxisDisplayInfo::GetAxisDisplayName(EAxisList::X), 
+					FText::Format(LOCTEXT("ActivateScaleChannel_Tooltip", "Causes this section to affect the {0} channel of the transform's scale"), AxisDisplayInfo::GetAxisDisplayName(EAxisList::X)),
 					FSlateIcon(), MakeUIAction(EMovieSceneTransformChannel::ScaleX), NAME_None, EUserInterfaceActionType::ToggleButton);
 				SubMenuBuilder.AddMenuEntry(
-					LOCTEXT("ScaleY", "Y"), LOCTEXT("ScaleY_ToolTip", "Causes this section to affect the Y channel of the transform's scale"),
+					AxisDisplayInfo::GetAxisDisplayName(EAxisList::Y), 
+					FText::Format(LOCTEXT("ActivateScaleChannel_Tooltip", "Causes this section to affect the {0} channel of the transform's scale"), AxisDisplayInfo::GetAxisDisplayName(EAxisList::Y)),
 					FSlateIcon(), MakeUIAction(EMovieSceneTransformChannel::ScaleY), NAME_None, EUserInterfaceActionType::ToggleButton);
 				SubMenuBuilder.AddMenuEntry(
-					LOCTEXT("ScaleZ", "Z"), LOCTEXT("ScaleZ_ToolTip", "Causes this section to affect the Z channel of the transform's scale"),
+					AxisDisplayInfo::GetAxisDisplayName(EAxisList::Z), 
+					FText::Format(LOCTEXT("ActivateScaleChannel_Tooltip", "Causes this section to affect the {0} channel of the transform's scale"), AxisDisplayInfo::GetAxisDisplayName(EAxisList::Z)),
 					FSlateIcon(), MakeUIAction(EMovieSceneTransformChannel::ScaleZ), NAME_None, EUserInterfaceActionType::ToggleButton);
 				}),
 			MakeUIAction(EMovieSceneTransformChannel::Scale),
diff --git a/Engine/Source/Editor/MovieSceneTools/Private/Sections/CinematicShotSection.cpp b/Engine/Source/Editor/MovieSceneTools/Private/Sections/CinematicShotSection.cpp
--- a/Engine/Source/Editor/MovieSceneTools/Private/Sections/CinematicShotSection.cpp
+++ b/Engine/Source/Editor/MovieSceneTools/Private/Sections/CinematicShotSection.cpp
@@ -11,6 +11,7 @@
 #include "MovieScene.h"
 #include "TrackEditors/CinematicShotTrackEditor.h"
 #include "SequencerSectionPainter.h"
+#include "Misc/AxisDisplayInfo.h"
 #include "Styling/AppStyle.h"
 #include "MovieSceneToolHelpers.h"
 #include "MovieSceneTimeHelpers.h"
@@ -395,18 +396,23 @@ void FCinematicShotSection::AddShotMenuSection(FMenuBuilder& MenuBuilder, const
 
 	TSharedPtr<ISequencer> Sequencer = GetSequencer();
 	
+
 	MenuBuilder.BeginSection(NAME_None, LOCTEXT("OriginChannelsText", "Active Channels"));
 	MenuBuilder.AddSubMenu(
 		LOCTEXT("AllTranslation", "Translation"), LOCTEXT("AllTranslation_ToolTip", "Causes this section to affect the translation of the transform"),
 		FNewMenuDelegate::CreateLambda([Sequencer, MakeUIAction](FMenuBuilder& SubMenuBuilder){
+
 			SubMenuBuilder.AddMenuEntry(
-				LOCTEXT("TranslationX", "X"), LOCTEXT("TranslationX_ToolTip", "Causes this section to affect the X channel of the transform's translation"),
+				AxisDisplayInfo::GetAxisDisplayName(EAxisList::X), 
+				FText::Format(LOCTEXT("ActivateTranslationChannel_Tooltip", "Causes this section to affect the {0} channel of the transform's translation"), AxisDisplayInfo::GetAxisDisplayName(EAxisList::X)),
 				FSlateIcon(), MakeUIAction(EMovieSceneTransformChannel::TranslationX, Sequencer), NAME_None, EUserInterfaceActionType::ToggleButton);
 			SubMenuBuilder.AddMenuEntry(
-				LOCTEXT("TranslationY", "Y"), LOCTEXT("TranslationY_ToolTip", "Causes this section to affect the Y channel of the transform's translation"),
+				AxisDisplayInfo::GetAxisDisplayName(EAxisList::Y), 
+				FText::Format(LOCTEXT("ActivateTranslationChannel_Tooltip", "Causes this section to affect the {0} channel of the transform's translation"), AxisDisplayInfo::GetAxisDisplayName(EAxisList::Y)),
 				FSlateIcon(), MakeUIAction(EMovieSceneTransformChannel::TranslationY, Sequencer), NAME_None, EUserInterfaceActionType::ToggleButton);
 			SubMenuBuilder.AddMenuEntry(
-				LOCTEXT("TranslationZ", "Z"), LOCTEXT("TranslationZ_ToolTip", "Causes this section to affect the Z channel of the transform's translation"),
+				AxisDisplayInfo::GetAxisDisplayName(EAxisList::Z), 
+				FText::Format(LOCTEXT("ActivateTranslationChannel_Tooltip", "Causes this section to affect the {0} channel of the transform's translation"), AxisDisplayInfo::GetAxisDisplayName(EAxisList::Z)),
 				FSlateIcon(), MakeUIAction(EMovieSceneTransformChannel::TranslationZ, Sequencer), NAME_None, EUserInterfaceActionType::ToggleButton);
 		}),
 		MakeUIAction(EMovieSceneTransformChannel::Translation, Sequencer),
@@ -417,13 +423,13 @@ void FCinematicShotSection::AddShotMenuSection(FMenuBuilder& MenuBuilder, const
 		LOCTEXT("AllRotation", "Rotation"), LOCTEXT("AllRotation_ToolTip", "Causes this section to affect the rotation of the transform"),
 		FNewMenuDelegate::CreateLambda([Sequencer, MakeUIAction](FMenuBuilder& SubMenuBuilder){
 			SubMenuBuilder.AddMenuEntry(
-				LOCTEXT("RotationX", "Roll (X)"), LOCTEXT("RotationX_ToolTip", "Causes this section to affect the roll (X) channel the transform's rotation"),
+				LOCTEXT("RotationX", "Roll"), LOCTEXT("RotationX_ToolTip", "Causes this section to affect the roll channel the transform's rotation"),
 				FSlateIcon(), MakeUIAction(EMovieSceneTransformChannel::RotationX, Sequencer), NAME_None, EUserInterfaceActionType::ToggleButton);
 			SubMenuBuilder.AddMenuEntry(
-				LOCTEXT("RotationY", "Pitch (Y)"), LOCTEXT("RotationY_ToolTip", "Causes this section to affect the pitch (Y) channel the transform's rotation"),
+				LOCTEXT("RotationY", "Pitch"), LOCTEXT("RotationY_ToolTip", "Causes this section to affect the pitch channel the transform's rotation"),
 				FSlateIcon(), MakeUIAction(EMovieSceneTransformChannel::RotationY, Sequencer), NAME_None, EUserInterfaceActionType::ToggleButton);
 			SubMenuBuilder.AddMenuEntry(
-				LOCTEXT("RotationZ", "Yaw (Z)"), LOCTEXT("RotationZ_ToolTip", "Causes this section to affect the yaw (Z) channel the transform's rotation"),
+				LOCTEXT("RotationZ", "Yaw"), LOCTEXT("RotationZ_ToolTip", "Causes this section to affect the yaw channel the transform's rotation"),
 				FSlateIcon(), MakeUIAction(EMovieSceneTransformChannel::RotationZ, Sequencer), NAME_None, EUserInterfaceActionType::ToggleButton);
 		}),
 		MakeUIAction(EMovieSceneTransformChannel::Rotation, Sequencer),
diff --git a/Engine/Source/Editor/MovieSceneTools/Private/Sections/TransformPropertySection.cpp b/Engine/Source/Editor/MovieSceneTools/Private/Sections/TransformPropertySection.cpp
--- a/Engine/Source/Editor/MovieSceneTools/Private/Sections/TransformPropertySection.cpp
+++ b/Engine/Source/Editor/MovieSceneTools/Private/Sections/TransformPropertySection.cpp
@@ -9,6 +9,7 @@
 #include "Framework/MultiBox/MultiBoxBuilder.h"
 #include "ISequencer.h"
 #include "Internationalization/Internationalization.h"
+#include "Misc/AxisDisplayInfo.h"
 #include "Misc/EnumClassFlags.h"
 #include "Misc/Guid.h"
 #include "MovieSceneSection.h"
@@ -90,13 +91,16 @@ void FTransformSection::BuildSectionContextMenu(FMenuBuilder& MenuBuilder, const
 			LOCTEXT("AllTranslation", "Translation"), LOCTEXT("AllTranslation_ToolTip", "Causes this section to affect the translation of the transform"),
 			FNewMenuDelegate::CreateLambda([=](FMenuBuilder& SubMenuBuilder){
 				SubMenuBuilder.AddMenuEntry(
-					LOCTEXT("TranslationX", "X"), LOCTEXT("TranslationX_ToolTip", "Causes this section to affect the X channel of the transform's translation"),
+					AxisDisplayInfo::GetAxisDisplayName(EAxisList::X), 
+					FText::Format(LOCTEXT("ActivateTranslationChannel_Tooltip", "Causes this section to affect the {0} channel of the transform's translation"), AxisDisplayInfo::GetAxisDisplayName(EAxisList::X)),
 					FSlateIcon(), MakeUIAction(EMovieSceneTransformChannel::TranslationX), NAME_None, EUserInterfaceActionType::ToggleButton);
 				SubMenuBuilder.AddMenuEntry(
-					LOCTEXT("TranslationY", "Y"), LOCTEXT("TranslationY_ToolTip", "Causes this section to affect the Y channel of the transform's translation"),
+					AxisDisplayInfo::GetAxisDisplayName(EAxisList::Y), 
+					FText::Format(LOCTEXT("ActivateTranslationChannel_Tooltip", "Causes this section to affect the {0} channel of the transform's translation"), AxisDisplayInfo::GetAxisDisplayName(EAxisList::Y)),
 					FSlateIcon(), MakeUIAction(EMovieSceneTransformChannel::TranslationY), NAME_None, EUserInterfaceActionType::ToggleButton);
 				SubMenuBuilder.AddMenuEntry(
-					LOCTEXT("TranslationZ", "Z"), LOCTEXT("TranslationZ_ToolTip", "Causes this section to affect the Z channel of the transform's translation"),
+					AxisDisplayInfo::GetAxisDisplayName(EAxisList::Z), 
+					FText::Format(LOCTEXT("ActivateTranslationChannel_Tooltip", "Causes this section to affect the {0} channel of the transform's translation"), AxisDisplayInfo::GetAxisDisplayName(EAxisList::Z)),
 					FSlateIcon(), MakeUIAction(EMovieSceneTransformChannel::TranslationZ), NAME_None, EUserInterfaceActionType::ToggleButton);
 			}),
 			MakeUIAction(EMovieSceneTransformChannel::Translation),
@@ -107,13 +111,13 @@ void FTransformSection::BuildSectionContextMenu(FMenuBuilder& MenuBuilder, const
 			LOCTEXT("AllRotation", "Rotation"), LOCTEXT("AllRotation_ToolTip", "Causes this section to affect the rotation of the transform"),
 			FNewMenuDelegate::CreateLambda([=](FMenuBuilder& SubMenuBuilder){
 				SubMenuBuilder.AddMenuEntry(
-					LOCTEXT("RotationX", "Roll (X)"), LOCTEXT("RotationX_ToolTip", "Causes this section to affect the roll (X) channel the transform's rotation"),
+					LOCTEXT("RotationX", "Roll"), LOCTEXT("RotationX_ToolTip", "Causes this section to affect the roll channel the transform's rotation"),
 					FSlateIcon(), MakeUIAction(EMovieSceneTransformChannel::RotationX), NAME_None, EUserInterfaceActionType::ToggleButton);
 				SubMenuBuilder.AddMenuEntry(
-					LOCTEXT("RotationY", "Pitch (Y)"), LOCTEXT("RotationY_ToolTip", "Causes this section to affect the pitch (Y) channel the transform's rotation"),
+					LOCTEXT("RotationY", "Pitch"), LOCTEXT("RotationY_ToolTip", "Causes this section to affect the pitch channel the transform's rotation"),
 					FSlateIcon(), MakeUIAction(EMovieSceneTransformChannel::RotationY), NAME_None, EUserInterfaceActionType::ToggleButton);
 				SubMenuBuilder.AddMenuEntry(
-					LOCTEXT("RotationZ", "Yaw (Z)"), LOCTEXT("RotationZ_ToolTip", "Causes this section to affect the yaw (Z) channel the transform's rotation"),
+					LOCTEXT("RotationZ", "Yaw"), LOCTEXT("RotationZ_ToolTip", "Causes this section to affect the yaw channel the transform's rotation"),
 					FSlateIcon(), MakeUIAction(EMovieSceneTransformChannel::RotationZ), NAME_None, EUserInterfaceActionType::ToggleButton);
 			}),
 			MakeUIAction(EMovieSceneTransformChannel::Rotation),
@@ -124,13 +128,16 @@ void FTransformSection::BuildSectionContextMenu(FMenuBuilder& MenuBuilder, const
 			LOCTEXT("AllScale", "Scale"), LOCTEXT("AllScale_ToolTip", "Causes this section to affect the scale of the transform"),
 			FNewMenuDelegate::CreateLambda([=](FMenuBuilder& SubMenuBuilder){
 				SubMenuBuilder.AddMenuEntry(
-					LOCTEXT("ScaleX", "X"), LOCTEXT("ScaleX_ToolTip", "Causes this section to affect the X channel of the transform's scale"),
+					AxisDisplayInfo::GetAxisDisplayName(EAxisList::X), 
+					FText::Format(LOCTEXT("ActivateScaleChannel_Tooltip", "Causes this section to affect the {0} channel of the transform's scale"), AxisDisplayInfo::GetAxisDisplayName(EAxisList::X)),
 					FSlateIcon(), MakeUIAction(EMovieSceneTransformChannel::ScaleX), NAME_None, EUserInterfaceActionType::ToggleButton);
 				SubMenuBuilder.AddMenuEntry(
-					LOCTEXT("ScaleY", "Y"), LOCTEXT("ScaleY_ToolTip", "Causes this section to affect the Y channel of the transform's scale"),
+					AxisDisplayInfo::GetAxisDisplayName(EAxisList::Y), 
+					FText::Format(LOCTEXT("ActivateScaleChannel_Tooltip", "Causes this section to affect the {0} channel of the transform's scale"), AxisDisplayInfo::GetAxisDisplayName(EAxisList::Y)),
 					FSlateIcon(), MakeUIAction(EMovieSceneTransformChannel::ScaleY), NAME_None, EUserInterfaceActionType::ToggleButton);
 				SubMenuBuilder.AddMenuEntry(
-					LOCTEXT("ScaleZ", "Z"), LOCTEXT("ScaleZ_ToolTip", "Causes this section to affect the Z channel of the transform's scale"),
+					AxisDisplayInfo::GetAxisDisplayName(EAxisList::Z), 
+					FText::Format(LOCTEXT("ActivateScaleChannel_Tooltip", "Causes this section to affect the {0} channel of the transform's scale"), AxisDisplayInfo::GetAxisDisplayName(EAxisList::Z)),
 					FSlateIcon(), MakeUIAction(EMovieSceneTransformChannel::ScaleZ), NAME_None, EUserInterfaceActionType::ToggleButton);
 			}),
 			MakeUIAction(EMovieSceneTransformChannel::Scale),
diff --git a/Engine/Source/Editor/MovieSceneTools/Private/TrackEditors/SubTrackEditor.cpp b/Engine/Source/Editor/MovieSceneTools/Private/TrackEditors/SubTrackEditor.cpp
--- a/Engine/Source/Editor/MovieSceneTools/Private/TrackEditors/SubTrackEditor.cpp
+++ b/Engine/Source/Editor/MovieSceneTools/Private/TrackEditors/SubTrackEditor.cpp
@@ -29,6 +29,7 @@
 #include "MovieSceneSequence.h"
 #include "MovieSceneToolHelpers.h"
 #include "MovieSceneToolsProjectSettings.h"
+#include "Misc/AxisDisplayInfo.h"
 #include "Misc/QualifiedFrameTime.h"
 #include "MovieSceneTimeHelpers.h"
 #include "EngineAnalytics.h"
@@ -176,13 +177,16 @@ public:
 			LOCTEXT("AllTranslation", "Translation"), LOCTEXT("AllTranslation_ToolTip", "Causes this section to affect the translation of the transform"),
 			FNewMenuDelegate::CreateLambda([Sequencer, MakeUIAction](FMenuBuilder& SubMenuBuilder){
 				SubMenuBuilder.AddMenuEntry(
-					LOCTEXT("TranslationX", "X"), LOCTEXT("TranslationX_ToolTip", "Causes this section to affect the X channel of the transform's translation"),
+					AxisDisplayInfo::GetAxisDisplayName(EAxisList::X), 
+					FText::Format(LOCTEXT("ActivateTranslationChannel_Tooltip", "Causes this section to affect the {0} channel of the transform's translation"), AxisDisplayInfo::GetAxisDisplayName(EAxisList::X)),
 					FSlateIcon(), MakeUIAction(EMovieSceneTransformChannel::TranslationX, Sequencer), NAME_None, EUserInterfaceActionType::ToggleButton);
 				SubMenuBuilder.AddMenuEntry(
-					LOCTEXT("TranslationY", "Y"), LOCTEXT("TranslationY_ToolTip", "Causes this section to affect the Y channel of the transform's translation"),
+					AxisDisplayInfo::GetAxisDisplayName(EAxisList::Y), 
+					FText::Format(LOCTEXT("ActivateTranslationChannel_Tooltip", "Causes this section to affect the {0} channel of the transform's translation"), AxisDisplayInfo::GetAxisDisplayName(EAxisList::Y)),
 					FSlateIcon(), MakeUIAction(EMovieSceneTransformChannel::TranslationY, Sequencer), NAME_None, EUserInterfaceActionType::ToggleButton);
 				SubMenuBuilder.AddMenuEntry(
-					LOCTEXT("TranslationZ", "Z"), LOCTEXT("TranslationZ_ToolTip", "Causes this section to affect the Z channel of the transform's translation"),
+					AxisDisplayInfo::GetAxisDisplayName(EAxisList::Z),
+					FText::Format(LOCTEXT("ActivateTranslationChannel_Tooltip", "Causes this section to affect the {0} channel of the transform's translation"), AxisDisplayInfo::GetAxisDisplayName(EAxisList::Z)),
 					FSlateIcon(), MakeUIAction(EMovieSceneTransformChannel::TranslationZ, Sequencer), NAME_None, EUserInterfaceActionType::ToggleButton);
 			}),
 			MakeUIAction(EMovieSceneTransformChannel::Translation, Sequencer),
@@ -193,13 +197,13 @@ public:
 			LOCTEXT("AllRotation", "Rotation"), LOCTEXT("AllRotation_ToolTip", "Causes this section to affect the rotation of the transform"),
 			FNewMenuDelegate::CreateLambda([Sequencer, MakeUIAction](FMenuBuilder& SubMenuBuilder){
 				SubMenuBuilder.AddMenuEntry(
-					LOCTEXT("RotationX", "Roll (X)"), LOCTEXT("RotationX_ToolTip", "Causes this section to affect the roll (X) channel the transform's rotation"),
+					LOCTEXT("RotationX", "Roll"), LOCTEXT("RotationX_ToolTip", "Causes this section to affect the roll channel the transform's rotation"),
 					FSlateIcon(), MakeUIAction(EMovieSceneTransformChannel::RotationX, Sequencer), NAME_None, EUserInterfaceActionType::ToggleButton);
 				SubMenuBuilder.AddMenuEntry(
-					LOCTEXT("RotationY", "Pitch (Y)"), LOCTEXT("RotationY_ToolTip", "Causes this section to affect the pitch (Y) channel the transform's rotation"),
+					LOCTEXT("RotationY", "Pitch"), LOCTEXT("RotationY_ToolTip", "Causes this section to affect the pitch channel the transform's rotation"),
 					FSlateIcon(), MakeUIAction(EMovieSceneTransformChannel::RotationY, Sequencer), NAME_None, EUserInterfaceActionType::ToggleButton);
 				SubMenuBuilder.AddMenuEntry(
-					LOCTEXT("RotationZ", "Yaw (Z)"), LOCTEXT("RotationZ_ToolTip", "Causes this section to affect the yaw (Z) channel the transform's rotation"),
+					LOCTEXT("RotationZ", "Yaw"), LOCTEXT("RotationZ_ToolTip", "Causes this section to affect the yaw channel the transform's rotation"),
 					FSlateIcon(), MakeUIAction(EMovieSceneTransformChannel::RotationZ, Sequencer), NAME_None, EUserInterfaceActionType::ToggleButton);
 			}),
 			MakeUIAction(EMovieSceneTransformChannel::Rotation, Sequencer),
diff --git a/Engine/Source/Runtime/MovieScene/Private/Sections/MovieSceneSubSection.cpp b/Engine/Source/Runtime/MovieScene/Private/Sections/MovieSceneSubSection.cpp
--- a/Engine/Source/Runtime/MovieScene/Private/Sections/MovieSceneSubSection.cpp
+++ b/Engine/Source/Runtime/MovieScene/Private/Sections/MovieSceneSubSection.cpp
@@ -11,6 +11,7 @@
 #include "Channels/MovieSceneChannelProxy.h"
 #include "EntitySystem/BuiltInComponentTypes.h"
 #include "Evaluation/MovieSceneEvaluationTemplate.h"
+#include "Misc/AxisDisplayInfo.h"
 #include "Misc/FrameRate.h"
 #include "Logging/MessageLog.h"
 #include "MovieSceneTransformTypes.h"
@@ -40,46 +41,46 @@ struct FSubSectionEditorData
 
 	FSubSectionEditorData(EMovieSceneTransformChannel Mask, UMovieSceneSubSection* SubSection)
 	{
-		MetaData[0].SetIdentifiers("Override.Location.X", FCommonChannelData::ChannelX, LocationGroup);
+		MetaData[0].SetIdentifiers("Override.Location.X", AxisDisplayInfo::GetAxisDisplayName(EAxisList::X), LocationGroup);
 		MetaData[0].SubPropertyPath = TEXT("Location.X");
 		MetaData[0].SortOrder = 0;
 		MetaData[0].bEnabled = EnumHasAllFlags(Mask, EMovieSceneTransformChannel::TranslationX);
-		MetaData[0].Color = FCommonChannelData::RedChannelColor;
+		MetaData[0].Color = AxisDisplayInfo::GetAxisColor(EAxisList::X);
 		MetaData[0].bCanCollapseToTrack = false;
 
-		MetaData[1].SetIdentifiers("Override.Location.Y", FCommonChannelData::ChannelY, LocationGroup);
+		MetaData[1].SetIdentifiers("Override.Location.Y", AxisDisplayInfo::GetAxisDisplayName(EAxisList::Y), LocationGroup);
 		MetaData[1].SubPropertyPath = TEXT("Location.Y");
 		MetaData[1].SortOrder = 1;
 		MetaData[1].bEnabled = EnumHasAllFlags(Mask, EMovieSceneTransformChannel::TranslationY);
-		MetaData[1].Color = FCommonChannelData::GreenChannelColor;
+		MetaData[1].Color = AxisDisplayInfo::GetAxisColor(EAxisList::Y);
 		MetaData[1].bCanCollapseToTrack = false;
 		
-		MetaData[2].SetIdentifiers("Override.Location.Z", FCommonChannelData::ChannelZ, LocationGroup);
+		MetaData[2].SetIdentifiers("Override.Location.Z", AxisDisplayInfo::GetAxisDisplayName(EAxisList::Z), LocationGroup);
 		MetaData[2].SubPropertyPath = TEXT("Location.Z");
 		MetaData[2].SortOrder = 2;
 		MetaData[2].bEnabled = EnumHasAllFlags(Mask, EMovieSceneTransformChannel::TranslationZ);
-		MetaData[2].Color = FCommonChannelData::BlueChannelColor;
+		MetaData[2].Color = AxisDisplayInfo::GetAxisColor(EAxisList::Z);
 		MetaData[2].bCanCollapseToTrack = false;
 
-		MetaData[3].SetIdentifiers("Override.Rotation.X", FCommonChannelData::ChannelX, RotationGroup);
+		MetaData[3].SetIdentifiers("Override.Rotation.X", NSLOCTEXT("MovieSceneSubSection", "RotationX", "Roll"), RotationGroup);
 		MetaData[3].SubPropertyPath = TEXT("Rotation.X");
 		MetaData[3].SortOrder = 3;
 		MetaData[3].bEnabled = EnumHasAllFlags(Mask, EMovieSceneTransformChannel::RotationX);
-		MetaData[3].Color = FCommonChannelData::RedChannelColor;
+		MetaData[3].Color = AxisDisplayInfo::GetAxisColor(EAxisList::X);
 		MetaData[3].bCanCollapseToTrack = false;
 
-		MetaData[4].SetIdentifiers("Override.Rotation.Y", FCommonChannelData::ChannelY, RotationGroup);
+		MetaData[4].SetIdentifiers("Override.Rotation.Y", NSLOCTEXT("MovieSceneSubSection", "RotationY", "Pitch"), RotationGroup);
 		MetaData[4].SubPropertyPath = TEXT("Rotation.Y");
 		MetaData[4].SortOrder = 4;
 		MetaData[4].bEnabled = EnumHasAllFlags(Mask, EMovieSceneTransformChannel::RotationY);
-		MetaData[4].Color = FCommonChannelData::GreenChannelColor;
+		MetaData[4].Color = AxisDisplayInfo::GetAxisColor(EAxisList::Y);
 		MetaData[4].bCanCollapseToTrack = false;
 
-		MetaData[5].SetIdentifiers("Override.Rotation.Z", FCommonChannelData::ChannelZ, RotationGroup);
+		MetaData[5].SetIdentifiers("Override.Rotation.Z", NSLOCTEXT("MovieSceneSubSection", "RotationZ", "Yaw"), RotationGroup);
 		MetaData[5].SubPropertyPath = TEXT("Rotation.Z");
 		MetaData[5].SortOrder = 5;
 		MetaData[5].bEnabled = EnumHasAllFlags(Mask, EMovieSceneTransformChannel::RotationZ);
-		MetaData[5].Color = FCommonChannelData::BlueChannelColor;
+		MetaData[5].Color = AxisDisplayInfo::GetAxisColor(EAxisList::Z);
 		MetaData[5].bCanCollapseToTrack = false;
 
 		ExternalValues[0].OnGetExternalValue = [SubSection](UObject& InObject, FTrackInstancePropertyBindings* Bindings) { return GetValue(SubSection, 0); };
diff --git a/Engine/Source/Runtime/MovieSceneTracks/Private/Sections/MovieScene3DTransformSection.cpp b/Engine/Source/Runtime/MovieSceneTracks/Private/Sections/MovieScene3DTransformSection.cpp
--- a/Engine/Source/Runtime/MovieSceneTracks/Private/Sections/MovieScene3DTransformSection.cpp
+++ b/Engine/Source/Runtime/MovieSceneTracks/Private/Sections/MovieScene3DTransformSection.cpp
@@ -19,7 +19,7 @@
 #include "TransformConstraint.h"
 #include "TransformableHandle.h"
 #include "UObject/ObjectSaveContext.h"
-
+#include "Misc/AxisDisplayInfo.h"
 
 #include UE_INLINE_GENERATED_CPP_BY_NAME(MovieScene3DTransformSection)
 
@@ -48,10 +48,10 @@ struct F3DTransformChannelEditorData
 		FText RotationGroup = NSLOCTEXT("MovieSceneTransformSection", "Rotation", "Rotation");
 		FText ScaleGroup    = NSLOCTEXT("MovieSceneTransformSection", "Scale",    "Scale");
 		{
-			MetaData[0].SetIdentifiers("Location.X", FCommonChannelData::ChannelX, LocationGroup);
+			MetaData[0].SetIdentifiers("Location.X", AxisDisplayInfo::GetAxisDisplayName(EAxisList::X), LocationGroup);
 			MetaData[0].SubPropertyPathMap = MakeSubPropertyMap(TEXT("Translation.X"), TEXT("Location.X"));
 			MetaData[0].bEnabled = EnumHasAllFlags(Mask, EMovieSceneTransformChannel::TranslationX);
-			MetaData[0].Color = FCommonChannelData::RedChannelColor;
+			MetaData[0].Color = AxisDisplayInfo::GetAxisColor(EAxisList::X);
 			MetaData[0].SortOrder = SortOrderStart;
 			MetaData[0].bCanCollapseToTrack = false;
 			for (const FName& PropertyMetaDataKey : PropertyMetaDataKeys)
@@ -59,10 +59,10 @@ struct F3DTransformChannelEditorData
 				MetaData[0].PropertyMetaData.Add(PropertyMetaDataKey, RelativeLocationProperty->GetMetaData(PropertyMetaDataKey));
 			}
 
-			MetaData[1].SetIdentifiers("Location.Y", FCommonChannelData::ChannelY, LocationGroup);
+			MetaData[1].SetIdentifiers("Location.Y", AxisDisplayInfo::GetAxisDisplayName(EAxisList::Y), LocationGroup);
 			MetaData[1].SubPropertyPathMap = MakeSubPropertyMap(TEXT("Translation.Y"), TEXT("Location.Y"));
 			MetaData[1].bEnabled = EnumHasAllFlags(Mask, EMovieSceneTransformChannel::TranslationY);
-			MetaData[1].Color = FCommonChannelData::GreenChannelColor;
+			MetaData[1].Color = AxisDisplayInfo::GetAxisColor(EAxisList::Y);
 			MetaData[1].SortOrder = SortOrderStart + 1;
 			MetaData[1].bCanCollapseToTrack = false;
 			for (const FName& PropertyMetaDataKey : PropertyMetaDataKeys)
@@ -70,10 +70,10 @@ struct F3DTransformChannelEditorData
 				MetaData[1].PropertyMetaData.Add(PropertyMetaDataKey, RelativeLocationProperty->GetMetaData(PropertyMetaDataKey));
 			}
 
-			MetaData[2].SetIdentifiers("Location.Z", FCommonChannelData::ChannelZ, LocationGroup);
+			MetaData[2].SetIdentifiers("Location.Z", AxisDisplayInfo::GetAxisDisplayName(EAxisList::Z), LocationGroup);
 			MetaData[2].SubPropertyPathMap = MakeSubPropertyMap(TEXT("Translation.Z"), TEXT("Location.Z"));
 			MetaData[2].bEnabled = EnumHasAllFlags(Mask, EMovieSceneTransformChannel::TranslationZ);
-			MetaData[2].Color = FCommonChannelData::BlueChannelColor;
+			MetaData[2].Color = AxisDisplayInfo::GetAxisColor(EAxisList::Z);
 			MetaData[2].SortOrder = SortOrderStart + 2;
 			MetaData[2].bCanCollapseToTrack = false;
 			for (const FName& PropertyMetaDataKey : PropertyMetaDataKeys)
@@ -85,7 +85,7 @@ struct F3DTransformChannelEditorData
 			MetaData[3].SetIdentifiers("Rotation.X", NSLOCTEXT("MovieSceneTransformSection", "RotationX", "Roll"), RotationGroup);
 			MetaData[3].SubPropertyPathMap = MakeSubPropertyMap(TEXT("Rotation.X"), TEXT("Rotation.Roll"));
 			MetaData[3].bEnabled = EnumHasAllFlags(Mask, EMovieSceneTransformChannel::RotationX);
-			MetaData[3].Color = FCommonChannelData::RedChannelColor;
+			MetaData[3].Color = AxisDisplayInfo::GetAxisColor(EAxisList::X);
 			MetaData[3].SortOrder = SortOrderStart + 3;
 			MetaData[3].bCanCollapseToTrack = false;
 			for (const FName& PropertyMetaDataKey : PropertyMetaDataKeys)
@@ -96,7 +96,7 @@ struct F3DTransformChannelEditorData
 			MetaData[4].SetIdentifiers("Rotation.Y", NSLOCTEXT("MovieSceneTransformSection", "RotationY", "Pitch"), RotationGroup);
 			MetaData[4].SubPropertyPathMap = MakeSubPropertyMap(TEXT("Rotation.Y"), TEXT("Rotation.Pitch"));
 			MetaData[4].bEnabled = EnumHasAllFlags(Mask, EMovieSceneTransformChannel::RotationY);
-			MetaData[4].Color = FCommonChannelData::GreenChannelColor;
+			MetaData[4].Color = AxisDisplayInfo::GetAxisColor(EAxisList::Y);
 			MetaData[4].SortOrder = SortOrderStart + 4;
 			MetaData[4].bCanCollapseToTrack = false;
 			for (const FName& PropertyMetaDataKey : PropertyMetaDataKeys)
@@ -107,7 +107,7 @@ struct F3DTransformChannelEditorData
 			MetaData[5].SetIdentifiers("Rotation.Z", NSLOCTEXT("MovieSceneTransformSection", "RotationZ", "Yaw"), RotationGroup);
 			MetaData[5].SubPropertyPathMap = MakeSubPropertyMap(TEXT("Rotation.Z"), TEXT("Rotation.Yaw"));
 			MetaData[5].bEnabled = EnumHasAllFlags(Mask, EMovieSceneTransformChannel::RotationZ);
-			MetaData[5].Color = FCommonChannelData::BlueChannelColor;
+			MetaData[5].Color = AxisDisplayInfo::GetAxisColor(EAxisList::Z);
 			MetaData[5].SortOrder = SortOrderStart + 5;
 			MetaData[5].bCanCollapseToTrack = false;
 			for (const FName& PropertyMetaDataKey : PropertyMetaDataKeys)
@@ -116,10 +116,10 @@ struct F3DTransformChannelEditorData
 			}
 		}
 		{
-			MetaData[6].SetIdentifiers("Scale.X", FCommonChannelData::ChannelX, ScaleGroup);
+			MetaData[6].SetIdentifiers("Scale.X", AxisDisplayInfo::GetAxisDisplayName(EAxisList::X), ScaleGroup);
 			MetaData[6].SubPropertyPathMap = MakeSubPropertyMap(TEXT("Scale3D.X"), TEXT("Scale.X"));
 			MetaData[6].bEnabled = EnumHasAllFlags(Mask, EMovieSceneTransformChannel::ScaleX);
-			MetaData[6].Color = FCommonChannelData::RedChannelColor;
+			MetaData[6].Color = AxisDisplayInfo::GetAxisColor(EAxisList::X);
 			MetaData[6].SortOrder = SortOrderStart + 6;
 			MetaData[6].bCanCollapseToTrack = false;
 			for (const FName& PropertyMetaDataKey : PropertyMetaDataKeys)
@@ -127,10 +127,10 @@ struct F3DTransformChannelEditorData
 				MetaData[6].PropertyMetaData.Add(PropertyMetaDataKey, RelativeScale3DProperty->GetMetaData(PropertyMetaDataKey));
 			}
 
-			MetaData[7].SetIdentifiers("Scale.Y", FCommonChannelData::ChannelY, ScaleGroup);
+			MetaData[7].SetIdentifiers("Scale.Y", AxisDisplayInfo::GetAxisDisplayName(EAxisList::Y), ScaleGroup);
 			MetaData[7].SubPropertyPathMap = MakeSubPropertyMap(TEXT("Scale3D.Y"), TEXT("Scale.Y"));
 			MetaData[7].bEnabled = EnumHasAllFlags(Mask, EMovieSceneTransformChannel::ScaleY);
-			MetaData[7].Color = FCommonChannelData::GreenChannelColor;
+			MetaData[7].Color = AxisDisplayInfo::GetAxisColor(EAxisList::Y);
 			MetaData[7].SortOrder = SortOrderStart + 7;
 			MetaData[7].bCanCollapseToTrack = false;
 			for (const FName& PropertyMetaDataKey : PropertyMetaDataKeys)
@@ -138,10 +138,10 @@ struct F3DTransformChannelEditorData
 				MetaData[7].PropertyMetaData.Add(PropertyMetaDataKey, RelativeScale3DProperty->GetMetaData(PropertyMetaDataKey));
 			}
 
-			MetaData[8].SetIdentifiers("Scale.Z", FCommonChannelData::ChannelZ, ScaleGroup);
+			MetaData[8].SetIdentifiers("Scale.Z", AxisDisplayInfo::GetAxisDisplayName(EAxisList::Z), ScaleGroup);
 			MetaData[8].SubPropertyPathMap = MakeSubPropertyMap(TEXT("Scale3D.Z"), TEXT("Scale.Z"));
 			MetaData[8].bEnabled = EnumHasAllFlags(Mask, EMovieSceneTransformChannel::ScaleZ);
-			MetaData[8].Color = FCommonChannelData::BlueChannelColor;
+			MetaData[8].Color = AxisDisplayInfo::GetAxisColor(EAxisList::Z);
 			MetaData[8].SortOrder = SortOrderStart + 8;
 			MetaData[8].bCanCollapseToTrack = false;
 			for (const FName& PropertyMetaDataKey : PropertyMetaDataKeys)
