Control Rig: Add semantic labels to parent constraint + available spaces

#rnx

diff --git a/Engine/Plugins/Animation/ControlRig/Source/ControlRig/Private/Rigs/RigHierarchy.cpp b/Engine/Plugins/Animation/ControlRig/Source/ControlRig/Private/Rigs/RigHierarchy.cpp
--- a/Engine/Plugins/Animation/ControlRig/Source/ControlRig/Private/Rigs/RigHierarchy.cpp
+++ b/Engine/Plugins/Animation/ControlRig/Source/ControlRig/Private/Rigs/RigHierarchy.cpp
@@ -2859,6 +2859,51 @@ FRigBaseElement* URigHierarchy::GetActiveParent(const FRigBaseElement* InElement
 	return nullptr;
 }
 
+FName URigHierarchy::GetDisplayLabelForParent(const FRigElementKey& InChildKey, const FRigElementKey& InParentKey) const
+{
+	if(InParentKey == GetDefaultParentKey())
+	{
+		return DefaultParentKeyLabel.Resolve();
+	}
+	if(InParentKey == GetWorldSpaceReferenceKey())
+	{
+		return WorldSpaceKeyLabel.Resolve();
+	}
+
+	if(!InChildKey.IsValid() || !InParentKey.IsValid())
+	{
+		return NAME_None;
+	}
+
+	if(const FRigControlElement* ControlElement = Find<FRigControlElement>(InChildKey))
+	{
+		if(const FRigElementKeyWithLabel* AvailableSpace = ControlElement->Settings.Customization.AvailableSpaces.FindByKey(InParentKey))
+		{
+			if(!AvailableSpace->Label.IsNone())
+			{
+				return AvailableSpace->Label;
+			}
+		}
+	}
+	if(const FRigMultiParentElement* MultiParentElement = Find<FRigMultiParentElement>(InChildKey))
+	{
+		if(const int32* ParentIndexPtr = MultiParentElement->IndexLookup.Find(InParentKey))
+		{
+			if(GetDefaultParent(InChildKey) == InParentKey)
+			{
+				return DefaultParentKeyLabel.Resolve();
+			}
+			
+			const FName& Label = MultiParentElement->ParentConstraints[*ParentIndexPtr].DisplayLabel;
+			if(!Label.IsNone())
+			{
+				return Label;
+			}
+		}
+	}
+
+	return InParentKey.Name;
+}
 
 bool URigHierarchy::SetParentWeight(FRigElementKey InChild, FRigElementKey InParent, FRigElementWeight InWeight, bool bInitial, bool bAffectChildren)
 {
@@ -3360,7 +3405,7 @@ TArray<FRigControlElement*> URigHierarchy::GetAnimationChannels(const FRigContro
 			{
 				if(PotentialAnimationChannel->IsAnimationChannel())
 				{
-					if(PotentialAnimationChannel->Settings.Customization.AvailableSpaces.Contains(InElement->Key))
+					if(PotentialAnimationChannel->Settings.Customization.AvailableSpaces.FindByKey(InElement->Key))
 					{
 						return true;
 					}
diff --git a/Engine/Plugins/Animation/ControlRig/Source/ControlRig/Private/Rigs/RigHierarchyController.cpp b/Engine/Plugins/Animation/ControlRig/Source/ControlRig/Private/Rigs/RigHierarchyController.cpp
--- a/Engine/Plugins/Animation/ControlRig/Source/ControlRig/Private/Rigs/RigHierarchyController.cpp
+++ b/Engine/Plugins/Animation/ControlRig/Source/ControlRig/Private/Rigs/RigHierarchyController.cpp
@@ -1637,7 +1637,12 @@ FString URigHierarchyController::ExportToText(TArray<FRigElementKey> InKeys) con
 
 		FRigHierarchyCopyPasteContentPerElement PerElementData;
 		PerElementData.Key = Key;
-		PerElementData.Parents = Hierarchy->GetParents(Key);
+		const TArray<FRigElementKey> ParentKeys = Hierarchy->GetParents(Key);
+		PerElementData.Parents.Reserve(ParentKeys.Num());
+		for(const FRigElementKey& ParentKey : ParentKeys)
+		{
+			PerElementData.Parents.Emplace(ParentKey, Hierarchy->GetDisplayLabelForParent(Key, ParentKey));
+		}
 
 		if(FRigMultiParentElement* MultiParentElement = Cast<FRigMultiParentElement>(Element))
 		{
@@ -1906,9 +1911,9 @@ TArray<FRigElementKey> URigHierarchyController::ImportFromText(FString InContent
 				{
 					RemoveAllParents(ExistingElement->GetKey(), true, bSetupUndo);
 
-					for(const FRigElementKey& NewParent : PerElementData.Parents)
+					for(const FRigElementKeyWithLabel& NewParent : PerElementData.Parents)
 					{
-						AddParent(ExistingElement->GetKey(), NewParent, 0.f, true, bSetupUndo);
+						AddParent(ExistingElement->GetKey(), NewParent.Key, 0.f, true, NewParent.Label, bSetupUndo);
 					}
 				}
 				
@@ -1945,15 +1950,15 @@ TArray<FRigElementKey> URigHierarchyController::ImportFromText(FString InContent
 		FRigElementKey MappedKey = KeyMap.FindChecked(PerElementData.Key);
 		FRigBaseElement* NewElement = Hierarchy->FindChecked(MappedKey);
 
-		for(const FRigElementKey& OriginalParent : PerElementData.Parents)
+		for(const FRigElementKeyWithLabel& OriginalParent : PerElementData.Parents)
 		{
-			FRigElementKey Parent = OriginalParent;
+			FRigElementKey Parent = OriginalParent.Key;
 			if(const FRigElementKey* RemappedParent = KeyMap.Find(Parent))
 			{
 				Parent = *RemappedParent;
 			}
 
-			AddParent(NewElement->GetKey(), Parent, 0.f, true, bSetupUndo);
+			AddParent(NewElement->GetKey(), Parent, 0.f, true, OriginalParent.Label, bSetupUndo);
 		}
 		
 		for(int32 ParentIndex = 0; ParentIndex < PerElementData.ParentWeights.Num(); ParentIndex++)
@@ -2678,7 +2683,11 @@ bool URigHierarchyController::RemoveElement(FRigBaseElement* InElement)
 	{
 		if(FRigControlElement* ControlElement = Cast<FRigControlElement>(Element))
 		{
-			ControlElement->Settings.Customization.AvailableSpaces.Remove(InElement->GetKey());
+			const int32 ExistingSpaceIndex = ControlElement->Settings.Customization.AvailableSpaces.IndexOfByKey(InElement->GetKey());
+			if(ExistingSpaceIndex != INDEX_NONE)
+			{
+				ControlElement->Settings.Customization.AvailableSpaces.RemoveAt(ExistingSpaceIndex);
+			}
 			ControlElement->Settings.Customization.RemovedSpaces.Remove(InElement->GetKey());
 			ControlElement->Settings.DrivenControls.Remove(InElement->GetKey());
 		}
@@ -2932,11 +2941,11 @@ bool URigHierarchyController::RenameElement(FRigBaseElement* InElement, const FN
 
 		if(FRigControlElement* ControlElement = Cast<FRigControlElement>(Element))
 		{
-			for(FRigElementKey& Favorite : ControlElement->Settings.Customization.AvailableSpaces)
+			for(FRigElementKeyWithLabel& Favorite : ControlElement->Settings.Customization.AvailableSpaces)
 			{
-				if(Favorite == OldKey)
+				if(Favorite.Key == OldKey)
 				{
-					Favorite.Name = NewKey.Name;
+					Favorite.Key.Name = NewKey.Name;
 				}
 			}
 
@@ -3078,7 +3087,7 @@ FName URigHierarchyController::SetDisplayName(FRigControlElement* InControlEleme
 	return InControlElement->Settings.DisplayName;
 }
 
-bool URigHierarchyController::AddParent(FRigElementKey InChild, FRigElementKey InParent, float InWeight, bool bMaintainGlobalTransform, bool bSetupUndo)
+bool URigHierarchyController::AddParent(FRigElementKey InChild, FRigElementKey InParent, float InWeight, bool bMaintainGlobalTransform, const FName& InDisplayLabel, bool bSetupUndo)
 {
 	if(!IsValid())
 	{
@@ -3116,7 +3125,7 @@ bool URigHierarchyController::AddParent(FRigElementKey InChild, FRigElementKey I
 	}
 #endif
 
-	const bool bAdded = AddParent(Child, Parent, InWeight, bMaintainGlobalTransform);
+	const bool bAdded = AddParent(Child, Parent, InWeight, bMaintainGlobalTransform, false, InDisplayLabel);
 
 #if WITH_EDITOR
 	if(!bAdded && TransactionPtr.IsValid())
@@ -3129,7 +3138,7 @@ bool URigHierarchyController::AddParent(FRigElementKey InChild, FRigElementKey I
 	return bAdded;
 }
 
-bool URigHierarchyController::AddParent(FRigBaseElement* InChild, FRigBaseElement* InParent, float InWeight, bool bMaintainGlobalTransform, bool bRemoveAllParents)
+bool URigHierarchyController::AddParent(FRigBaseElement* InChild, FRigBaseElement* InParent, float InWeight, bool bMaintainGlobalTransform, bool bRemoveAllParents, const FName& InDisplayLabel)
 {
 	if(InChild == nullptr || InParent == nullptr)
 	{
@@ -3177,7 +3186,7 @@ bool URigHierarchyController::AddParent(FRigBaseElement* InChild, FRigBaseElemen
 
 		if(ChildControlElement->Settings.bRestrictSpaceSwitching)
 		{
-			if(ChildControlElement->Settings.Customization.AvailableSpaces.Contains(InParent->GetKey()))
+			if(ChildControlElement->Settings.Customization.AvailableSpaces.FindByKey(InParent->GetKey()) != nullptr)
 			{
 				return false;
 			}
@@ -3234,6 +3243,7 @@ bool URigHierarchyController::AddParent(FRigBaseElement* InChild, FRigBaseElemen
 	}
 	Constraint.InitialWeight = InWeight;
 	Constraint.Weight = InWeight;
+	Constraint.DisplayLabel = InDisplayLabel;
 
 	if(FRigSingleParentElement* SingleParentElement = Cast<FRigSingleParentElement>(InChild))
 	{
@@ -3654,7 +3664,7 @@ bool URigHierarchyController::SetParent(FRigElementKey InChild, FRigElementKey I
 	return bParentSet;
 }
 
-bool URigHierarchyController::AddAvailableSpace(FRigElementKey InControl, FRigElementKey InSpace, bool bSetupUndo, bool bPrintPythonCommand)
+bool URigHierarchyController::AddAvailableSpace(FRigElementKey InControl, FRigElementKey InSpace, const FName& InDisplayLabel, bool bSetupUndo, bool bPrintPythonCommand)
 {
 	if(!IsValid())
 	{
@@ -3700,7 +3710,7 @@ bool URigHierarchyController::AddAvailableSpace(FRigElementKey InControl, FRigEl
 	}
 #endif
 
-	const bool bSuccess = AddAvailableSpace(Control, Space);
+	const bool bSuccess = AddAvailableSpace(Control, Space, InDisplayLabel);
 
 #if WITH_EDITOR
 	if(!bSuccess && TransactionPtr.IsValid())
@@ -3714,9 +3724,10 @@ bool URigHierarchyController::AddAvailableSpace(FRigElementKey InControl, FRigEl
 		if (const UBlueprint* Blueprint = GetTypedOuter<UBlueprint>())
 		{
 			RigVMPythonUtils::Print(Blueprint->GetFName().ToString(),
-				FString::Printf(TEXT("hierarchy_controller.add_available_space(%s, %s)"),
+				FString::Printf(TEXT("hierarchy_controller.add_available_space(%s, %s, '%s')"),
 				*InControl.ToPythonString(),
-				*InSpace.ToPythonString()));
+				*InSpace.ToPythonString(),
+				*InDisplayLabel.ToString()));
 		}
 	}
 #endif
@@ -3835,7 +3846,7 @@ bool URigHierarchyController::SetAvailableSpaceIndex(FRigElementKey InControl, F
 	TSharedPtr<FScopedTransaction> TransactionPtr;
 	if(bSetupUndo)
 	{
-		TransactionPtr = MakeShared<FScopedTransaction>(NSLOCTEXT("RigHierarchyController", "RemoveAvailableSpace", "Remove Available Space"));
+		TransactionPtr = MakeShared<FScopedTransaction>(NSLOCTEXT("RigHierarchyController", "SetAvailableSpaceIndex", "Reorder Available Space"));
 		Hierarchy->Modify();
 	}
 #endif
@@ -3864,6 +3875,77 @@ bool URigHierarchyController::SetAvailableSpaceIndex(FRigElementKey InControl, F
 	return bSuccess;
 }
 
+bool URigHierarchyController::SetAvailableSpaceLabel(FRigElementKey InControl, FRigElementKey InSpace, FName InDisplayLabel, bool bSetupUndo, bool bPrintPythonCommand)
+{
+	if(!IsValid())
+	{
+		return false;
+	}
+
+	URigHierarchy* Hierarchy = GetHierarchy();
+
+	FRigBaseElement* ControlBase = Hierarchy->Find(InControl);
+	if(ControlBase == nullptr)
+	{
+		ReportWarningf(TEXT("Cannot Set Available Space Label, Control '%s' not found."), *InControl.ToString());
+		return false;
+	}
+
+	FRigControlElement* Control = Cast<FRigControlElement>(ControlBase);
+	if(Control == nullptr)
+	{
+		ReportWarningf(TEXT("Cannot Set Available Space Label, '%s' is not a Control."), *InControl.ToString());
+		return false;
+	}
+
+	const FRigBaseElement* SpaceBase = Hierarchy->Find(InSpace);
+	if(SpaceBase == nullptr)
+	{
+		ReportWarningf(TEXT("Cannot Set Available Space Label, Space '%s' not found."), *InSpace.ToString());
+		return false;
+	}
+
+	const FRigTransformElement* Space = Cast<FRigTransformElement>(SpaceBase);
+	if(Space == nullptr)
+	{
+		ReportWarningf(TEXT("Cannot Set Available Space Label, '%s' is not a Transform."), *InSpace.ToString());
+		return false;
+	}
+
+#if WITH_EDITOR
+	TSharedPtr<FScopedTransaction> TransactionPtr;
+	if(bSetupUndo)
+	{
+		TransactionPtr = MakeShared<FScopedTransaction>(NSLOCTEXT("RigHierarchyController", "SetAvailableSpaceLabel", "Set Available Space Label"));
+		Hierarchy->Modify();
+	}
+#endif
+
+	const bool bSuccess = SetAvailableSpaceLabel(Control, Space, InDisplayLabel);
+
+#if WITH_EDITOR
+	if(!bSuccess && TransactionPtr.IsValid())
+	{
+		TransactionPtr->Cancel();
+	}
+	TransactionPtr.Reset();
+
+	if (bSuccess && bPrintPythonCommand && !bSuspendPythonPrinting)
+	{
+		if (const UBlueprint* Blueprint = GetTypedOuter<UBlueprint>())
+		{
+			RigVMPythonUtils::Print(Blueprint->GetFName().ToString(),
+				FString::Printf(TEXT("hierarchy_controller.set_available_space_label(%s, %s, '%s')"),
+				*InControl.ToPythonString(),
+				*InSpace.ToPythonString(),
+				*InDisplayLabel.ToString()));
+		}
+	}
+#endif
+
+	return bSuccess;
+}
+
 bool URigHierarchyController::AddChannelHost(FRigElementKey InChannel, FRigElementKey InHost, bool bSetupUndo, bool bPrintPythonCommand)
 {
 	if(!IsValid())
@@ -3920,7 +4002,7 @@ bool URigHierarchyController::AddChannelHost(FRigElementKey InChannel, FRigEleme
 		return false;
 	}
 
-	if(Channel->Settings.Customization.AvailableSpaces.Contains(Host->GetKey()))
+	if(Channel->Settings.Customization.AvailableSpaces.FindByKey(Host->GetKey()))
 	{
 		ReportWarningf(TEXT("Cannot Add Channel Host, '%s' is already a host for channel '%s'."), *InHost.ToString(), *InChannel.ToString());
 		return false;
@@ -4008,7 +4090,7 @@ bool URigHierarchyController::RemoveChannelHost(FRigElementKey InChannel, FRigEl
 		return false;
 	}
 
-	if(!Channel->Settings.Customization.AvailableSpaces.Contains(Host->GetKey()))
+	if(Channel->Settings.Customization.AvailableSpaces.FindByKey(Host->GetKey()) == nullptr)
 	{
 		ReportWarningf(TEXT("Cannot Remove Channel Host, '%s' is not a host for channel '%s'."), *InHost.ToString(), *InChannel.ToString());
 		return false;
@@ -4218,7 +4300,7 @@ bool URigHierarchyController::SetParent(FRigBaseElement* InChild, FRigBaseElemen
 	return AddParent(InChild, InParent, 1.f, bMaintainGlobalTransform, true);
 }
 
-bool URigHierarchyController::AddAvailableSpace(FRigControlElement* InControlElement, const FRigTransformElement* InSpaceElement)
+bool URigHierarchyController::AddAvailableSpace(FRigControlElement* InControlElement, const FRigTransformElement* InSpaceElement, const FName& InDisplayLabel)
 {
 	if(InControlElement == nullptr || InSpaceElement == nullptr)
 	{
@@ -4250,13 +4332,13 @@ bool URigHierarchyController::AddAvailableSpace(FRigControlElement* InControlEle
 	}
 
 	FRigControlSettings Settings = InControlElement->Settings;
-	TArray<FRigElementKey>& AvailableSpaces = Settings.Customization.AvailableSpaces;
-	if(AvailableSpaces.Contains(InSpaceElement->GetKey()))
+	TArray<FRigElementKeyWithLabel>& AvailableSpaces = Settings.Customization.AvailableSpaces;
+	if(AvailableSpaces.FindByKey(InSpaceElement->GetKey()))
 	{
 		return false;
 	}
 
-	AvailableSpaces.Add(InSpaceElement->GetKey());
+	AvailableSpaces.Emplace(InSpaceElement->GetKey(), InDisplayLabel);
 
 	GetHierarchy()->SetControlSettings(InControlElement, Settings, false, false, false);
 	return true;
@@ -4270,12 +4352,13 @@ bool URigHierarchyController::RemoveAvailableSpace(FRigControlElement* InControl
 	}
 
 	FRigControlSettings Settings = InControlElement->Settings;
-	TArray<FRigElementKey>& AvailableSpaces =Settings.Customization.AvailableSpaces;
-	const int32 NumRemoved = AvailableSpaces.Remove(InSpaceElement->GetKey());
-	if(NumRemoved == 0)
+	TArray<FRigElementKeyWithLabel>& AvailableSpaces =Settings.Customization.AvailableSpaces;
+	const int32 ExistingSpaceIndex = AvailableSpaces.IndexOfByKey(InSpaceElement->GetKey());
+	if(ExistingSpaceIndex == INDEX_NONE)
 	{
 		return false;
 	}
+	AvailableSpaces.RemoveAt(ExistingSpaceIndex);
 	
 	GetHierarchy()->SetControlSettings(InControlElement, Settings, false, false, false);
 	return true;
@@ -4289,39 +4372,87 @@ bool URigHierarchyController::SetAvailableSpaceIndex(FRigControlElement* InContr
 	}
 
 	FRigControlSettings Settings = InControlElement->Settings;
-	TArray<FRigElementKey>& AvailableSpaces = Settings.Customization.AvailableSpaces;
+	TArray<FRigElementKeyWithLabel>& AvailableSpaces = Settings.Customization.AvailableSpaces;
 
 	bool bAddedAvailableSpace = false;
-	if(!AvailableSpaces.Contains(InSpaceElement->GetKey()))
+	if(AvailableSpaces.FindByKey(InSpaceElement->GetKey()) == nullptr)
 	{
 		bAddedAvailableSpace = AddAvailableSpace(InControlElement, InSpaceElement);
 		if(!bAddedAvailableSpace)
 		{
 			return false;
 		}
+		Settings = InControlElement->Settings;
+		AvailableSpaces = Settings.Customization.AvailableSpaces;
 	}
 
-	if(AvailableSpaces.Find(InSpaceElement->GetKey()) == InIndex)
+	const int32 ExistingSpaceIndex = AvailableSpaces.IndexOfByKey(InSpaceElement->GetKey());
+	if(ExistingSpaceIndex == InIndex)
 	{
 		return bAddedAvailableSpace;
 	}
 
-	(void)AvailableSpaces.Remove(InSpaceElement->GetKey());
+	const FRigElementKeyWithLabel AvailableSpaceToMove = AvailableSpaces[ExistingSpaceIndex];
+	AvailableSpaces.RemoveAt(ExistingSpaceIndex);
+	
 	InIndex = FMath::Max(InIndex, 0);
-
 	if(InIndex >= AvailableSpaces.Num())
 	{
-		AvailableSpaces.Add(InSpaceElement->GetKey());
+		AvailableSpaces.Add(AvailableSpaceToMove);
 	}
 	else
 	{
-		AvailableSpaces.Insert(InSpaceElement->GetKey(), InIndex);
+		AvailableSpaces.Insert(AvailableSpaceToMove, InIndex);
 	}
 
 	GetHierarchy()->SetControlSettings(InControlElement, Settings, false, false, false);
 	return true;
 }
 
+bool URigHierarchyController::SetAvailableSpaceLabel(FRigControlElement* InControlElement, const FRigTransformElement* InSpaceElement, const FName& InDisplayLabel)
+{
+	if(InControlElement == nullptr || InSpaceElement == nullptr)
+	{
+		return false;
+	}
+
+	FRigControlSettings Settings = InControlElement->Settings;
+	TArray<FRigElementKeyWithLabel>& AvailableSpaces = Settings.Customization.AvailableSpaces;
+
+	// first let's see if this is an available space that's registered in the control's settings
+	if(FRigElementKeyWithLabel* AvailableSpace = AvailableSpaces.FindByKey(InSpaceElement->GetKey()))
+	{
+		if(AvailableSpace->Label.IsEqual(InDisplayLabel))
+		{
+			return false;
+		}
+		AvailableSpace->Label = InDisplayLabel;
+		GetHierarchy()->SetControlSettings(InControlElement, Settings, false, false, false);
+		return true;
+	}
+
+	// now we'll need to take a look at the parent constraints of the control
+	if(const int32* ParentConstraintIndexPtr = InControlElement->IndexLookup.Find(InSpaceElement->GetKey()))
+	{
+		if(GetHierarchy()->GetDefaultParent(InControlElement->GetKey()) == InSpaceElement->GetKey())
+		{
+			return false;
+		}
+		
+		FRigElementParentConstraint& ParentConstraint = InControlElement->ParentConstraints[*ParentConstraintIndexPtr];
+		if(ParentConstraint.DisplayLabel.IsEqual(InDisplayLabel))
+		{
+			return false;
+		}
+		ParentConstraint.DisplayLabel = InDisplayLabel;
+
+		Notify(ERigHierarchyNotification::ControlSettingChanged, InControlElement);
+		return true;
+	}
+
+	return false;
+}
+
 void URigHierarchyController::AddElementToDirty(FRigBaseElement* InParent, FRigBaseElement* InElementToAdd, int32 InHierarchyDistance) const
 {
 	if(InParent == nullptr)
diff --git a/Engine/Plugins/Animation/ControlRig/Source/ControlRig/Private/Rigs/RigHierarchyElements.cpp b/Engine/Plugins/Animation/ControlRig/Source/ControlRig/Private/Rigs/RigHierarchyElements.cpp
--- a/Engine/Plugins/Animation/ControlRig/Source/ControlRig/Private/Rigs/RigHierarchyElements.cpp
+++ b/Engine/Plugins/Animation/ControlRig/Source/ControlRig/Private/Rigs/RigHierarchyElements.cpp
@@ -1187,6 +1187,7 @@ void FRigMultiParentElement::Save(FArchive& Ar, const FRigHierarchySerialization
 			Ar << ParentKey;
 			Ar << ParentConstraints[ParentIndex].InitialWeight;
 			Ar << ParentConstraints[ParentIndex].Weight;
+			Ar << ParentConstraints[ParentIndex].DisplayLabel;
 		}
 	}
 }
@@ -1235,6 +1236,15 @@ void FRigMultiParentElement::Load(FArchive& Ar, const FRigHierarchySerialization
 				Ar << Weight;
 				ParentConstraints[ParentIndex].Weight = FRigElementWeight(Weight);
 			}
+			
+			if (Ar.CustomVer(FControlRigObjectVersion::GUID) < FControlRigObjectVersion::RigHierarchyParentContraintWithLabel)
+			{
+				ParentConstraints[ParentIndex].DisplayLabel = NAME_None;
+			}
+			else
+			{
+				Ar << ParentConstraints[ParentIndex].DisplayLabel;
+			}
 
 			IndexLookup.Add(ParentKey, ParentIndex);
 		}
@@ -1547,10 +1557,21 @@ void FRigControlSettings::Load(FArchive& Ar)
 		}
 	}
 
-	if (Ar.CustomVer(FControlRigObjectVersion::GUID) >= FControlRigObjectVersion::RigHierarchyControlSpaceFavorites)
+	if (Ar.CustomVer(FControlRigObjectVersion::GUID) >= FControlRigObjectVersion::RigHierarchyParentContraintWithLabel)
 	{
 		Ar << Customization.AvailableSpaces;
 	}
+	else if (Ar.CustomVer(FControlRigObjectVersion::GUID) >= FControlRigObjectVersion::RigHierarchyControlSpaceFavorites)
+	{
+		TArray<FRigElementKey> AvailableSpaceKeys;
+		Ar << AvailableSpaceKeys;
+		Customization.AvailableSpaces.Reset();
+		Customization.AvailableSpaces.Reserve(AvailableSpaceKeys.Num());
+		for(const FRigElementKey& AvailableSpaceKey : AvailableSpaceKeys)
+		{
+			Customization.AvailableSpaces.Emplace(AvailableSpaceKey, NAME_None);
+		}
+	}
 	else
 	{
 		Customization.AvailableSpaces.Reset();
diff --git a/Engine/Plugins/Animation/ControlRig/Source/ControlRig/Private/Units/Execution/RigUnit_DynamicHierarchy.cpp b/Engine/Plugins/Animation/ControlRig/Source/ControlRig/Private/Units/Execution/RigUnit_DynamicHierarchy.cpp
--- a/Engine/Plugins/Animation/ControlRig/Source/ControlRig/Private/Units/Execution/RigUnit_DynamicHierarchy.cpp
+++ b/Engine/Plugins/Animation/ControlRig/Source/ControlRig/Private/Units/Execution/RigUnit_DynamicHierarchy.cpp
@@ -70,7 +70,7 @@ FRigUnit_AddParent_Execute()
 	FRigHierarchyEnableControllerBracket EnableController(ExecuteContext.Hierarchy, true);
 	if(URigHierarchyController* Controller = ExecuteContext.Hierarchy->GetController(true))
 	{
-		Controller->AddParent(ChildElement, ParentElement, 0.f, true, false);
+		Controller->AddParent(ChildElement, ParentElement, 0.f, true, false, DisplayLabel);
 	}
 }
 
@@ -102,6 +102,37 @@ FRigUnit_SetDefaultParent_Execute()
 	}
 }
 
+FRigUnit_AddAvailableSpaces_Execute()
+{
+	if(!FRigUnit_DynamicHierarchyBase::IsValidToRunInContext(ExecuteContext, true))
+	{
+		return;
+	}
+	
+	FRigControlElement* ControlElement = ExecuteContext.Hierarchy->Find<FRigControlElement>(Control);
+	if(ControlElement == nullptr)
+	{
+		UE_CONTROLRIG_RIGUNIT_REPORT_WARNING(TEXT("Control item %s does not exist."), *Control.ToString())
+		return;
+	}
+
+	FRigHierarchyEnableControllerBracket EnableController(ExecuteContext.Hierarchy, true);
+	if(URigHierarchyController* Controller = ExecuteContext.Hierarchy->GetController(true))
+	{
+		for(const FRigElementKeyWithLabel& Space : Spaces)
+		{
+			if(const FRigTransformElement* TransformElement = ExecuteContext.Hierarchy->Find<FRigTransformElement>(Space.Key))
+			{
+				(void)Controller->AddAvailableSpace(ControlElement, TransformElement, Space.Label);
+			}
+			else
+			{
+				UE_CONTROLRIG_RIGUNIT_REPORT_WARNING(TEXT("Space %s does not exist."), *Space.Key.ToString())
+			}
+		}
+	}
+}
+
 FRigUnit_SetChannelHosts_Execute()
 {
 	if(!FRigUnit_DynamicHierarchyBase::IsValidToRunInContext(ExecuteContext, true))
diff --git a/Engine/Plugins/Animation/ControlRig/Source/ControlRig/Public/ControlRigObjectVersion.h b/Engine/Plugins/Animation/ControlRig/Source/ControlRig/Public/ControlRigObjectVersion.h
--- a/Engine/Plugins/Animation/ControlRig/Source/ControlRig/Public/ControlRigObjectVersion.h
+++ b/Engine/Plugins/Animation/ControlRig/Source/ControlRig/Public/ControlRigObjectVersion.h
@@ -119,6 +119,9 @@ struct CONTROLRIG_API FControlRigObjectVersion
 		// Connectors to support arrays
 		RigHierarchyArrayConnectors,
 
+		// Parent constraints offering a display label
+		RigHierarchyParentContraintWithLabel,
+
 		// -----<new versions can be added above this line>-------------------------------------------------
 		VersionPlusOne,
 		LatestVersion = VersionPlusOne - 1
diff --git a/Engine/Plugins/Animation/ControlRig/Source/ControlRig/Public/Rigs/RigHierarchy.h b/Engine/Plugins/Animation/ControlRig/Source/ControlRig/Public/Rigs/RigHierarchy.h
--- a/Engine/Plugins/Animation/ControlRig/Source/ControlRig/Public/Rigs/RigHierarchy.h
+++ b/Engine/Plugins/Animation/ControlRig/Source/ControlRig/Public/Rigs/RigHierarchy.h
@@ -3355,6 +3355,14 @@ public:
 	 */
 	FRigBaseElement* GetActiveParent(const FRigBaseElement* InElement) const;
 
+	/**
+	 * Returns the display label to use for the space given a parent key
+	 * @param InChildKey The key of the child element to look up the parent for
+	 * @param InParentKey The key of the parent to look up the display label for
+	 * @return The display label to use for the user interface for a given parent space
+	 */
+	FName GetDisplayLabelForParent(const FRigElementKey& InChildKey, const FRigElementKey& InParentKey) const;
+
 	/**
 	 * Sets the weight of a parent below a multi parent element
 	 * @param InChild The key of the multi parented element
@@ -3522,6 +3530,8 @@ public:
 	
 	static FRigElementKey GetDefaultParentKey();
 	static FRigElementKey GetWorldSpaceReferenceKey();
+	static inline const FLazyName DefaultParentKeyLabel = TEXT("Parent");
+	static inline const FLazyName WorldSpaceKeyLabel = TEXT("World");
 
 	/**
 	 * Returns true if an element is parented to another element
@@ -5259,6 +5269,7 @@ private:
 	friend struct FRigUnit_AddParents;
 	friend struct FRigUnit_SetDefaultParent;
 	friend struct FRigUnit_SetChannelHosts;
+	friend struct FRigUnit_AddAvailableSpaces;
 
 private:
 	TGuardValue<bool> GuardIsControllerAvailable;
diff --git a/Engine/Plugins/Animation/ControlRig/Source/ControlRig/Public/Rigs/RigHierarchyController.h b/Engine/Plugins/Animation/ControlRig/Source/ControlRig/Public/Rigs/RigHierarchyController.h
--- a/Engine/Plugins/Animation/ControlRig/Source/ControlRig/Public/Rigs/RigHierarchyController.h
+++ b/Engine/Plugins/Animation/ControlRig/Source/ControlRig/Public/Rigs/RigHierarchyController.h
@@ -586,11 +586,12 @@ public:
 	 * @param InParent The key of the new parent to add
 	 * @param InWeight The initial weight to give to the parent
 	 * @param bMaintainGlobalTransform If set to true the child will stay in the same place spatially, otherwise it will maintain it's local transform (and potential move).
+	 * @param InDisplayLabel The optional display label to use for the parent constraint / space.
 	 * @param bSetupUndo If set to true the stack will record the change for undo / redo
 	 * @return Returns true if successful.
 	 */
 	UFUNCTION(BlueprintCallable, Category = URigHierarchyController)
-	bool AddParent(FRigElementKey InChild, FRigElementKey InParent, float InWeight = 0.f, bool bMaintainGlobalTransform = true, bool bSetupUndo = false);
+	bool AddParent(FRigElementKey InChild, FRigElementKey InParent, float InWeight = 0.f, bool bMaintainGlobalTransform = true, const FName& InDisplayLabel = NAME_None, bool bSetupUndo = false);
 
 	/**
 	* Adds a new parent to an element. For elements that allow only one parent the parent will be replaced (Same as ::SetParent).
@@ -599,9 +600,10 @@ public:
 	* @param InWeight The initial weight to give to the parent
 	* @param bMaintainGlobalTransform If set to true the child will stay in the same place spatially, otherwise it will maintain it's local transform (and potential move).
 	* @param bRemoveAllParents If set to true all parents of the child will be removed first.
+	* @param InDisplayLabel The optional display label to use for the parent constraint / space.
 	* @return Returns true if successful.
 	*/
-	bool AddParent(FRigBaseElement* InChild, FRigBaseElement* InParent, float InWeight = 0.f, bool bMaintainGlobalTransform = true, bool bRemoveAllParents = false);
+	bool AddParent(FRigBaseElement* InChild, FRigBaseElement* InParent, float InWeight = 0.f, bool bMaintainGlobalTransform = true, bool bRemoveAllParents = false, const FName& InDisplayLabel = NAME_None);
 
 	/**
 	 * Removes an existing parent from an element in the hierarchy. For elements that allow only one parent the element will be unparented (same as ::RemoveAllParents)
@@ -642,12 +644,22 @@ public:
 	 * Adds a new available space to the given control
 	 * @param InControl The control to add the available space for
 	 * @param InSpace The space to add to the available spaces list
+	 * @param InDisplayLabel The optional display label to use for this space
 	 * @param bSetupUndo If set to true the stack will record the change for undo / redo
 	 * @param bPrintPythonCommand If set to true a python command equivalent to this call will be printed out
 	 * @return Returns true if successful.
 	 */
 	UFUNCTION(BlueprintCallable, Category = URigHierarchyController)
-	bool AddAvailableSpace(FRigElementKey InControl, FRigElementKey InSpace, bool bSetupUndo = false, bool bPrintPythonCommand = false);
+	bool AddAvailableSpace(FRigElementKey InControl, FRigElementKey InSpace, const FName& InDisplayLabel = NAME_None, bool bSetupUndo = false, bool bPrintPythonCommand = false);
+
+	/**
+	 * Adds a new available space to the given control
+	 * @param InControlElement The control element to add the available space for
+	 * @param InSpaceElement The space element to add to the available spaces list
+	 * @param InDisplayLabel The optional display label to use for the space
+	 * @return Returns true if successful.
+	 */
+	bool AddAvailableSpace(FRigControlElement* InControlElement, const FRigTransformElement* InSpaceElement, const FName& InDisplayLabel = NAME_None);
 
 	/**
 	 * Removes an available space from the given control
@@ -672,6 +684,18 @@ public:
 	UFUNCTION(BlueprintCallable, Category = URigHierarchyController)
 	bool SetAvailableSpaceIndex(FRigElementKey InControl, FRigElementKey InSpace, int32 InIndex, bool bSetupUndo = false, bool bPrintPythonCommand = false);
 
+	/**
+	 * Updates the label on an available space
+	 * @param InControl The control to reorder the host for
+	 * @param InSpace The space to set the new index for
+	 * @param InDisplayLabel The new label of the available space
+	 * @param bSetupUndo If set to true the stack will record the change for undo / redo
+	 * @param bPrintPythonCommand If set to true a python command equivalent to this call will be printed out
+	 * @return Returns true if successful.
+	 */
+	UFUNCTION(BlueprintCallable, Category = URigHierarchyController)
+	bool SetAvailableSpaceLabel(FRigElementKey InControl, FRigElementKey InSpace, FName InDisplayLabel, bool bSetupUndo = false, bool bPrintPythonCommand = false);
+
 	/**
 	 * Adds a new channel host to the animation channel
 	 * @note This is just an overload of AddAvailableSpace for readability
@@ -891,14 +915,6 @@ private:
  	 */
 	bool SetParent(FRigBaseElement* InChild, FRigBaseElement* InParent, bool bMaintainGlobalTransform = true);
 
-	/**
-	 * Adds a new available space to the given control
-	 * @param InControlElement The control element to add the available space for
-	 * @param InSpaceElement The space element to add to the available spaces list
-	 * @return Returns true if successful.
-	 */
-	bool AddAvailableSpace(FRigControlElement* InControlElement, const FRigTransformElement* InSpaceElement);
-
 	/**
 	 * Removes an available space from the given control
 	 * @param InControlElement The control element to remove the available space from
@@ -916,6 +932,15 @@ private:
 	 */
 	bool SetAvailableSpaceIndex(FRigControlElement* InControlElement, const FRigTransformElement* InSpaceElement, int32 InIndex);
 
+	/**
+	 * Updates the label on an available space
+	 * @param InControlElement The control to reorder the host for
+	 * @param InSpaceElement The space to set the new index for
+	 * @param InDisplayLabel The new label of the available space
+	 * @return Returns true if successful.
+	 */
+	bool SetAvailableSpaceLabel(FRigControlElement* InControlElement, const FRigTransformElement* InSpaceElement, const FName& InDisplayLabel);
+
 	/**
 	 * Adds a new element to the dirty list of the given parent.
 	 * This function is recursive and will affect all parents in the tree.
diff --git a/Engine/Plugins/Animation/ControlRig/Source/ControlRig/Public/Rigs/RigHierarchyElements.h b/Engine/Plugins/Animation/ControlRig/Source/ControlRig/Public/Rigs/RigHierarchyElements.h
--- a/Engine/Plugins/Animation/ControlRig/Source/ControlRig/Public/Rigs/RigHierarchyElements.h
+++ b/Engine/Plugins/Animation/ControlRig/Source/ControlRig/Public/Rigs/RigHierarchyElements.h
@@ -1226,11 +1226,13 @@ struct CONTROLRIG_API FRigElementParentConstraint
 	FRigTransformElement* ParentElement;
 	FRigElementWeight Weight;
 	FRigElementWeight InitialWeight;
+	FName DisplayLabel;
 	mutable FTransform Cache;
 	mutable bool bCacheIsDirty;
 		
 	FRigElementParentConstraint()
 		: ParentElement(nullptr)
+		, DisplayLabel(NAME_None)
 	{
 		Cache = FTransform::Identity;
 		bCacheIsDirty = true;
@@ -1361,13 +1363,64 @@ private:
 	friend struct FRigBaseElement;
 };
 
+USTRUCT(BlueprintType)
+struct FRigElementKeyWithLabel
+{
+	GENERATED_BODY()
+
+	FRigElementKeyWithLabel()
+		: Key()
+		, Label(NAME_None)
+	{
+	}
+	
+	explicit FRigElementKeyWithLabel(const FRigElementKey& InKey, const FName& InLabel = NAME_None)
+		: Key(InKey)
+		, Label(InLabel)
+	{
+	}
+
+	UPROPERTY(BlueprintReadOnly, EditAnywhere, Category = Customization, meta=(DisplayName="Item"))
+	FRigElementKey Key;
+
+	UPROPERTY(BlueprintReadOnly, EditAnywhere, Category = Customization)
+	FName Label;
+
+	const FName& GetLabel() const
+	{
+		return Label.IsNone() ? Key.Name : Label;
+	}
+
+	friend uint32 GetTypeHash(const FRigElementKeyWithLabel& InKeyWithLabel)
+	{
+		return GetTypeHash(InKeyWithLabel.Key);
+	}
+
+	bool operator ==(const FRigElementKeyWithLabel& InOther) const
+	{
+		return InOther.Key == Key && InOther.Label == Label;
+	}
+
+	bool operator ==(const FRigElementKey& InOther) const
+	{
+		return InOther == Key;
+	}
+
+	friend FArchive& operator<<(FArchive& Ar, FRigElementKeyWithLabel& KeyWithLabel)
+	{
+		Ar << KeyWithLabel.Key;
+		Ar << KeyWithLabel.Label;
+		return Ar;
+	}
+};
+
 USTRUCT(BlueprintType)
 struct FRigControlElementCustomization
 {
 	GENERATED_BODY()
 
 	UPROPERTY(BlueprintReadOnly, EditAnywhere, Category = Customization)
-	TArray<FRigElementKey> AvailableSpaces;
+	TArray<FRigElementKeyWithLabel> AvailableSpaces;
 
 	UPROPERTY(BlueprintReadOnly, EditAnywhere, Category = Customization)
 	TArray<FRigElementKey> RemovedSpaces;
@@ -2046,7 +2099,7 @@ struct CONTROLRIG_API FRigHierarchyCopyPasteContentPerElement
 	FString Content;
 
 	UPROPERTY()
-	TArray<FRigElementKey> Parents;
+	TArray<FRigElementKeyWithLabel> Parents;
 
 	UPROPERTY()
 	TArray<FRigElementWeight> ParentWeights;
diff --git a/Engine/Plugins/Animation/ControlRig/Source/ControlRig/Public/Units/Execution/RigUnit_DynamicHierarchy.h b/Engine/Plugins/Animation/ControlRig/Source/ControlRig/Public/Units/Execution/RigUnit_DynamicHierarchy.h
--- a/Engine/Plugins/Animation/ControlRig/Source/ControlRig/Public/Units/Execution/RigUnit_DynamicHierarchy.h
+++ b/Engine/Plugins/Animation/ControlRig/Source/ControlRig/Public/Units/Execution/RigUnit_DynamicHierarchy.h
@@ -26,8 +26,10 @@ struct CONTROLRIG_API FRigUnit_DynamicHierarchyBaseMutable : public FRigUnitMuta
 /**
 * Adds a new parent to an element. The weight for the new parent will be 0.0.
 * You can use the SetParentWeights node to change the parent weights later.
+*
+* If you just want to add a space to a control you can use the 'Add Spaces' node.
 */
-USTRUCT(meta=(DisplayName="Add Parent", Keywords="Children,Parent,Constraint,Space", Varying))
+USTRUCT(meta=(DisplayName="Create Parent Relationship", Keywords="Children,Parent,Constraint,Space,AddParent", Varying))
 struct CONTROLRIG_API FRigUnit_AddParent : public FRigUnit_DynamicHierarchyBaseMutable
 {
 	GENERATED_BODY()
@@ -35,6 +37,7 @@ struct CONTROLRIG_API FRigUnit_AddParent : public FRigUnit_DynamicHierarchyBaseM
 	FRigUnit_AddParent()
 	{
 		Child = Parent = FRigElementKey(NAME_None, ERigElementType::Control);
+		DisplayLabel = NAME_None;
 	}
 
 	RIGVM_METHOD()
@@ -51,6 +54,12 @@ struct CONTROLRIG_API FRigUnit_AddParent : public FRigUnit_DynamicHierarchyBaseM
 	 */
 	UPROPERTY(meta = (Input, ExpandByDefault))
 	FRigElementKey Parent;
+
+	/*
+	 * The optional display label for the parent constraint / space
+	 */
+	UPROPERTY(meta = (Input, ExpandByDefault))
+	FName DisplayLabel;
 };
 
 /**
@@ -82,6 +91,36 @@ struct CONTROLRIG_API FRigUnit_SetDefaultParent : public FRigUnit_DynamicHierarc
 	FRigElementKey Parent;
 };
 
+/**
+* Adds available spaces for a given control. This is used for the space switching widget.
+*/
+USTRUCT(meta=(DisplayName="Add Spaces", Keywords="Children,AddParent,Constraint,Space", Varying))
+struct CONTROLRIG_API FRigUnit_AddAvailableSpaces : public FRigUnit_DynamicHierarchyBaseMutable
+{
+	GENERATED_BODY()
+
+	FRigUnit_AddAvailableSpaces()
+	{
+		Control = FRigElementKey(NAME_None, ERigElementType::Control);
+		Spaces = { FRigElementKeyWithLabel(FRigElementKey(NAME_None, ERigElementType::Control), NAME_None)};
+	}
+
+	RIGVM_METHOD()
+	virtual void Execute() override;
+
+	/*
+	 * The available spaces to add for a given control
+	 */
+	UPROPERTY(meta = (Input, ExpandByDefault))
+	FRigElementKey Control;
+
+	/*
+	 * The spaces to add
+	 */
+	UPROPERTY(meta = (Input, ExpandByDefault))
+	TArray<FRigElementKeyWithLabel> Spaces;
+};
+
 /**
  * Allows an animation channel to be hosted by multiple controls
  */
diff --git a/Engine/Plugins/Animation/ControlRig/Source/ControlRigEditor/Private/ControlRigElementDetails.cpp b/Engine/Plugins/Animation/ControlRig/Source/ControlRigEditor/Private/ControlRigElementDetails.cpp
--- a/Engine/Plugins/Animation/ControlRig/Source/ControlRigEditor/Private/ControlRigElementDetails.cpp
+++ b/Engine/Plugins/Animation/ControlRig/Source/ControlRigEditor/Private/ControlRigElementDetails.cpp
@@ -4247,8 +4247,14 @@ void FRigControlElementDetails::CustomizeAvailableSpaces(IDetailLayoutBuilder& D
 	];
 	Category.HeaderContent(HeaderContentWidget);
 
-	TArray<FRigElementKey> AvailableSpaces = { Hierarchy->GetDefaultParent(ControlElement->GetKey()) };
-	for(const FRigElementKey& AvailableSpace : ControlElement->Settings.Customization.AvailableSpaces)
+	TArray<FRigElementKeyWithLabel> AvailableSpaces;
+	const FRigElementKey DefaultParent = Hierarchy->GetDefaultParent(ControlElement->GetKey());
+	if(DefaultParent.IsValid())
+	{
+		const FName SpaceLabel = Hierarchy->GetDisplayLabelForParent(ControlElement->GetKey(), DefaultParent);
+		AvailableSpaces.Emplace(DefaultParent, SpaceLabel);
+	}
+	for(const FRigElementKeyWithLabel& AvailableSpace : ControlElement->Settings.Customization.AvailableSpaces)
 	{
 		AvailableSpaces.AddUnique(AvailableSpace);
 	}
@@ -4261,14 +4267,14 @@ void FRigControlElementDetails::CustomizeAvailableSpaces(IDetailLayoutBuilder& D
 	for(int32 Index = 0; Index < AvailableSpaces.Num(); Index++)
 	{
 		const FRigElementKey ControlKey = ControlElement->GetKey();
-		const FRigElementKey AvailableSpace = AvailableSpaces[Index];
+		const FRigElementKeyWithLabel AvailableSpace = AvailableSpaces[Index];
 		const bool bIsParentSpace = Index == 0;
-		const TPair<const FSlateBrush*, FSlateColor> BrushAndColor = SRigHierarchyItem::GetBrushForElementType(Hierarchy, AvailableSpace);
+		const TPair<const FSlateBrush*, FSlateColor> BrushAndColor = SRigHierarchyItem::GetBrushForElementType(Hierarchy, AvailableSpace.Key);
 
 		TSharedPtr<SButton> SelectSpaceButton, RemoveSpaceButton, MoveSpaceUpButton, MoveSpaceDownButton;
 		TSharedPtr<SImage> SelectSpaceImage, RemoveSpaceImage, MoveSpaceUpImage, MoveSpaceDownImage;
 
-		FDetailWidgetRow& WidgetRow = Category.AddCustomRow(FText::FromString(AvailableSpace.ToString()))
+		FDetailWidgetRow& WidgetRow = Category.AddCustomRow(FText::FromString(AvailableSpace.Key.ToString()))
 		.NameContent()
 		.MinDesiredWidth(200.f)
 		.MaxDesiredWidth(800.f)
@@ -4293,11 +4299,41 @@ void FRigControlElementDetails::CustomizeAvailableSpaces(IDetailLayoutBuilder& D
 			.VAlign(VAlign_Center)
 			.Padding(FMargin(0.f, 0.f, 8.f, 0.f))
 			[
-				SNew(STextBlock)
+				SNew(SEditableText)
 				.Font(IDetailLayoutBuilder::GetDetailFont())
+				.IsReadOnly((Index == 0) && bIsEnabled)
 				.Text_Lambda([this, AvailableSpace]() -> FText
 				{
-					return GetDisplayNameForElement(AvailableSpace);
+					const FName Label = PerElementInfos[0].GetHierarchy()->GetDisplayLabelForParent(
+						PerElementInfos[0].Element.GetKey(),
+						AvailableSpace.Key);
+					if(Label.IsNone())
+					{
+						return GetDisplayNameForElement(AvailableSpace.Key);
+					}
+					return FText::FromName(Label);
+				})
+				.OnTextCommitted_Lambda([this, AvailableSpace](const FText& InText, ETextCommit::Type InCommitType)
+				{
+					if(InCommitType == ETextCommit::OnCleared)
+					{
+						return;
+					}
+					UControlRigBlueprint* Blueprint = PerElementInfos[0].GetBlueprint();
+					URigHierarchyController* Controller = Blueprint->GetHierarchyController();
+					(void)Controller->SetAvailableSpaceLabel(PerElementInfos[0].Element.GetKey(), AvailableSpace.Key, *InText.ToString(), true);
+				})
+				.ToolTipText_Lambda([this, Index, AvailableSpace]()
+				{
+					if(Index == 0)
+					{
+						return FText::Format(
+							LOCTEXT("AvailableSpaceToolTipDefaultParentFormat", "{0}\n\nThis is the default parent - the label cannot be edited."),
+							GetDisplayNameForElement(AvailableSpace.Key));
+					}
+					return FText::Format(
+						LOCTEXT("AvailableSpaceToolTipFormat", "{0}\n\nDouble-click here to edit the label of this space."),
+						GetDisplayNameForElement(AvailableSpace.Key));
 				})
 			]
 
@@ -4311,7 +4347,7 @@ void FRigControlElementDetails::CustomizeAvailableSpaces(IDetailLayoutBuilder& D
 				.ButtonStyle( FAppStyle::Get(), "NoBorder" )
 				.OnClicked_Lambda([this, AvailableSpace]() -> FReply
 				{
-					return OnSelectElementClicked(AvailableSpace);
+					return OnSelectElementClicked(AvailableSpace.Key);
 				})
 				.ContentPadding(0)
 				.ToolTipText(LOCTEXT("SelectElementInHierarchy", "Select Element in hierarchy"))
@@ -4340,7 +4376,7 @@ void FRigControlElementDetails::CustomizeAvailableSpaces(IDetailLayoutBuilder& D
 					{
 						FScopedTransaction Transaction(LOCTEXT("MoveAvailableSpaceUp", "Move Available Space Up"));
 						HierarchyToChange->Modify();
-						Controller->SetAvailableSpaceIndex(ControlKey, AvailableSpace, Index - 2);
+						Controller->SetAvailableSpaceIndex(ControlKey, AvailableSpace.Key, Index - 2);
 						PropertyUtilities->ForceRefresh();
 					}
 					return FReply::Handled();
@@ -4366,7 +4402,7 @@ void FRigControlElementDetails::CustomizeAvailableSpaces(IDetailLayoutBuilder& D
 					{
 						FScopedTransaction Transaction(LOCTEXT("MoveAvailableSpaceDown", "Move Available Space Down"));
 						HierarchyToChange->Modify();
-						Controller->SetAvailableSpaceIndex(ControlKey, AvailableSpace, Index);
+						Controller->SetAvailableSpaceIndex(ControlKey, AvailableSpace.Key, Index);
 						PropertyUtilities->ForceRefresh();
 					}
 					return FReply::Handled();
@@ -4395,11 +4431,11 @@ void FRigControlElementDetails::CustomizeAvailableSpaces(IDetailLayoutBuilder& D
 						HierarchyToChange->Modify();
 						if(bIsAnimationChannel)
 						{
-							Controller->RemoveChannelHost(ControlKey, AvailableSpace);
+							Controller->RemoveChannelHost(ControlKey, AvailableSpace.Key);
 						}
 						else
 						{
-							Controller->RemoveAvailableSpace(ControlKey, AvailableSpace);
+							Controller->RemoveAvailableSpace(ControlKey, AvailableSpace.Key);
 						}
 						PropertyUtilities->ForceRefresh();
 					}
@@ -4429,7 +4465,7 @@ void FRigControlElementDetails::CustomizeAvailableSpaces(IDetailLayoutBuilder& D
 						{
 							FScopedTransaction Transaction(LOCTEXT("MoveAvailableSpaceUp", "Move Available Space Up"));
 							HierarchyToChange->Modify();
-							Controller->SetAvailableSpaceIndex(ControlKey, AvailableSpace, Index - 2);
+							Controller->SetAvailableSpaceIndex(ControlKey, AvailableSpace.Key, Index - 2);
 							PropertyUtilities->ForceRefresh();
 						}
 					}),
@@ -4446,7 +4482,7 @@ void FRigControlElementDetails::CustomizeAvailableSpaces(IDetailLayoutBuilder& D
 						{
 							FScopedTransaction Transaction(LOCTEXT("MoveAvailableSpaceDown", "Move Available Space Down"));
 							HierarchyToChange->Modify();
-							Controller->SetAvailableSpaceIndex(ControlKey, AvailableSpace, Index);
+							Controller->SetAvailableSpaceIndex(ControlKey, AvailableSpace.Key, Index);
 							PropertyUtilities->ForceRefresh();
 						}
 					}),
@@ -4466,11 +4502,11 @@ void FRigControlElementDetails::CustomizeAvailableSpaces(IDetailLayoutBuilder& D
 						HierarchyToChange->Modify();
 						if(bIsAnimationChannel)
 						{
-							Controller->RemoveChannelHost(ControlKey, AvailableSpace);
+							Controller->RemoveChannelHost(ControlKey, AvailableSpace.Key);
 						}
 						else
 						{
-							Controller->RemoveAvailableSpace(ControlKey, AvailableSpace);
+							Controller->RemoveAvailableSpace(ControlKey, AvailableSpace.Key);
 						}
 						PropertyUtilities->ForceRefresh();
 					}
diff --git a/Engine/Plugins/Animation/ControlRig/Source/ControlRigEditor/Private/EditMode/ControlRigEditMode.cpp b/Engine/Plugins/Animation/ControlRig/Source/ControlRigEditor/Private/EditMode/ControlRigEditMode.cpp
--- a/Engine/Plugins/Animation/ControlRig/Source/ControlRigEditor/Private/EditMode/ControlRigEditMode.cpp
+++ b/Engine/Plugins/Animation/ControlRig/Source/ControlRigEditor/Private/EditMode/ControlRigEditMode.cpp
@@ -3274,7 +3274,7 @@ void FControlRigEditMode::OpenSpacePickerWidget()
 		}
 		
 	})
-	.OnSpaceListChanged_Lambda([this, SelectedControls, RuntimeRig](URigHierarchy* InHierarchy, const FRigElementKey& InControlKey, const TArray<FRigElementKey>& InSpaceList)
+	.OnSpaceListChanged_Lambda([this, SelectedControls, RuntimeRig](URigHierarchy* InHierarchy, const FRigElementKey& InControlKey, const TArray<FRigElementKeyWithLabel>& InSpaceList)
 	{
 		check(SelectedControls.Contains(InControlKey));
 
@@ -3321,11 +3321,11 @@ void FControlRigEditMode::OpenSpacePickerWidget()
 				ControlCustomization.RemovedSpaces.Reset();
 
 				// remember  the elements which are in the asset's available list but removed by the user
-				for(const FRigElementKey& AvailableSpace : ControlElement->Settings.Customization.AvailableSpaces)
+				for(const FRigElementKeyWithLabel& AvailableSpace : ControlElement->Settings.Customization.AvailableSpaces)
 				{
-					if(!ControlCustomization.AvailableSpaces.Contains(AvailableSpace))
+					if(ControlCustomization.AvailableSpaces.FindByKey(AvailableSpace.Key) == nullptr)
 					{
-						ControlCustomization.RemovedSpaces.Add(AvailableSpace);
+						ControlCustomization.RemovedSpaces.Add(AvailableSpace.Key);
 					}
 				}
 
@@ -4655,7 +4655,7 @@ void FControlRigEditMode::OnHierarchyModified(ERigHierarchyNotification InNotif,
 														return true;
 													}
 												}
-												if(AnimationChannelControl->Settings.Customization.AvailableSpaces.Contains(ControlElement->GetKey()))
+												if(AnimationChannelControl->Settings.Customization.AvailableSpaces.FindByKey(ControlElement->GetKey()))
 												{
 													ControlProxy->SelectProxy(ControlRig, AnimationChannelControl, bSelected);
 													return true;
diff --git a/Engine/Plugins/Animation/ControlRig/Source/ControlRigEditor/Private/EditMode/SControlRigDetails.cpp b/Engine/Plugins/Animation/ControlRig/Source/ControlRigEditor/Private/EditMode/SControlRigDetails.cpp
--- a/Engine/Plugins/Animation/ControlRig/Source/ControlRigEditor/Private/EditMode/SControlRigDetails.cpp
+++ b/Engine/Plugins/Animation/ControlRig/Source/ControlRigEditor/Private/EditMode/SControlRigDetails.cpp
@@ -163,7 +163,7 @@ static TArray<UControlRigControlsProxy*> GetParentProxies(UControlRigControlsPro
 					}
 					if(const FRigControlElement* ChildControlElement = Cast<FRigControlElement>(ChildProxy->OwnerControlElement.GetElement()))
 					{
-						if(ChildControlElement->Settings.Customization.AvailableSpaces.Contains(OwnerControlElement->GetKey()))
+						if(ChildControlElement->Settings.Customization.AvailableSpaces.FindByKey(OwnerControlElement->GetKey()))
 						{
 							ParentProxies.AddUnique(Proxy);
 						}
diff --git a/Engine/Plugins/Animation/ControlRig/Source/ControlRigEditor/Private/EditMode/SControlRigEditModeTools.cpp b/Engine/Plugins/Animation/ControlRig/Source/ControlRigEditor/Private/EditMode/SControlRigEditModeTools.cpp
--- a/Engine/Plugins/Animation/ControlRig/Source/ControlRigEditor/Private/EditMode/SControlRigEditModeTools.cpp
+++ b/Engine/Plugins/Animation/ControlRig/Source/ControlRigEditor/Private/EditMode/SControlRigEditModeTools.cpp
@@ -960,7 +960,7 @@ void SControlRigEditModeTools::HandleActiveSpaceChanged(URigHierarchy* InHierarc
 }
 
 void SControlRigEditModeTools::HandleSpaceListChanged(URigHierarchy* InHierarchy, const FRigElementKey& InControlKey,
-	const TArray<FRigElementKey>& InSpaceList)
+	const TArray<FRigElementKeyWithLabel>& InSpaceList)
 {
 	FScopedTransaction Transaction(LOCTEXT("ChangeControlRigSpace", "Change Control Rig Space"));
 
@@ -977,11 +977,11 @@ void SControlRigEditModeTools::HandleSpaceListChanged(URigHierarchy* InHierarchy
 				ControlCustomization.RemovedSpaces.Reset();
 
 				// remember  the elements which are in the asset's available list but removed by the user
-				for (const FRigElementKey& AvailableSpace : ControlElement->Settings.Customization.AvailableSpaces)
+				for (const FRigElementKeyWithLabel& AvailableSpace : ControlElement->Settings.Customization.AvailableSpaces)
 				{
-					if (!ControlCustomization.AvailableSpaces.Contains(AvailableSpace))
+					if (ControlCustomization.AvailableSpaces.FindByKey(AvailableSpace.Key) == nullptr)
 					{
-						ControlCustomization.RemovedSpaces.Add(AvailableSpace);
+						ControlCustomization.RemovedSpaces.Add(AvailableSpace.Key);
 					}
 				}
 
diff --git a/Engine/Plugins/Animation/ControlRig/Source/ControlRigEditor/Private/EditMode/SControlRigEditModeTools.h b/Engine/Plugins/Animation/ControlRig/Source/ControlRigEditor/Private/EditMode/SControlRigEditModeTools.h
--- a/Engine/Plugins/Animation/ControlRig/Source/ControlRigEditor/Private/EditMode/SControlRigEditModeTools.h
+++ b/Engine/Plugins/Animation/ControlRig/Source/ControlRigEditor/Private/EditMode/SControlRigEditModeTools.h
@@ -111,7 +111,7 @@ private:
 
 	const FRigControlElementCustomization* HandleGetControlElementCustomization(URigHierarchy* InHierarchy, const FRigElementKey& InControlKey);
 	void HandleActiveSpaceChanged(URigHierarchy* InHierarchy, const FRigElementKey& InControlKey, const FRigElementKey& InSpaceKey);
-	void HandleSpaceListChanged(URigHierarchy* InHierarchy, const FRigElementKey& InControlKey, const TArray<FRigElementKey>& InSpaceList);
+	void HandleSpaceListChanged(URigHierarchy* InHierarchy, const FRigElementKey& InControlKey, const TArray<FRigElementKeyWithLabel>& InSpaceList);
 	FReply HandleAddSpaceClicked();
 	EVisibility GetAddSpaceButtonVisibility() const;
 	bool IsSpaceSwitchingRestricted() const;
diff --git a/Engine/Plugins/Animation/ControlRig/Source/ControlRigEditor/Private/EditMode/SControlRigSpacePicker.cpp b/Engine/Plugins/Animation/ControlRig/Source/ControlRigEditor/Private/EditMode/SControlRigSpacePicker.cpp
--- a/Engine/Plugins/Animation/ControlRig/Source/ControlRigEditor/Private/EditMode/SControlRigSpacePicker.cpp
+++ b/Engine/Plugins/Animation/ControlRig/Source/ControlRigEditor/Private/EditMode/SControlRigSpacePicker.cpp
@@ -194,7 +194,7 @@ void SControlRigSpacePicker::HandleActiveSpaceChanged(URigHierarchy* InHierarchy
 }
 
 void SControlRigSpacePicker::HandleSpaceListChanged(URigHierarchy* InHierarchy, const FRigElementKey& InControlKey,
-	const TArray<FRigElementKey>& InSpaceList)
+	const TArray<FRigElementKeyWithLabel>& InSpaceList)
 {
 	if (UControlRig* ControlRig = GetControlRig())
 	{
@@ -205,11 +205,11 @@ void SControlRigSpacePicker::HandleSpaceListChanged(URigHierarchy* InHierarchy,
 			ControlCustomization.RemovedSpaces.Reset();
 
 			// remember  the elements which are in the asset's available list but removed by the user
-			for (const FRigElementKey& AvailableSpace : ControlElement->Settings.Customization.AvailableSpaces)
+			for (const FRigElementKeyWithLabel& AvailableSpace : ControlElement->Settings.Customization.AvailableSpaces)
 			{
-				if (!ControlCustomization.AvailableSpaces.Contains(AvailableSpace))
+				if (ControlCustomization.AvailableSpaces.FindByKey(AvailableSpace.Key) == nullptr)
 				{
-					ControlCustomization.RemovedSpaces.Add(AvailableSpace);
+					ControlCustomization.RemovedSpaces.Add(AvailableSpace.Key);
 				}
 			}
 
diff --git a/Engine/Plugins/Animation/ControlRig/Source/ControlRigEditor/Private/EditMode/SControlRigSpacePicker.h b/Engine/Plugins/Animation/ControlRig/Source/ControlRigEditor/Private/EditMode/SControlRigSpacePicker.h
--- a/Engine/Plugins/Animation/ControlRig/Source/ControlRigEditor/Private/EditMode/SControlRigSpacePicker.h
+++ b/Engine/Plugins/Animation/ControlRig/Source/ControlRigEditor/Private/EditMode/SControlRigSpacePicker.h
@@ -37,7 +37,7 @@ private:
 
 	const FRigControlElementCustomization* HandleGetControlElementCustomization(URigHierarchy* InHierarchy, const FRigElementKey& InControlKey);
 	void HandleActiveSpaceChanged(URigHierarchy* InHierarchy, const FRigElementKey& InControlKey, const FRigElementKey& InSpaceKey);
-	void HandleSpaceListChanged(URigHierarchy* InHierarchy, const FRigElementKey& InControlKey, const TArray<FRigElementKey>& InSpaceList);
+	void HandleSpaceListChanged(URigHierarchy* InHierarchy, const FRigElementKey& InControlKey, const TArray<FRigElementKeyWithLabel>& InSpaceList);
 	FReply HandleAddSpaceClicked();
 	FReply OnBakeControlsToNewSpaceButtonClicked();
 	FReply OnCompensateKeyClicked();
diff --git a/Engine/Plugins/Animation/ControlRig/Source/ControlRigEditor/Private/Editor/SRigSpacePickerWidget.cpp b/Engine/Plugins/Animation/ControlRig/Source/ControlRigEditor/Private/Editor/SRigSpacePickerWidget.cpp
--- a/Engine/Plugins/Animation/ControlRig/Source/ControlRigEditor/Private/Editor/SRigSpacePickerWidget.cpp
+++ b/Engine/Plugins/Animation/ControlRig/Source/ControlRigEditor/Private/Editor/SRigSpacePickerWidget.cpp
@@ -401,7 +401,7 @@ void SRigSpacePickerWidget::Tick(const FGeometry& AllottedGeometry, const double
 		if (Hierarchy.IsValid())
 		{
 			URigHierarchy* StrongHierarchy = Hierarchy.Get();
-			TArray<FRigElementKey> CurrentAdditionalSpaces;
+			TArray<FRigElementKeyWithLabel> CurrentAdditionalSpaces;
 			for(const FRigElementKey& ControlKey: ControlKeys)
 			{
 				CurrentAdditionalSpaces.Append(GetAdditionalSpacesDelegate.Execute(StrongHierarchy, ControlKey));
@@ -420,19 +420,20 @@ const TArray<FRigElementKey>& SRigSpacePickerWidget::GetActiveSpaces() const
 	return ActiveSpaceKeys;
 }
 
-TArray<FRigElementKey> SRigSpacePickerWidget::GetDefaultSpaces() const
+const TArray<FRigElementKeyWithLabel>& SRigSpacePickerWidget::GetDefaultSpaces() const
 {
-	TArray<FRigElementKey> DefaultSpaces;
-	DefaultSpaces.Add(URigHierarchy::GetDefaultParentKey());
-	DefaultSpaces.Add(URigHierarchy::GetWorldSpaceReferenceKey());
+	static const TArray<FRigElementKeyWithLabel> DefaultSpaces = {
+		FRigElementKeyWithLabel(URigHierarchy::GetDefaultParentKey(), URigHierarchy::DefaultParentKeyLabel),
+		FRigElementKeyWithLabel(URigHierarchy::GetWorldSpaceReferenceKey(), URigHierarchy::WorldSpaceKeyLabel)
+	};
 	return DefaultSpaces;
 }
 
-TArray<FRigElementKey> SRigSpacePickerWidget::GetSpaceList(bool bIncludeDefaultSpaces) const
+TArray<FRigElementKeyWithLabel> SRigSpacePickerWidget::GetSpaceList(bool bIncludeDefaultSpaces) const
 {
 	if(bIncludeDefaultSpaces && bShowDefaultSpaces)
 	{
-		TArray<FRigElementKey> Spaces;
+		TArray<FRigElementKeyWithLabel> Spaces;
 		Spaces.Append(GetDefaultSpaces());
 		Spaces.Append(CurrentSpaceKeys);
 		return Spaces;
@@ -503,6 +504,7 @@ void SRigSpacePickerWidget::AddSpacePickerRow(
 						SNew( STextBlock )
 						.Text( InTitle )
 						.Font( IDetailLayoutBuilder::GetDetailFont() )
+						.ToolTipText(FText::FromName(InKey.Name))
 					]
 
 					+ SHorizontalBox::Slot()
@@ -628,12 +630,12 @@ FReply SRigSpacePickerWidget::HandleSpaceMoveUp(FRigElementKey InKey)
 {
 	if(CurrentSpaceKeys.Num() > 1)
 	{
-		const int32 Index = CurrentSpaceKeys.Find(InKey);
+		const int32 Index = CurrentSpaceKeys.IndexOfByKey(InKey);
 		if(CurrentSpaceKeys.IsValidIndex(Index))
 		{
 			if(Index > 0)
 			{
-				TArray<FRigElementKey> ChangedSpaceKeys = CurrentSpaceKeys;
+				TArray<FRigElementKeyWithLabel> ChangedSpaceKeys = CurrentSpaceKeys;
 				ChangedSpaceKeys.Swap(Index, Index - 1);
 
 				if (Hierarchy.IsValid())
@@ -656,12 +658,12 @@ FReply SRigSpacePickerWidget::HandleSpaceMoveDown(FRigElementKey InKey)
 {
 	if(CurrentSpaceKeys.Num() > 1)
 	{
-		const int32 Index = CurrentSpaceKeys.Find(InKey);
+		const int32 Index = CurrentSpaceKeys.IndexOfByKey(InKey);
 		if(CurrentSpaceKeys.IsValidIndex(Index))
 		{
 			if(Index < CurrentSpaceKeys.Num() - 1)
 			{
-				TArray<FRigElementKey> ChangedSpaceKeys = CurrentSpaceKeys;
+				TArray<FRigElementKeyWithLabel> ChangedSpaceKeys = CurrentSpaceKeys;
 				ChangedSpaceKeys.Swap(Index, Index + 1);
 
 				if (Hierarchy.IsValid())
@@ -682,9 +684,11 @@ FReply SRigSpacePickerWidget::HandleSpaceMoveDown(FRigElementKey InKey)
 
 void SRigSpacePickerWidget::HandleSpaceDelete(FRigElementKey InKey)
 {
-	TArray<FRigElementKey> ChangedSpaceKeys = CurrentSpaceKeys;
-	if(ChangedSpaceKeys.Remove(InKey) > 0)
+	TArray<FRigElementKeyWithLabel> ChangedSpaceKeys = CurrentSpaceKeys;
+	const int32 ExistingSpaceIndex = ChangedSpaceKeys.IndexOfByKey(InKey);
+	if(ExistingSpaceIndex != INDEX_NONE)
 	{
+		ChangedSpaceKeys.RemoveAt(ExistingSpaceIndex);
 		if (Hierarchy.IsValid())
 		{
 			URigHierarchy* StrongHierarchy = Hierarchy.Get();
@@ -746,8 +750,18 @@ FReply SRigSpacePickerWidget::HandleAddElementClicked()
 						}
 					}
 					
-					TArray<FRigElementKey> ChangedSpaceKeys = CurrentSpaceKeys;
-					ChangedSpaceKeys.AddUnique(Key);
+					TArray<FRigElementKeyWithLabel> ChangedSpaceKeys = CurrentSpaceKeys;
+					if(ChangedSpaceKeys.FindByKey(Key) == nullptr)
+					{
+						if(ControlKeys.IsEmpty())
+						{
+							ChangedSpaceKeys.Emplace(Key);
+						}
+						else
+						{
+							ChangedSpaceKeys.Emplace(Key, StrongHierarchy->GetDisplayLabelForParent(ControlKeys[0], Key));
+						}
+					}
 
 					for(const FRigElementKey& ControlKey : ControlKeys)
 					{
@@ -841,7 +855,7 @@ bool SRigSpacePickerWidget::IsSpaceMoveUpEnabled(FRigElementKey InKey) const
 	{
 		return false;
 	}
-	return CurrentSpaceKeys[0] != InKey;
+	return CurrentSpaceKeys[0].Key != InKey;
 }
 
 bool SRigSpacePickerWidget::IsSpaceMoveDownEnabled(FRigElementKey InKey) const
@@ -850,7 +864,7 @@ bool SRigSpacePickerWidget::IsSpaceMoveDownEnabled(FRigElementKey InKey) const
 	{
 		return false;
 	}
-	return CurrentSpaceKeys.Last( )!= InKey;
+	return CurrentSpaceKeys.Last().Key != InKey;
 }
 
 void SRigSpacePickerWidget::OnHierarchyModified(ERigHierarchyNotification InNotif, URigHierarchy* InHierarchy,
@@ -931,16 +945,17 @@ FRigElementKey SRigSpacePickerWidget::GetActiveSpace_Private(URigHierarchy* InHi
 	return URigHierarchy::GetDefaultParentKey();
 }
 
-TArray<FRigElementKey> SRigSpacePickerWidget::GetCurrentParents_Private(URigHierarchy* InHierarchy,
+TArray<FRigElementKeyWithLabel> SRigSpacePickerWidget::GetCurrentParents_Private(URigHierarchy* InHierarchy,
                                                                 const FRigElementKey& InControlKey) const
 {
 	if(!InControlKey.IsValid() || InHierarchy == nullptr)
 	{
-		return TArray<FRigElementKey>();
+		return TArray<FRigElementKeyWithLabel>();
 	}
 
 	check(ControlKeys.Contains(InControlKey));
 	TArray<FRigElementKey> Parents = InHierarchy->GetParents(InControlKey);
+	TArray<FRigElementKeyWithLabel> ParentSpaces;
 	if(Parents.Num() > 0)
 	{
 		if(!IsDefaultSpace(Parents[0]))
@@ -948,7 +963,12 @@ TArray<FRigElementKey> SRigSpacePickerWidget::GetCurrentParents_Private(URigHier
 			Parents[0] = URigHierarchy::GetDefaultParentKey();
 		}
 	}
-	return Parents;
+	ParentSpaces.Reserve(Parents.Num());
+	for(const FRigElementKey& ParentKey : Parents)
+	{
+		ParentSpaces.Emplace(ParentKey, InHierarchy->GetDisplayLabelForParent(InControlKey, ParentKey));
+	}
+	return ParentSpaces;
 }
 
 void SRigSpacePickerWidget::RepopulateItemSpaces()
@@ -964,7 +984,7 @@ void SRigSpacePickerWidget::RepopulateItemSpaces()
 
 	URigHierarchy* StrongHierarchy = Hierarchy.Get();
 	
-	TArray<FRigElementKey> FavoriteKeys, SpacesFromDelegate;
+	TArray<FRigElementKeyWithLabel> FavoriteKeys, SpacesFromDelegate;
 
 	if(bShowFavoriteSpaces)
 	{
@@ -978,13 +998,13 @@ void SRigSpacePickerWidget::RepopulateItemSpaces()
 
 			if(Customization)
 			{
-				for(const FRigElementKey& Key : Customization->AvailableSpaces)
+				for(const FRigElementKeyWithLabel& AvailableSpace : Customization->AvailableSpaces)
 				{
-					if(IsDefaultSpace(Key) || !IsValidKey(Key))
+					if(IsDefaultSpace(AvailableSpace.Key) || !IsValidKey(AvailableSpace.Key))
 					{
 						continue;
 					}
-					FavoriteKeys.AddUnique(Key);
+					FavoriteKeys.AddUnique(AvailableSpace);
 				}
 			}
 			
@@ -993,25 +1013,25 @@ void SRigSpacePickerWidget::RepopulateItemSpaces()
 			{
 				if(Customization != &ControlElement->Settings.Customization)
 				{
-					for(const FRigElementKey& Key : ControlElement->Settings.Customization.AvailableSpaces)
+					for(const FRigElementKeyWithLabel& AvailableSpace : ControlElement->Settings.Customization.AvailableSpaces)
 					{
-						if(IsDefaultSpace(Key) || !IsValidKey(Key))
+						if(IsDefaultSpace(AvailableSpace.Key) || !IsValidKey(AvailableSpace.Key))
 						{
 							continue;
 						}
 
 						if(Customization)
 						{
-							if(Customization->AvailableSpaces.Contains(Key))
+							if(Customization->AvailableSpaces.FindByKey(AvailableSpace.Key) != nullptr)
 							{
 								continue;
 							}
-							if(Customization->RemovedSpaces.Contains(Key))
+							if(Customization->RemovedSpaces.Contains(AvailableSpace.Key))
 							{
 								continue;
 							}
 						}
-						FavoriteKeys.AddUnique(Key);
+						FavoriteKeys.AddUnique(AvailableSpace);
 					}
 				}
 			}
@@ -1027,13 +1047,13 @@ void SRigSpacePickerWidget::RepopulateItemSpaces()
 			AdditionalSpaces.Append(GetAdditionalSpacesDelegate.Execute(StrongHierarchy, ControlKey));
 		}
 		
-		for(const FRigElementKey& Key : AdditionalSpaces)
+		for(const FRigElementKeyWithLabel& AdditionalSpace : AdditionalSpaces)
 		{
-			if(IsDefaultSpace(Key)  || !IsValidKey(Key))
+			if(IsDefaultSpace(AdditionalSpace.Key)  || !IsValidKey(AdditionalSpace.Key))
 			{
 				continue;
 			}
-			SpacesFromDelegate.AddUnique(Key);
+			SpacesFromDelegate.AddUnique(AdditionalSpace);
 		}
 	}
 
@@ -1065,35 +1085,38 @@ void SRigSpacePickerWidget::RepopulateItemSpaces()
 	SpacesFromDelegate.Sort(FKeySortPredicate());
 	*/
 
-	TArray<FRigElementKey> Keys = FavoriteKeys;
-	for(const FRigElementKey& Key : SpacesFromDelegate)
+	TArray<FRigElementKeyWithLabel> CombinedSpaces = FavoriteKeys;
+	for(const FRigElementKeyWithLabel& Space : SpacesFromDelegate)
 	{
-		Keys.AddUnique(Key);
+		if(CombinedSpaces.FindByKey(Space.Key) == nullptr)
+		{
+			CombinedSpaces.AddUnique(Space);
+		}
 	}
 
-	if(Keys == CurrentSpaceKeys)
+	if(CombinedSpaces == CurrentSpaceKeys)
 	{
 		return;
 	}
 
 	ClearListBox(ItemSpacesListBox);
 
-	for(const FRigElementKey& Key : Keys)
+	for(const FRigElementKeyWithLabel& Space : CombinedSpaces)
 	{
-		TPair<const FSlateBrush*, FSlateColor> IconAndColor = SRigHierarchyItem::GetBrushForElementType(StrongHierarchy, Key);
+		TPair<const FSlateBrush*, FSlateColor> IconAndColor = SRigHierarchyItem::GetBrushForElementType(StrongHierarchy, Space.Key);
 		
 		AddSpacePickerRow(
 			ItemSpacesListBox,
 			ESpacePickerType_Item,
-			Key,
+			Space.Key,
 			IconAndColor.Key,
 			IconAndColor.Value,
-			FText::FromName(Key.Name),
-			FOnClicked::CreateSP(this, &SRigSpacePickerWidget::HandleElementSpaceClicked, Key)
+			FText::FromName(Space.GetLabel()),
+			FOnClicked::CreateSP(this, &SRigSpacePickerWidget::HandleElementSpaceClicked, Space.Key)
 		);
 	}
 
-	CurrentSpaceKeys = Keys;
+	CurrentSpaceKeys = CombinedSpaces;
 }
 
 void SRigSpacePickerWidget::ClearListBox(TSharedPtr<SVerticalBox> InListBox)
@@ -1211,7 +1234,7 @@ void SRigSpacePickerBakeWidget::Construct(const FArguments& InArgs)
 				{
 					return &Customization;
 				})
-				.OnSpaceListChanged_Lambda([this](URigHierarchy*, const FRigElementKey&, const TArray<FRigElementKey>& InSpaceList)
+				.OnSpaceListChanged_Lambda([this](URigHierarchy*, const FRigElementKey&, const TArray<FRigElementKeyWithLabel>& InSpaceList)
 				{
 					if(Customization.AvailableSpaces != InSpaceList)
 					{
diff --git a/Engine/Plugins/Animation/ControlRig/Source/ControlRigEditor/Private/Editor/SRigSpacePickerWidget.h b/Engine/Plugins/Animation/ControlRig/Source/ControlRigEditor/Private/Editor/SRigSpacePickerWidget.h
--- a/Engine/Plugins/Animation/ControlRig/Source/ControlRigEditor/Private/Editor/SRigSpacePickerWidget.h
+++ b/Engine/Plugins/Animation/ControlRig/Source/ControlRigEditor/Private/Editor/SRigSpacePickerWidget.h
@@ -17,8 +17,8 @@ class IMenu;
 DECLARE_DELEGATE_RetVal_TwoParams(FRigElementKey, FRigSpacePickerGetActiveSpace, URigHierarchy*, const FRigElementKey&);
 DECLARE_DELEGATE_RetVal_TwoParams(const FRigControlElementCustomization*, FRigSpacePickerGetControlCustomization, URigHierarchy*, const FRigElementKey&);
 DECLARE_EVENT_ThreeParams(SRigSpacePickerWidget, FRigSpacePickerActiveSpaceChanged, URigHierarchy*, const FRigElementKey&, const FRigElementKey&);
-DECLARE_EVENT_ThreeParams(SRigSpacePickerWidget, FRigSpacePickerSpaceListChanged, URigHierarchy*, const FRigElementKey&, const TArray<FRigElementKey>&);
-DECLARE_DELEGATE_RetVal_TwoParams(TArray<FRigElementKey>, FRigSpacePickerGetAdditionalSpaces, URigHierarchy*, const FRigElementKey&);
+DECLARE_EVENT_ThreeParams(SRigSpacePickerWidget, FRigSpacePickerSpaceListChanged, URigHierarchy*, const FRigElementKey&, const TArray<FRigElementKeyWithLabel>&);
+DECLARE_DELEGATE_RetVal_TwoParams(TArray<FRigElementKeyWithLabel>, FRigSpacePickerGetAdditionalSpaces, URigHierarchy*, const FRigElementKey&);
 DECLARE_DELEGATE_RetVal_ThreeParams(FReply, SRigSpacePickerOnBake, URigHierarchy*, TArray<FRigElementKey> /* Controls */, FRigSpacePickerBakeSettings);
 
 /** Widget allowing picking of a space source for space switching */
@@ -87,8 +87,8 @@ public:
 	
 	const TArray<FRigElementKey>& GetControls() const { return ControlKeys; }
 	const TArray<FRigElementKey>& GetActiveSpaces() const;
-	TArray<FRigElementKey> GetDefaultSpaces() const;
-	TArray<FRigElementKey> GetSpaceList(bool bIncludeDefaultSpaces = false) const;
+	const TArray<FRigElementKeyWithLabel>& GetDefaultSpaces() const;
+	TArray<FRigElementKeyWithLabel> GetSpaceList(bool bIncludeDefaultSpaces = false) const;
 	FRigSpacePickerActiveSpaceChanged& OnActiveSpaceChanged() { return ActiveSpaceChangedEvent; }
 	FRigSpacePickerSpaceListChanged& OnSpaceListChanged() { return SpaceListChangedEvent; }
 	void RefreshContents();
@@ -143,14 +143,14 @@ private:
 
 	FSlateColor GetButtonColor(ESpacePickerType InType, FRigElementKey InKey) const;
 	FRigElementKey GetActiveSpace_Private(URigHierarchy* InHierarchy, const FRigElementKey& InControlKey) const;
-	TArray<FRigElementKey> GetCurrentParents_Private(URigHierarchy* InHierarchy, const FRigElementKey& InControlKey) const;
+	TArray<FRigElementKeyWithLabel> GetCurrentParents_Private(URigHierarchy* InHierarchy, const FRigElementKey& InControlKey) const;
 	
 	FRigSpacePickerActiveSpaceChanged ActiveSpaceChangedEvent;
 	FRigSpacePickerSpaceListChanged SpaceListChangedEvent;
 
 	TWeakObjectPtr<URigHierarchy> Hierarchy;
 	TArray<FRigElementKey> ControlKeys;
-	TArray<FRigElementKey> CurrentSpaceKeys;
+	TArray<FRigElementKeyWithLabel> CurrentSpaceKeys;
 	TArray<FRigElementKey> ActiveSpaceKeys;
 	bool bRepopulateRequired;
 
@@ -166,7 +166,7 @@ private:
 	FRigSpacePickerGetControlCustomization GetControlCustomizationDelegate;
 	FRigSpacePickerGetActiveSpace GetActiveSpaceDelegate;
 	FRigSpacePickerGetAdditionalSpaces GetAdditionalSpacesDelegate; 
-	TArray<FRigElementKey> AdditionalSpaces;
+	TArray<FRigElementKeyWithLabel> AdditionalSpaces;
 
 	TSharedPtr<SVerticalBox> TopLevelListBox;
 	TSharedPtr<SVerticalBox> ItemSpacesListBox;
diff --git a/Engine/Plugins/Animation/ControlRig/Source/ControlRigEditor/Private/Sequencer/ControlRigParameterTrackEditor.cpp b/Engine/Plugins/Animation/ControlRig/Source/ControlRigEditor/Private/Sequencer/ControlRigParameterTrackEditor.cpp
--- a/Engine/Plugins/Animation/ControlRig/Source/ControlRigEditor/Private/Sequencer/ControlRigParameterTrackEditor.cpp
+++ b/Engine/Plugins/Animation/ControlRig/Source/ControlRigEditor/Private/Sequencer/ControlRigParameterTrackEditor.cpp
@@ -3674,7 +3674,7 @@ void FControlRigParameterTrackEditor::HandleControlSelected(UControlRig* Subject
 						}
 					}
 
-					if(OtherControlElement->IsAnimationChannel() && OtherControlElement->Settings.Customization.AvailableSpaces.Contains(ControlElement->GetKey()))
+					if(OtherControlElement->IsAnimationChannel() && OtherControlElement->Settings.Customization.AvailableSpaces.FindByKey(ControlElement->GetKey()))
 					{
 						Controller->SelectElement(OtherControlElement->GetKey(), bSelected);
 						return true;
diff --git a/Engine/Plugins/Animation/ControlRig/Source/ControlRigEditor/Private/Sequencer/ControlRigSpaceChannelEditors.cpp b/Engine/Plugins/Animation/ControlRig/Source/ControlRigEditor/Private/Sequencer/ControlRigSpaceChannelEditors.cpp
--- a/Engine/Plugins/Animation/ControlRig/Source/ControlRigEditor/Private/Sequencer/ControlRigSpaceChannelEditors.cpp
+++ b/Engine/Plugins/Animation/ControlRig/Source/ControlRigEditor/Private/Sequencer/ControlRigSpaceChannelEditors.cpp
@@ -95,7 +95,7 @@ static FKeyHandle SequencerOpenSpaceSwitchDialog(UControlRig* ControlRig, TArray
 		Handle = FControlRigSpaceChannelHelpers::SequencerKeyControlRigSpaceChannel(ControlRig, Sequencer, Channel, SectionToKey, Time, RigHierarchy, ControlKey, SpaceKey);
 
 	})
-	.OnSpaceListChanged_Lambda([SelectedControls, ControlRig](URigHierarchy* InHierarchy, const FRigElementKey& InControlKey, const TArray<FRigElementKey>& InSpaceList)
+	.OnSpaceListChanged_Lambda([SelectedControls, ControlRig](URigHierarchy* InHierarchy, const FRigElementKey& InControlKey, const TArray<FRigElementKeyWithLabel>& InSpaceList)
 	{
 		check(SelectedControls.Contains(InControlKey));
 				
@@ -111,11 +111,11 @@ static FKeyHandle SequencerOpenSpaceSwitchDialog(UControlRig* ControlRig, TArray
 			ControlCustomization.RemovedSpaces.Reset();
 
 			// remember  the elements which are in the asset's available list but removed by the user
-			for (const FRigElementKey& AvailableSpace : ControlElement->Settings.Customization.AvailableSpaces)
+			for (const FRigElementKeyWithLabel& AvailableSpace : ControlElement->Settings.Customization.AvailableSpaces)
 			{
-				if (!ControlCustomization.AvailableSpaces.Contains(AvailableSpace))
+				if (ControlCustomization.AvailableSpaces.FindByKey(AvailableSpace.Key) == nullptr)
 				{
-					ControlCustomization.RemovedSpaces.Add(AvailableSpace);
+					ControlCustomization.RemovedSpaces.Add(AvailableSpace.Key);
 				}
 			}
 
