Control Rig: Hierarchy components to save using their CPPType and fail gracefully on load

#rnx

diff --git a/Engine/Plugins/Animation/ControlRig/Source/ControlRig/Private/Rigs/RigHierarchy.cpp b/Engine/Plugins/Animation/ControlRig/Source/ControlRig/Private/Rigs/RigHierarchy.cpp
--- a/Engine/Plugins/Animation/ControlRig/Source/ControlRig/Private/Rigs/RigHierarchy.cpp
+++ b/Engine/Plugins/Animation/ControlRig/Source/ControlRig/Private/Rigs/RigHierarchy.cpp
@@ -292,14 +292,61 @@ void URigHierarchy::Save(FArchive& Ar)
 		}
 	}
 	*ArchiveForElements << NumComponents;
-	for(FInstancedStruct& InstancedStruct : ElementComponents)
+	if(NumComponents > 0)
 	{
-		if(InstancedStruct.IsValid())
+		TArray<const UScriptStruct*> ComponentScriptStructs;
+		for(FInstancedStruct& InstancedStruct : ElementComponents)
+		{
+			if(InstancedStruct.IsValid())
+			{
+				FRigBaseComponent* Component = InstancedStruct.GetMutablePtr<FRigBaseComponent>();
+				UScriptStruct* ScriptStruct = Component->GetScriptStruct();
+				ComponentScriptStructs.AddUnique(ScriptStruct);
+
+				// if the reference collector is running also provide the script struct itself.
+				// this makes sure that dependency plugins will be loaded and understood as a dependency
+				// of this asset.
+				if(Ar.IsObjectReferenceCollector())
+				{
+					Ar << ScriptStruct;
+				}
+			}
+		}
+
+		TArray<FString> ScriptStructNames;
+		ScriptStructNames.Reserve(ComponentScriptStructs.Num());
+		for(const UScriptStruct* ScriptStruct : ComponentScriptStructs)
 		{
-			FRigBaseComponent* Component = InstancedStruct.GetMutablePtr<FRigBaseComponent>();
-			UScriptStruct* ScriptStruct = Component->GetScriptStruct();
-			*ArchiveForElements << ScriptStruct;
-			Component->Serialize(*ArchiveForElements);
+			ScriptStructNames.Add(ScriptStruct->GetStructCPPName());
+		}
+		*ArchiveForElements << ScriptStructNames;
+		
+		for(FInstancedStruct& InstancedStruct : ElementComponents)
+		{
+			if(InstancedStruct.IsValid())
+			{
+				FRigBaseComponent* Component = InstancedStruct.GetMutablePtr<FRigBaseComponent>();
+				const UScriptStruct* ScriptStruct = Component->GetScriptStruct();
+				
+				int32 IndexOfScriptStruct = ComponentScriptStructs.Find(ScriptStruct);
+				check(ComponentScriptStructs.IsValidIndex(IndexOfScriptStruct));
+				*ArchiveForElements << IndexOfScriptStruct;
+
+				// we store the archive position after the component. we do this so we can recover from
+				// loading content which cannot resolve the component's script struct.
+				int64 ArchivePositionBeforeComponent = ArchiveForElements->GetArchiveState().Tell();
+				int64 ArchivePositionAfterComponent = 0; // now known yet
+				*ArchiveForElements << ArchivePositionAfterComponent;
+
+				// serialize the component and take note of the pos of the archive
+				Component->Serialize(*ArchiveForElements);
+				ArchivePositionAfterComponent = ArchiveForElements->GetArchiveState().Tell();
+
+				// save this again so we can skip during load
+				ArchiveForElements->Seek(ArchivePositionBeforeComponent);
+				*ArchiveForElements << ArchivePositionAfterComponent;
+				ArchiveForElements->Seek(ArchivePositionAfterComponent);
+			}
 		}
 	}
 
@@ -598,30 +645,62 @@ void URigHierarchy::Load(FArchive& Ar)
 		int32 NumComponents = 0;
 		*ArchiveForElements << NumComponents;
 
-		ElementComponents.AddZeroed(NumComponents);
-		ComponentIndexLookup.Reserve(NumComponents);
-		for (int32 ComponentIndex = 0; ComponentIndex < ElementComponents.Num(); ComponentIndex++)
+		if(NumComponents > 0)
 		{
-			UScriptStruct* ScriptStruct = nullptr;
-			*ArchiveForElements << ScriptStruct;
-			if (ScriptStruct)
+			ElementComponents.AddZeroed(NumComponents);
+			ComponentIndexLookup.Reserve(NumComponents);
+
+			TArray<FString> ScriptStructNames;
+			*ArchiveForElements << ScriptStructNames;
+
+			TArray<UScriptStruct*> ComponentScriptStructs;
+			ComponentScriptStructs.AddZeroed(ScriptStructNames.Num());
+
+			for(int32 ScriptStructIndex = 0; ScriptStructIndex < ScriptStructNames.Num(); ScriptStructIndex++)
 			{
-				ArchiveForElements->Preload(ScriptStruct);
+				// this lookup makes sure to avoid loading things syncronously,
+				// and it also relies on FindGlobally - which internally uses
+				// the core redirects in case the structure has been renamed.
+				if(UScriptStruct* ScriptStruct = Cast<UScriptStruct>(RigVMTypeUtils::ObjectFromCPPType(ScriptStructNames[ScriptStructIndex])))
+				{
+					ComponentScriptStructs[ScriptStructIndex] = ScriptStruct;
+				}
 			}
 
-			FInstancedStruct& InstancedStruct = ElementComponents[ComponentIndex];
-			InstancedStruct.InitializeAs(ScriptStruct);
-			
-			FRigBaseComponent* Component = InstancedStruct.GetMutablePtr<FRigBaseComponent>();
-			Component->Serialize(*ArchiveForElements);
+			for (int32 ComponentIndex = 0; ComponentIndex < ElementComponents.Num(); ComponentIndex++)
+			{
+				int32 IndexOfScriptStruct = 0;
+				*ArchiveForElements << IndexOfScriptStruct;
+				check(ComponentScriptStructs.IsValidIndex(IndexOfScriptStruct));
 
-			FRigBaseElement* Element = Find(Component->GetElementKey());
-			check(Element);
-			
-			Component->IndexInHierarchy = ComponentIndex;
-			Component->Element = Element;
-			Component->IndexInElement = Element->ComponentIndices.Add(Component->IndexInHierarchy);
-			ComponentIndexLookup.Add(Component->GetKey(), ComponentIndex);
+				// we know the position of the archive so we can skip the data
+				// in case we cannot resolve the component
+				int64 ArchivePositionAfterComponent = 0;
+				*ArchiveForElements << ArchivePositionAfterComponent;
+
+				UScriptStruct* ScriptStruct = ComponentScriptStructs[IndexOfScriptStruct];
+				if (ScriptStruct == nullptr)
+				{
+					ArchiveForElements->Seek(ArchivePositionAfterComponent);
+					continue;
+				}
+
+				ArchiveForElements->Preload(ScriptStruct);
+
+				FInstancedStruct& InstancedStruct = ElementComponents[ComponentIndex];
+				InstancedStruct.InitializeAs(ScriptStruct);
+				
+				FRigBaseComponent* Component = InstancedStruct.GetMutablePtr<FRigBaseComponent>();
+				Component->Serialize(*ArchiveForElements);
+
+				FRigBaseElement* Element = Find(Component->GetElementKey());
+				check(Element);
+				
+				Component->IndexInHierarchy = ComponentIndex;
+				Component->Element = Element;
+				Component->IndexInElement = Element->ComponentIndices.Add(Component->IndexInHierarchy);
+				ComponentIndexLookup.Add(Component->GetKey(), ComponentIndex);
+			}
 		}
 
 		IncrementTopologyVersion();
