Sequencer: Speculative fix for pin() on a weak pointer

#rnx

diff --git a/Engine/Plugins/Animation/ControlRig/Source/ControlRigEditor/Private/EditMode/ControlRigEditMode.cpp b/Engine/Plugins/Animation/ControlRig/Source/ControlRigEditor/Private/EditMode/ControlRigEditMode.cpp
--- a/Engine/Plugins/Animation/ControlRig/Source/ControlRigEditor/Private/EditMode/ControlRigEditMode.cpp
+++ b/Engine/Plugins/Animation/ControlRig/Source/ControlRigEditor/Private/EditMode/ControlRigEditMode.cpp
@@ -325,7 +325,7 @@ bool FControlRigEditMode::AddControlRigObject(UControlRig* InControlRig, const T
 	return false;
 }
 
-void FControlRigEditMode::SetObjects(UControlRig* ControlRig,  UObject* BindingObject, TWeakPtr<ISequencer> InSequencer)
+void FControlRigEditMode::SetObjects(UControlRig* ControlRig,  UObject* BindingObject, const TWeakPtr<ISequencer>& InSequencer)
 {
 	TArray<TWeakObjectPtr<UControlRig>> PreviousRuntimeRigs = RuntimeControlRigs;
 	for (int32 PreviousRuntimeRigIndex = 0; PreviousRuntimeRigIndex < PreviousRuntimeRigs.Num(); PreviousRuntimeRigIndex++)
diff --git a/Engine/Plugins/Animation/ControlRig/Source/ControlRigEditor/Private/Sequencer/ControlRigParameterTrackEditor.cpp b/Engine/Plugins/Animation/ControlRig/Source/ControlRigEditor/Private/Sequencer/ControlRigParameterTrackEditor.cpp
--- a/Engine/Plugins/Animation/ControlRig/Source/ControlRigEditor/Private/Sequencer/ControlRigParameterTrackEditor.cpp
+++ b/Engine/Plugins/Animation/ControlRig/Source/ControlRigEditor/Private/Sequencer/ControlRigParameterTrackEditor.cpp
@@ -340,191 +340,218 @@ FText FControlRigParameterTrackEditor::GetDisplayName() const
 
 void FControlRigParameterTrackEditor::OnInitialize()
 {
-	if (GetSequencer().IsValid() == false)
+	TSharedPtr<ISequencer> Sequencer = GetSequencer();
+	if (!Sequencer)
 	{
 		return;
 	}
-	ISequencer* SequencerPtr = GetSequencer().Get();
-
-	UMovieScene* MovieScene = SequencerPtr->GetFocusedMovieSceneSequence()->GetMovieScene();
-	SelectionChangedHandle = SequencerPtr->GetSelectionChangedTracks().AddRaw(this, &FControlRigParameterTrackEditor::OnSelectionChanged);
-	SequencerChangedHandle = SequencerPtr->OnMovieSceneDataChanged().AddRaw(this, &FControlRigParameterTrackEditor::OnSequencerDataChanged);
-	OnActivateSequenceChangedHandle = SequencerPtr->OnActivateSequence().AddRaw(this, &FControlRigParameterTrackEditor::OnActivateSequenceChanged);
-	OnChannelChangedHandle = SequencerPtr->OnChannelChanged().AddRaw(this, &FControlRigParameterTrackEditor::OnChannelChanged);
-	OnMovieSceneChannelChangedHandle = MovieScene->OnChannelChanged().AddRaw(this, &FControlRigParameterTrackEditor::OnChannelChanged);
-	OnActorAddedToSequencerHandle = SequencerPtr->OnActorAddedToSequencer().AddRaw(this, &FControlRigParameterTrackEditor::HandleActorAdded);
+	
+	UMovieSceneSequence* OwnerSequence = Sequencer->GetFocusedMovieSceneSequence();
+	UMovieScene* MovieScene = OwnerSequence ? OwnerSequence->GetMovieScene() : nullptr;
+	
+	SelectionChangedHandle = Sequencer->GetSelectionChangedTracks().AddRaw(this, &FControlRigParameterTrackEditor::OnSelectionChanged);
+	SequencerChangedHandle = Sequencer->OnMovieSceneDataChanged().AddRaw(this, &FControlRigParameterTrackEditor::OnSequencerDataChanged);
+	OnActivateSequenceChangedHandle = Sequencer->OnActivateSequence().AddRaw(this, &FControlRigParameterTrackEditor::OnActivateSequenceChanged);
+	OnChannelChangedHandle = Sequencer->OnChannelChanged().AddRaw(this, &FControlRigParameterTrackEditor::OnChannelChanged);
+	if (MovieScene)
+	{
+		OnMovieSceneChannelChangedHandle = MovieScene->OnChannelChanged().AddRaw(this, &FControlRigParameterTrackEditor::OnChannelChanged);
+	}
+	OnActorAddedToSequencerHandle = Sequencer->OnActorAddedToSequencer().AddRaw(this, &FControlRigParameterTrackEditor::HandleActorAdded);
 
 	{
 		//we check for two things, one if the control rig has been replaced if so we need to switch.
 		//the other is if bound object on the edit mode is null we request a re-evaluate which will reset it up.
 		const FDelegateHandle OnObjectsReplacedHandle = FCoreUObjectDelegates::OnObjectsReplaced.AddLambda([this](const TMap<UObject*, UObject*>& ReplacementMap)
 		{
-			if (GetSequencer().IsValid() && GetSequencer()->GetFocusedMovieSceneSequence() && GetSequencer()->GetFocusedMovieSceneSequence()->GetMovieScene())
+			TSharedPtr<ISequencer> Sequencer = GetSequencer();
+			UMovieSceneSequence* OwnerSequence = Sequencer ? Sequencer->GetFocusedMovieSceneSequence() : nullptr;
+			UMovieScene* MovieScene = OwnerSequence ? OwnerSequence->GetMovieScene() : nullptr;
+
+			if (!MovieScene)
 			{
-				static bool bHasEnteredSilent = false;
-				
-				bool bRequestEvaluate = false;
-				TMap<UControlRig*, UControlRig*> OldToNewControlRigs;
-				FControlRigEditMode* ControlRigEditMode = GetEditMode();
-				if (ControlRigEditMode)
-				{
-					UMovieScene* MovieScene = GetSequencer()->GetFocusedMovieSceneSequence()->GetMovieScene();
-					TArray<FString>& MuteNodes = MovieScene->GetMuteNodes();
+				return;
+			}
+
+			static bool bHasEnteredSilent = false;
+			
+			bool bRequestEvaluate = false;
+			TMap<UControlRig*, UControlRig*> OldToNewControlRigs;
+			
+			FControlRigEditMode* ControlRigEditMode = GetEditMode();
+			if (ControlRigEditMode)
+			{
+				const TArray<FString>& MuteNodes = MovieScene->GetMuteNodes();
 
-					TArrayView<TWeakObjectPtr<UControlRig>> ControlRigs = ControlRigEditMode->GetControlRigs();
-					for (TWeakObjectPtr<UControlRig>& ControlRigPtr : ControlRigs)
+				const TArrayView<const TWeakObjectPtr<UControlRig>> ControlRigs = ControlRigEditMode->GetControlRigs();
+				for (const TWeakObjectPtr<UControlRig>& ControlRigPtr : ControlRigs)
+				{
+					if (const UControlRig* ControlRig = ControlRigPtr.Get())
 					{
-						if (UControlRig* ControlRig = ControlRigPtr.Get())
+						if (ControlRig->GetObjectBinding() && ControlRig->GetObjectBinding()->GetBoundObject() == nullptr)
 						{
-							if (ControlRig->GetObjectBinding() && ControlRig->GetObjectBinding()->GetBoundObject() == nullptr)
+							IterateTracks([this, ControlRig, &MuteNodes, &bRequestEvaluate](UMovieSceneControlRigParameterTrack* Track)
 							{
-								IterateTracks([this, ControlRig, &MuteNodes, &bRequestEvaluate](UMovieSceneControlRigParameterTrack* Track)
-									{
-										if (ControlRig == Track->GetControlRig())
-										{
-											if (MuteNodes.Find(Track->GetName()) != INDEX_NONE) //only re-evalute if  not muted, TODO, need function to test to see if a track is evaluating at a certain time
-											{
-												bRequestEvaluate = true;
-											}
-										}
-										return false;
-									});
-								if (bRequestEvaluate == true)
+								if (ControlRig == Track->GetControlRig())
 								{
-									break;
+									if (MuteNodes.Find(Track->GetName()) != INDEX_NONE) //only re-evalute if  not muted, TODO, need function to test to see if a track is evaluating at a certain time
+									{
+										bRequestEvaluate = true;
+									}
 								}
+								return false;
+							});
+							
+							if (bRequestEvaluate == true)
+							{
+								break;
 							}
 						}
 					}
 				}
-				//Reset Bindings for replaced objects.
-				for (TPair<UObject*, UObject*> ReplacedObject : ReplacementMap)
+			}
+			
+			//Reset Bindings for replaced objects.
+			for (TPair<UObject*, UObject*> ReplacedObject : ReplacementMap)
+			{
+				if (UControlRigComponent* OldControlRigComponent = Cast<UControlRigComponent>(ReplacedObject.Key))
 				{
-					if (UControlRigComponent* OldControlRigComponent = Cast<UControlRigComponent>(ReplacedObject.Key))
+					UControlRigComponent* NewControlRigComponent = Cast<UControlRigComponent>(ReplacedObject.Value);
+					if (OldControlRigComponent->GetControlRig())
 					{
-						UControlRigComponent* NewControlRigComponent = Cast<UControlRigComponent>(ReplacedObject.Value);
-						if (OldControlRigComponent->GetControlRig())
+						UControlRig* NewControlRig = nullptr;
+						if (NewControlRigComponent)
 						{
-							UControlRig* NewControlRig = nullptr;
-							if (NewControlRigComponent)
-							{
-								NewControlRig = NewControlRigComponent->GetControlRig();
-							}
-							OldToNewControlRigs.Emplace(OldControlRigComponent->GetControlRig(), NewControlRig);
+							NewControlRig = NewControlRigComponent->GetControlRig();
 						}
-					}
-					else if (UControlRig* OldControlRig = Cast<UControlRig>(ReplacedObject.Key))
-					{
-						UControlRig* NewControlRig = Cast<UControlRig>(ReplacedObject.Value);
-						OldToNewControlRigs.Emplace(OldControlRig, NewControlRig);
+						OldToNewControlRigs.Emplace(OldControlRigComponent->GetControlRig(), NewControlRig);
 					}
 				}
-				if (OldToNewControlRigs.Num() > 0)
+				else if (UControlRig* OldControlRig = Cast<UControlRig>(ReplacedObject.Key))
 				{
-					//need to avoid any evaluations when doing this replacement otherwise we will evaluate sequencer
-					if (bHasEnteredSilent == false)
-					{
-						GetSequencer()->EnterSilentMode();
-						bHasEnteredSilent = true;
-					}
+					UControlRig* NewControlRig = Cast<UControlRig>(ReplacedObject.Value);
+					OldToNewControlRigs.Emplace(OldControlRig, NewControlRig);
+				}
+			}
+			
+			if (OldToNewControlRigs.Num() > 0)
+			{
+				//need to avoid any evaluations when doing this replacement otherwise we will evaluate sequencer
+				if (bHasEnteredSilent == false)
+				{
+					Sequencer->EnterSilentMode();
+					bHasEnteredSilent = true;
+				}
 
-					IterateTracks([this, &OldToNewControlRigs, ControlRigEditMode, &bRequestEvaluate](UMovieSceneControlRigParameterTrack* Track)
+				IterateTracks([this, &OldToNewControlRigs, ControlRigEditMode, &bRequestEvaluate](UMovieSceneControlRigParameterTrack* Track)
+				{
+					if (UControlRig* OldControlRig = Track->GetControlRig())
 					{
-						if (UControlRig* OldControlRig = Track->GetControlRig())
+						UControlRig** FoundNewControlRig = OldToNewControlRigs.Find(OldControlRig);
+						if (FoundNewControlRig)
 						{
-							UControlRig** NewControlRig = OldToNewControlRigs.Find(OldControlRig);
+							UControlRig* NewControlRig = (*FoundNewControlRig); 
+							
+							TArray<FName> SelectedControls = OldControlRig->CurrentControlSelection();
+							OldControlRig->ClearControlSelection();
+							UnbindControlRig(OldControlRig);
+
 							if (NewControlRig)
 							{
-								TArray<FName> SelectedControls = OldControlRig->CurrentControlSelection();
-								OldControlRig->ClearControlSelection();
-								UnbindControlRig(OldControlRig);
-								if (*NewControlRig)
-								{
-									Track->Modify();
-									Track->ReplaceControlRig(*NewControlRig, OldControlRig->GetClass() != (*NewControlRig)->GetClass());
-									BindControlRig(*NewControlRig);
-									bRequestEvaluate = true;
-								}
-								else
-								{
-									Track->ReplaceControlRig(nullptr, true);
-								}
-								if (ControlRigEditMode)
-								{
-									ControlRigEditMode->ReplaceControlRig(OldControlRig, *NewControlRig);
-
-									UControlRig* PtrNewControlRig = *NewControlRig;
+								Track->Modify();
+								Track->ReplaceControlRig(NewControlRig, OldControlRig->GetClass() != NewControlRig->GetClass());
+								BindControlRig(NewControlRig);
+								bRequestEvaluate = true;
+							}
+							else
+							{
+								Track->ReplaceControlRig(nullptr, true);
+							}
+							
+							if (ControlRigEditMode)
+							{
+								ControlRigEditMode->ReplaceControlRig(OldControlRig, NewControlRig);
 
-									auto UpdateSelectionDelegate = [this, SelectedControls, PtrNewControlRig, bRequestEvaluate]()
+								auto UpdateSelectionDelegate = [this, SelectedControls, NewControlRig, bRequestEvaluate]()
+								{
+									if (!(FSlateApplication::Get().HasAnyMouseCaptor() || GUnrealEd->IsUserInteracting()))
 									{
+										TSharedPtr<ISequencer> Sequencer = GetSequencer();
+										
+										TGuardValue<bool> Guard(bIsDoingSelection, true);
+										if (Sequencer)
+										{
+											Sequencer->ExternalSelectionHasChanged();
+										}
 
-										if (!(FSlateApplication::Get().HasAnyMouseCaptor() || GUnrealEd->IsUserInteracting()))
+										if (NewControlRig)
 										{
-											TGuardValue<bool> Guard(bIsDoingSelection, true);
-											GetSequencer()->ExternalSelectionHasChanged();
-											if (PtrNewControlRig)
+											GEditor->GetTimerManager()->SetTimerForNextTick([this, SelectedControls, NewControlRig]()
 											{
-												GEditor->GetTimerManager()->SetTimerForNextTick([this, SelectedControls, PtrNewControlRig, bRequestEvaluate]()
-													{
-														PtrNewControlRig->ClearControlSelection();
-														for (const FName& ControlName : SelectedControls)
-														{
-															PtrNewControlRig->SelectControl(ControlName, true);
-														}
-													});
-											}
-											
-											if (bHasEnteredSilent == true)
-											{
-												GetSequencer()->ExitSilentMode();
-												bHasEnteredSilent = false;
-											}
-											
-											if (bRequestEvaluate)
+												NewControlRig->ClearControlSelection();
+												for (const FName& ControlName : SelectedControls)
+												{
+													NewControlRig->SelectControl(ControlName, true);
+												}
+											});
+										}
+										
+										if (bHasEnteredSilent)
+										{
+											if (Sequencer)
 											{
-												GetSequencer()->NotifyMovieSceneDataChanged(EMovieSceneDataChangeType::MovieSceneStructureItemsChanged);
+												Sequencer->ExitSilentMode();
 											}
-											if (UpdateSelectionTimerHandle.IsValid())
+											bHasEnteredSilent = false;
+										}
+										
+										if (bRequestEvaluate)
+										{
+											if (Sequencer)
 											{
-												GEditor->GetTimerManager()->ClearTimer(UpdateSelectionTimerHandle);
+												Sequencer->NotifyMovieSceneDataChanged(EMovieSceneDataChangeType::MovieSceneStructureItemsChanged);
 											}
 										}
+										
+										if (UpdateSelectionTimerHandle.IsValid())
+										{
+											GEditor->GetTimerManager()->ClearTimer(UpdateSelectionTimerHandle);
+										}
+									}
 
-									};
+								};
 
-									GEditor->GetTimerManager()->SetTimer(UpdateSelectionTimerHandle, UpdateSelectionDelegate, 0.01f, true);
-								}
+								GEditor->GetTimerManager()->SetTimer(UpdateSelectionTimerHandle, UpdateSelectionDelegate, 0.01f, true);
 							}
 						}
+					}
 
-						return false;
-					});
+					return false;
+				});
 
-					if (!ControlRigEditMode)
+				if (!ControlRigEditMode)
+				{
+					if (bRequestEvaluate)
 					{
-						if (bRequestEvaluate)
-						{
-							GetSequencer()->NotifyMovieSceneDataChanged(EMovieSceneDataChangeType::MovieSceneStructureItemsChanged);
-						}
+						Sequencer->NotifyMovieSceneDataChanged(EMovieSceneDataChangeType::MovieSceneStructureItemsChanged);
 					}
 				}
+			}
 
-				if (ControlRigEditMode && bRequestEvaluate)
-				{
-					GetSequencer()->NotifyMovieSceneDataChanged(EMovieSceneDataChangeType::MovieSceneStructureItemsChanged);
-				}
+			if (ControlRigEditMode && bRequestEvaluate)
+			{
+				Sequencer->NotifyMovieSceneDataChanged(EMovieSceneDataChangeType::MovieSceneStructureItemsChanged);
+			}
 
-				// ensure we exit silent mode if it has been entered
-				if (bHasEnteredSilent)
-				{
-					GetSequencer()->ExitSilentMode();
-					bHasEnteredSilent = false;
-				}
+			// ensure we exit silent mode if it has been entered
+			if (bHasEnteredSilent)
+			{
+				Sequencer->ExitSilentMode();
+				bHasEnteredSilent = false;
 			}
 		});
 	
-
-		AcquiredResources.Add([=] { FCoreUObjectDelegates::OnObjectsReplaced.Remove(OnObjectsReplacedHandle); });
+		AcquiredResources.Add([OnObjectsReplacedHandle] { FCoreUObjectDelegates::OnObjectsReplaced.Remove(OnObjectsReplacedHandle); });
 	}
 
 	// Register all modified/selections for control rigs
@@ -665,6 +692,9 @@ void FControlRigParameterTrackEditor::OnRelease()
 
 	FControlRigEditMode* ControlRigEditMode = GetEditMode();
 	PreviousSelectedControlRigs.Reset();
+
+	TSharedPtr<ISequencer> Sequencer = GetSequencer();
+	
 	if (ControlRigEditMode)
 	{
 		bControlRigEditModeWasOpen = true;
@@ -679,56 +709,71 @@ void FControlRigParameterTrackEditor::OnRelease()
 			}
 		}
 		ControlRigEditMode->Exit(); //deactive mode below doesn't exit for some reason so need to make sure things are cleaned up
+
 		if (FEditorModeTools* Tools = GetEditorModeTools())
 		{
 			Tools->DeactivateMode(FControlRigEditMode::ModeName);
 		}
 
-		ControlRigEditMode->SetObjects(nullptr, nullptr, GetSequencer());
+		ControlRigEditMode->SetObjects(nullptr, nullptr, Sequencer);
 	}
 	else
 	{
 		bControlRigEditModeWasOpen = false;
 	}
+	
 	UnbindAllControlRigs();
-	if (GetSequencer().IsValid())
+
+	if (Sequencer)
 	{
 		if (SelectionChangedHandle.IsValid())
 		{
-			GetSequencer()->GetSelectionChangedTracks().Remove(SelectionChangedHandle);
+			Sequencer->GetSelectionChangedTracks().Remove(SelectionChangedHandle);
+			SelectionChangedHandle.Reset();
 		}
+		
 		if (SequencerChangedHandle.IsValid())
 		{
-			GetSequencer()->OnMovieSceneDataChanged().Remove(SequencerChangedHandle);
+			Sequencer->OnMovieSceneDataChanged().Remove(SequencerChangedHandle);
+			SequencerChangedHandle.Reset();
 		}
+		
 		if (OnActivateSequenceChangedHandle.IsValid())
 		{
-			GetSequencer()->OnActivateSequence().Remove(OnActivateSequenceChangedHandle);
+			Sequencer->OnActivateSequence().Remove(OnActivateSequenceChangedHandle);
+			OnActivateSequenceChangedHandle.Reset();
 		}
+		
 		if (CurveChangedHandle.IsValid())
 		{
-			GetSequencer()->GetCurveDisplayChanged().Remove(CurveChangedHandle);
+			Sequencer->GetCurveDisplayChanged().Remove(CurveChangedHandle);
+			CurveChangedHandle.Reset();
 		}
+		
 		if (OnActorAddedToSequencerHandle.IsValid())
 		{
-			GetSequencer()->OnActorAddedToSequencer().Remove(OnActorAddedToSequencerHandle);
+			Sequencer->OnActorAddedToSequencer().Remove(OnActorAddedToSequencerHandle);
+			OnActorAddedToSequencerHandle.Reset();
 		}
+		
 		if (OnChannelChangedHandle.IsValid())
 		{
-			GetSequencer()->OnChannelChanged().Remove(OnChannelChangedHandle);
+			Sequencer->OnChannelChanged().Remove(OnChannelChangedHandle);
+			OnChannelChangedHandle.Reset();
 		}
 
-		if (GetSequencer()->GetFocusedMovieSceneSequence() && GetSequencer()->GetFocusedMovieSceneSequence()->GetMovieScene())
+		if (OnMovieSceneChannelChangedHandle.IsValid())
 		{
-			UMovieScene* MovieScene = GetSequencer()->GetFocusedMovieSceneSequence()->GetMovieScene();
-			if (OnMovieSceneChannelChangedHandle.IsValid())
+			UMovieSceneSequence* OwnerSequence = Sequencer->GetFocusedMovieSceneSequence();
+			if (UMovieScene* MovieScene = OwnerSequence ? OwnerSequence->GetMovieScene() : nullptr)
 			{
 				MovieScene->OnChannelChanged().Remove(OnMovieSceneChannelChangedHandle);
 			}
+			OnMovieSceneChannelChangedHandle.Reset();
 		}
 	}
+	
 	AcquiredResources.Release();
-
 }
 
 TSharedRef<ISequencerTrackEditor> FControlRigParameterTrackEditor::CreateTrackEditor(TSharedRef<ISequencer> InSequencer)
@@ -772,10 +817,10 @@ void FControlRigParameterTrackEditor::BuildObjectBindingContextMenu(FMenuBuilder
 		ObjectClass->IsChildOf(AActor::StaticClass()) ||
 		ObjectClass->IsChildOf(UChildActorComponent::StaticClass()))
 	{
-		const TSharedPtr<ISequencer> ParentSequencer = GetSequencer();
+		const TSharedPtr<ISequencer> Sequencer = GetSequencer();
 		UObject* BoundObject = nullptr;
-		USkeleton* Skeleton = AcquireSkeletonFromObjectGuid(ObjectBindings[0], &BoundObject, ParentSequencer);
-		USkeletalMeshComponent* SkelMeshComp = AcquireSkeletalMeshFromObject(BoundObject, ParentSequencer);
+		USkeleton* Skeleton = AcquireSkeletonFromObjectGuid(ObjectBindings[0], &BoundObject, Sequencer);
+		USkeletalMeshComponent* SkelMeshComp = AcquireSkeletalMeshFromObject(BoundObject, Sequencer);
 
 		if (Skeleton && SkelMeshComp)
 		{
@@ -857,8 +902,6 @@ private:
 
 void FControlRigParameterTrackEditor::BakeToControlRigSubMenu(FMenuBuilder& MenuBuilder, FGuid ObjectBinding, UObject* BoundObject, USkeletalMeshComponent* SkelMeshComp, USkeleton* Skeleton)
 {
-	const TSharedPtr<ISequencer> ParentSequencer = GetSequencer();
-
 	if (Skeleton)
 	{
 		FClassViewerInitializationOptions Options;
@@ -1089,14 +1132,14 @@ FReply SBakeToAnimAndControlRigOptionsWindow::OnResetToDefaultClick() const
 	return FReply::Handled();
 }
 
-void FControlRigParameterTrackEditor::SmartReduce(TSharedPtr<ISequencer>& InSequencer, const FSmartReduceParams& InParams, UMovieSceneControlRigParameterSection * ParamSection)
+void FControlRigParameterTrackEditor::SmartReduce(const TSharedPtr<ISequencer>& InSequencer, const FSmartReduceParams& InParams, UMovieSceneControlRigParameterSection * ParamSection)
 {
 	using namespace UE::Sequencer;
 
 	if (ParamSection)
 	{
 		ISequencerModule& SequencerModule = FModuleManager::LoadModuleChecked<ISequencerModule>("Sequencer");
-		const bool bNeedToTestExisting = false;
+		constexpr bool bNeedToTestExisting = false;
 		TOptional<FKeyHandleSet> KeyHandleSet;
 
 		FMovieSceneChannelProxy& ChannelProxy = ParamSection->GetChannelProxy();
@@ -1120,12 +1163,13 @@ void FControlRigParameterTrackEditor::SmartReduce(TSharedPtr<ISequencer>& InSequ
 	}
 }
 
-bool FControlRigParameterTrackEditor::LoadAnimationIntoSection(TSharedPtr<ISequencer>& SequencerPtr, UAnimSequence* AnimSequence, USkeletalMeshComponent* SkelMeshComp,
+bool FControlRigParameterTrackEditor::LoadAnimationIntoSection(const TSharedPtr<ISequencer>& SequencerPtr, UAnimSequence* AnimSequence, USkeletalMeshComponent* SkelMeshComp,
 	FFrameNumber StartFrame, bool bReduceKeys, const FSmartReduceParams& ReduceParams, bool bResetControls, UMovieSceneControlRigParameterSection* ParamSection)
 {
 	EMovieSceneKeyInterpolation DefaultInterpolation = SequencerPtr->GetKeyInterpolation();
 	UMovieSceneSequence* OwnerSequence = SequencerPtr->GetFocusedMovieSceneSequence();
-	UMovieScene* OwnerMovieScene = OwnerSequence->GetMovieScene();
+	UMovieScene* OwnerMovieScene = OwnerSequence ? OwnerSequence->GetMovieScene() : nullptr;
+	
 	if (ParamSection->LoadAnimSequenceIntoThisSection(AnimSequence, FFrameNumber(0), OwnerMovieScene, SkelMeshComp,
 		false, 0.0, bResetControls, StartFrame, DefaultInterpolation))
 	{
@@ -1141,24 +1185,29 @@ bool FControlRigParameterTrackEditor::LoadAnimationIntoSection(TSharedPtr<ISeque
 void FControlRigParameterTrackEditor::BakeToControlRig(UClass* InClass, FGuid ObjectBinding, UObject* BoundActor, USkeletalMeshComponent* SkelMeshComp, USkeleton* Skeleton)
 {
 	FSlateApplication::Get().DismissAllMenus();
-	const TSharedPtr<ISequencer> SequencerParent = GetSequencer();
+	
+	const TSharedPtr<ISequencer> Sequencer = GetSequencer();
 
-	if (InClass && InClass->IsChildOf(UControlRig::StaticClass()) && SequencerParent.IsValid())
+	if (InClass && InClass->IsChildOf(UControlRig::StaticClass()) && Sequencer)
 	{
-
-		UMovieSceneSequence* OwnerSequence = GetSequencer()->GetFocusedMovieSceneSequence();
-		UMovieSceneSequence* RootSequence = GetSequencer()->GetRootMovieSceneSequence();
-		UMovieScene* OwnerMovieScene = OwnerSequence->GetMovieScene();
+		UMovieSceneSequence* OwnerSequence = Sequencer->GetFocusedMovieSceneSequence();
+		UMovieSceneSequence* RootSequence = Sequencer->GetRootMovieSceneSequence();
+		UMovieScene* OwnerMovieScene = OwnerSequence ? OwnerSequence->GetMovieScene() : nullptr;
+		
+		if (!ensure(OwnerMovieScene))
+		{
+			return;
+		}
+		
 		{
 			UAnimSequence* TempAnimSequence = NewObject<UAnimSequence>(GetTransientPackage(), NAME_None);
 			TempAnimSequence->SetSkeleton(Skeleton);
-			const TSharedPtr<ISequencer> ParentSequencer = GetSequencer();
-			FMovieSceneSequenceIDRef Template = ParentSequencer->GetFocusedTemplateID();
-			FMovieSceneSequenceTransform RootToLocalTransform = ParentSequencer->GetFocusedMovieSceneSequenceTransform();
+			
+			FMovieSceneSequenceTransform RootToLocalTransform = Sequencer->GetFocusedMovieSceneSequenceTransform();
 			UAnimSeqExportOption* AnimSeqExportOption = NewObject<UAnimSeqExportOption>(GetTransientPackage(), NAME_None);
 			UBakeToControlRigSettings* BakeSettings = GetMutableDefault<UBakeToControlRigSettings>();
 			AnimSeqExportOption->bTransactRecording = false;
-			AnimSeqExportOption->CustomDisplayRate = ParentSequencer->GetFocusedDisplayRate();
+			AnimSeqExportOption->CustomDisplayRate = Sequencer->GetFocusedDisplayRate();
 
 			TSharedPtr<SWindow> ParentWindow;
 			if (FModuleManager::Get().IsModuleLoaded("MainFrame"))
@@ -1187,11 +1236,11 @@ void FControlRigParameterTrackEditor::BakeToControlRig(UClass* InClass, FGuid Ob
 			if (OptionWindow.Get()->ShouldExport())
 			{
 				FAnimExportSequenceParameters AESP;
-				AESP.Player = ParentSequencer.Get();
+				AESP.Player = Sequencer.Get();
 				AESP.RootToLocalTransform = RootToLocalTransform;
 				AESP.MovieSceneSequence = OwnerSequence;
 				AESP.RootMovieSceneSequence = RootSequence;
-				AESP.bForceUseOfMovieScenePlaybackRange = ParentSequencer.Get()->GetSequencerSettings()->ShouldEvaluateSubSequencesInIsolation();
+				AESP.bForceUseOfMovieScenePlaybackRange = Sequencer->GetSequencerSettings()->ShouldEvaluateSubSequencesInIsolation();
 				bool bResult = MovieSceneToolHelpers::ExportToAnimSequence(TempAnimSequence, AnimSeqExportOption, AESP, SkelMeshComp);
 				if (bResult == false)
 				{
@@ -1265,7 +1314,7 @@ void FControlRigParameterTrackEditor::BakeToControlRig(UClass* InClass, FGuid Ob
 						ControlRig->Evaluate_AnyThread();
 					}
 
-					const bool bSequencerOwnsControlRig = true;
+					constexpr bool bSequencerOwnsControlRig = true;
 					UMovieSceneSection* NewSection = Track->CreateControlRigSection(0, ControlRig, bSequencerOwnsControlRig);
 					UMovieSceneControlRigParameterSection* ParamSection = Cast<UMovieSceneControlRigParameterSection>(NewSection);
 
@@ -1273,15 +1322,16 @@ void FControlRigParameterTrackEditor::BakeToControlRig(UClass* InClass, FGuid Ob
 					Track->SetTrackName(FName(*ObjectName));
 					Track->SetDisplayName(FText::FromString(ObjectName));
 
-					TSharedPtr<ISequencer> SequencerPtr = GetSequencer();
-					SequencerPtr->EmptySelection();
-					SequencerPtr->SelectSection(NewSection);
-					SequencerPtr->ThrobSectionSelection();
-					SequencerPtr->NotifyMovieSceneDataChanged(EMovieSceneDataChangeType::MovieSceneStructureItemAdded);
-					TOptional<TRange<FFrameNumber>> OptionalRange = SequencerPtr->GetSubSequenceRange();
+					Sequencer->EmptySelection();
+					Sequencer->SelectSection(NewSection);
+					Sequencer->ThrobSectionSelection();
+					Sequencer->NotifyMovieSceneDataChanged(EMovieSceneDataChangeType::MovieSceneStructureItemAdded);
+					TOptional<TRange<FFrameNumber>> OptionalRange = Sequencer->GetSubSequenceRange();
 					FFrameNumber StartFrame = OptionalRange.IsSet() ? OptionalRange.GetValue().GetLowerBoundValue() : OwnerMovieScene->GetPlaybackRange().GetLowerBoundValue();
-					LoadAnimationIntoSection(SequencerPtr, TempAnimSequence, SkelMeshComp, StartFrame,
+
+					LoadAnimationIntoSection(Sequencer, TempAnimSequence, SkelMeshComp, StartFrame,
 						BakeSettings->bReduceKeys, BakeSettings->SmartReduce, BakeSettings->bResetControls, ParamSection);
+
 					//Turn Off Any Skeletal Animation Tracks
 					TArray<UMovieSceneSkeletalAnimationTrack*> SkelAnimationTracks;
 					if (const FMovieSceneBinding* Binding = OwnerMovieScene->FindBinding(ObjectBinding))
@@ -1294,7 +1344,8 @@ void FControlRigParameterTrackEditor::BakeToControlRig(UClass* InClass, FGuid Ob
 							}
 						}
 					}
-					FGuid SkelMeshGuid = GetSequencer()->FindObjectId(*SkelMeshComp, GetSequencer()->GetFocusedTemplateID());
+
+					const FGuid SkelMeshGuid = Sequencer->FindObjectId(*SkelMeshComp, Sequencer->GetFocusedTemplateID());
 					if (const FMovieSceneBinding* Binding = OwnerMovieScene->FindBinding(SkelMeshGuid))
 					{
 						for (UMovieSceneTrack* MovieSceneTrack : Binding->GetTracks())
@@ -1309,6 +1360,7 @@ void FControlRigParameterTrackEditor::BakeToControlRig(UClass* InClass, FGuid Ob
 					for (UMovieSceneSkeletalAnimationTrack* SkelTrack : SkelAnimationTracks)
 					{
 						SkelTrack->Modify();
+						
 						//can't just turn off the track so need to mute the sections
 						const TArray<UMovieSceneSection*>& Sections = SkelTrack->GetAllSections();
 						for (UMovieSceneSection* Section : Sections)
@@ -1324,50 +1376,48 @@ void FControlRigParameterTrackEditor::BakeToControlRig(UClass* InClass, FGuid Ob
 					//Finish Setup
 					if (ControlRigEditMode)
 					{
-						ControlRigEditMode->AddControlRigObject(ControlRig, GetSequencer());
+						ControlRigEditMode->AddControlRigObject(ControlRig, Sequencer);
 					}
 					BindControlRig(ControlRig);
 
 					TempAnimSequence->MarkAsGarbage();
 					AnimSeqExportOption->MarkAsGarbage();
-					GetSequencer()->NotifyMovieSceneDataChanged(EMovieSceneDataChangeType::MovieSceneStructureItemAdded);
-
-
+					Sequencer->NotifyMovieSceneDataChanged(EMovieSceneDataChangeType::MovieSceneStructureItemAdded);
 				}
 			}
 		}
 	}
 }
 
-static void IterateTracksInMovieScene(UMovieScene& MovieScene, TFunctionRef<bool(UMovieSceneControlRigParameterTrack*)> Callback)
+static void IterateTracksInMovieScene(const UMovieScene& InMovieScene, TFunctionRef<bool(UMovieSceneControlRigParameterTrack*)> InCallback)
 {
-	TArray<UMovieSceneControlRigParameterTrack*> Tracks;
-	
-	const TArray<FMovieSceneBinding>& Bindings = MovieScene.GetBindings();
+	const TArray<FMovieSceneBinding>& Bindings = InMovieScene.GetBindings();
 	for (const FMovieSceneBinding& Binding : Bindings)
 	{
-		TArray<UMovieSceneTrack*> FoundTracks = MovieScene.FindTracks(UMovieSceneControlRigParameterTrack::StaticClass(), Binding.GetObjectGuid(), NAME_None);
+		TArray<UMovieSceneTrack*> FoundTracks = InMovieScene.FindTracks(UMovieSceneControlRigParameterTrack::StaticClass(), Binding.GetObjectGuid(), NAME_None);
 		for(UMovieSceneTrack* Track : FoundTracks)
 		{
 			if (UMovieSceneControlRigParameterTrack* CRTrack = Cast<UMovieSceneControlRigParameterTrack>(Track))
 			{
-				Callback(CRTrack);
+				InCallback(CRTrack);
 			}
 		}
 	}
 	
-	for (UMovieSceneTrack* Track : MovieScene.GetTracks())
+	for (UMovieSceneTrack* Track : InMovieScene.GetTracks())
 	{
 		if (UMovieSceneControlRigParameterTrack* CRTrack = Cast<UMovieSceneControlRigParameterTrack>(Track))
 		{
-			Callback(CRTrack);
+			InCallback(CRTrack);
 		}
 	}
 }
 
 void FControlRigParameterTrackEditor::IterateTracks(TFunctionRef<bool(UMovieSceneControlRigParameterTrack*)> Callback) const
 {
-	UMovieScene* MovieScene = GetSequencer().IsValid() && GetSequencer()->GetFocusedMovieSceneSequence() ? GetSequencer()->GetFocusedMovieSceneSequence()->GetMovieScene() : nullptr;
+	const TSharedPtr<ISequencer> Sequencer = GetSequencer();
+	UMovieSceneSequence* OwnerSequence = Sequencer ? Sequencer->GetFocusedMovieSceneSequence() : nullptr;
+	UMovieScene* MovieScene = OwnerSequence ? OwnerSequence->GetMovieScene() : nullptr;
 	if (!MovieScene)
 	{
 		return;
@@ -1384,33 +1434,36 @@ void FControlRigParameterTrackEditor::BakeInvertedPose(UControlRig* InControlRig
 		return;
 	}
 
-	UMovieSceneControlRigParameterSection* Section = Cast<UMovieSceneControlRigParameterSection>(Track->GetSectionToKey());
+	const TSharedPtr<ISequencer> Sequencer = GetSequencer();
+	if (!Sequencer)
+	{
+		return;
+	}
+
 	USkeletalMeshComponent* SkelMeshComp = Cast<USkeletalMeshComponent>(InControlRig->GetObjectBinding()->GetBoundObject());
-	UMovieSceneSequence* RootMovieSceneSequence = GetSequencer()->GetRootMovieSceneSequence();
-	UMovieSceneSequence* MovieSceneSequence = GetSequencer()->GetFocusedMovieSceneSequence();
-	UMovieScene* MovieScene = MovieSceneSequence->GetMovieScene();
+	UMovieSceneSequence* RootMovieSceneSequence = Sequencer->GetRootMovieSceneSequence();
+	UMovieSceneSequence* MovieSceneSequence = Sequencer->GetFocusedMovieSceneSequence();
+	UMovieScene* MovieScene = MovieSceneSequence ? MovieSceneSequence->GetMovieScene() : nullptr;
 	UAnimSeqExportOption* ExportOptions = NewObject<UAnimSeqExportOption>(GetTransientPackage(), NAME_None);
 
-	//@sara to do, not sure if you want to key reduce after, but BakeSettings isn't used
-	//UBakeToControlRigSettings* BakeSettings = GetMutableDefault<UBakeToControlRigSettings>();
-	const TSharedPtr<ISequencer> ParentSequencer = GetSequencer();
-	FMovieSceneSequenceIDRef Template = ParentSequencer->GetFocusedTemplateID();
-	FMovieSceneSequenceTransform RootToLocalTransform = ParentSequencer->GetFocusedMovieSceneSequenceTransform();
-
 	if (ExportOptions == nullptr || MovieScene == nullptr || SkelMeshComp == nullptr)
 	{
 		UE_LOG(LogMovieScene, Error, TEXT("FControlRigParameterTrackEditor::BakeInvertedPose All parameters must be valid."));
 		return;
 	}
-	
+
+	//@sara to do, not sure if you want to key reduce after, but BakeSettings isn't used
+	//UBakeToControlRigSettings* BakeSettings = GetMutableDefault<UBakeToControlRigSettings>();
+	FMovieSceneSequenceTransform RootToLocalTransform = Sequencer->GetFocusedMovieSceneSequenceTransform();
+
 	const FScopedTransaction Transaction(LOCTEXT("BakeInvertedPose_Transaction", "Bake Inverted Pose"));
 
 	UnFbx::FLevelSequenceAnimTrackAdapter::FAnimTrackSettings Settings;
-	Settings.MovieScenePlayer = ParentSequencer.Get();
+	Settings.MovieScenePlayer = Sequencer.Get();
 	Settings.MovieSceneSequence = MovieSceneSequence;
 	Settings.RootMovieSceneSequence = RootMovieSceneSequence;
 	Settings.RootToLocalTransform = RootToLocalTransform;
-	Settings.bForceUseOfMovieScenePlaybackRange = ParentSequencer.Get()->GetSequencerSettings()->ShouldEvaluateSubSequencesInIsolation();
+	Settings.bForceUseOfMovieScenePlaybackRange = Sequencer->GetSequencerSettings()->ShouldEvaluateSubSequencesInIsolation();
 
 	UnFbx::FLevelSequenceAnimTrackAdapter AnimTrackAdapter(Settings);
 	int32 AnimationLength = AnimTrackAdapter.GetLength();
@@ -1422,38 +1475,42 @@ void FControlRigParameterTrackEditor::BakeInvertedPose(UControlRig* InControlRig
 		InControlRig->InvertInputPose();
 	});
 
-	auto KeyFrame = [this, ParentSequencer, InControlRig, SkelMeshComp](const FFrameNumber FrameNumber)
+	auto KeyFrame = [this, Sequencer, InControlRig, SkelMeshComp](const FFrameNumber FrameNumber)
 	{
-		const FFrameNumber NewTime = ConvertFrameTime(FrameNumber, ParentSequencer->GetFocusedDisplayRate(), ParentSequencer->GetFocusedTickResolution()).FrameNumber;
-		float LocalTime = ParentSequencer->GetFocusedTickResolution().AsSeconds(FFrameTime(NewTime));
+		const FFrameNumber NewTime = ConvertFrameTime(FrameNumber, Sequencer->GetFocusedDisplayRate(), Sequencer->GetFocusedTickResolution()).FrameNumber;
+		float LocalTime = Sequencer->GetFocusedTickResolution().AsSeconds(FFrameTime(NewTime));
 
 		AddControlKeys(SkelMeshComp, InControlRig, InControlRig->GetFName(), NAME_None, EControlRigContextChannelToKey::AllTransform, 
 				ESequencerKeyMode::ManualKeyForced, LocalTime);
 	};
 
 	FInitAnimationCB InitCallback = FInitAnimationCB::CreateLambda([]{});
+
 	FStartAnimationCB StartCallback = FStartAnimationCB::CreateLambda([AnimTrackAdapter, KeyFrame](FFrameNumber FrameNumber)
 	{
 		KeyFrame(AnimTrackAdapter.GetLocalStartFrame());
 	});
+	
 	FTickAnimationCB TickCallback = FTickAnimationCB::CreateLambda([KeyFrame, &Progress](float DeltaTime, FFrameNumber FrameNumber)
 	{
 		KeyFrame(FrameNumber);
 		Progress.EnterProgressFrame(1);
 	});
+	
 	FEndAnimationCB EndCallback = FEndAnimationCB::CreateLambda([]{});
 
 	FAnimExportSequenceParameters AESP;
-	AESP.Player = ParentSequencer.Get();
+	AESP.Player = Sequencer.Get();
 	AESP.RootToLocalTransform = RootToLocalTransform;
 	AESP.MovieSceneSequence = MovieSceneSequence;
 	AESP.RootMovieSceneSequence = RootMovieSceneSequence;
-	AESP.bForceUseOfMovieScenePlaybackRange = ParentSequencer.Get()->GetSequencerSettings()->ShouldEvaluateSubSequencesInIsolation();
+	AESP.bForceUseOfMovieScenePlaybackRange = Sequencer->GetSequencerSettings()->ShouldEvaluateSubSequencesInIsolation();
+	
 	MovieSceneToolHelpers::BakeToSkelMeshToCallbacks(AESP,SkelMeshComp, ExportOptions,
 		InitCallback, StartCallback, TickCallback, EndCallback);
 
 	InControlRig->OnPreAdditiveValuesApplication_AnyThread().Remove(DelegateHandle);
-	ParentSequencer->NotifyMovieSceneDataChanged(EMovieSceneDataChangeType::TrackValueChanged);
+	Sequencer->NotifyMovieSceneDataChanged(EMovieSceneDataChangeType::TrackValueChanged);
 }
 
 bool FControlRigParameterTrackEditor::IsLayered(UMovieSceneControlRigParameterTrack* Track) const
@@ -1513,7 +1570,10 @@ void FControlRigParameterTrackEditor::BuildObjectBindingTrackMenu(FMenuBuilder&
 		if (Skeleton)
 		{
 			UMovieSceneTrack* Track = nullptr;
-			MenuBuilder.AddSubMenu(LOCTEXT("ControlRigText", "Control Rig"), FText(), FNewMenuDelegate::CreateSP(this, &FControlRigParameterTrackEditor::HandleAddTrackSubMenu, ObjectBindings, Track));
+			MenuBuilder.AddSubMenu(
+				LOCTEXT("ControlRigText", "Control Rig"),
+				FText(),
+				FNewMenuDelegate::CreateSP(this, &FControlRigParameterTrackEditor::HandleAddTrackSubMenu, ObjectBindings, Track));
 		}
 	}
 }
@@ -1829,38 +1889,6 @@ bool FControlRigParameterTrackEditor::IsControlRigAllowed(const FAssetData& Asse
 	return true;	
 }
 
-/*
-bool FControlRigParameterTrackEditor::ShouldFilterAsset(const FAssetData& AssetData)
-{
-	if (AssetData.AssetClassPath == UControlRig::StaticClass()->GetClassPathName())
-	{
-		return true;
-	}
-	return false;
-}
-
-void FControlRigParameterTrackEditor::OnControlRigAssetSelected(const FAssetData& AssetData, TArray<FGuid> ObjectBindings, USkeleton* Skeleton)
-{
-	FSlateApplication::Get().DismissAllMenus();
-
-	UObject* SelectedObject = AssetData.GetAsset();todo
-	TSharedPtr<ISequencer> SequencerPtr = GetSequencer();
-
-	if (SelectedObject && SelectedObject->IsA(UControlRig::StaticClass()) && SequencerPtr.IsValid())
-	{
-		UControlRig* ControlRig = CastChecked<UControlRig>(AssetData.GetAsset());
-
-		const FScopedTransaction Transaction(LOCTEXT("AddAnimation_Transaction", "Add Animation"));
-
-		for (FGuid ObjectBinding : ObjectBindings)
-		{
-			UObject* Object = SequencerPtr->FindSpawnedObjectOrTemplate(ObjectBinding);
-			int32 RowIndex = INDEX_NONE;
-			//AnimatablePropertyChanged(FOnKeyProperty::CreateRaw(this, &FControlRigParameterTrackEditor::AddControlRig, BoundObject, ObjectBindings[0], Skeleton)));
-		}
-	}
-*/
-
 static TArray<UMovieSceneControlRigParameterTrack*> GetExistingControlRigTracksForSkeletalMeshComponent(UMovieScene* MovieScene, USkeletalMeshComponent* SkeletalMeshComponent)
 {
 	TArray<UMovieSceneControlRigParameterTrack*> ExistingControlRigTracks;
@@ -2056,13 +2084,19 @@ static UMovieSceneTrack* FindOrCreateControlRigTrack(TSharedPtr<ISequencer>& Seq
 
 void FControlRigParameterTrackEditor::AddControlRig(const UClass* InClass, UObject* BoundActor, FGuid ObjectBinding, UControlRig* InExistingControlRig)
 {
+	TSharedPtr<ISequencer> Sequencer = GetSequencer();
+	if (!Sequencer)
+	{
+		return;
+	}
+	
 	UWorld* World = GCurrentLevelEditingViewportClient ? GCurrentLevelEditingViewportClient->GetWorld() : nullptr;
-	UMovieSceneSequence* Sequence = GetSequencer()->GetFocusedMovieSceneSequence();
+	UMovieSceneSequence* Sequence = Sequencer->GetFocusedMovieSceneSequence();
 	FMovieSceneBindingProxy BindingProxy(ObjectBinding, Sequence);
-	TSharedPtr<ISequencer> SequencerPtr = GetSequencer();
+
 	//UControlRigSequencerEditorLibrary::FindOrCreateControlRigTrack.. in 5.5 we will redo this but for 
 	//5.4.4 we can't change headers so for now we just make the change here locally
-	if (UMovieSceneTrack* Track = FindOrCreateControlRigTrack(SequencerPtr, World, InClass, BindingProxy, bIsLayeredControlRig))
+	if (UMovieSceneTrack* Track = FindOrCreateControlRigTrack(Sequencer, World, InClass, BindingProxy, bIsLayeredControlRig))
 	{
 		BindControlRig(CastChecked<UMovieSceneControlRigParameterTrack>(Track)->GetControlRig());
 	}
@@ -2101,8 +2135,8 @@ void FControlRigParameterTrackEditor::AddControlRig(UClass* InClass, UObject* Bo
 //This now adds all of the control rig components, not just the first one
 void FControlRigParameterTrackEditor::AddControlRigFromComponent(FGuid InGuid)
 {
-	const TSharedPtr<ISequencer> ParentSequencer = GetSequencer();
-	UObject* BoundObject = ParentSequencer.IsValid() ? ParentSequencer->FindSpawnedObjectOrTemplate(InGuid) : nullptr;
+	const TSharedPtr<ISequencer> Sequencer = GetSequencer();
+	UObject* BoundObject = Sequencer ? Sequencer->FindSpawnedObjectOrTemplate(InGuid) : nullptr;
 
 	if (AActor* BoundActor = Cast<AActor>(BoundObject))
 	{
@@ -2115,7 +2149,6 @@ void FControlRigParameterTrackEditor::AddControlRigFromComponent(FGuid InGuid)
 				AddControlRig(CR->GetClass(), BoundActor, InGuid, CR);
 			}
 		}
-
 	}
 }
 
@@ -2136,7 +2169,8 @@ bool FControlRigParameterTrackEditor::CanAddTransformKeysForSelectedObjects() co
 		}
 	}
 
-	if (!GetSequencer()->IsAllowedToChange())
+	const TSharedPtr<ISequencer> Sequencer = GetSequencer();
+	if (!Sequencer || !Sequencer->IsAllowedToChange())
 	{
 		return false;
 	}
@@ -2147,13 +2181,15 @@ bool FControlRigParameterTrackEditor::CanAddTransformKeysForSelectedObjects() co
 		TMap<UControlRig*, TArray<FRigElementKey>> SelectedControls;
 		ControlRigEditMode->GetAllSelectedControls(SelectedControls);
 		return (SelectedControls.Num() > 0);
-		}
+	}
+	
 	return false;
 }
 
 void FControlRigParameterTrackEditor::OnAddTransformKeysForSelectedObjects(EMovieSceneTransformChannel Channel)
 {
-	if (!GetSequencer()->IsAllowedToChange())
+	const TSharedPtr<ISequencer> Sequencer = GetSequencer();
+	if (!Sequencer || !Sequencer->IsAllowedToChange())
 	{
 		return;
 	}
@@ -2371,10 +2407,14 @@ void FControlRigParameterTrackEditor::OnChannelChanged(const FMovieSceneChannelM
 {
 	static FChannelChangedStruct ChannelChanged;
 
-	UMovieSceneControlRigParameterSection* Section = Cast<UMovieSceneControlRigParameterSection>(InSection);
-	TSharedPtr<ISequencer> SequencerPtr = GetSequencer();
+	const TSharedPtr<ISequencer> Sequencer = GetSequencer();
+	if (!Sequencer)
+	{
+		return;
+	}
 
-	if (Section && Section->GetControlRig() && MetaData && SequencerPtr.IsValid())
+	UMovieSceneControlRigParameterSection* Section = Cast<UMovieSceneControlRigParameterSection>(InSection);
+	if (Section && Section->GetControlRig() && MetaData)
 	{
 		TArray<FString> StringArray;
 		FString String = MetaData->Name.ToString();
@@ -2388,8 +2428,8 @@ void FControlRigParameterTrackEditor::OnChannelChanged(const FMovieSceneChannelM
 				ChannelChanged.SectionControlNames.Reset();
 				auto ChannelChangedDelegate = [this]()
 				{
-					TSharedPtr<ISequencer> SequencerPtr = GetSequencer();
-					if (!SequencerPtr.IsValid())
+					const TSharedPtr<ISequencer> Sequencer = GetSequencer();
+					if (!Sequencer)
 					{
 						GEditor->GetTimerManager()->ClearTimer(ChannelChanged.TimerHandle);
 						return;
@@ -2399,12 +2439,12 @@ void FControlRigParameterTrackEditor::OnChannelChanged(const FMovieSceneChannelM
 					{
 						if (ChannelChanged.TimerHandle.IsValid())
 						{
-							FFrameTime Time = SequencerPtr->GetLocalTime().Time;
+							FFrameTime Time = Sequencer->GetLocalTime().Time;
 							GEditor->GetTimerManager()->ClearTimer(ChannelChanged.TimerHandle);
 							ChannelChanged.TimerHandle.Invalidate();
 							ChannelChanged.bWasSetAlready = false;
 							TOptional<FFrameNumber> Optional;	
-							ISequencer* SequencerRaw = SequencerPtr.Get();
+
 							for (TPair <UMovieSceneControlRigParameterSection*, TSet<FName>>& Pair : ChannelChanged.SectionControlNames)
 							{
 								for (FName& ControlName : Pair.Value)
@@ -2458,29 +2498,32 @@ void FControlRigParameterTrackEditor::AddTrackForComponent(USceneComponent* InCo
 			UObject* Object = SkelMeshComp->GetDefaultAnimatingRig().LoadSynchronous();
 			if (Object != nullptr && (Object->IsA<UControlRigBlueprint>() || Object->IsA<UControlRigComponent>() || Object->IsA<URigVMBlueprintGeneratedClass>()))
 			{
-				FGuid Binding = InBinding.IsValid() ? InBinding : GetSequencer()->GetHandleToObject(InComponent, true /*bCreateHandle*/);
-				if (Binding.IsValid())
+				if (TSharedPtr<ISequencer> Sequencer = GetSequencer())
 				{
-					UMovieSceneSequence* OwnerSequence = GetSequencer()->GetFocusedMovieSceneSequence();
-					UMovieScene* OwnerMovieScene = OwnerSequence->GetMovieScene();
-					UMovieSceneControlRigParameterTrack* Track = Cast<UMovieSceneControlRigParameterTrack>(OwnerMovieScene->FindTrack(UMovieSceneControlRigParameterTrack::StaticClass(), Binding, NAME_None));
-					if (Track == nullptr)
+					FGuid Binding = InBinding.IsValid() ? InBinding : Sequencer->GetHandleToObject(InComponent, true /*bCreateHandle*/);
+					if (Binding.IsValid())
 					{
-						URigVMBlueprintGeneratedClass* RigClass = nullptr;
-						if (UControlRigBlueprint* BPControlRig = Cast<UControlRigBlueprint>(Object))
+						UMovieSceneSequence* OwnerSequence = Sequencer->GetFocusedMovieSceneSequence();
+						UMovieScene* OwnerMovieScene = OwnerSequence->GetMovieScene();
+						UMovieSceneControlRigParameterTrack* Track = Cast<UMovieSceneControlRigParameterTrack>(OwnerMovieScene->FindTrack(UMovieSceneControlRigParameterTrack::StaticClass(), Binding, NAME_None));
+						if (Track == nullptr)
 						{
-							RigClass = BPControlRig->GetRigVMBlueprintGeneratedClass();
-						}
-						else
-						{
-							RigClass = Cast<URigVMBlueprintGeneratedClass>(Object);
-						}
+							URigVMBlueprintGeneratedClass* RigClass = nullptr;
+							if (UControlRigBlueprint* BPControlRig = Cast<UControlRigBlueprint>(Object))
+							{
+								RigClass = BPControlRig->GetRigVMBlueprintGeneratedClass();
+							}
+							else
+							{
+								RigClass = Cast<URigVMBlueprintGeneratedClass>(Object);
+							}
 
-						if (RigClass)
-						{
-							if (UControlRig* CDO = Cast<UControlRig>(RigClass->GetDefaultObject(true /* create if needed */)))
+							if (RigClass)
 							{
-								AddControlRig(CDO->GetClass(), InComponent, Binding);
+								if (UControlRig* CDO = Cast<UControlRig>(RigClass->GetDefaultObject(true /* create if needed */)))
+								{
+									AddControlRig(CDO->GetClass(), InComponent, Binding);
+								}
 							}
 						}
 					}
@@ -2488,6 +2531,7 @@ void FControlRigParameterTrackEditor::AddTrackForComponent(USceneComponent* InCo
 			}
 		}
 	}
+	
 	TArray<USceneComponent*> ChildComponents;
 	InComponent->GetChildrenComponents(false, ChildComponents);
 	for (USceneComponent* ChildComponent : ChildComponents)
@@ -2506,19 +2550,25 @@ void FControlRigParameterTrackEditor::HandleActorAdded(AActor* Actor, FGuid Targ
 	}
 
 	//test for constraint
-	const FConstraintsManagerController& Controller = FConstraintsManagerController::Get(Actor->GetWorld());
-	const TArray< TWeakObjectPtr<UTickableConstraint>> Constraints = Controller.GetAllConstraints();
-	for (const TWeakObjectPtr<UTickableConstraint>& WeakConstraint : Constraints)
+	if (bAutoGenerateControlRigTrack)
 	{
-		if (UTickableTransformConstraint* Constraint = WeakConstraint.IsValid() ? Cast<UTickableTransformConstraint>(WeakConstraint.Get()) : nullptr)
+		if (TSharedPtr<ISequencer> Sequencer = GetSequencer())
 		{
-			if (UObject* Child = Constraint->ChildTRSHandle ? Constraint->ChildTRSHandle->GetTarget().Get() : nullptr)
+			const FConstraintsManagerController& Controller = FConstraintsManagerController::Get(Actor->GetWorld());
+			const TArray< TWeakObjectPtr<UTickableConstraint>> Constraints = Controller.GetAllConstraints();
+			for (const TWeakObjectPtr<UTickableConstraint>& WeakConstraint : Constraints)
 			{
-				const AActor* TargetActor = Child->IsA<AActor>() ? Cast<AActor>(Child) : Child->GetTypedOuter<AActor>();
-				if (TargetActor == Actor)
+				if (UTickableTransformConstraint* Constraint = WeakConstraint.IsValid() ? Cast<UTickableTransformConstraint>(WeakConstraint.Get()) : nullptr)
 				{
-					AddConstraintToSequencer(GetSequencer(), Constraint);
-				}		
+					if (UObject* Child = Constraint->ChildTRSHandle ? Constraint->ChildTRSHandle->GetTarget().Get() : nullptr)
+					{
+						const AActor* TargetActor = Child->IsA<AActor>() ? Cast<AActor>(Child) : Child->GetTypedOuter<AActor>();
+						if (TargetActor == Actor)
+						{
+							AddConstraintToSequencer(Sequencer, Constraint);
+						}		
+					}
+				}
 			}
 		}
 	}
@@ -2562,7 +2612,9 @@ void FControlRigParameterTrackEditor::OnActivateSequenceChanged(FMovieSceneSeque
 		
 		return false;
 	});
-	if (bControlRigEditModeWasOpen && GetSequencer()->IsLevelEditorSequencer())
+
+	TSharedPtr<ISequencer> Sequencer = GetSequencer();
+	if (bControlRigEditModeWasOpen && Sequencer && Sequencer->IsLevelEditorSequencer())
 	{
 		GEditor->GetTimerManager()->SetTimerForNextTick([this]()
 		{
@@ -2577,6 +2629,7 @@ void FControlRigParameterTrackEditor::OnActivateSequenceChanged(FMovieSceneSeque
 					EditorModeTools->Tick(EditorViewClient, 0.033f);
 				}
 			}
+			
 			GEditor->GetTimerManager()->SetTimerForNextTick([this]()
 			{
 				//now we can recreate it
@@ -2584,12 +2637,15 @@ void FControlRigParameterTrackEditor::OnActivateSequenceChanged(FMovieSceneSeque
 				if (ControlRigEditMode)
 				{
 					bool bSequencerSet = false;
+
+					TSharedPtr<ISequencer> Sequencer = GetSequencer();
 					for (TWeakObjectPtr<UControlRig>& ControlRig : BoundControlRigs)
 					{
 						if (ControlRig.IsValid())
 						{
-							ControlRigEditMode->AddControlRigObject(ControlRig.Get(), GetSequencer());
+							ControlRigEditMode->AddControlRigObject(ControlRig.Get(), Sequencer);
 							bSequencerSet = true;
+							
 							for (int32 Index = 0; Index < PreviousSelectedControlRigs.Num(); ++Index)
 							{
 								if (ControlRig.Get()->GetClass() == PreviousSelectedControlRigs[Index].Key)
@@ -2604,9 +2660,10 @@ void FControlRigParameterTrackEditor::OnActivateSequenceChanged(FMovieSceneSeque
 							}
 						}
 					}
-					if (bSequencerSet == false)
+
+					if (!bSequencerSet)
 					{
-						ControlRigEditMode->SetSequencer(GetSequencer());
+						ControlRigEditMode->SetSequencer(Sequencer);
 					}
 				
 				}
@@ -2619,16 +2676,29 @@ void FControlRigParameterTrackEditor::OnActivateSequenceChanged(FMovieSceneSeque
 
 void FControlRigParameterTrackEditor::OnSequencerDataChanged(EMovieSceneDataChangeType DataChangeType)
 {
-	const UMovieScene* MovieScene = GetSequencer()->GetFocusedMovieSceneSequence()->GetMovieScene();
+	const TSharedPtr<ISequencer> Sequencer = GetSequencer();
+	UMovieSceneSequence* OwnerSequence = Sequencer ? Sequencer->GetFocusedMovieSceneSequence() : nullptr;
+	UMovieScene* MovieScene = OwnerSequence ? OwnerSequence->GetMovieScene() : nullptr;
+
+	if (!MovieScene)
+	{
+		return;
+	}
+
 	FControlRigEditMode* ControlRigEditMode = GetEditMode();
+	static const TArray<UControlRig*> NoRigs;
+	const TArray<UControlRig*>& ControlRigs = ControlRigEditMode ? ControlRigEditMode->GetControlRigsArray(false /*bIsVisible*/) : NoRigs;
 
+	if (ControlRigs.IsEmpty())
+	{
+		return;
+	}
+	
 	//if we have a valid control rig edit mode need to check and see the control rig in that mode is still in a track
 	//if not we get rid of it.
-	if (ControlRigEditMode && ControlRigEditMode->GetControlRigsArray(false /*bIsVisible*/).Num() != 0 && MovieScene && (DataChangeType == EMovieSceneDataChangeType::MovieSceneStructureItemRemoved ||
-		DataChangeType == EMovieSceneDataChangeType::Unknown))
+	if (DataChangeType == EMovieSceneDataChangeType::MovieSceneStructureItemRemoved || DataChangeType == EMovieSceneDataChangeType::Unknown)
 	{
-		TArray<UControlRig*> ControlRigs = ControlRigEditMode->GetControlRigsArray(false /*bIsVisible*/);
-		const float FPS = (float)GetSequencer()->GetFocusedDisplayRate().AsDecimal();
+		const float FPS = static_cast<float>(Sequencer->GetFocusedDisplayRate().AsDecimal());
 		for (UControlRig* ControlRig : ControlRigs)
 		{
 			if (ControlRig)
@@ -2647,7 +2717,7 @@ void FControlRigParameterTrackEditor::OnSequencerDataChanged(EMovieSceneDataChan
 					return true;
 				});
 
-				if (bControlRigInTrack == false)
+				if (!bControlRigInTrack)
 				{
 					ControlRigEditMode->RemoveControlRig(ControlRig);
 				}
@@ -2680,7 +2750,13 @@ void FControlRigParameterTrackEditor::PostEvaluation(UMovieScene* MovieScene, FF
 
 void FControlRigParameterTrackEditor::OnSelectionChanged(TArray<UMovieSceneTrack*> InTracks)
 {
-	if (bIsDoingSelection || GetSequencer().IsValid() == false)
+	if (bIsDoingSelection)
+	{
+		return;
+	}
+
+	TSharedPtr<ISequencer> Sequencer = GetSequencer();
+	if (!Sequencer)
 	{
 		return;
 	}
@@ -2699,7 +2775,8 @@ void FControlRigParameterTrackEditor::OnSelectionChanged(TArray<UMovieSceneTrack
 
 	TArray<const IKeyArea*> KeyAreas;
 	const bool UseSelectedKeys = CVarSelectedKeysSelectControls.GetValueOnGameThread();
-	GetSequencer()->GetSelectedKeyAreas(KeyAreas, UseSelectedKeys);
+	Sequencer->GetSelectedKeyAreas(KeyAreas, UseSelectedKeys);
+	
 	if (KeyAreas.Num() <= 0)
 	{ 
 		if (FSlateApplication::Get().GetModifierKeys().IsShiftDown() == false && 
@@ -2717,6 +2794,7 @@ void FControlRigParameterTrackEditor::OnSelectionChanged(TArray<UMovieSceneTrack
 				}
 			}
 		}
+
 		for (UMovieSceneTrack* Track : InTracks)
 		{
 			UMovieSceneControlRigParameterTrack* CRTrack = Cast<UMovieSceneControlRigParameterTrack>(Track);
@@ -2727,7 +2805,7 @@ void FControlRigParameterTrackEditor::OnSelectionChanged(TArray<UMovieSceneTrack
 				{
 					if (ControlRigEditMode)
 					{
-						const bool bAdded = ControlRigEditMode->AddControlRigObject(TrackControlRig, GetSequencer());
+						const bool bAdded = ControlRigEditMode->AddControlRigObject(TrackControlRig, Sequencer);
 						if (bAdded)
 						{
 							ControlRigEditMode->RequestToRecreateControlShapeActors(TrackControlRig);
@@ -2741,7 +2819,7 @@ void FControlRigParameterTrackEditor::OnSelectionChanged(TArray<UMovieSceneTrack
 						{
 							if (ControlRigEditMode)
 							{
-								const bool bAdded = ControlRigEditMode->AddControlRigObject(TrackControlRig, GetSequencer());
+								const bool bAdded = ControlRigEditMode->AddControlRigObject(TrackControlRig, Sequencer);
 								if (bAdded)
 								{
 									ControlRigEditMode->RequestToRecreateControlShapeActors(TrackControlRig);
@@ -2752,12 +2830,14 @@ void FControlRigParameterTrackEditor::OnSelectionChanged(TArray<UMovieSceneTrack
 				}
 			}
 		}
+		
 		const bool bSelectedSectionSetsSectionToKey = CVarSelectedSectionSetsSectionToKey.GetValueOnGameThread();
 		if (bSelectedSectionSetsSectionToKey)
 		{
 			TMap<UMovieSceneTrack*, TSet<UMovieSceneControlRigParameterSection*>> TracksAndSections;
 			using namespace UE::Sequencer;
-			for (FViewModelPtr ViewModel : GetSequencer()->GetViewModel()->GetSelection()->Outliner)
+
+			for (FViewModelPtr ViewModel : Sequencer->GetViewModel()->GetSelection()->Outliner)
 			{
 				if (TViewModelPtr<FTrackRowModel> TrackRowModel = ViewModel.ImplicitCast())
 				{
@@ -2786,12 +2866,13 @@ void FControlRigParameterTrackEditor::OnSelectionChanged(TArray<UMovieSceneTrack
 		}
 		return;
 	}
+	
 	SelectRigsAndControls(ControlRig, KeyAreas);
 
 	// If the edit mode has been activated, we need to synchronize the external selection (possibly again to account for control rig control actors selection)
 	if (!bEditModeExisted && GetEditMode() != nullptr)
 	{
-		FSequencerUtilities::SynchronizeExternalSelectionWithSequencerSelection(GetSequencer().ToSharedRef());
+		FSequencerUtilities::SynchronizeExternalSelectionWithSequencerSelection(Sequencer.ToSharedRef());
 	}
 	
 }
@@ -2804,6 +2885,8 @@ void FControlRigParameterTrackEditor::SelectRigsAndControls(UControlRig* Control
 	const bool bSelectedSectionSetsSectionToKey = CVarSelectedSectionSetsSectionToKey.GetValueOnGameThread();
 	TMap<UMovieSceneTrack*, TSet<UMovieSceneControlRigParameterSection*>> TracksAndSections;
 
+	const TSharedPtr<ISequencer> Sequencer = GetSequencer();
+	
 	TArray<FString> StringArray;
 	//we have two sets here one to see if selection has really changed that contains the attirbutes, the other to select just the parent
 	TMap<UControlRig*, TSet<FName>> RigsAndControls;
@@ -2821,18 +2904,18 @@ void FControlRigParameterTrackEditor::SelectRigsAndControls(UControlRig* Control
 				{
 					if (ControlRigEditMode)
 					{
-						ControlRigEditMode->AddControlRigObject(ControlRig, GetSequencer());
+						ControlRigEditMode->AddControlRigObject(ControlRig, Sequencer);
 					}
 				}
 			}
 			else
 			{
-				if (ControlRigEditMode->AddControlRigObject(ControlRig, GetSequencer()))
+				if (ControlRigEditMode->AddControlRigObject(ControlRig, Sequencer))
 				{
 					//force an evaluation, this will get the control rig setup so edit mode looks good.
-					if (GetSequencer().IsValid())
+					if (Sequencer)
 					{
-						GetSequencer()->NotifyMovieSceneDataChanged(EMovieSceneDataChangeType::Unknown);
+						Sequencer->NotifyMovieSceneDataChanged(EMovieSceneDataChangeType::Unknown);
 					}
 				}
 			}
@@ -2879,6 +2962,7 @@ void FControlRigParameterTrackEditor::SelectRigsAndControls(UControlRig* Control
 					}
 				}
 			}
+			
 			if (bSelectedSectionSetsSectionToKey)
 			{
 				if (UMovieSceneTrack* Track = MovieSection->GetTypedOuter<UMovieSceneTrack>())
@@ -2998,16 +3082,24 @@ void FControlRigParameterTrackEditor::SelectRigsAndControls(UControlRig* Control
 
 FMovieSceneTrackEditor::FFindOrCreateHandleResult FControlRigParameterTrackEditor::FindOrCreateHandleToObject(UObject* InObj, UControlRig* InControlRig)
 {
-	const bool bCreateHandleIfMissing = false;
+	FFindOrCreateHandleResult Result;
+	Result.bWasCreated = false;
+	
+	TSharedPtr<ISequencer> Sequencer = GetSequencer();
+	if (!Sequencer)
+	{
+		return Result;
+	}
+	
+	constexpr bool bCreateHandleIfMissing = false;
 	FName CreatedFolderName = NAME_None;
 
-	FFindOrCreateHandleResult Result;
-	bool bHandleWasValid = GetSequencer()->GetHandleToObject(InObj, bCreateHandleIfMissing).IsValid();
+	bool bHandleWasValid = Sequencer->GetHandleToObject(InObj, bCreateHandleIfMissing).IsValid();
 
-	Result.Handle = GetSequencer()->GetHandleToObject(InObj, bCreateHandleIfMissing, CreatedFolderName);
+	Result.Handle = Sequencer->GetHandleToObject(InObj, bCreateHandleIfMissing, CreatedFolderName);
 	Result.bWasCreated = bHandleWasValid == false && Result.Handle.IsValid();
 
-	const UMovieScene* MovieScene = GetSequencer()->GetFocusedMovieSceneSequence()->GetMovieScene();
+	const UMovieScene* MovieScene = Sequencer->GetFocusedMovieSceneSequence()->GetMovieScene();
 
 	// Prioritize a control rig parameter track on this component if it matches the handle
 	if (Result.Handle.IsValid())
@@ -3026,7 +3118,7 @@ FMovieSceneTrackEditor::FFindOrCreateHandleResult FControlRigParameterTrackEdito
 	{
 		// If the owner has a control rig parameter track, let's use it
 		UObject* OwnerObject = SceneComponent->GetOwner();
-		const FGuid OwnerHandle = GetSequencer()->GetHandleToObject(OwnerObject, bCreateHandleIfMissing);
+		const FGuid OwnerHandle = Sequencer->GetHandleToObject(OwnerObject, bCreateHandleIfMissing);
 	    bHandleWasValid = OwnerHandle.IsValid();
 	    if (OwnerHandle.IsValid())
 	    {
@@ -3053,6 +3145,12 @@ FMovieSceneTrackEditor::FFindOrCreateHandleResult FControlRigParameterTrackEdito
 
 void FControlRigParameterTrackEditor::SelectSequencerNodeInSection(UMovieSceneControlRigParameterSection* ParamSection, const FName& ControlName, bool bSelected)
 {
+	const TSharedPtr<ISequencer> Sequencer = GetSequencer();
+	if (!Sequencer)
+	{
+		return;
+	}
+	
 	if (ParamSection)
 	{
 		FChannelMapInfo* pChannelIndex = ParamSection->ControlChannelMap.Find(ControlName);
@@ -3063,7 +3161,7 @@ void FControlRigParameterTrackEditor::SelectSequencerNodeInSection(UMovieSceneCo
 				int32 CategoryIndex = ParamSection->GetActiveCategoryIndex(ControlName);
 				if (CategoryIndex != INDEX_NONE)
 				{
-					GetSequencer()->SelectByNthCategoryNode(ParamSection, CategoryIndex, bSelected);
+					Sequencer->SelectByNthCategoryNode(ParamSection, CategoryIndex, bSelected);
 				}
 			}
 			else
@@ -3079,7 +3177,7 @@ void FControlRigParameterTrackEditor::SelectSequencerNodeInSection(UMovieSceneCo
 						FMovieSceneChannelHandle Channel = ChannelProxy.MakeHandle(ChannelTypeName, pChannelIndex->ChannelIndex);
 						TArray<FMovieSceneChannelHandle> Channels;
 						Channels.Add(Channel);
-						GetSequencer()->SelectByChannels(ParamSection, Channels, false, bSelected);
+						Sequencer->SelectByChannels(ParamSection, Channels, false, bSelected);
 						break;
 					}
 				}
@@ -3106,8 +3204,11 @@ FMovieSceneTrackEditor::FFindOrCreateTrackResult FControlRigParameterTrackEditor
 	// Only create track if the object handle is valid
 	if (!Result.Track && bCreateTrackIfMissing && ObjectHandle.IsValid())
 	{
-		UMovieScene* MovieScene = GetSequencer()->GetFocusedMovieSceneSequence()->GetMovieScene();
-		Result.Track = AddTrack(MovieScene, ObjectHandle, UMovieSceneControlRigParameterTrack::StaticClass(), PropertyName);
+		if(TSharedPtr<ISequencer> Sequencer = GetSequencer())
+		{
+			UMovieScene* MovieScene = Sequencer->GetFocusedMovieSceneSequence()->GetMovieScene();
+			Result.Track = AddTrack(MovieScene, ObjectHandle, UMovieSceneControlRigParameterTrack::StaticClass(), PropertyName);
+		}
 	}
 
 	Result.bWasCreated = bTrackExisted == false && Result.Track != nullptr;
@@ -3249,8 +3350,9 @@ void FControlRigParameterTrackEditor::HandleSpaceKeyMoved(
 void FControlRigParameterTrackEditor::ClearOutAllSpaceAndConstraintDelegates(const UControlRig* InOptionalControlRig) const
 {
 	UTickableTransformConstraint::GetOnConstraintChanged().RemoveAll(this);
-	
-	const UMovieScene* MovieScene = GetSequencer().IsValid() && GetSequencer()->GetFocusedMovieSceneSequence() ? GetSequencer()->GetFocusedMovieSceneSequence()->GetMovieScene() : nullptr;
+
+	const TSharedPtr<ISequencer> Sequencer = GetSequencer();
+	const UMovieScene* MovieScene = Sequencer && Sequencer->GetFocusedMovieSceneSequence() ? Sequencer->GetFocusedMovieSceneSequence()->GetMovieScene() : nullptr;
 	if (!MovieScene)
 	{
 		return;
@@ -3407,6 +3509,12 @@ void FControlRigParameterTrackEditor::HandleConstraintKeyDeleted(
 	{
 		return;
 	}
+
+	const TSharedPtr<ISequencer> Sequencer = GetSequencer();
+	if (!Sequencer)
+	{
+		return;
+	}
 	
 	UMovieSceneControlRigParameterSection* Section = Cast<UMovieSceneControlRigParameterSection>(InSection);
 
@@ -3416,18 +3524,13 @@ void FControlRigParameterTrackEditor::HandleConstraintKeyDeleted(
 		return;
 	}
 	
-	const FConstraintAndControlData ConstraintAndControlData =
-		FConstraintAndControlData::CreateFromSection(Section, InConstraintChannel);
+	const FConstraintAndControlData ConstraintAndControlData = FConstraintAndControlData::CreateFromSection(Section, InConstraintChannel);
 	if (ConstraintAndControlData.IsValid())
 	{
-
 		UTickableTransformConstraint* Constraint = ConstraintAndControlData.Constraint.Get();
 		for (const FKeyAddOrDeleteEventItem& EventItem: InDeletedItems)
 		{
-			FMovieSceneConstraintChannelHelper::HandleConstraintKeyDeleted(
-				Constraint, InConstraintChannel,
-				GetSequencer(), Section,
-				EventItem.Frame);
+			FMovieSceneConstraintChannelHelper::HandleConstraintKeyDeleted(Constraint, InConstraintChannel, Sequencer, Section, EventItem.Frame);
 		}
 	}
 }
@@ -3461,11 +3564,12 @@ void FControlRigParameterTrackEditor::HandleConstraintRemoved(IMovieSceneConstra
 	if (const UControlRig* ControlRig = Section->GetControlRig())
 	{
 		FConstraintsManagerController& Controller = FConstraintsManagerController::Get(ControlRig->GetWorld());
+
 		if (!InSection->OnConstraintRemovedHandle.IsValid())
 		{
 			InSection->OnConstraintRemovedHandle =
 			Controller.GetNotifyDelegate().AddLambda(
-				[InSection,Section, this](EConstraintsManagerNotifyType InNotifyType, UObject *InObject)
+				[InSection, Section, this](EConstraintsManagerNotifyType InNotifyType, UObject *InObject)
 			{
 				switch (InNotifyType)
 				{
@@ -3486,16 +3590,18 @@ void FControlRigParameterTrackEditor::HandleConstraintRemoved(IMovieSceneConstra
 								return;
 							}
 
-							TSharedPtr<ISequencer> Sequencer = GetSequencer();
-							
-							const bool bCompensate = (InNotifyType == EConstraintsManagerNotifyType::ConstraintRemovedWithCompensation);
-							if (bCompensate && ConstraintChannel->GetConstraint().Get())
+							const TSharedPtr<ISequencer> Sequencer = GetSequencer();
+							if (Sequencer)
 							{
-								FMovieSceneConstraintChannelHelper::HandleConstraintRemoved(
-									ConstraintChannel->GetConstraint().Get(),
-									&ConstraintChannel->ActiveChannel,
-									Sequencer,
-									Section);
+								const bool bCompensate = (InNotifyType == EConstraintsManagerNotifyType::ConstraintRemovedWithCompensation);
+							   if (bCompensate && ConstraintChannel->GetConstraint().Get())
+							   {
+								   FMovieSceneConstraintChannelHelper::HandleConstraintRemoved(
+									   ConstraintChannel->GetConstraint().Get(),
+									   &ConstraintChannel->ActiveChannel,
+									   Sequencer,
+									   Section);
+							   }
 							}
 
 							InSection->RemoveConstraintChannel(Constraint);
@@ -3513,6 +3619,7 @@ void FControlRigParameterTrackEditor::HandleConstraintRemoved(IMovieSceneConstra
 						break;
 				}
 			});
+
 			ConstraintHandlesToClear.Add(InSection->OnConstraintRemovedHandle);
 		}
 	}
@@ -3538,8 +3645,9 @@ void FControlRigParameterTrackEditor::HandleConstraintPropertyChanged(UTickableT
 	{
 		return;
 	}
-	
-	UMovieSceneSection* Section = Interface->GetHandleConstraintSection(Handle, GetSequencer());
+
+	const TSharedPtr<ISequencer> Sequencer = GetSequencer();
+	UMovieSceneSection* Section = Sequencer ? Interface->GetHandleConstraintSection(Handle, Sequencer) : nullptr;
 	IMovieSceneConstrainedSection* ConstraintSection = Cast<IMovieSceneConstrainedSection>(Section);
 	if (!ConstraintSection)
 	{
@@ -3559,7 +3667,7 @@ void FControlRigParameterTrackEditor::HandleConstraintPropertyChanged(UTickableT
 	}
 
 	FMovieSceneConstraintChannelHelper::HandleConstraintPropertyChanged(
-			InConstraint, Channel->ActiveChannel, InPropertyChangedEvent, GetSequencer(), Section);
+			InConstraint, Channel->ActiveChannel, InPropertyChangedEvent, Sequencer, Section);
 }
 
 void FControlRigParameterTrackEditor::SetUpEditModeIfNeeded(UControlRig* ControlRig)
@@ -3567,28 +3675,30 @@ void FControlRigParameterTrackEditor::SetUpEditModeIfNeeded(UControlRig* Control
 	if (ControlRig)
 	{
 		//this could clear the selection so if it does reset it
-		TArray<FName> ControlRigSelection = ControlRig->CurrentControlSelection();
+		const TArray<FName> ControlRigSelection = ControlRig->CurrentControlSelection();
 
+		const TSharedPtr<ISequencer> Sequencer = GetSequencer();
 		FControlRigEditMode* ControlRigEditMode = GetEditMode();
+		
 		if (!ControlRigEditMode)
 		{
 			ControlRigEditMode = GetEditMode(true);
-			if (TSharedPtr<IControlRigObjectBinding> ObjectBinding = ControlRig->GetObjectBinding())
+			if (ControlRigEditMode)
 			{
-				if (ControlRigEditMode)
+				if (TSharedPtr<IControlRigObjectBinding> ObjectBinding = ControlRig->GetObjectBinding())
 				{
-					ControlRigEditMode->AddControlRigObject(ControlRig, GetSequencer());
+					ControlRigEditMode->AddControlRigObject(ControlRig, Sequencer);
 				}
 			}
 		}
 		else
 		{
-			if (ControlRigEditMode->AddControlRigObject(ControlRig, GetSequencer()))
+			if (ControlRigEditMode->AddControlRigObject(ControlRig, Sequencer))
 			{
 				//force an evaluation, this will get the control rig setup so edit mode looks good.
-				if (GetSequencer().IsValid())
+				if (Sequencer)
 				{
-					GetSequencer()->NotifyMovieSceneDataChanged(EMovieSceneDataChangeType::Unknown);
+					Sequencer->NotifyMovieSceneDataChanged(EMovieSceneDataChangeType::Unknown);
 				}
 			}
 		}
@@ -3689,10 +3799,9 @@ void FControlRigParameterTrackEditor::HandleControlSelected(UControlRig* Subject
 	{
 		return;
 	}
+	
 	TGuardValue<bool> Guard(bIsDoingSelection, true);
 
-	FControlRigEditMode* ControlRigEditMode = GetEditMode();
-
 	const FName ControlRigName(*Subject->GetName());
 	if (TSharedPtr<IControlRigObjectBinding> ObjectBinding = Subject->GetObjectBinding())
 	{
@@ -3706,9 +3815,11 @@ void FControlRigParameterTrackEditor::HandleControlSelected(UControlRig* Subject
 		const FFindOrCreateHandleResult HandleResult = FindOrCreateHandleToObject(Object, Subject);
 		FFindOrCreateTrackResult TrackResult = FindOrCreateControlRigTrackForObject(HandleResult.Handle, Subject, ControlRigName, bCreateTrack);
 		UMovieSceneControlRigParameterTrack* Track = CastChecked<UMovieSceneControlRigParameterTrack>(TrackResult.Track, ECastCheckedType::NullAllowed);
-		if (Track)
+
+		const TSharedPtr<ISequencer> Sequencer = GetSequencer();
+		if (Track && Sequencer)
 		{
-			FSelectionEventSuppressor SuppressSelectionEvents = GetSequencer()->GetViewModel()->GetSelection()->SuppressEvents();
+			FSelectionEventSuppressor SuppressSelectionEvents = Sequencer->GetViewModel()->GetSelection()->SuppressEvents();
 
 			//Just select in section to key, if deselecting makes sure deselected everywhere
 			if (bSelected == true)
@@ -3733,8 +3844,8 @@ void FControlRigParameterTrackEditor::HandleControlSelected(UControlRig* Subject
 			//Force refresh later, not now
 			bSkipNextSelectionFromTimer = bSkipNextSelectionFromTimer ||
 				(bIsSelectingIndirectControl && ControlElement->Settings.AnimationType == ERigControlAnimationType::AnimationControl);
-			GetSequencer()->NotifyMovieSceneDataChanged(EMovieSceneDataChangeType::RefreshTree);
-
+			
+			Sequencer->NotifyMovieSceneDataChanged(EMovieSceneDataChangeType::RefreshTree);
 		}
 	}
 }
@@ -3745,38 +3856,47 @@ void FControlRigParameterTrackEditor::HandleOnPostConstructed(UControlRig* Subje
 	{
 		UControlRig* ControlRig = CastChecked<UControlRig>(Subject);
 
-		if (GetSequencer().IsValid())
+		if (TSharedPtr<ISequencer> Sequencer = GetSequencer())
 		{
 			//refresh tree for ANY control rig may be FK or procedural
-			GetSequencer()->NotifyMovieSceneDataChanged(EMovieSceneDataChangeType::RefreshTree);
+			Sequencer->NotifyMovieSceneDataChanged(EMovieSceneDataChangeType::RefreshTree);
 		}
 	}
 }
 
 void FControlRigParameterTrackEditor::HandleControlModified(UControlRig* ControlRig, FRigControlElement* ControlElement, const FRigControlModifiedContext& Context)
 {
-	if(ControlElement == nullptr)
+	if (Context.SetKey == EControlRigSetKey::Never)
+	{
+		return;
+	}
+	
+	if (!ControlElement ||
+		ControlElement->Settings.AnimationType == ERigControlAnimationType::ProxyControl ||
+		ControlElement->Settings.AnimationType == ERigControlAnimationType::VisualCue ) 
 	{
 		return;
 	}
+	
 	if (IsInGameThread() == false)
 	{
 		return;
 	}
-	if (!GetSequencer().IsValid() || !GetSequencer()->IsAllowedToChange() || Context.SetKey == EControlRigSetKey::Never 
-		|| ControlElement->Settings.AnimationType == ERigControlAnimationType::ProxyControl
-		|| ControlElement->Settings.AnimationType == ERigControlAnimationType::VisualCue)
+
+	const TSharedPtr<ISequencer> Sequencer = GetSequencer();
+	if (!Sequencer || !Sequencer->IsAllowedToChange())
 	{
 		return;
 	}
+
 	FTransform  Transform = ControlRig->GetControlLocalTransform(ControlElement->GetFName());
-	UMovieScene* MovieScene = GetSequencer()->GetFocusedMovieSceneSequence()->GetMovieScene();
 
 	IterateTracks([this, ControlRig, ControlElement, Context](UMovieSceneControlRigParameterTrack* Track)
 	{
 		if (Track && Track->GetControlRig() == ControlRig)
 		{
 			FName Name(*ControlRig->GetName());
+			
 			if (TSharedPtr<IControlRigObjectBinding> ObjectBinding = ControlRig->GetObjectBinding())
 			{
 				if (UObject* Object = ObjectBinding->GetBoundObject())
@@ -3787,10 +3907,10 @@ void FControlRigParameterTrackEditor::HandleControlModified(UControlRig* Control
 						KeyMode = ESequencerKeyMode::ManualKeyForced;
 					}
 
-					UE_LOG(LogControlRigEditor, Warning, TEXT("ADD KEY TO %s"), *ControlElement->GetName());
 					AddControlKeys(Object, ControlRig, Name, ControlElement->GetFName(), (EControlRigContextChannelToKey)Context.KeyMask, 
 						KeyMode, Context.LocalTime);
 					ControlChangedDuringUndoBracket++;
+					
 					return true;
 				}
 			}
@@ -3924,12 +4044,20 @@ void FControlRigParameterTrackEditor::GetControlRigKeys(
 	FGeneratedTrackKeys& OutGeneratedKeys,
 	const bool bInConstraintSpace)
 {
-	EMovieSceneTransformChannel TransformMask = SectionToKey->GetTransformMask().GetChannels();
+	const TSharedPtr<ISequencer> Sequencer = GetSequencer();
+	if (!Sequencer)
+	{
+		return;	
+	}
+
+	const EKeyGroupMode KeyGroupMode = Sequencer->GetKeyGroupMode();
+	const EMovieSceneTransformChannel TransformMask = SectionToKey->GetTransformMask().GetChannels();
 
 	TArray<FRigControlElement*> Controls;
 	InControlRig->GetControlsInOrder(Controls);
+
 	// If key all is enabled, for a key on all the channels
-	if (KeyMode != ESequencerKeyMode::ManualKeyForced && GetSequencer()->GetKeyGroupMode() == EKeyGroupMode::KeyAll)
+	if (KeyMode != ESequencerKeyMode::ManualKeyForced && KeyGroupMode == EKeyGroupMode::KeyAll)
 	{
 		ChannelsToKey = EControlRigContextChannelToKey::AllTransform;
 	}
@@ -4089,7 +4217,7 @@ void FControlRigParameterTrackEditor::GetControlRigKeys(
 				bool bKeyX = bSetKey && EnumHasAnyFlags(ChannelsToKey, EControlRigContextChannelToKey::TranslationX);
 				bool bKeyY = bSetKey && EnumHasAnyFlags(ChannelsToKey, EControlRigContextChannelToKey::TranslationY);
 				bool bKeyZ = bSetKey && EnumHasAnyFlags(ChannelsToKey, EControlRigContextChannelToKey::TranslationZ);
-				if (KeyMode != ESequencerKeyMode::ManualKeyForced && (GetSequencer()->GetKeyGroupMode() == EKeyGroupMode::KeyGroup && (bKeyX || bKeyY || bKeyZ)))
+				if (KeyMode != ESequencerKeyMode::ManualKeyForced && (KeyGroupMode == EKeyGroupMode::KeyGroup && (bKeyX || bKeyY || bKeyZ)))
 				{
 					bKeyX = bKeyY = bKeyZ = true;
 				}
@@ -4126,7 +4254,7 @@ void FControlRigParameterTrackEditor::GetControlRigKeys(
 				bKeyX = bSetKey && EnumHasAnyFlags(ChannelsToKey, EControlRigContextChannelToKey::RotationX);
 				bKeyY = bSetKey && EnumHasAnyFlags(ChannelsToKey, EControlRigContextChannelToKey::RotationY);
 				bKeyZ = bSetKey && EnumHasAnyFlags(ChannelsToKey, EControlRigContextChannelToKey::RotationZ);
-				if (KeyMode != ESequencerKeyMode::ManualKeyForced && (GetSequencer()->GetKeyGroupMode() == EKeyGroupMode::KeyGroup && (bKeyX || bKeyY || bKeyZ)))
+				if (KeyMode != ESequencerKeyMode::ManualKeyForced && (KeyGroupMode == EKeyGroupMode::KeyGroup && (bKeyX || bKeyY || bKeyZ)))
 				{
 					bKeyX = bKeyY = bKeyZ = true;
 				}
@@ -4153,7 +4281,7 @@ void FControlRigParameterTrackEditor::GetControlRigKeys(
 					bKeyX = bSetKey && EnumHasAnyFlags(ChannelsToKey, EControlRigContextChannelToKey::ScaleX);
 					bKeyY = bSetKey && EnumHasAnyFlags(ChannelsToKey, EControlRigContextChannelToKey::ScaleY);
 					bKeyZ = bSetKey && EnumHasAnyFlags(ChannelsToKey, EControlRigContextChannelToKey::ScaleZ);
-					if (KeyMode != ESequencerKeyMode::ManualKeyForced && (GetSequencer()->GetKeyGroupMode() == EKeyGroupMode::KeyGroup && (bKeyX || bKeyY || bKeyZ)))
+					if (KeyMode != ESequencerKeyMode::ManualKeyForced && (KeyGroupMode == EKeyGroupMode::KeyGroup && (bKeyX || bKeyY || bKeyZ)))
 					{
 						bKeyX = bKeyY = bKeyZ = true;
 					}
@@ -4184,8 +4312,16 @@ FKeyPropertyResult FControlRigParameterTrackEditor::AddKeysToControlRigHandle(UO
 	FGuid ObjectHandle, FFrameNumber KeyTime, FFrameNumber EvaluateTime, FGeneratedTrackKeys& GeneratedKeys,
 	ESequencerKeyMode KeyMode, TSubclassOf<UMovieSceneTrack> TrackClass, FName ControlRigName, FName RigControlName)
 {
-	EAutoChangeMode AutoChangeMode = GetSequencer()->GetAutoChangeMode();
-	EAllowEditsMode AllowEditsMode = GetSequencer()->GetAllowEditsMode();
+	FKeyPropertyResult KeyPropertyResult;
+	
+	const TSharedPtr<ISequencer> Sequencer = GetSequencer();
+	if (!Sequencer)
+	{
+		return KeyPropertyResult;
+	}
+
+	const EAutoChangeMode AutoChangeMode = Sequencer->GetAutoChangeMode();
+	const EAllowEditsMode AllowEditsMode = Sequencer->GetAllowEditsMode();
 
 	bool bCreateTrack =
 		(KeyMode == ESequencerKeyMode::AutoKey && (AutoChangeMode == EAutoChangeMode::AutoTrack || AutoChangeMode == EAutoChangeMode::All)) ||
@@ -4204,13 +4340,8 @@ FKeyPropertyResult FControlRigParameterTrackEditor::AddKeysToControlRigHandle(UO
 
 	bool bTrackCreated = TrackResult.bWasCreated;
 
-	bool bSectionCreated = false;
-	FKeyPropertyResult KeyPropertyResult;
-
 	if (Track)
 	{
-		float Weight = 1.0f;
-
 		UMovieSceneSection* SectionToKey = Track->GetSectionToKey(RigControlName);
 		if (SectionToKey)
 		{
@@ -4223,6 +4354,9 @@ FKeyPropertyResult FControlRigParameterTrackEditor::AddKeysToControlRigHandle(UO
 				SectionToKey->SetStartFrame(KeyTime);
 			}
 		}
+		
+		bool bSectionCreated = false;
+		float Weight = 1.0f;
 
 		// If there's no overlapping section to key, create one only if a track was newly created. Otherwise, skip keying altogether
 		// so that the user is forced to create a section to key on.
@@ -4230,7 +4364,7 @@ FKeyPropertyResult FControlRigParameterTrackEditor::AddKeysToControlRigHandle(UO
 		{
 			Track->Modify();
 			SectionToKey = Track->FindOrAddSection(KeyTime, bSectionCreated);
-			if (bSectionCreated && GetSequencer()->GetInfiniteKeyAreas())
+			if (bSectionCreated && Sequencer->GetInfiniteKeyAreas())
 			{
 				SectionToKey->SetRange(TRange<FFrameNumber>::All());
 			}
@@ -4263,12 +4397,12 @@ FKeyPropertyResult FControlRigParameterTrackEditor::AddKeysToControlRigHandle(UO
 
 				// compensate spaces
 				FControlRigSpaceChannelHelpers::CompensateIfNeeded(
-					SectionControlRig, GetSequencer().Get(), ParamSection,
+					SectionControlRig, Sequencer.Get(), ParamSection,
 					OptionalKeyTime, true /*comp previous*/);
 
 				// compensate constraints
 				const uint32 ControlHash = UTransformableControlHandle::ComputeHash(SectionControlRig, RigControlName);
-				FMovieSceneConstraintChannelHelper::CompensateIfNeeded(GetSequencer(), ParamSection, OptionalKeyTime, true /*bCompPreviousTick*/, ControlHash);
+				FMovieSceneConstraintChannelHelper::CompensateIfNeeded(Sequencer, ParamSection, OptionalKeyTime, true /*bCompPreviousTick*/, ControlHash);
 			}
 		}
 	}
@@ -4280,8 +4414,16 @@ FKeyPropertyResult FControlRigParameterTrackEditor::AddKeysToControlRig(
 	ESequencerKeyMode KeyMode, TSubclassOf<UMovieSceneTrack> TrackClass, FName ControlRigName, FName RigControlName)
 {
 	FKeyPropertyResult KeyPropertyResult;
-	EAutoChangeMode AutoChangeMode = GetSequencer()->GetAutoChangeMode();
-	EAllowEditsMode AllowEditsMode = GetSequencer()->GetAllowEditsMode();
+	
+	const TSharedPtr<ISequencer> Sequencer = GetSequencer();
+	if (!Sequencer)
+	{
+		return KeyPropertyResult;
+	}
+	
+	const EAutoChangeMode AutoChangeMode = Sequencer->GetAutoChangeMode();
+	const EAllowEditsMode AllowEditsMode = Sequencer->GetAllowEditsMode();
+
 	bool bCreateHandle =
 		(KeyMode == ESequencerKeyMode::AutoKey && (AutoChangeMode == EAutoChangeMode::All)) ||
 		KeyMode == ESequencerKeyMode::ManualKey ||
@@ -4306,10 +4448,17 @@ void FControlRigParameterTrackEditor::AddControlKeys(
 	float InLocalTime,
 	const bool bInConstraintSpace)
 {
-	if (KeyMode == ESequencerKeyMode::ManualKey || (GetSequencer().IsValid() && !GetSequencer()->IsAllowedToChange()))
+	const TSharedPtr<ISequencer> Sequencer = GetSequencer();
+	if (!Sequencer)
 	{
 		return;
 	}
+	
+	if (KeyMode == ESequencerKeyMode::ManualKey || !Sequencer->IsAllowedToChange())
+	{
+		return;
+	}
+	
 	bool bCreateTrack = false;
 	bool bCreateHandle = false;
 	FFindOrCreateHandleResult HandleResult = FindOrCreateHandleToObject(InObject, InControlRig);
@@ -4323,9 +4472,10 @@ void FControlRigParameterTrackEditor::AddControlKeys(
 		if (InLocalTime != FLT_MAX)
 		{
 			//convert from frame time since conversion may give us one frame less, e.g 1.53333330 * 24000.0/1.0 = 36799.999199999998
-			FFrameTime LocalFrameTime = GetSequencer()->GetFocusedTickResolution().AsFrameTime((double)InLocalTime);
+			FFrameTime LocalFrameTime = Sequencer->GetFocusedTickResolution().AsFrameTime((double)InLocalTime);
 			BeginKeying(LocalFrameTime.RoundToFrame());
 		}
+		
 		const FFrameNumber FrameTime = GetTimeForKey();
 		UMovieSceneSection* Section = Track->GetSectionToKey(RigControlName);
 		ParamSection = Cast<UMovieSceneControlRigParameterSection>(Section);
@@ -4348,19 +4498,25 @@ void FControlRigParameterTrackEditor::AddControlKeys(
 
 	auto OnKeyProperty = [=, this](FFrameNumber Time) -> FKeyPropertyResult
 	{
-		FFrameNumber LocalTime = Time;
-		//for modify weights we evaluate so need to make sure we use the evaluated time
-		FFrameNumber EvaluateTime = GetSequencer()->GetLastEvaluatedLocalTime().RoundToFrame();
-		//if InLocalTime is specified that means time value was set with SetControlValue, so we don't use sequencer times at all, but this time instead
-		if (InLocalTime != FLT_MAX)
+		if (const TSharedPtr<ISequencer> Sequencer = GetSequencer())
 		{
-			//convert from frame time since conversion may give us one frame less, e.g 1.53333330 * 24000.0/1.0 = 36799.999199999998
-			FFrameTime LocalFrameTime = GetSequencer()->GetFocusedTickResolution().AsFrameTime((double)InLocalTime);
-			LocalTime = LocalFrameTime.RoundToFrame();
-			EvaluateTime = LocalTime;
+			FFrameNumber LocalTime = Time;
+			//for modify weights we evaluate so need to make sure we use the evaluated time
+			FFrameNumber EvaluateTime = Sequencer->GetLastEvaluatedLocalTime().RoundToFrame();
+			//if InLocalTime is specified that means time value was set with SetControlValue, so we don't use sequencer times at all, but this time instead
+			if (InLocalTime != FLT_MAX)
+			{
+				//convert from frame time since conversion may give us one frame less, e.g 1.53333330 * 24000.0/1.0 = 36799.999199999998
+				FFrameTime LocalFrameTime = Sequencer->GetFocusedTickResolution().AsFrameTime((double)InLocalTime);
+				LocalTime = LocalFrameTime.RoundToFrame();
+				EvaluateTime = LocalTime;
+			}
+		
+			return this->AddKeysToControlRig(InObject, InControlRig, LocalTime, EvaluateTime, *GeneratedKeys, KeyMode, UMovieSceneControlRigParameterTrack::StaticClass(), ControlRigName, RigControlName);
 		}
 		
-		return this->AddKeysToControlRig(InObject, InControlRig, LocalTime, EvaluateTime, *GeneratedKeys, KeyMode, UMovieSceneControlRigParameterTrack::StaticClass(), ControlRigName, RigControlName);
+		static const FKeyPropertyResult Dummy;
+		return Dummy;
 	};
 
 	AnimatablePropertyChanged(FOnKeyProperty::CreateLambda(OnKeyProperty));
@@ -4371,17 +4527,23 @@ bool FControlRigParameterTrackEditor::ModifyOurGeneratedKeysByCurrentAndWeight(U
 {
 	using namespace UE::MovieScene;
 
+	TSharedPtr<ISequencer> Sequencer = GetSequencer();
+	if (!Sequencer)
+	{
+		return false;
+	}
+	
 	// Start off with stable names for population since we shouldn't encounter any duplicates
 	FControlRigParameterValues ParameterBufferValues(EControlRigParameterBufferIndexStability::Stable);
+	const FFrameRate TickResolution = Sequencer->GetFocusedTickResolution();
 
 	if (UMovieSceneControlRigParameterTrack::ShouldUseLegacyTemplate())
 	{
-		FFrameRate TickResolution = GetSequencer()->GetFocusedTickResolution();
 		FMovieSceneEvaluationTrack EvalTrack = CastChecked<UMovieSceneControlRigParameterTrack>(Track)->GenerateTrackTemplate(Track);
 		FMovieSceneInterrogationData InterrogationData;
-		GetSequencer()->GetEvaluationTemplate().CopyActuators(InterrogationData.GetAccumulator());
+		Sequencer->GetEvaluationTemplate().CopyActuators(InterrogationData.GetAccumulator());
 		//use the EvaluateTime to do the evaluation, may be different than the actually time we key
-		FMovieSceneContext Context(FMovieSceneEvaluationRange(EvaluateTime, GetSequencer()->GetFocusedTickResolution()));
+		FMovieSceneContext Context(FMovieSceneEvaluationRange(EvaluateTime, TickResolution));
 		EvalTrack.Interrogate(Context, InterrogationData, Object);
 		int32 ChannelIndex = 0;
 		FChannelMapInfo* pChannelIndex = nullptr;
@@ -4439,8 +4601,6 @@ bool FControlRigParameterTrackEditor::ModifyOurGeneratedKeysByCurrentAndWeight(U
 	// Make searching faster by hashing the values
 	ParameterBufferValues.OptimizeForLookup();
 
-	FFrameRate TickResolution = GetSequencer()->GetFocusedTickResolution();
-
 	TArray<FRigControlElement*> Controls = InControlRig->AvailableControls();
 	UMovieSceneControlRigParameterSection* Section = Cast<UMovieSceneControlRigParameterSection>(SectionToKey);
 	FMovieSceneChannelProxy& Proxy = SectionToKey->GetChannelProxy();
@@ -4550,7 +4710,7 @@ void FControlRigParameterTrackEditor::ProcessKeyOperation(FFrameNumber InKeyTime
 	}
 
 	// ECS system performs a full recomposition on the ECS data
-	auto Iterator = [this, InKeyTime, &Operation, &InSequencer, OutResults](UMovieSceneTrack* Track, TArrayView<const UE::Sequencer::FKeySectionOperation> Operations)
+	auto Iterator = [this, InKeyTime, &InSequencer, OutResults](UMovieSceneTrack* Track, TArrayView<const UE::Sequencer::FKeySectionOperation> Operations)
 	{
 		UMovieSceneControlRigParameterTrack* ControlRigTrack = CastChecked<UMovieSceneControlRigParameterTrack>(Track);
 
@@ -5002,7 +5162,13 @@ bool FControlRigParameterTrackEditor::HandleAssetAdded(UObject* Asset, const FGu
 		return false;
 	}
 
-	UMovieScene* MovieScene = GetSequencer()->GetFocusedMovieSceneSequence()->GetMovieScene();
+	const TSharedPtr<ISequencer> Sequencer = GetSequencer();
+	if (!Sequencer)
+	{
+		return false;
+	}
+	
+	UMovieScene* MovieScene = Sequencer->GetFocusedMovieSceneSequence()->GetMovieScene();
 	if (!MovieScene)
 	{
 		return false;
@@ -5029,7 +5195,7 @@ bool FControlRigParameterTrackEditor::HandleAssetAdded(UObject* Asset, const FGu
 	UE::Sequencer::FCreateBindingParams CreateBindingParams;
 	CreateBindingParams.bSpawnable = true;
 	CreateBindingParams.bAllowCustomBinding = true;
-	FGuid NewGuid = GetSequencer()->CreateBinding(*ASkeletalMeshActor::StaticClass(), CreateBindingParams);
+	FGuid NewGuid = Sequencer->CreateBinding(*ASkeletalMeshActor::StaticClass(), CreateBindingParams);
 
 	// CreateBinding can fail if spawnables are not allowed
 	if (!NewGuid.IsValid())
@@ -5037,7 +5203,7 @@ bool FControlRigParameterTrackEditor::HandleAssetAdded(UObject* Asset, const FGu
 		return false;
 	}
 	
-	ASkeletalMeshActor* SpawnedSkeletalMeshActor = Cast<ASkeletalMeshActor>(GetSequencer()->FindSpawnedObjectOrTemplate(NewGuid));
+	ASkeletalMeshActor* SpawnedSkeletalMeshActor = Cast<ASkeletalMeshActor>(Sequencer->FindSpawnedObjectOrTemplate(NewGuid));
 	if (!ensure(SpawnedSkeletalMeshActor))
 	{
 		return false;
@@ -5050,7 +5216,7 @@ bool FControlRigParameterTrackEditor::HandleAssetAdded(UObject* Asset, const FGu
 	
 	// Save Spawnable state as the default (with new name and skeletal mesh asset)
 	{
-		GetSequencer()->GetSpawnRegister().SaveDefaultSpawnableState(NewGuid, GetSequencer()->GetFocusedTemplateID(), GetSequencer()->GetSharedPlaybackState());
+		Sequencer->GetSpawnRegister().SaveDefaultSpawnableState(NewGuid, Sequencer->GetFocusedTemplateID(), Sequencer->GetSharedPlaybackState());
 	}
 
 	UMovieSceneControlRigParameterTrack* Track = Cast<UMovieSceneControlRigParameterTrack>(MovieScene->FindTrack(UMovieSceneControlRigParameterTrack::StaticClass(), NewGuid, NAME_None));
@@ -6224,37 +6390,33 @@ void FControlRigParameterSection::OnAnimationAssetEnterPressedForFK(const TArray
 
 FEditorModeTools* FControlRigParameterTrackEditor::GetEditorModeTools() const
 {
-	TSharedPtr<ISequencer> SharedSequencer = GetSequencer();
-	if (SharedSequencer.IsValid())
+	if (const TSharedPtr<ISequencer> Sequencer = GetSequencer())
 	{
-		TSharedPtr<IToolkitHost> ToolkitHost = SharedSequencer->GetToolkitHost();
-		if (ToolkitHost.IsValid())
+		if (const TSharedPtr<IToolkitHost> ToolkitHost = Sequencer->GetToolkitHost())
 		{
 			return &ToolkitHost->GetEditorModeManager();
 		}
 	}
-
+	
 	return nullptr;
 }
 
 FControlRigEditMode* FControlRigParameterTrackEditor::GetEditMode(bool bForceActivate /*= false*/) const
 {
-	if (FEditorModeTools* EditorModetools = GetEditorModeTools())
+	if (FEditorModeTools* EditorModeTools = GetEditorModeTools())
 	{
-		if (bForceActivate && !EditorModetools->IsModeActive(FControlRigEditMode::ModeName))
+		if (bForceActivate && !EditorModeTools->IsModeActive(FControlRigEditMode::ModeName))
 		{
-			EditorModetools->ActivateMode(FControlRigEditMode::ModeName);
-			if (bForceActivate)
+			EditorModeTools->ActivateMode(FControlRigEditMode::ModeName);
+
+			FControlRigEditMode* EditMode = static_cast<FControlRigEditMode*>(EditorModeTools->GetActiveMode(FControlRigEditMode::ModeName));
+			if (EditMode && EditMode->GetToolkit().IsValid() == false)
 			{
-				FControlRigEditMode* EditMode = static_cast<FControlRigEditMode*>(EditorModetools->GetActiveMode(FControlRigEditMode::ModeName));
-				if (EditMode && EditMode->GetToolkit().IsValid() == false)
-				{
-					EditMode->Enter();
-				}
+				EditMode->Enter();
 			}
 		}
 
-		return static_cast<FControlRigEditMode*>(EditorModetools->GetActiveMode(FControlRigEditMode::ModeName));
+		return static_cast<FControlRigEditMode*>(EditorModeTools->GetActiveMode(FControlRigEditMode::ModeName));
 	}
 
 	return nullptr;
diff --git a/Engine/Plugins/Animation/ControlRig/Source/ControlRigEditor/Private/Sequencer/ControlRigParameterTrackEditor.h b/Engine/Plugins/Animation/ControlRig/Source/ControlRigEditor/Private/Sequencer/ControlRigParameterTrackEditor.h
--- a/Engine/Plugins/Animation/ControlRig/Source/ControlRigEditor/Private/Sequencer/ControlRigParameterTrackEditor.h
+++ b/Engine/Plugins/Animation/ControlRig/Source/ControlRigEditor/Private/Sequencer/ControlRigParameterTrackEditor.h
@@ -300,11 +300,11 @@ public:
 	static bool CollapseAllLayers(TSharedPtr<ISequencer>& SequencerPtr, UMovieSceneTrack* OwnerTrack, const FBakingAnimationKeySettings& InSettings);
 
 	//** Function to load animation into a section, returns true if successful
-	static bool LoadAnimationIntoSection(TSharedPtr<ISequencer>& SequencerPtr,  UAnimSequence* AnimSequence, USkeletalMeshComponent* SkelMeshComp,
+	static bool LoadAnimationIntoSection(const TSharedPtr<ISequencer>& InSequencer,  UAnimSequence* AnimSequence, USkeletalMeshComponent* SkelMeshComp,
 		FFrameNumber StartFrame, bool bReduceKeys, const FSmartReduceParams& ReduceParams, bool bResetControls, UMovieSceneControlRigParameterSection* ParamSection);
 
 	//** Function to smart reduce all keys on all constraints based on the parameters
-	static void SmartReduce(TSharedPtr<ISequencer>& SequencerPtr, const FSmartReduceParams& InParams, UMovieSceneControlRigParameterSection* InSection);
+	static void SmartReduce(const TSharedPtr<ISequencer>& InSequencer, const FSmartReduceParams& InParams, UMovieSceneControlRigParameterSection* InSection);
 
 	/** Adds the constraint to sequencer and create the bindings if needed. Note that ControlRig Track auto-generation is disabled in this function. */
 	static void AddConstraintToSequencer(const TSharedPtr<ISequencer>& InSequencer, UTickableTransformConstraint* InConstraint);
diff --git a/Engine/Plugins/Animation/ControlRig/Source/ControlRigEditor/Public/EditMode/ControlRigEditMode.h b/Engine/Plugins/Animation/ControlRig/Source/ControlRigEditor/Public/EditMode/ControlRigEditMode.h
--- a/Engine/Plugins/Animation/ControlRig/Source/ControlRigEditor/Public/EditMode/ControlRigEditMode.h
+++ b/Engine/Plugins/Animation/ControlRig/Source/ControlRigEditor/Public/EditMode/ControlRigEditMode.h
@@ -138,7 +138,7 @@ public:
 
 	/** Set the Control Rig Object to be active in the edit mode. You set both the Control Rig and a possible binding together with an optional Sequencer
 	 This will remove all other control rigs present and should be called for stand alone editors, like the Control Rig Editor*/
-	void SetObjects(UControlRig* InControlRig, UObject* BindingObject, TWeakPtr<ISequencer> InSequencer);
+	void SetObjects(UControlRig* InControlRig, UObject* BindingObject, const TWeakPtr<ISequencer>& InSequencer);
 
 	/** Add a Control Rig object if it doesn't exist, will return true if it was added, false if it wasn't since it's already there. You can also set the Sequencer.*/
 	bool AddControlRigObject(UControlRig* InControlRig, const TWeakPtr<ISequencer>& InSequencer);
diff --git a/Engine/Source/Editor/Sequencer/Private/MovieSceneTrackEditor.cpp b/Engine/Source/Editor/Sequencer/Private/MovieSceneTrackEditor.cpp
--- a/Engine/Source/Editor/Sequencer/Private/MovieSceneTrackEditor.cpp
+++ b/Engine/Source/Editor/Sequencer/Private/MovieSceneTrackEditor.cpp
@@ -45,7 +45,7 @@ namespace UE::Sequencer
 }
 
 FMovieSceneTrackEditor::FMovieSceneTrackEditor(TSharedRef<ISequencer> InSequencer)
-	: Sequencer(InSequencer)
+	: WeakSequencerPtr(InSequencer)
 { 
 }
 
@@ -55,7 +55,7 @@ FMovieSceneTrackEditor::~FMovieSceneTrackEditor()
 
 UMovieSceneSequence* FMovieSceneTrackEditor::GetMovieSceneSequence() const
 {
-	return Sequencer.Pin()->GetFocusedMovieSceneSequence();
+	return WeakSequencerPtr.Pin()->GetFocusedMovieSceneSequence();
 }
 
 TOptional<FFrameNumber> FMovieSceneTrackEditor::NextKeyTime;
@@ -69,7 +69,7 @@ FFrameNumber FMovieSceneTrackEditor::GetTimeForKey()
 	}
 
 	// Otherwise, key at the current time
-	TSharedPtr<ISequencer> SequencerPin = Sequencer.Pin();
+	TSharedPtr<ISequencer> SequencerPin = WeakSequencerPtr.Pin();
 	return SequencerPin.IsValid() ? SequencerPin->GetLocalTime().Time.FrameNumber : FFrameNumber(0);
 }
 
@@ -87,7 +87,7 @@ void FMovieSceneTrackEditor::EndKeying()
 
 void FMovieSceneTrackEditor::UpdatePlaybackRange()
 {
-	TSharedPtr<ISequencer> SequencerPin = Sequencer.Pin();
+	TSharedPtr<ISequencer> SequencerPin = WeakSequencerPtr.Pin();
 	if( SequencerPin.IsValid()  )
 	{
 		SequencerPin->UpdatePlaybackRange();
@@ -135,32 +135,32 @@ void FMovieSceneTrackEditor::AnimatablePropertyChanged( FOnKeyProperty OnKeyProp
 		if (KeyPropertyResult.bTrackCreated)
 		{
 			// If a track is created evaluate immediately so that the pre-animated state can be stored
-			Sequencer.Pin()->NotifyMovieSceneDataChanged( EMovieSceneDataChangeType::RefreshAllImmediately );
+			WeakSequencerPtr.Pin()->NotifyMovieSceneDataChanged( EMovieSceneDataChangeType::RefreshAllImmediately );
 		}
 		else if (KeyPropertyResult.bTrackModified)
 		{
-			Sequencer.Pin()->NotifyMovieSceneDataChanged( EMovieSceneDataChangeType::MovieSceneStructureItemAdded );
+			WeakSequencerPtr.Pin()->NotifyMovieSceneDataChanged( EMovieSceneDataChangeType::MovieSceneStructureItemAdded );
 		}
 		else if (KeyPropertyResult.bKeyCreated)
 		{
-			Sequencer.Pin()->NotifyMovieSceneDataChanged( EMovieSceneDataChangeType::TrackValueChanged );
+			WeakSequencerPtr.Pin()->NotifyMovieSceneDataChanged( EMovieSceneDataChangeType::TrackValueChanged );
 		}
 		else
 		{
 			// If the only thing we changed as a result of the external change were channel defaults, we suppress automatic
 			// re-evaluation of the sequence for this change to ensure that the object does not have the change immediately overwritten
 			// by animated channels that have keys, but did not have keys added
-			UMovieSceneSequence* FocusedSequence = Sequencer.Pin()->GetFocusedMovieSceneSequence();
+			UMovieSceneSequence* FocusedSequence = WeakSequencerPtr.Pin()->GetFocusedMovieSceneSequence();
 			if (FocusedSequence)
 			{
-				Sequencer.Pin()->SuppressAutoEvaluation(FocusedSequence, FocusedSequence->GetSignature());
+				WeakSequencerPtr.Pin()->SuppressAutoEvaluation(FocusedSequence, FocusedSequence->GetSignature());
 			}
 		}
 
 		// If a handle was created, something new has been added to Sequencer, so its selection state must be synchronized
 		if (KeyPropertyResult.bHandleCreated)
 		{
-			Sequencer.Pin()->ExternalSelectionHasChanged();
+			WeakSequencerPtr.Pin()->ExternalSelectionHasChanged();
 		}
 
 		UpdatePlaybackRange();
@@ -199,10 +199,9 @@ FMovieSceneTrackEditor::FFindOrCreateTrackResult FMovieSceneTrackEditor::FindOrC
 	return Result;
 }
 
-
-const TSharedPtr<ISequencer> FMovieSceneTrackEditor::GetSequencer() const
+TSharedPtr<ISequencer> FMovieSceneTrackEditor::GetSequencer() const
 {
-	return Sequencer.Pin();
+	return WeakSequencerPtr.IsValid() ? WeakSequencerPtr.Pin() : nullptr;
 }
 
 FText FMovieSceneTrackEditor::GetDisplayName() const
diff --git a/Engine/Source/Editor/Sequencer/Private/SequencerUtilities.cpp b/Engine/Source/Editor/Sequencer/Private/SequencerUtilities.cpp
--- a/Engine/Source/Editor/Sequencer/Private/SequencerUtilities.cpp
+++ b/Engine/Source/Editor/Sequencer/Private/SequencerUtilities.cpp
@@ -311,8 +311,8 @@ void FSequencerUtilities::PopulateMenu_CreateNewSection(FMenuBuilder& MenuBuilde
 	
 	auto CreateNewSection = [Track, InSequencer, RowIndex](EMovieSceneBlendType BlendType)
 	{
-		TSharedPtr<ISequencer> Sequencer = InSequencer.Pin();
-		if (!Sequencer.IsValid())
+		TSharedPtr<ISequencer> Sequencer = InSequencer.IsValid() ? InSequencer.Pin() : nullptr;
+		if (!Sequencer)
 		{
 			return;
 		}
diff --git a/Engine/Source/Editor/Sequencer/Public/MovieSceneTrackEditor.h b/Engine/Source/Editor/Sequencer/Public/MovieSceneTrackEditor.h
--- a/Engine/Source/Editor/Sequencer/Public/MovieSceneTrackEditor.h
+++ b/Engine/Source/Editor/Sequencer/Public/MovieSceneTrackEditor.h
@@ -165,7 +165,7 @@ public:
 	}
 
 	/** @return The sequencer bound to this handler */
-	const TSharedPtr<ISequencer> GetSequencer() const;
+	TSharedPtr<ISequencer> GetSequencer() const;
 
 public:
 
@@ -208,7 +208,7 @@ protected:
 private:
 
 	/** The sequencer bound to this handler.  Used to access movie scene and time info during auto-key */
-	TWeakPtr<ISequencer> Sequencer;
+	TWeakPtr<ISequencer> WeakSequencerPtr;
 
 	/** The key time to use during a multi key operation. Only used if bKeying is true */
 	static TOptional<FFrameNumber> NextKeyTime;
