Anim In Engine: Added new hot key to mirror unselected controls.

diff --git a/Engine/Plugins/Animation/ControlRig/Source/ControlRigEditor/Private/EditMode/ControlRigEditMode.cpp b/Engine/Plugins/Animation/ControlRig/Source/ControlRigEditor/Private/EditMode/ControlRigEditMode.cpp
--- a/Engine/Plugins/Animation/ControlRig/Source/ControlRigEditor/Private/EditMode/ControlRigEditMode.cpp
+++ b/Engine/Plugins/Animation/ControlRig/Source/ControlRigEditor/Private/EditMode/ControlRigEditMode.cpp
@@ -2919,6 +2919,10 @@ void FControlRigEditMode::BindCommands()
 		Commands.MirrorSelectedControls,
 		FExecuteAction::CreateRaw(this, &FControlRigEditMode::MirrorSelectedControls));
 
+	CommandBindings->MapAction(
+		Commands.MirrorUnselectedControls,
+		FExecuteAction::CreateRaw(this, &FControlRigEditMode::MirrorUnselectedControls));
+
 	CommandBindings->MapAction(
 		Commands.SelectAllControls,
 		FExecuteAction::CreateRaw(this, &FControlRigEditMode::SelectAllControls));
@@ -6007,11 +6011,12 @@ void FControlRigEditMode::SelectMirroredControls()
 	FScopedTransaction ScopedTransaction(LOCTEXT("SelectMirroredControls", "Select Mirrored Controls"), !GIsTransacting);
 	FGCScopeGuard Guard;
 	UControlRigPoseAsset* TempPose = NewObject<UControlRigPoseAsset>(GetTransientPackage(), NAME_None);
-
+	
 	for (UControlRig* ControlRig : ControlRigs)
 	{
 		if (ControlRig)
 		{
+			ControlRig->Modify();
 			TempPose->SavePose(ControlRig, false);
 			TempPose->SelectControls(ControlRig, true);
 		}
@@ -6035,6 +6040,7 @@ void FControlRigEditMode::AddMirroredControlsToSelection()
 	{
 		if (ControlRig)
 		{
+			ControlRig->Modify();
 			TempPose->SavePose(ControlRig, false);
 			TempPose->SelectControls(ControlRig, true, false);
 		}
@@ -6058,12 +6064,39 @@ void FControlRigEditMode::MirrorSelectedControls()
 	{
 		if (ControlRig)
 		{
+			ControlRig->Modify();
 			TempPose->SavePose(ControlRig, true);
 			TempPose->PastePose(ControlRig, /*setkey*/false, true);
 		}
 	}
 }
 
+/** Put Unselected Controls To Mirrored selcted controls*/
+void FControlRigEditMode::MirrorUnselectedControls()
+{
+	TArray<UControlRig*> ControlRigs = GetControlRigsArray(true);
+	ControlRigs = GetControlRigsWithSelectedControls(ControlRigs);
+	if (ControlRigs.Num() == 0)
+	{
+		return;
+	}
+	FScopedTransaction ScopedTransaction(LOCTEXT("MirrorSelectedControls", "Mirror Selected Controls"), !GIsTransacting);
+	FGCScopeGuard Guard;
+	UControlRigPoseAsset* TempPose = NewObject<UControlRigPoseAsset>(GetTransientPackage(), NAME_None);
+
+	for (UControlRig* ControlRig : ControlRigs)
+	{
+		if (ControlRig)
+		{
+			ControlRig->Modify();
+			TempPose->SavePose(ControlRig, false);
+			TempPose->SelectControls(ControlRig, true); //select mirrored controls
+			TempPose->PastePose(ControlRig, /*setkey*/false, true); //paste it
+			TempPose->SelectControls(ControlRig, false);//put it back
+		}
+	}
+}
+
 /** Select All Controls*/
 void FControlRigEditMode::SelectAllControls()
 {
@@ -6080,6 +6113,7 @@ void FControlRigEditMode::SelectAllControls()
 	{
 		if (ControlRig)
 		{
+			ControlRig->Modify();
 			TempPose->SavePose(ControlRig, true); //this will save the whole pose
 			TempPose->SelectControls(ControlRig);
 		}
@@ -6117,6 +6151,7 @@ void FControlRigEditMode::SelectPose(bool bMirror, int32 PoseNum)
 		else if (ControlRigs.Num() == 1 && ControlRigs[0])
 		{
 			FScopedTransaction ScopedTransaction(LOCTEXT("SelectPoseControls", "Select Pose Controls"), !GIsTransacting);
+			ControlRigs[0]->Modify();
 			StoredPose->SelectControls(ControlRigs[0], bMirror);
 			return;
 		}
@@ -6131,11 +6166,8 @@ void FControlRigEditMode::SelectPose(bool bMirror, int32 PoseNum)
 		{
 			if (ControlRig)
 			{
-				//if (ControlRig->CurrentControlSelection().Num() > 0)
-				{
-					StoredPose->SelectControls(ControlRigs[0], bMirror);
-					//return;
-				}
+				ControlRig->Modify();
+				StoredPose->SelectControls(ControlRig, bMirror);
 			}
 		}
 	}
@@ -6157,6 +6189,7 @@ void FControlRigEditMode::PastePose(bool bMirror, int32 PoseNum)
 		{
 			if (ControlRig)
 			{
+				ControlRig->Modify();
 				StoredPose->PastePose(ControlRig,/*setkey*/false, bMirror);
 				return;
 			}
diff --git a/Engine/Plugins/Animation/ControlRig/Source/ControlRigEditor/Private/EditMode/ControlRigEditModeCommands.cpp b/Engine/Plugins/Animation/ControlRig/Source/ControlRigEditor/Private/EditMode/ControlRigEditModeCommands.cpp
--- a/Engine/Plugins/Animation/ControlRig/Source/ControlRigEditor/Private/EditMode/ControlRigEditModeCommands.cpp
+++ b/Engine/Plugins/Animation/ControlRig/Source/ControlRigEditor/Private/EditMode/ControlRigEditModeCommands.cpp
@@ -25,7 +25,10 @@ void FControlRigEditModeCommands::RegisterCommands()
 
 	UI_COMMAND(AddMirroredControlsToSelection, "Add Mirrored Controls to Selection", "Select mirrored Controls on current selection, keeping current selection", EUserInterfaceActionType::Button, FInputChord(EKeys::M, EModifierKey::Shift));
 
-	UI_COMMAND(MirrorSelectedControls, "Mirror Selected Controls", "Put selected Controls to mirrored Pose", EUserInterfaceActionType::Button, FInputChord(EKeys::M));
+	UI_COMMAND(MirrorSelectedControls, "Mirror Selected Controls", "Put selected Controls to mirrored Pose", EUserInterfaceActionType::Button, FInputChord(EKeys::M, EModifierKey::Shift | EModifierKey::Control));
+
+	UI_COMMAND(MirrorUnselectedControls, "Mirror Unselected Controls", "Put unselected Controls to mirrored Selection Pose", EUserInterfaceActionType::Button, FInputChord(EKeys::M, EModifierKey::Alt));
+
 
 	UI_COMMAND(SelectAllControls, "Select All Controls", "Select All Controls", EUserInterfaceActionType::Button, FInputChord(EKeys::Escape, EModifierKey::Shift));
 
diff --git a/Engine/Plugins/Animation/ControlRig/Source/ControlRigEditor/Private/EditMode/ControlRigEditModeCommands.h b/Engine/Plugins/Animation/ControlRig/Source/ControlRigEditor/Private/EditMode/ControlRigEditModeCommands.h
--- a/Engine/Plugins/Animation/ControlRig/Source/ControlRigEditor/Private/EditMode/ControlRigEditModeCommands.h
+++ b/Engine/Plugins/Animation/ControlRig/Source/ControlRigEditor/Private/EditMode/ControlRigEditModeCommands.h
@@ -65,9 +65,12 @@ public:
 	/** Select mirrored Controls on current selection, keeping current selection*/
 	TSharedPtr< FUICommandInfo > AddMirroredControlsToSelection;
 
-	/** Put selected Controls To mirrored Pose*/
+	/** Put selected Controls to mirrored Pose*/
 	TSharedPtr< FUICommandInfo > MirrorSelectedControls;
 
+	/** Put unselected Controls to mirrored selected Controls*/
+	TSharedPtr< FUICommandInfo > MirrorUnselectedControls;
+
 	/** Select all Controls*/
 	TSharedPtr< FUICommandInfo > SelectAllControls;
 
diff --git a/Engine/Plugins/Animation/ControlRig/Source/ControlRigEditor/Public/EditMode/ControlRigEditMode.h b/Engine/Plugins/Animation/ControlRig/Source/ControlRigEditor/Public/EditMode/ControlRigEditMode.h
--- a/Engine/Plugins/Animation/ControlRig/Source/ControlRigEditor/Public/EditMode/ControlRigEditMode.h
+++ b/Engine/Plugins/Animation/ControlRig/Source/ControlRigEditor/Public/EditMode/ControlRigEditMode.h
@@ -301,6 +301,9 @@ public:
 	/** Put Selected Controls To Mirrored Pose*/
 	void MirrorSelectedControls();
 
+	/** Put Unselected Controls To Mirrored Pose*/
+	void MirrorUnselectedControls();
+
 	/** Select All Controls*/
 	void SelectAllControls();
 
