Anim In Engine: Pivot tool not working with multiple selected objects, remove unneeded Modify calls.

diff --git a/Engine/Plugins/Animation/ControlRig/Source/ControlRigEditor/Private/ControlRigSequencerEditorLibrary.cpp b/Engine/Plugins/Animation/ControlRig/Source/ControlRigEditor/Private/ControlRigSequencerEditorLibrary.cpp
--- a/Engine/Plugins/Animation/ControlRig/Source/ControlRigEditor/Private/ControlRigSequencerEditorLibrary.cpp
+++ b/Engine/Plugins/Animation/ControlRig/Source/ControlRigEditor/Private/ControlRigSequencerEditorLibrary.cpp
@@ -1048,7 +1048,6 @@ static void LocalSetControlRigWorldTransforms(ULevelSequence* LevelSequence,UCon
 					ConvertFramesToTickResolution(LevelSequence->GetMovieScene(), InFrames, TimeUnit, Frames);
 
 					UMovieScene* MovieScene = WeakSequencer.Pin().Get()->GetFocusedMovieSceneSequence()->GetMovieScene();
-					MovieScene->Modify();
 					FFrameRate TickResolution = MovieScene->GetTickResolution();
 					FRigControlModifiedContext Context;
 					Context.SetKey = SetKey;
diff --git a/Engine/Plugins/Animation/ControlRig/Source/ControlRigEditor/Private/Tools/ControlRigSnapper.cpp b/Engine/Plugins/Animation/ControlRig/Source/ControlRigEditor/Private/Tools/ControlRigSnapper.cpp
--- a/Engine/Plugins/Animation/ControlRig/Source/ControlRigEditor/Private/Tools/ControlRigSnapper.cpp
+++ b/Engine/Plugins/Animation/ControlRig/Source/ControlRigEditor/Private/Tools/ControlRigSnapper.cpp
@@ -171,7 +171,7 @@ bool FControlRigSnapper::GetControlRigControlTransforms(ISequencer* Sequencer,
 	{
 		FMovieSceneSequenceIDRef Template = Sequencer->GetFocusedTemplateID();
 		FMovieSceneSequenceTransform RootToLocalTransform = Sequencer->GetFocusedMovieSceneSequenceTransform();
-		TOptional<FFrameNumber> FrameNumber;
+		TOptional<FFrameNumber> FrameNumber = Sequencer->GetLocalTime().Time.RoundToFrame();
 		return LocalGetControlRigControlTransforms(Sequencer, FrameNumber, Sequencer->GetFocusedMovieSceneSequence(), Template, RootToLocalTransform,
 			ControlRig, ControlName, Frames, ParentTransforms, OutTransforms);
 	
@@ -182,7 +182,7 @@ bool FControlRigSnapper::GetControlRigControlTransforms(ISequencer* Sequencer,
 bool FControlRigSnapper::GetControlRigControlTransforms(UWorld* World,ULevelSequence* LevelSequence, UControlRig* ControlRig, const FName& ControlName,
 	const TArray<FFrameNumber>& Frames, const TArray<FTransform>& ParentTransforms, TArray<FTransform>& OutTransforms)
 {
-	if (LevelSequence)
+	if (LevelSequence && Frames.Num() > 0)
 	{
 		ALevelSequenceActor* OutActor;
 		FMovieSceneSequencePlaybackSettings Settings;
