[AutoRTFM] Wrap another thread_local variable with an AutoRTFM-aware macro.

The CL at https://p4-swarm.epicgames.net/changes/39867862 went through several iterations, and the fixup for this thread_local FString was accidentally missed as part of that reshuffle.

diff --git a/Engine/Plugins/Animation/ControlRig/Source/ControlRig/Private/Rigs/FKControlRig.cpp b/Engine/Plugins/Animation/ControlRig/Source/ControlRig/Private/Rigs/FKControlRig.cpp
--- a/Engine/Plugins/Animation/ControlRig/Source/ControlRig/Private/Rigs/FKControlRig.cpp
+++ b/Engine/Plugins/Animation/ControlRig/Source/ControlRig/Private/Rigs/FKControlRig.cpp
@@ -40,7 +40,7 @@ FName UFKControlRig::GetControlName(const FName& InName, const ERigElementType&
 			}
 			else
 			{
-				static thread_local FString ScratchString;
+				UE_AUTORTFM_DECLARE_THREAD_LOCAL_VAR(FString, ScratchString);
 
 				// ToString performs ScratchString.Reset() internally
 				InName.ToString(ScratchString);
@@ -84,7 +84,7 @@ FName UFKControlRig::GetControlTargetName(const FName& InName, const ERigElement
 		}
 		else
 		{
-			static thread_local FString ScratchString;
+			UE_AUTORTFM_DECLARE_THREAD_LOCAL_VAR(FString, ScratchString);
 			
 			// ToString performs ScratchString.Reset() internally
 			InName.ToString(ScratchString);
