Control Rig: Recover the reset, use and select buttons for modular details panel

#rnx

diff --git a/Engine/Plugins/Animation/ControlRig/Source/ControlRigEditor/Private/ControlRigModuleDetails.cpp b/Engine/Plugins/Animation/ControlRig/Source/ControlRigEditor/Private/ControlRigModuleDetails.cpp
--- a/Engine/Plugins/Animation/ControlRig/Source/ControlRigEditor/Private/ControlRigModuleDetails.cpp
+++ b/Engine/Plugins/Animation/ControlRig/Source/ControlRigEditor/Private/ControlRigModuleDetails.cpp
@@ -469,23 +469,168 @@ void FRigModuleInstanceDetails::CustomizeDetails(IDetailLayoutBuilder& DetailBui
 					]
 				];
 
-
 				if(TargetKeys.IsSet())
 				{
+					TSharedPtr<SHorizontalBox> HorizontalBox;
+					
 					ConnectorRow
 					.ValueContent()
 					.HAlign(HAlign_Left)
 					.VAlign(VAlign_Center)
 					[
-						SNew(SRigConnectorTargetWidget)
-						.Outer(PerModuleInfos[0].GetBlueprint())
-						.ConnectorKey(ConnectorKey)
-						.IsArray(bIsArrayConnector)
-						.ExpandArrayByDefault(true)
-						.Targets(TargetKeys.GetValue())
-						.OnSetTargetArray(FRigConnectorTargetWidget_SetTargetArray::CreateSP(this, &FRigModuleInstanceDetails::OnConnectorTargetChanged, Connector))
-						.RigTreeDelegates(RigTreeDelegates)
+						SAssignNew(HorizontalBox, SHorizontalBox)
+						+SHorizontalBox::Slot()
+						.AutoWidth()
+						.Padding(0.f, 0.f, 0.f, 0.f)
+						[
+							SNew(SRigConnectorTargetWidget)
+							.Outer(PerModuleInfos[0].GetBlueprint())
+							.ConnectorKey(ConnectorKey)
+							.IsArray(bIsArrayConnector)
+							.ExpandArrayByDefault(true)
+							.Targets(TargetKeys.GetValue())
+							.OnSetTargetArray(FRigConnectorTargetWidget_SetTargetArray::CreateSP(this, &FRigModuleInstanceDetails::OnConnectorTargetChanged, Connector))
+							.RigTreeDelegates(RigTreeDelegates)
+						]
 					];
+
+					if(!bIsArrayConnector)
+					{
+						// Reset button
+						HorizontalBox->AddSlot()
+						.AutoWidth()
+						.Padding(4.f, 0.f, 0.f, 0.f)
+						.VAlign(VAlign_Center)
+						[
+							SAssignNew(ResetConnectorButton.FindOrAdd(Connector.Name), SButton)
+							.ButtonStyle( FAppStyle::Get(), "NoBorder" )
+							.ButtonColorAndOpacity_Lambda([this, Connector]()
+							{
+								const TSharedPtr<SButton>& Button = ResetConnectorButton.FindRef(Connector.Name);
+								return Button.IsValid() && Button->IsHovered()
+									? FSlateColor(FLinearColor(1.f, 1.f, 1.f, 0.8))
+									: FSlateColor(FLinearColor(1.f, 1.f, 1.f, 0.4));
+							})
+							.OnClicked_Lambda([this, Connector]()
+							{
+								for (FPerModuleInfo& Info : PerModuleInfos)
+								{
+									FString ConnectorPath = FString::Printf(TEXT("%s:%s"), *Info.Path, *Connector.Name);
+									FRigElementKey ConnectorKey(*ConnectorPath, ERigElementType::Connector);
+									Info.GetBlueprint()->GetModularRigController()->DisconnectConnector(ConnectorKey);
+								}
+								return FReply::Handled();
+							})
+							.ContentPadding(1.f)
+							.ToolTipText(NSLOCTEXT("ControlRigModuleDetails", "Reset_Connector", "Reset Connector"))
+							[
+								SNew(SImage)
+								//.DesiredSizeOverride(FVector2D(16,16))
+								.ColorAndOpacity_Lambda( [this, Connector]()
+								{
+									const TSharedPtr<SButton>& Button = ResetConnectorButton.FindRef(Connector.Name);
+									return Button.IsValid() && Button->IsHovered()
+									? FSlateColor(FLinearColor(1.f, 1.f, 1.f, 0.8))
+									: FSlateColor(FLinearColor(1.f, 1.f, 1.f, 0.4));
+								})
+								.Image(FSlateIcon(FAppStyle::Get().GetStyleSetName(), "PropertyWindow.DiffersFromDefault").GetIcon())
+							]
+						];
+
+						// Use button
+						HorizontalBox->AddSlot()
+						.AutoWidth()
+						.Padding(0.f, 0.f, 0.f, 0.f)
+						.VAlign(VAlign_Center)
+						[
+							SAssignNew(UseSelectedButton.FindOrAdd(Connector.Name), SButton)
+							.ButtonStyle( FAppStyle::Get(), "NoBorder" )
+							.ButtonColorAndOpacity_Lambda([this, Connector]()
+							{
+								const TSharedPtr<SButton>& Button = UseSelectedButton.FindRef(Connector.Name);
+								return Button.IsValid() && Button->IsHovered()
+									? FSlateColor(FLinearColor(1.f, 1.f, 1.f, 0.8))
+									: FSlateColor(FLinearColor(1.f, 1.f, 1.f, 0.4));
+							})
+							.OnClicked_Lambda([this, Connector]()
+							{
+								if (UModularRig* ModularRig = PerModuleInfos[0].GetModularRig())
+								{
+									const TArray<FRigElementKey>& Selected = ModularRig->GetHierarchy()->GetSelectedKeys();
+									if (Selected.Num() > 0)
+									{
+										for (FPerModuleInfo& Info : PerModuleInfos)
+										{
+											FString ConnectorPath = FString::Printf(TEXT("%s:%s"), *Info.Path, *Connector.Name);
+											FRigElementKey ConnectorKey(*ConnectorPath, ERigElementType::Connector);
+											Info.GetBlueprint()->GetModularRigController()->ConnectConnectorToElements(ConnectorKey, Selected, true, ModularRig->GetModularRigSettings().bAutoResolve);
+										}
+									}
+								}
+								return FReply::Handled();
+							})
+							.ContentPadding(1.f)
+							.ToolTipText(NSLOCTEXT("ControlRigModuleDetails", "Use_Selected", "Use Selected"))
+							[
+								SNew(SImage)
+								//.DesiredSizeOverride(FVector2D(16,16))
+								.ColorAndOpacity_Lambda( [this, Connector]()
+								{
+									const TSharedPtr<SButton>& Button = UseSelectedButton.FindRef(Connector.Name);
+									return Button.IsValid() && Button->IsHovered()
+									? FSlateColor(FLinearColor(1.f, 1.f, 1.f, 0.8))
+									: FSlateColor(FLinearColor(1.f, 1.f, 1.f, 0.4));
+								})
+								.Image(FAppStyle::GetBrush("Icons.CircleArrowLeft"))
+							]
+						];
+
+						// Select in hierarchy button
+						HorizontalBox->AddSlot()
+						.AutoWidth()
+						.Padding(0.f, 0.f, 0.f, 0.f)
+						.VAlign(VAlign_Center)
+						[
+							SAssignNew(SelectElementButton.FindOrAdd(Connector.Name), SButton)
+							.ButtonStyle( FAppStyle::Get(), "NoBorder" )
+							.ButtonColorAndOpacity_Lambda([this, Connector]()
+							{
+								const TSharedPtr<SButton>& Button = SelectElementButton.FindRef(Connector.Name);
+								return Button.IsValid() && Button->IsHovered()
+									? FSlateColor(FLinearColor(1.f, 1.f, 1.f, 0.8))
+									: FSlateColor(FLinearColor(1.f, 1.f, 1.f, 0.4));
+							})
+							.OnClicked_Lambda([this, Connector]()
+							{
+								if (UModularRig* ModularRig = PerModuleInfos[0].GetModularRig())
+								{
+									FString ConnectorPath = FString::Printf(TEXT("%s:%s"), *PerModuleInfos[0].Path, *Connector.Name);
+									FRigElementKey ConnectorKey(*ConnectorPath, ERigElementType::Connector);
+									if (const FRigElementKeyRedirector::FKeyArray* TargetKeys = ModularRig->GetElementKeyRedirector().FindExternalKey(ConnectorKey))
+									{
+										TArray<FRigElementKey> KeysToSelect;
+										KeysToSelect.Append(*TargetKeys);
+										ModularRig->GetHierarchy()->GetController()->SetSelection(KeysToSelect, true);
+									}
+								}
+								return FReply::Handled();
+							})
+							.ContentPadding(1.f)
+							.ToolTipText(NSLOCTEXT("ControlRigModuleDetails", "Select_Element", "Select Element"))
+							[
+								SNew(SImage)
+								//.DesiredSizeOverride(FVector2D(16,16))
+								.ColorAndOpacity_Lambda( [this, Connector]()
+								{
+									const TSharedPtr<SButton>& Button = SelectElementButton.FindRef(Connector.Name);
+									return Button.IsValid() && Button->IsHovered()
+									? FSlateColor(FLinearColor(1.f, 1.f, 1.f, 0.8))
+									: FSlateColor(FLinearColor(1.f, 1.f, 1.f, 0.4));
+								})
+								.Image(FAppStyle::GetBrush("Icons.Search"))
+							]
+						];
+					}
 				}
 				else
 				{
diff --git a/Engine/Plugins/Animation/ControlRig/Source/ControlRigEditor/Private/Editor/SRigConnectorTargetWidget.cpp b/Engine/Plugins/Animation/ControlRig/Source/ControlRigEditor/Private/Editor/SRigConnectorTargetWidget.cpp
--- a/Engine/Plugins/Animation/ControlRig/Source/ControlRigEditor/Private/Editor/SRigConnectorTargetWidget.cpp
+++ b/Engine/Plugins/Animation/ControlRig/Source/ControlRigEditor/Private/Editor/SRigConnectorTargetWidget.cpp
@@ -58,7 +58,13 @@ void SRigConnectorTargetComboButton::Construct(const FArguments& InArgs)
 		.Visibility(EVisibility::Visible)
 		.BorderImage(FAppStyle::GetBrush("Menu.Background"))
 		[
-			SearchableTreeView.ToSharedRef()
+			SNew(SHorizontalBox)
+			+ SHorizontalBox::Slot()
+			.Padding(0)
+			.MaxWidth(900)
+			[
+				SearchableTreeView.ToSharedRef()
+			]
 		]
 	];
 
