Control Rig: Hierarchy Top Level Components

diff --git a/Engine/Plugins/Animation/ControlRig/Source/ControlRig/Private/Rigs/RigHierarchy.cpp b/Engine/Plugins/Animation/ControlRig/Source/ControlRig/Private/Rigs/RigHierarchy.cpp
--- a/Engine/Plugins/Animation/ControlRig/Source/ControlRig/Private/Rigs/RigHierarchy.cpp
+++ b/Engine/Plugins/Animation/ControlRig/Source/ControlRig/Private/Rigs/RigHierarchy.cpp
@@ -641,6 +641,7 @@ void URigHierarchy::Load(FArchive& Ar)
 	{
 		ElementComponents.Reset();
 		ComponentIndexLookup.Reset();
+		TopLevelComponentIndices.Reset();
 
 		int32 NumComponents = 0;
 		*ArchiveForElements << NumComponents;
@@ -693,12 +694,21 @@ void URigHierarchy::Load(FArchive& Ar)
 				FRigBaseComponent* Component = InstancedStruct.GetMutablePtr<FRigBaseComponent>();
 				Component->Serialize(*ArchiveForElements);
 
-				FRigBaseElement* Element = Find(Component->GetElementKey());
-				check(Element);
-				
 				Component->IndexInHierarchy = ComponentIndex;
-				Component->Element = Element;
-				Component->IndexInElement = Element->ComponentIndices.Add(Component->IndexInHierarchy);
+
+				if(Component->GetElementKey() == GetTopLevelComponentElementKey())
+				{
+					Component->Element = nullptr;
+					Component->IndexInElement = TopLevelComponentIndices.Add(Component->IndexInHierarchy);
+				}
+				else
+				{
+					FRigBaseElement* Element = Find(Component->GetElementKey());
+					check(Element);
+					
+					Component->Element = Element;
+					Component->IndexInElement = Element->ComponentIndices.Add(Component->IndexInHierarchy);
+				}
 				ComponentIndexLookup.Add(Component->GetKey(), ComponentIndex);
 			}
 		}
@@ -1086,6 +1096,8 @@ void URigHierarchy::CopyHierarchy(URigHierarchy* InHierarchy)
 		ElementComponents.Reserve(InHierarchy->ElementComponents.Num());
 		ComponentIndexLookup.Reset();
 		ComponentIndexLookup.Reserve(InHierarchy->ElementComponents.Num());
+		TopLevelComponentIndices.Reset();
+		TopLevelComponentIndices.Reserve(InHierarchy->TopLevelComponentIndices.Num());
 
 		for(int32 SourceComponentIndex = 0; SourceComponentIndex < InHierarchy->ElementComponents.Num(); SourceComponentIndex++)
 		{
@@ -1095,9 +1107,18 @@ void URigHierarchy::CopyHierarchy(URigHierarchy* InHierarchy)
 			}
 			const int32 TargetComponentIndex = ElementComponents.Add(InHierarchy->ElementComponents[SourceComponentIndex]);
 			FRigBaseComponent* TargetComponent = ElementComponents[TargetComponentIndex].GetMutablePtr<FRigBaseComponent>();
-			TargetComponent->Element = Find(TargetComponent->GetElementKey());
 			TargetComponent->IndexInHierarchy = TargetComponentIndex;
-			TargetComponent->IndexInElement = TargetComponent->Element->ComponentIndices.Add(TargetComponent->IndexInHierarchy);
+
+			if(TargetComponent->IsTopLevel())
+			{
+				TargetComponent->Element = nullptr;
+				TargetComponent->IndexInElement = TopLevelComponentIndices.Add(TargetComponentIndex);
+			}
+			else
+			{
+				TargetComponent->Element = Find(TargetComponent->GetElementKey());
+				TargetComponent->IndexInElement = TargetComponent->Element->ComponentIndices.Add(TargetComponent->IndexInHierarchy);
+			}
 			ComponentIndexLookup.Add(TargetComponent->GetKey(), TargetComponentIndex);
 		}
 
@@ -1585,9 +1606,37 @@ FRigBaseComponent* URigHierarchy::FindComponent(const FRigComponentKey& InKey)
 	return nullptr;
 }
 
+const FRigBaseComponent* URigHierarchy::FindTopLevelComponent(const FName& InName) const
+{
+	return FindComponent(FRigComponentKey(GetTopLevelComponentElementKey(), InName));
+}
+
+FRigBaseComponent* URigHierarchy::FindTopLevelComponent(const FName& InName)
+{
+	return FindComponent(FRigComponentKey(GetTopLevelComponentElementKey(), InName));
+}
+
+const FRigBaseComponent* URigHierarchy::GetTopLevelComponent(int32 InTopLevelComponentIndex) const
+{
+	if(TopLevelComponentIndices.IsValidIndex(InTopLevelComponentIndex))
+	{
+		return GetComponent(TopLevelComponentIndices[InTopLevelComponentIndex]);
+	}
+	return nullptr;
+}
+
+FRigBaseComponent* URigHierarchy::GetTopLevelComponent(int32 InTopLevelComponentIndex)
+{
+	if(TopLevelComponentIndices.IsValidIndex(InTopLevelComponentIndex))
+	{
+		return GetComponent(TopLevelComponentIndices[InTopLevelComponentIndex]);
+	}
+	return nullptr;
+}
+
 int32 URigHierarchy::GetComponentIndex(const FRigComponentKey& InComponentKey, bool bFollowRedirector) const
 {
-	if(bFollowRedirector && ElementKeyRedirector)
+	if(bFollowRedirector && ElementKeyRedirector && !InComponentKey.IsTopLevel())
 	{
 		const int32 ElementIndex = GetIndex(InComponentKey.ElementKey);
 		if(ElementIndex != INDEX_NONE)
@@ -1661,6 +1710,13 @@ TArray<FRigComponentKey> URigHierarchy::GetComponentKeys(FRigElementKey InElemen
 
 FRigComponentKey URigHierarchy::GetComponentKey(FRigElementKey InElement, int32 InComponentIndex) const
 {
+	if(InElement == GetTopLevelComponentElementKey())
+	{
+		if(const FRigBaseComponent* Component = GetTopLevelComponent(InComponentIndex))
+		{
+			return Component->GetKey();
+		}
+	}
 	if(const FRigBaseElement* Element = Find(InElement))
 	{
 		if(const FRigBaseComponent* Component = Element->GetComponent(InComponentIndex))
@@ -1673,6 +1729,13 @@ FRigComponentKey URigHierarchy::GetComponentKey(FRigElementKey InElement, int32
 
 FName URigHierarchy::GetComponentName(FRigElementKey InElement, int32 InComponentIndex) const
 {
+	if(InElement == GetTopLevelComponentElementKey())
+	{
+		if(const FRigBaseComponent* Component = GetTopLevelComponent(InComponentIndex))
+		{
+			return Component->GetFName();
+		}
+	}
 	if(const FRigBaseElement* Element = Find(InElement))
 	{
 		if(const FRigBaseComponent* Component = Element->GetComponent(InComponentIndex))
@@ -1685,6 +1748,13 @@ FName URigHierarchy::GetComponentName(FRigElementKey InElement, int32 InComponen
 
 UScriptStruct* URigHierarchy::GetComponentType(FRigElementKey InElement, int32 InComponentIndex) const
 {
+	if(InElement == GetTopLevelComponentElementKey())
+	{
+		if(const FRigBaseComponent* Component = GetTopLevelComponent(InComponentIndex))
+		{
+			return Component->GetScriptStruct();
+		}
+	}
 	if(const FRigBaseElement* Element = Find(InElement))
 	{
 		if(const FRigBaseComponent* Component = Element->GetComponent(InComponentIndex))
@@ -1697,6 +1767,15 @@ UScriptStruct* URigHierarchy::GetComponentType(FRigElementKey InElement, int32 I
 
 FString URigHierarchy::GetComponentContent(FRigElementKey InElement, int32 InComponentIndex) const
 {
+	if(InElement == GetTopLevelComponentElementKey())
+	{
+		if(const FRigBaseComponent* Component = GetTopLevelComponent(InComponentIndex))
+		{
+			FString Content;
+			Component->GetScriptStruct()->ExportText(Content, Component, Component, nullptr, PPF_None, nullptr);
+			return Content;
+		}
+	}
 	if(const FRigBaseElement* Element = Find(InElement))
 	{
 		if(const FRigBaseComponent* Component = Element->GetComponent(InComponentIndex))
@@ -1709,6 +1788,48 @@ FString URigHierarchy::GetComponentContent(FRigElementKey InElement, int32 InCom
 	return FString();
 }
 
+int32 URigHierarchy::NumTopLevelComponents() const
+{
+	return TopLevelComponentIndices.Num();
+}
+
+TArray<FRigComponentKey> URigHierarchy::GetTopLevelComponentKeys() const
+{
+	TArray<FRigComponentKey> Keys;
+	Keys.Reserve(NumTopLevelComponents());
+	for(int32 Index = 0; Index < NumTopLevelComponents(); Index++)
+	{
+		Keys.Add(GetTopLevelComponentKey(Index));
+	}
+	return Keys;
+}
+
+FRigComponentKey URigHierarchy::GetTopLevelComponentKey(int32 InTopLevelComponentIndex) const
+{
+	return GetComponentKey(GetTopLevelComponentElementKey(), InTopLevelComponentIndex);
+}
+
+FName URigHierarchy::GetTopLevelComponentName(int32 InTopLevelComponentIndex) const
+{
+	return GetComponentName(GetTopLevelComponentElementKey(), InTopLevelComponentIndex);
+}
+
+UScriptStruct* URigHierarchy::GetTopLevelComponentType(int32 InTopLevelComponentIndex) const
+{
+	return GetComponentType(GetTopLevelComponentElementKey(), InTopLevelComponentIndex);
+}
+
+FString URigHierarchy::GetTopLevelComponentContent(int32 InTopLevelComponentIndex) const
+{
+	return GetComponentContent(GetTopLevelComponentElementKey(), InTopLevelComponentIndex);
+}
+
+const FRigElementKey& URigHierarchy::GetTopLevelComponentElementKey()
+{
+	static const FRigElementKey TopLevelComponentElementKey = FRigElementKey(TEXT("TopLevel"), ERigElementType::None);
+	return TopLevelComponentElementKey;
+}
+
 TArray<FRigSocketState> URigHierarchy::GetSocketStates() const
 {
 	const TArray<FRigElementKey> Keys = GetSocketKeys(true);
@@ -6266,8 +6387,10 @@ FRigBaseComponent* URigHierarchy::MakeComponent(const UScriptStruct* InComponent
 {
 	check(InComponentStruct);
 	check(InComponentStruct->IsChildOf(FRigBaseComponent::StaticStruct()));
-	check(InElement);
-	check(InElement->GetOwner() == this);
+	if(InElement)
+	{
+		check(InElement->GetOwner() == this);
+	}
 	
 	int32 ComponentIndex = INDEX_NONE;
 	for(int32 ExistingIndex = 0; ExistingIndex < ElementComponents.Num(); ExistingIndex++)
@@ -6291,10 +6414,21 @@ FRigBaseComponent* URigHierarchy::MakeComponent(const UScriptStruct* InComponent
 	FRigName UniqueName = SanitizedName;
 
 	TArray<FName> ExistingNames;
-	ExistingNames.Reserve(InElement->NumComponents());
-	for(int32 ExistingComponentIndex = 0; ExistingComponentIndex < InElement->NumComponents(); ExistingComponentIndex++)
+	if(InElement)
+	{
+		ExistingNames.Reserve(InElement->NumComponents());
+		for(int32 ExistingComponentIndex = 0; ExistingComponentIndex < InElement->NumComponents(); ExistingComponentIndex++)
+		{
+			ExistingNames.Add(InElement->GetComponent(ExistingComponentIndex)->GetFName());
+		}
+	}
+	else
 	{
-		ExistingNames.Add(InElement->GetComponent(ExistingComponentIndex)->GetFName());
+		ExistingNames.Reserve(NumTopLevelComponents());
+		for(int32 ExistingComponentIndex = 0; ExistingComponentIndex < NumTopLevelComponents(); ExistingComponentIndex++)
+		{
+			ExistingNames.Add(GetTopLevelComponentName(ExistingComponentIndex));
+		}
 	}
 
 	int32 Suffix = 1;
@@ -6312,10 +6446,19 @@ FRigBaseComponent* URigHierarchy::MakeComponent(const UScriptStruct* InComponent
 	}
 	
 	FRigBaseComponent* Component = ElementComponents[ComponentIndex].GetMutablePtr<FRigBaseComponent>();
-	Component->Key = FRigComponentKey(InElement->GetKey(), UniqueName.GetFName());
-	Component->Element = InElement;
 	Component->IndexInHierarchy = ComponentIndex;
-	Component->IndexInElement = InElement->ComponentIndices.Add(Component->IndexInHierarchy);
+	if(InElement)
+	{
+		Component->Key = FRigComponentKey(InElement->GetKey(), UniqueName.GetFName());
+		Component->Element = InElement;
+		Component->IndexInElement = InElement->ComponentIndices.Add(Component->IndexInHierarchy);
+	}
+	else
+	{
+		Component->Key = FRigComponentKey(GetTopLevelComponentElementKey(), UniqueName.GetFName());
+		Component->Element = nullptr;
+		Component->IndexInElement = TopLevelComponentIndices.Add(Component->IndexInHierarchy);
+	}
 
 	ComponentIndexLookup.Add(Component->GetKey(), Component->GetIndexInHierarchy());
 	
@@ -6330,7 +6473,19 @@ void URigHierarchy::DestroyComponent(FRigBaseComponent*& InComponent)
 		const FRigBaseComponent* Component = InstancedStruct.GetMutablePtr<FRigBaseComponent>();
 		if(Component == InComponent)
 		{
-			if(FRigBaseElement* Element = Find(Component->GetElementKey()))
+			if(Component->GetElementKey() == GetTopLevelComponentElementKey())
+			{
+				TopLevelComponentIndices.Remove(ComponentIndex);
+				
+				for(int32 TopLevelComponentIndex = 0; TopLevelComponentIndex < NumTopLevelComponents(); TopLevelComponentIndex++)
+				{
+					if(FRigBaseComponent* RemainingComponent = GetTopLevelComponent(TopLevelComponentIndex))
+					{
+						RemainingComponent->IndexInElement = TopLevelComponentIndex;
+					}
+				}
+			}
+			else if(FRigBaseElement* Element = Find(Component->GetElementKey()))
 			{
 				Element->ComponentIndices.Remove(ComponentIndex);
 				
@@ -6379,6 +6534,8 @@ void URigHierarchy::ShrinkComponentStorage()
 		Element->ComponentIndices.Reset();
 	}
 
+	TopLevelComponentIndices.Reset();
+
 	ElementComponents.RemoveAll([](const FInstancedStruct& InstancedStruct) -> bool
 	{
 		return !InstancedStruct.IsValid();
@@ -6391,8 +6548,15 @@ void URigHierarchy::ShrinkComponentStorage()
 		FInstancedStruct& InstancedStruct = ElementComponents[ComponentIndex];
 		FRigBaseComponent* Component = InstancedStruct.GetMutablePtr<FRigBaseComponent>();
 		Component->IndexInHierarchy = ComponentIndex;
-		Component->Element = Find(Component->GetElementKey());
-		Component->Element->ComponentIndices.Add(ComponentIndex);
+		if(Component->Element)
+		{
+			Component->Element = Find(Component->GetElementKey());
+			Component->Element->ComponentIndices.Add(ComponentIndex);
+		}
+		else if(Component->GetElementKey() == GetTopLevelComponentElementKey())
+		{
+			TopLevelComponentIndices.Add(ComponentIndex);
+		}
 		ComponentIndexLookup.Add(Component->GetKey(), ComponentIndex);
 	}
 }
diff --git a/Engine/Plugins/Animation/ControlRig/Source/ControlRig/Private/Rigs/RigHierarchyController.cpp b/Engine/Plugins/Animation/ControlRig/Source/ControlRig/Private/Rigs/RigHierarchyController.cpp
--- a/Engine/Plugins/Animation/ControlRig/Source/ControlRig/Private/Rigs/RigHierarchyController.cpp
+++ b/Engine/Plugins/Animation/ControlRig/Source/ControlRig/Private/Rigs/RigHierarchyController.cpp
@@ -860,12 +860,57 @@ FRigComponentKey URigHierarchyController::AddComponent(UScriptStruct* InComponen
 
 	URigHierarchy* Hierarchy = GetHierarchy();
 	check(Hierarchy);
-	
-	FRigBaseElement* Element = Hierarchy->Find(InElement);
-	if(Element == nullptr)
+
+	FRigBaseElement* Element = nullptr;
+	if(InElement != URigHierarchy::GetTopLevelComponentElementKey())
 	{
-		ReportErrorf(TEXT("The element '%s' could not be found."), *InElement.ToString());
-		return FRigComponentKey();
+		Element = Hierarchy->Find(InElement);
+		if(Element == nullptr)
+		{
+			ReportErrorf(TEXT("The element '%s' could not be found."), *InElement.ToString());
+			return FRigComponentKey();
+		}
+	}
+
+	const FStructOnScope StructOnScope(InComponentStruct);
+	if(const FRigBaseComponent* StructMemory = reinterpret_cast<const FRigBaseComponent*>(StructOnScope.GetStructMemory()))
+	{
+		if(StructMemory->IsSingleton())
+		{
+			if(Element)
+			{
+				for(int32 Index = 0; Index < Element->NumComponents(); Index++)
+				{
+					if(const FRigBaseComponent* ExistingComponent = Element->GetComponent(Index))
+					{
+						if(ExistingComponent->GetScriptStruct() == InComponentStruct)
+						{
+							ReportErrorf(TEXT("A component of type '%s' already exists on element '%s'."), *InComponentStruct->GetName(), *InElement.ToString());
+							return FRigComponentKey();
+						}
+					}
+				}
+			}
+			else
+			{
+				for(int32 Index = 0; Index < Hierarchy->NumTopLevelComponents(); Index++)
+				{
+					if(const FRigBaseComponent* ExistingComponent = Hierarchy->GetTopLevelComponent(Index))
+					{
+						if(ExistingComponent->GetScriptStruct() == InComponentStruct)
+						{
+							ReportErrorf(TEXT("A component of type '%s' already exists at the top level."), *InComponentStruct->GetName());
+							return FRigComponentKey();
+						}
+					}
+				}
+			}
+		}
+		if(!StructMemory->CanBeAddedTo(InElement, Hierarchy))
+		{
+			ReportErrorf(TEXT("The component of type '%s' cannot be added to element '%s'."), *InComponentStruct->GetName(), *InElement.ToString());
+			return FRigComponentKey();
+		}
 	}
 
 #if WITH_EDITOR
@@ -906,6 +951,12 @@ FRigComponentKey URigHierarchyController::AddComponent(UScriptStruct* InComponen
 	return NewComponent->GetKey();
 }
 
+FRigComponentKey URigHierarchyController::AddTopLevelComponent(UScriptStruct* InComponentStruct, FName InName, FString InContent, bool bSetupUndo,
+	bool bPrintPythonCommand)
+{
+	return AddComponent(InComponentStruct, InName, URigHierarchy::GetTopLevelComponentElementKey(), InContent, bSetupUndo, bPrintPythonCommand);
+}
+
 bool URigHierarchyController::RemoveComponent(FRigComponentKey InComponent, bool bSetupUndo, bool bPrintPythonCommand)
 {
 	if(!IsValid())
diff --git a/Engine/Plugins/Animation/ControlRig/Source/ControlRig/Private/Rigs/RigHierarchyDefines.cpp b/Engine/Plugins/Animation/ControlRig/Source/ControlRig/Private/Rigs/RigHierarchyDefines.cpp
--- a/Engine/Plugins/Animation/ControlRig/Source/ControlRig/Private/Rigs/RigHierarchyDefines.cpp
+++ b/Engine/Plugins/Animation/ControlRig/Source/ControlRig/Private/Rigs/RigHierarchyDefines.cpp
@@ -174,6 +174,12 @@ void FRigComponentKey::Load(FArchive& Ar)
 	Ar << ElementKey;
 }
 
+bool FRigComponentKey::IsValid() const
+{
+	return Name.IsValid() && Name != NAME_None &&
+		(ElementKey.IsValid() || ElementKey == URigHierarchy::GetTopLevelComponentElementKey());
+}
+
 FString FRigComponentKey::ToPythonString() const
 {
 #if WITH_EDITOR
@@ -185,6 +191,11 @@ FString FRigComponentKey::ToPythonString() const
 #endif
 }
 
+bool FRigComponentKey::IsTopLevel() const
+{
+	return ElementKey == URigHierarchy::GetTopLevelComponentElementKey();
+}
+
 ////////////////////////////////////////////////////////////////////////////////
 // FRigElementKeyAndIndex
 ////////////////////////////////////////////////////////////////////////////////
diff --git a/Engine/Plugins/Animation/ControlRig/Source/ControlRig/Public/Rigs/RigHierarchy.h b/Engine/Plugins/Animation/ControlRig/Source/ControlRig/Public/Rigs/RigHierarchy.h
--- a/Engine/Plugins/Animation/ControlRig/Source/ControlRig/Public/Rigs/RigHierarchy.h
+++ b/Engine/Plugins/Animation/ControlRig/Source/ControlRig/Public/Rigs/RigHierarchy.h
@@ -796,6 +796,86 @@ public:
 		return CastChecked<T>(GetComponent(InIndex));
 	}
 	
+	/**
+	 * Returns a component for a given top level name
+	 * @param InName The name of the component to retrieve.
+	 * @return The retrieved component or nullptr.
+	 */
+	const FRigBaseComponent* FindTopLevelComponent(const FName& InName) const;
+
+	/**
+	 * Returns a component for given top level name
+	 * @param InName The name of the component to retrieve.
+	 * @return The retrieved component or nullptr.
+	 */
+	FRigBaseComponent* FindTopLevelComponent(const FName& InName);
+
+	/**
+	 * Returns a component at a given index or nullptr.
+	 * @param InTopLevelComponentIndex The index of the top level component to retrieve.
+	 * @return The retrieved component or nullptr.
+	 */
+	const FRigBaseComponent* GetTopLevelComponent(int32 InTopLevelComponentIndex) const;
+
+	/**
+	 * Returns a component at a given index or nullptr.
+	 * @param InTopLevelComponentIndex The index of the top level component to retrieve.
+	 * @return The retrieved component or nullptr.
+	 */
+	FRigBaseComponent* GetTopLevelComponent(int32 InTopLevelComponentIndex);
+
+	/**
+	 * Returns a component at a given index or nullptr.
+	 * This templated method also casts to the chosen
+	 * component type but does not guarantee a valid result.
+	 * @param InTopLevelComponentIndex The index of the top level component to retrieve.
+	 * @return The retrieved component or nullptr.
+	 */
+	template<typename T>
+	const T* GetTopLevelComponent(int32 InTopLevelComponentIndex) const
+	{
+		return Cast<T>(GetTopLevelComponent(InTopLevelComponentIndex));
+	}
+
+	/**
+	 * Returns a component at a given index or nullptr.
+	 * This templated method also casts to the chosen
+	 * component type but does not guarantee a valid result.
+	 * @param InTopLevelComponentIndex The index of the top level component to retrieve.
+	 * @return The retrieved component or nullptr.
+	 */
+	template<typename T>
+	T* GetTopLevelComponent(int32 InTopLevelComponentIndex)
+	{
+		return Cast<T>(GetTopLevelComponent(InTopLevelComponentIndex));
+	}
+
+	/**
+	 * Returns a component at a given index.
+	 * This templated method also casts to the chosen
+	 * component type and checks for a the valid result.
+	 * @param InTopLevelComponentIndex The index of the top level component to retrieve.
+	 * @return The retrieved component or nullptr.
+	 */
+	template<typename T>
+	const T* GetTopLevelComponentChecked(int32 InTopLevelComponentIndex) const
+	{
+		return CastChecked<T>(GetTopLevelComponent(InTopLevelComponentIndex));
+	}
+
+	/**
+	 * Returns a component at a given index.
+	 * This templated method also casts to the chosen
+	 * component type and checks for a the valid result.
+	 * @param InTopLevelComponentIndex The index of the top level component to retrieve.
+	 * @return The retrieved component or nullptr.
+	 */
+	template<typename T>
+    T* GetTopLevelComponentChecked(int32 InTopLevelComponentIndex)
+	{
+		return CastChecked<T>(GetTopLevelComponent(InTopLevelComponentIndex));
+	}
+
 	UFUNCTION(BlueprintPure, Category = URigHierarchy)
 	int32 NumComponents(FRigElementKey InElement) const;
 
@@ -817,6 +897,26 @@ public:
 	UFUNCTION(BlueprintPure, Category = URigHierarchy)
 	FString GetComponentContent(FRigElementKey InElement, int32 InComponentIndex) const;
 
+	UFUNCTION(BlueprintPure, Category = URigHierarchy)
+	int32 NumTopLevelComponents() const;
+
+	UFUNCTION(BlueprintPure, Category = URigHierarchy)
+	TArray<FRigComponentKey> GetTopLevelComponentKeys() const;
+
+	UFUNCTION(BlueprintPure, Category = URigHierarchy)
+	FRigComponentKey GetTopLevelComponentKey(int32 InTopLevelComponentIndex) const;
+
+	UFUNCTION(BlueprintPure, Category = URigHierarchy)
+	FName GetTopLevelComponentName(int32 InTopLevelComponentIndex) const;
+	
+	UFUNCTION(BlueprintPure, Category = URigHierarchy)
+	UScriptStruct* GetTopLevelComponentType(int32 InTopLevelComponentIndex) const;
+
+	UFUNCTION(BlueprintPure, Category = URigHierarchy)
+	FString GetTopLevelComponentContent(int32 InTopLevelComponentIndex) const;
+
+	static const FRigElementKey& GetTopLevelComponentElementKey();
+
 private:
 	/**
 	* Returns bone element for a given key, for scripting purpose only, for cpp usage, use Find<FRigBoneElement>()
@@ -4573,6 +4673,9 @@ private:
 	// the component storage.
 	TArray<FInstancedStruct> ElementComponents;
 
+	// List of top level component indices
+	TArray<int32> TopLevelComponentIndices;
+
 	// A quick-lookup cache for elements' children. Each element that has a ChildCacheIndex
 	// not equal to INDEX_NONE is an index into the offset and count cache below, which in
 	// turn contains an offset into the ChildElementCache, which stores consecutive runs of
diff --git a/Engine/Plugins/Animation/ControlRig/Source/ControlRig/Public/Rigs/RigHierarchyComponents.h b/Engine/Plugins/Animation/ControlRig/Source/ControlRig/Public/Rigs/RigHierarchyComponents.h
--- a/Engine/Plugins/Animation/ControlRig/Source/ControlRig/Public/Rigs/RigHierarchyComponents.h
+++ b/Engine/Plugins/Animation/ControlRig/Source/ControlRig/Public/Rigs/RigHierarchyComponents.h
@@ -59,6 +59,14 @@ public:
 		bSelected = InOther.bSelected;
 		return *this;
 	}
+
+	// returns true if this component can only be added once
+	virtual bool IsSingleton() const { return false; }
+
+	// returns true if this component can be added to a given key.
+	// if you want to determine if a component can be added as a top level component
+	// you need to pass the URigHierarchy::GetTopLevelComponentElementKey() as InElementKey
+	virtual bool CanBeAddedTo(const FRigElementKey& InElementKey, const URigHierarchy* InHierarchy) const { return true; }
 	
 protected:
 
@@ -92,6 +100,10 @@ public:
 		}
 		return CachedNameString;
 	}
+	bool IsTopLevel() const
+	{
+		return Key.IsTopLevel();
+	}
 
 	virtual const FName& GetDisplayName() const { return GetFName(); }
 	int32 GetIndexInElement() const { return IndexInElement; }
diff --git a/Engine/Plugins/Animation/ControlRig/Source/ControlRig/Public/Rigs/RigHierarchyController.h b/Engine/Plugins/Animation/ControlRig/Source/ControlRig/Public/Rigs/RigHierarchyController.h
--- a/Engine/Plugins/Animation/ControlRig/Source/ControlRig/Public/Rigs/RigHierarchyController.h
+++ b/Engine/Plugins/Animation/ControlRig/Source/ControlRig/Public/Rigs/RigHierarchyController.h
@@ -275,6 +275,18 @@ public:
 	UFUNCTION(BlueprintCallable, Category = URigHierarchyController)
 	FRigComponentKey AddComponent(UScriptStruct* InComponentStruct, FName InName, FRigElementKey InElement, FString InContent = TEXT(""), bool bSetupUndo = false, bool bPrintPythonCommand = false);
 
+	/**
+	 * Adds a component to the hierarchy top level of the hierarchy
+	 * @param InComponentStruct The script struct of the component to add
+	 * @param InName The suggested name of the new component.
+	 * @param InContent The (optional) serialized text default for the component
+	 * @param bSetupUndo If set to true the stack will record the change for undo / redo
+	 * @param bPrintPythonCommand If set to true a python command equivalent to this call will be printed out
+	 * @return The name of the newly created component
+	 */
+	UFUNCTION(BlueprintCallable, Category = URigHierarchyController)
+	FRigComponentKey AddTopLevelComponent(UScriptStruct* InComponentStruct, FName InName, FString InContent = TEXT(""), bool bSetupUndo = false, bool bPrintPythonCommand = false);
+
 	/**
      * Removes a component from the hierarchy
      * @param InComponent The component to remove
diff --git a/Engine/Plugins/Animation/ControlRig/Source/ControlRig/Public/Rigs/RigHierarchyDefines.h b/Engine/Plugins/Animation/ControlRig/Source/ControlRig/Public/Rigs/RigHierarchyDefines.h
--- a/Engine/Plugins/Animation/ControlRig/Source/ControlRig/Public/Rigs/RigHierarchyDefines.h
+++ b/Engine/Plugins/Animation/ControlRig/Source/ControlRig/Public/Rigs/RigHierarchyDefines.h
@@ -1694,10 +1694,7 @@ public:
 		return Ar;
 	}
 
-	bool IsValid() const
-	{
-		return Name.IsValid() && Name != NAME_None && ElementKey.IsValid();
-	}
+	bool IsValid() const;
 
 	explicit operator bool() const
 	{
@@ -1764,6 +1761,8 @@ public:
 	}
 
 	FString ToPythonString() const;
+
+	bool IsTopLevel() const;
 };
 
 struct CONTROLRIG_API FRigElementKeyAndIndex
