Anim In Engine: When selecting anim layers select controls in the new layer in sequencer, deselecting the old ones.

diff --git a/Engine/Plugins/Animation/ControlRig/Source/ControlRigEditor/Private/Sequencer/AnimLayers/AnimLayers.cpp b/Engine/Plugins/Animation/ControlRig/Source/ControlRigEditor/Private/Sequencer/AnimLayers/AnimLayers.cpp
--- a/Engine/Plugins/Animation/ControlRig/Source/ControlRigEditor/Private/Sequencer/AnimLayers/AnimLayers.cpp
+++ b/Engine/Plugins/Animation/ControlRig/Source/ControlRigEditor/Private/Sequencer/AnimLayers/AnimLayers.cpp
@@ -610,46 +610,61 @@ void UAnimLayer::SetSelectedInList(bool bInValue)
 	bIsSelectedInList = bInValue;
 	if (TSharedPtr<ISequencer> SequencerPtr = UAnimLayers::GetSequencerFromAsset())
 	{
-		TArray<UMovieSceneSection*> Sections;
-		GetSections(Sections);
 		using namespace UE::Sequencer;
 		FOutlinerSelection& SelectedOutlinerItems = SequencerPtr->GetViewModel()->GetSelection()->Outliner;
-		if (bIsSelectedInList) //selected so select it's track row model
+
+		FTrackRowModelStorageExtension* TrackRowModelStorage = SequencerPtr->GetViewModel()->GetRootModel()->CastDynamic<FTrackRowModelStorageExtension>();
+		check(TrackRowModelStorage);
+		for (const TPair<TWeakObjectPtr<UObject>, FAnimLayerItem>& Pair : AnimLayerItems)
 		{
-			FTrackRowModelStorageExtension* TrackRowModelStorage = SequencerPtr->GetViewModel()->GetRootModel()->CastDynamic<FTrackRowModelStorageExtension>();
-			check(TrackRowModelStorage);
-			for (UMovieSceneSection* Section : Sections)
-			{
-				if (UMovieSceneTrack* Track =Section->GetTypedOuter<UMovieSceneTrack>())
-				{
-					const int32 RowIndex = Section->GetRowIndex();
-					TSharedPtr<FTrackRowModel> TrackRowModel = TrackRowModelStorage->FindModelForTrackRow(Track, RowIndex);
-					SelectedOutlinerItems.Select(TrackRowModel);
-				}
-			}
-		}
-		else //not selected, so clear it from outliner
-		{ 
-			TArray <TViewModelPtr<FTrackRowModel>> TrackRowModelsToDeselect;
-			for (FViewModelPtr ViewModel : SelectedOutlinerItems)
+			if (Pair.Key != nullptr)
 			{
-				if (TViewModelPtr<FTrackRowModel> TrackRowModel = ViewModel.ImplicitCast())
+				for (const FAnimLayerSectionItem& SectionItem : Pair.Value.SectionItems)
 				{
-					for (UMovieSceneSection* Section : TrackRowModel->GetSections())
+					if (SectionItem.Section.IsValid())
 					{
-						if (Sections.Contains(Section))
+						const UControlRig* ControlRig = Cast<UControlRig>(Pair.Key.Get());
+						if (ControlRig)
 						{
-							TrackRowModelsToDeselect.Add(TrackRowModel);
-							break;
+							if (UMovieSceneControlRigParameterSection* CRSection = Cast<UMovieSceneControlRigParameterSection>(SectionItem.Section.Get()))
+							{
+								for (const TPair<FName, FAnimLayerPropertyAndChannels>& ControlName : SectionItem.AnimLayerSet.Names)
+								{
+									FChannelMapInfo* pChannelIndex = CRSection->ControlChannelMap.Find(ControlName.Key);
+									if (pChannelIndex != nullptr)
+									{
+										if (pChannelIndex->ParentControlIndex == INDEX_NONE)
+										{
+											int32 CategoryIndex = CRSection->GetActiveCategoryIndex(ControlName.Key);
+											if (CategoryIndex != INDEX_NONE)
+											{
+												SequencerPtr->SelectByNthCategoryNode(CRSection, CategoryIndex, bIsSelectedInList);
+											}
+										}
+									}
+								}
+							}
+						}
+						if(bIsSelectedInList == false || ControlRig == nullptr) //if not a control rig we select the whole trackrow, always make sure to deselect it
+						{
+							if (UMovieSceneTrack* Track = SectionItem.Section->GetTypedOuter<UMovieSceneTrack>())
+							{
+								const int32 RowIndex = SectionItem.Section->GetRowIndex();
+								TSharedPtr<FTrackRowModel> TrackRowModel = TrackRowModelStorage->FindModelForTrackRow(Track, RowIndex);
+								if (bIsSelectedInList)
+								{
+									SelectedOutlinerItems.Select(TrackRowModel);
+								}
+								else
+								{
+									SelectedOutlinerItems.Deselect(TrackRowModel);
+								}
+							}
 						}
 					}
 				}
 			}
-			for (TViewModelPtr<FTrackRowModel>& TrackRowModel:  TrackRowModelsToDeselect)
-			{
-				SelectedOutlinerItems.Deselect(TrackRowModel);
-			}
-		}
+		}	
 	}
 }
 
