Control Rig: Support for components within the SRigHierarchy

#rnx

diff --git a/Engine/Plugins/Animation/ControlRig/Source/ControlRig/Private/Rigs/RigHierarchy.cpp b/Engine/Plugins/Animation/ControlRig/Source/ControlRig/Private/Rigs/RigHierarchy.cpp
--- a/Engine/Plugins/Animation/ControlRig/Source/ControlRig/Private/Rigs/RigHierarchy.cpp
+++ b/Engine/Plugins/Animation/ControlRig/Source/ControlRig/Private/Rigs/RigHierarchy.cpp
@@ -203,7 +203,7 @@ void URigHierarchy::Save(FArchive& Ar)
 			return;
 		}
 
-		TArray<FRigElementKey> SelectedKeys = GetSelectedKeys();
+		TArray<FRigHierarchyKey> SelectedKeys = GetSelectedHierarchyKeys();
 		Ar << SelectedKeys;
 		UE_RIGVM_ARCHIVETRACE_ENTRY(Ar, TEXT("SelectedKeys"));
 	}
@@ -394,7 +394,7 @@ void URigHierarchy::Load(FArchive& Ar)
 	// unlink any existing pose adapter
 	UnlinkPoseAdapter();
 
-	TArray<FRigElementKey> SelectedKeys;
+	TArray<FRigHierarchyKey> SelectedKeys;
 	if(Ar.IsTransacting())
 	{
 		bool bOnlySerializedTransformStackIndex = false;
@@ -602,12 +602,23 @@ void URigHierarchy::Load(FArchive& Ar)
 
 	if(Ar.IsTransacting())
 	{
-		for(const FRigElementKey& SelectedKey : SelectedKeys)
+		for(const FRigHierarchyKey& SelectedKey : SelectedKeys)
 		{
-			if(FRigBaseElement* Element = Find<FRigBaseElement>(SelectedKey))
+			if(SelectedKey.IsElement())
 			{
-				Element->bSelected = true;
-				OrderedSelection.Add(SelectedKey);
+				if(FRigBaseElement* Element = Find<FRigBaseElement>(SelectedKey.GetElement()))
+				{
+					Element->bSelected = true;
+					OrderedSelection.Add(SelectedKey);
+				}
+			}
+			else if(SelectedKey.IsComponent())
+			{
+				if(FRigBaseComponent* Component = FindComponent(SelectedKey.GetComponent()))
+				{
+					Component->bSelected = true;
+					OrderedSelection.Add(SelectedKey);
+				}
 			}
 		}
 	}
@@ -775,6 +786,7 @@ void URigHierarchy::Reset_Impl(bool bResetElements)
 		// clearing the components first will avoid cost of looping over
 		// each component per elements and destroying them one by onw
 		ElementComponents.Reset();
+		TopLevelComponentIndices.Reset();
 
 		// walk in reverse since certain elements might not have been allocated themselves
 		for(int32 ElementIndex = Elements.Num() - 1; ElementIndex >= 0; ElementIndex--)
@@ -1218,7 +1230,10 @@ uint32 URigHierarchy::GetTopologyHash(bool bIncludeTopologyVersion, bool bInclud
 			Hash = HashCombine(Hash, GetTypeHash(Component->GetScriptStruct()));
 			Hash = HashCombine(Hash, GetTypeHash(Component->GetIndexInHierarchy()));
 			Hash = HashCombine(Hash, GetTypeHash(Component->GetIndexInElement()));
-			Hash = HashCombine(Hash, GetTypeHash(Component->GetElement()->GetIndex()));
+			if(Component->GetElement())
+			{
+				Hash = HashCombine(Hash, GetTypeHash(Component->GetElement()->GetIndex()));
+			}
 		}
 	}
 
@@ -1586,6 +1601,33 @@ bool URigHierarchy::IsProcedural(const FRigBaseElement* InElement) const
 	return InElement->IsProcedural();
 }
 
+bool URigHierarchy::IsProcedural(const FRigComponentKey& InKey) const
+{
+	return IsProcedural(FindComponent(InKey));
+}
+
+bool URigHierarchy::IsProcedural(const FRigBaseComponent* InComponent) const
+{
+	if(InComponent == nullptr)
+	{
+		return false;
+	}
+	return InComponent->IsProcedural();
+}
+
+bool URigHierarchy::IsProcedural(const FRigHierarchyKey& InKey) const
+{
+	if(InKey.IsElement())
+	{
+		return IsProcedural(InKey.GetElement());
+	}
+	if(InKey.IsComponent())
+	{
+		return IsProcedural(InKey.GetComponent());
+	}
+	return false;
+}
+
 const FRigBaseComponent* URigHierarchy::FindComponent(const FRigComponentKey& InKey) const
 {
 	const int32 ComponentIndex = GetComponentIndex(InKey);
@@ -2026,27 +2068,62 @@ TArray<const FRigBaseElement*> URigHierarchy::GetSelectedElements(ERigElementTyp
 	LLM_SCOPE_BYNAME(TEXT("Animation/ControlRig"));
 	TArray<const FRigBaseElement*> Selection;
 
-	if(URigHierarchy* HierarchyForSelection = HierarchyForSelectionPtr.Get())
+	if(const URigHierarchy* HierarchyForSelection = HierarchyForSelectionPtr.Get())
 	{
-		TArray<FRigElementKey> SelectedKeys = HierarchyForSelection->GetSelectedKeys(InTypeFilter);
-		for(const FRigElementKey& SelectedKey : SelectedKeys)
+		TArray<const FRigBaseElement*> SelectedElements = HierarchyForSelection->GetSelectedElements(InTypeFilter);
+		for(const FRigBaseElement* SelectedElement : SelectedElements)
 		{
-			if(const FRigBaseElement* Element = Find(SelectedKey))
+			if(const FRigBaseElement* Element = Find(SelectedElement->GetKey()))
 			{
-				Selection.Add((FRigBaseElement*)Element);
+				Selection.Add(Element);
 			}
 		}
 		return Selection;
 	}
 
-	for (const FRigElementKey& SelectedKey : OrderedSelection)
+	for (const FRigHierarchyKey& SelectedKey : OrderedSelection)
 	{
-		if(SelectedKey.IsTypeOf(InTypeFilter))
+		if(SelectedKey.IsElement())
 		{
-			if(const FRigBaseElement* Element = FindChecked(SelectedKey))
+			if(SelectedKey.GetElement().IsTypeOf(InTypeFilter))
 			{
-				ensure(Element->IsSelected());
-				Selection.Add(Element);
+				if(const FRigBaseElement* Element = FindChecked(SelectedKey.GetElement()))
+				{
+					ensure(Element->IsSelected());
+					Selection.Add(Element);
+				}
+			}
+		}
+	}
+	return Selection;
+}
+
+TArray<const FRigBaseComponent*> URigHierarchy::GetSelectedComponents() const
+{
+	LLM_SCOPE_BYNAME(TEXT("Animation/ControlRig"));
+	TArray<const FRigBaseComponent*> Selection;
+
+	if(const URigHierarchy* HierarchyForSelection = HierarchyForSelectionPtr.Get())
+	{
+		TArray<const FRigBaseComponent*> SelectedComponents = HierarchyForSelection->GetSelectedComponents();
+		for(const FRigBaseComponent* SelectedComponent : SelectedComponents)
+		{
+			if(const FRigBaseComponent* Component = FindComponent(SelectedComponent->GetKey()))
+			{
+				Selection.Add(Component);
+			}
+		}
+		return Selection;
+	}
+
+	for (const FRigHierarchyKey& SelectedKey : OrderedSelection)
+	{
+		if(SelectedKey.IsComponent())
+		{
+			if(const FRigBaseComponent* Component = FindComponent(SelectedKey.GetComponent()))
+			{
+				ensure(Component->IsSelected());
+				Selection.Add(Component);
 			}
 		}
 	}
@@ -2057,47 +2134,74 @@ TArray<FRigElementKey> URigHierarchy::GetSelectedKeys(ERigElementType InTypeFilt
 {
 	LLM_SCOPE_BYNAME(TEXT("Animation/ControlRig"));
 	
-	if(URigHierarchy* HierarchyForSelection = HierarchyForSelectionPtr.Get())
+	if(const URigHierarchy* HierarchyForSelection = HierarchyForSelectionPtr.Get())
 	{
 		return HierarchyForSelection->GetSelectedKeys(InTypeFilter);
 	}
 
 	TArray<FRigElementKey> Selection;
-	for (const FRigElementKey& SelectedKey : OrderedSelection)
+	for (const FRigHierarchyKey& SelectedKey : OrderedSelection)
 	{
-		if(SelectedKey.IsTypeOf(InTypeFilter))
+		if(SelectedKey.IsElement())
 		{
-			Selection.Add(SelectedKey);
+			if(SelectedKey.GetElement().IsTypeOf(InTypeFilter))
+			{
+				Selection.Add(SelectedKey.GetElement());
+			}
 		}
 	}
 	
 	return Selection;
 }
 
+TArray<FRigHierarchyKey> URigHierarchy::GetSelectedHierarchyKeys() const
+{
+	LLM_SCOPE_BYNAME(TEXT("Animation/ControlRig"));
+	
+	if(const URigHierarchy* HierarchyForSelection = HierarchyForSelectionPtr.Get())
+	{
+		return HierarchyForSelection->GetSelectedHierarchyKeys();
+	}
+
+	return OrderedSelection;
+}
+
 bool URigHierarchy::HasAnythingSelectedByPredicate(TFunctionRef<bool(const FRigElementKey&)> InPredicate) const
 {
-	if (URigHierarchy* HierarchyForSelection = HierarchyForSelectionPtr.Get())
+	if (const URigHierarchy* HierarchyForSelection = HierarchyForSelectionPtr.Get())
 	{
 		return HierarchyForSelection->HasAnythingSelectedByPredicate(InPredicate);
 	}
 
-	return OrderedSelection.ContainsByPredicate([InPredicate](const FRigElementKey& SelectedKey)
+	return OrderedSelection.ContainsByPredicate([InPredicate](const FRigHierarchyKey& SelectedKey)
 	{
-		return InPredicate(SelectedKey);
+		if(SelectedKey.IsElement())
+		{
+			return InPredicate(SelectedKey.GetElement());
+		}
+		return false;
 	});
 }
 
 TArray<FRigElementKey> URigHierarchy::GetSelectedKeysByPredicate(TFunctionRef<bool(const FRigElementKey&)> InPredicate) const
 {
-	if(URigHierarchy* HierarchyForSelection = HierarchyForSelectionPtr.Get())
+	if(const URigHierarchy* HierarchyForSelection = HierarchyForSelectionPtr.Get())
 	{
 		return HierarchyForSelection->GetSelectedKeysByPredicate(InPredicate);
 	}
 
-	return OrderedSelection.FilterByPredicate([InPredicate](const FRigElementKey& SelectedKey)
+	TArray<FRigElementKey> Result;
+	for(const FRigHierarchyKey& Key : OrderedSelection)
 	{
-		return InPredicate(SelectedKey);
-	});
+		if(Key.IsElement())
+		{
+			if(InPredicate(Key.GetElement()))
+			{
+				Result.Add(Key.GetElement());
+			}
+		}
+	}
+	return Result;
 }
 
 FRigName URigHierarchy::JoinNameSpace(const FRigName& InLeft, const FRigName& InRight)
@@ -3645,8 +3749,7 @@ TArray<FRigElementKey> URigHierarchy::GetKeysByPredicate(
 			const TArray<FRigBaseElement*> RootElements = GetRootElements();
 			for (FRigBaseElement* Element : RootElements)
 			{
-				const int32 ElementIndex = Element->GetIndex();
-				Traverse(Element, true, [&ElementVisited, InProcessFunc, InPredicateFunc](FRigBaseElement* InElement, bool& bContinue)
+				Traverse(Element, true, [&ElementVisited, InProcessFunc, InPredicateFunc](const FRigBaseElement* InElement, bool& bContinue)
 				{
 					bContinue = !ElementVisited[InElement->GetIndex()];
 
@@ -6043,6 +6146,21 @@ bool URigHierarchy::IsSelected(const FRigBaseElement* InElement) const
 	return bIsSelected;
 }
 
+bool URigHierarchy::IsComponentSelected(const FRigBaseComponent* InComponent) const
+{
+	if(InComponent == nullptr)
+	{
+		return false;
+	}
+	if(const URigHierarchy* HierarchyForSelection = HierarchyForSelectionPtr.Get())
+	{
+		return HierarchyForSelection->IsComponentSelected(InComponent->GetKey());
+	}
+
+	const bool bIsSelected = OrderedSelection.Contains(InComponent->GetKey());
+	ensure(bIsSelected == InComponent->IsSelected());
+	return bIsSelected;
+}
 
 void URigHierarchy::EnsureCachedChildrenAreCurrent() const
 {
@@ -8025,7 +8143,7 @@ void URigHierarchy::UpdateVisibilityOnProxyControls()
 	}
 
 	// create a local map of visible things, starting with the selection
-	TSet<FRigElementKey> VisibleElements;
+	TSet<FRigHierarchyKey> VisibleElements;
 	VisibleElements.Append(HierarchyForSelection->OrderedSelection);
 
 	// if the control is visible - we should consider the
@@ -8039,7 +8157,10 @@ void URigHierarchy::UpdateVisibilityOnProxyControls()
 			{
 				if(VisibleElements.Contains(ControlElement->GetKey()))
 				{
-					VisibleElements.Append(ControlElement->Settings.DrivenControls);
+					for(const FRigElementKey& DrivenControlKey : ControlElement->Settings.DrivenControls)
+					{
+						VisibleElements.Add(DrivenControlKey);
+					}
 				}
 			}
 		}
diff --git a/Engine/Plugins/Animation/ControlRig/Source/ControlRig/Private/Rigs/RigHierarchyComponents.cpp b/Engine/Plugins/Animation/ControlRig/Source/ControlRig/Private/Rigs/RigHierarchyComponents.cpp
--- a/Engine/Plugins/Animation/ControlRig/Source/ControlRig/Private/Rigs/RigHierarchyComponents.cpp
+++ b/Engine/Plugins/Animation/ControlRig/Source/ControlRig/Private/Rigs/RigHierarchyComponents.cpp
@@ -2,7 +2,7 @@
 
 #include "Rigs/RigHierarchyComponents.h"
 #include "Rigs/RigHierarchyElements.h"
-#include "ControlRigObjectVersion.h"
+#include "UObject/UObjectIterator.h"
 
 #include UE_INLINE_GENERATED_CPP_BY_NAME(RigHierarchyComponents)
 
@@ -14,6 +14,31 @@ FRigBaseComponent::~FRigBaseComponent()
 {
 }
 
+TArray<UScriptStruct*> FRigBaseComponent::GetAllComponentScriptStructs(bool bSorted)
+{
+	TArray<UScriptStruct*> RigComponentStructs;
+	for (TObjectIterator<UScriptStruct> ScriptStructIt; ScriptStructIt; ++ScriptStructIt)
+	{
+		UScriptStruct* ScriptStruct = *ScriptStructIt;
+		if(ScriptStruct != StaticStruct() && ScriptStruct->IsChildOf(StaticStruct()))
+		{
+			if(ScriptStruct->GetStructCPPName().Contains(TEXT("Base")))
+			{
+				continue;
+			}
+			RigComponentStructs.Add(ScriptStruct);
+		}
+	}
+	if(bSorted)
+	{
+		Algo::Sort(RigComponentStructs, [](const UScriptStruct* A, const UScriptStruct* B) -> bool
+		{
+			return A->GetFName().LexicalLess(B->GetFName());
+		});
+	}
+	return RigComponentStructs;
+}
+
 bool FRigBaseComponent::Serialize(FArchive& Ar)
 {
 	if(Ar.IsLoading())
@@ -37,3 +62,13 @@ void FRigBaseComponent::Load(FArchive& Ar)
 	Ar << Key;
 	CachedNameString.Reset();
 }
+
+FString FRigBaseComponent::GetContentAsText() const
+{
+	FString Content;
+	if(const UScriptStruct* ScriptStruct = GetScriptStruct())
+	{
+		ScriptStruct->ExportText(Content, this, nullptr, nullptr, PPF_ExternalEditor, nullptr);
+	}
+	return Content;
+}
diff --git a/Engine/Plugins/Animation/ControlRig/Source/ControlRig/Private/Rigs/RigHierarchyController.cpp b/Engine/Plugins/Animation/ControlRig/Source/ControlRig/Private/Rigs/RigHierarchyController.cpp
--- a/Engine/Plugins/Animation/ControlRig/Source/ControlRig/Private/Rigs/RigHierarchyController.cpp
+++ b/Engine/Plugins/Animation/ControlRig/Source/ControlRig/Private/Rigs/RigHierarchyController.cpp
@@ -18,6 +18,8 @@
 #include "RigVMPythonUtils.h"
 #endif
 
+#include "Algo/Count.h"
+
 #include UE_INLINE_GENERATED_CPP_BY_NAME(RigHierarchyController)
 
 ////////////////////////////////////////////////////////////////////////////////
@@ -83,6 +85,16 @@ void URigHierarchyController::SetHierarchy(URigHierarchy* InHierarchy)
 }
 
 bool URigHierarchyController::SelectElement(FRigElementKey InKey, bool bSelect, bool bClearSelection)
+{
+	return SelectHierarchyKey(InKey, bSelect, bClearSelection);
+}
+
+bool URigHierarchyController::SelectComponent(FRigComponentKey InKey, bool bSelect, bool bClearSelection)
+{
+	return SelectHierarchyKey(InKey, bSelect, bClearSelection);
+}
+
+bool URigHierarchyController::SelectHierarchyKey(FRigHierarchyKey InKey, bool bSelect, bool bClearSelection)
 {
 	if(!IsValid())
 	{
@@ -91,9 +103,9 @@ bool URigHierarchyController::SelectElement(FRigElementKey InKey, bool bSelect,
 
 	if(bClearSelection)
 	{
-		TArray<FRigElementKey> KeysToSelect;
+		TArray<FRigHierarchyKey> KeysToSelect;
 		KeysToSelect.Add(InKey);
-		return SetSelection(KeysToSelect);
+		return SetHierarchySelection(KeysToSelect);
 	}
 
 	URigHierarchy* Hierarchy = GetHierarchy();
@@ -102,14 +114,14 @@ bool URigHierarchyController::SelectElement(FRigElementKey InKey, bool bSelect,
 	{
 		if(URigHierarchyController* ControllerForSelection = HierarchyForSelection->GetController())
 		{
-			return ControllerForSelection->SelectElement(InKey, bSelect, bClearSelection);
+			return ControllerForSelection->SelectHierarchyKey(InKey, bSelect, bClearSelection);
 		}
 	}
 
-	FRigElementKey Key = InKey;
-	if(Hierarchy->ElementKeyRedirector)
+	FRigHierarchyKey Key = InKey;
+	if(Hierarchy->ElementKeyRedirector && Key.IsElement())
 	{
-		if(const FRigElementKeyRedirector::FCachedKeyArray* Cache = Hierarchy->ElementKeyRedirector->Find(Key))
+		if(const FRigElementKeyRedirector::FCachedKeyArray* Cache = Hierarchy->ElementKeyRedirector->Find(Key.GetElement()))
 		{
 			if(!Cache->IsEmpty())
 			{
@@ -141,44 +153,107 @@ bool URigHierarchyController::SelectElement(FRigElementKey InKey, bool bSelect,
 		}
 	}
 
-	FRigBaseElement* Element = Hierarchy->Find(Key);
-	if(Element == nullptr)
+	if(Key.IsElement())
 	{
-		return false;
-	}
+		FRigBaseElement* Element = Hierarchy->Find(Key.GetElement());
+		if(Element == nullptr)
+		{
+			return false;
+		}
 
-	const bool bSelectionState = Hierarchy->OrderedSelection.Contains(Key);
-	ensure(bSelectionState == Element->bSelected);
-	if(Element->bSelected == bSelect)
-	{
-		return false;
-	}
+		const bool bSelectionState = Hierarchy->OrderedSelection.Contains(Key);
+		ensure(bSelectionState == Element->bSelected);
+		if(Element->bSelected == bSelect)
+		{
+			return false;
+		}
 
-	Element->bSelected = bSelect;
-	if(bSelect)
-	{
-		Hierarchy->OrderedSelection.Add(Key);
+		Element->bSelected = bSelect;
+	
+		if(bSelect)
+		{
+			Hierarchy->OrderedSelection.Add(Key);
+		}
+		else
+		{
+			Hierarchy->OrderedSelection.Remove(Key);
+		}
+
+		if(Element->bSelected)
+		{
+			Notify(ERigHierarchyNotification::ElementSelected, Element);
+		}
+		else
+		{
+			Notify(ERigHierarchyNotification::ElementDeselected, Element);
+		}
+
+		Hierarchy->UpdateVisibilityOnProxyControls();
 	}
-	else
+	else if(Key.IsComponent())
 	{
-		Hierarchy->OrderedSelection.Remove(Key);
+		FRigBaseComponent* Component = Hierarchy->FindComponent(Key.GetComponent());
+		if(Component == nullptr)
+		{
+			return false;
+		}
+
+		const bool bSelectionState = Hierarchy->OrderedSelection.Contains(Key);
+		ensure(bSelectionState == Component->bSelected);
+		if(Component->bSelected == bSelect)
+		{
+			return false;
+		}
+
+		Component->bSelected = bSelect;
+	
+		if(bSelect)
+		{
+			Hierarchy->OrderedSelection.Add(Key);
+		}
+		else
+		{
+			Hierarchy->OrderedSelection.Remove(Key);
+		}
+
+		if(Component->bSelected)
+		{
+			Notify(ERigHierarchyNotification::ComponentSelected, Component);
+		}
+		else
+		{
+			Notify(ERigHierarchyNotification::ComponentDeselected, Component);
+		}
+
+		Hierarchy->UpdateVisibilityOnProxyControls();
 	}
 
-	if(Element->bSelected)
+	return true;
+}
+
+bool URigHierarchyController::SetSelection(const TArray<FRigElementKey>& InKeys, bool bPrintPythonCommand)
+{
+	TArray<FRigHierarchyKey> Keys;
+	Keys.Reserve(InKeys.Num());
+	for(const FRigElementKey& Key : InKeys)
 	{
-		Notify(ERigHierarchyNotification::ElementSelected, Element);
+		Keys.Emplace(Key);
 	}
-	else
+	return SetHierarchySelection(Keys);
+}
+
+bool URigHierarchyController::SetComponentSelection(const TArray<FRigComponentKey>& InKeys, bool bPrintPythonCommand)
+{
+	TArray<FRigHierarchyKey> Keys;
+	Keys.Reserve(InKeys.Num());
+	for(const FRigComponentKey& Key : InKeys)
 	{
-		Notify(ERigHierarchyNotification::ElementDeselected, Element);
+		Keys.Emplace(Key);
 	}
-
-	Hierarchy->UpdateVisibilityOnProxyControls();
-
-	return true;
+	return SetHierarchySelection(Keys);
 }
 
-bool URigHierarchyController::SetSelection(const TArray<FRigElementKey>& InKeys, bool bPrintPythonCommand)
+bool URigHierarchyController::SetHierarchySelection(const TArray<FRigHierarchyKey>& InKeys, bool bPrintPythonCommand)
 {
 	if(!IsValid())
 	{
@@ -191,33 +266,33 @@ bool URigHierarchyController::SetSelection(const TArray<FRigElementKey>& InKeys,
 	{
 		if(URigHierarchyController* ControllerForSelection = HierarchyForSelection->GetController())
 		{
-			return ControllerForSelection->SetSelection(InKeys);
+			return ControllerForSelection->SetHierarchySelection(InKeys);
 		}
 	}
 
-	TArray<FRigElementKey> PreviousSelection = Hierarchy->GetSelectedKeys();
+	TArray<FRigHierarchyKey> PreviousSelection = Hierarchy->GetSelectedHierarchyKeys();
 	bool bResult = true;
 
 	{
 		// disable python printing here as we only want to print a single command instead of one per selected item
 		const TGuardValue<bool> Guard(bSuspendPythonPrinting, true);
 
-		for(const FRigElementKey& KeyToDeselect : PreviousSelection)
+		for(const FRigHierarchyKey& KeyToDeselect : PreviousSelection)
 		{
 			if(!InKeys.Contains(KeyToDeselect))
 			{
-				if(!SelectElement(KeyToDeselect, false))
+				if(!SelectHierarchyKey(KeyToDeselect, false))
 				{
 					bResult = false;
 				}
 			}
 		}
 
-		for(const FRigElementKey& KeyToSelect : InKeys)
+		for(const FRigHierarchyKey& KeyToSelect : InKeys)
 		{
 			if(!PreviousSelection.Contains(KeyToSelect))
 			{
-				if(!SelectElement(KeyToSelect, true))
+				if(!SelectHierarchyKey(KeyToSelect, true))
 				{
 					bResult = false;
 				}
@@ -231,14 +306,30 @@ bool URigHierarchyController::SetSelection(const TArray<FRigElementKey>& InKeys,
 		UBlueprint* Blueprint = GetTypedOuter<UBlueprint>();
 		if (Blueprint)
 		{
-			const FString Selection = FString::JoinBy( InKeys, TEXT(", "), [](const FRigElementKey& Key)
+			const int32 NumElements = Algo::CountIf(InKeys, [](const FRigHierarchyKey& Key) { return Key.IsElement(); });
+			const int32 NumComponents = Algo::CountIf(InKeys, [](const FRigHierarchyKey& Key) { return Key.IsComponent(); });
+			if(NumComponents == 0)
 			{
-				return Key.ToPythonString();
-			});
-			
-			RigVMPythonUtils::Print(Blueprint->GetFName().ToString(),
-				FString::Printf(TEXT("hierarchy_controller.set_selection([%s])"),
-				*Selection ) );
+				const FString Selection = FString::JoinBy( InKeys, TEXT(", "), [](const FRigHierarchyKey& Key)
+				{
+					return Key.GetElement().ToPythonString();
+				});
+				
+				RigVMPythonUtils::Print(Blueprint->GetFName().ToString(),
+					FString::Printf(TEXT("hierarchy_controller.set_selection([%s])"),
+					*Selection ) );
+			}
+			else if(NumElements == 0)
+			{
+				const FString Selection = FString::JoinBy( InKeys, TEXT(", "), [](const FRigHierarchyKey& Key)
+				{
+					return Key.GetComponent().ToPythonString();
+				});
+				
+				RigVMPythonUtils::Print(Blueprint->GetFName().ToString(),
+					FString::Printf(TEXT("hierarchy_controller.set_component_selection([%s])"),
+					*Selection ) );
+			}
 		}
 	}
 #endif
@@ -2432,7 +2523,10 @@ TArray<FString> URigHierarchyController::GetAddComponentPythonCommands(const FRi
 {
 	TArray<FString> Commands;
 
-	const FString ElementKeyStr = Component->GetElement()->GetKey().ToPythonString();
+	const FString ElementKeyStr =
+		Component->GetElement() ?
+		Component->GetElement()->GetKey().ToPythonString() :
+		URigHierarchy::GetTopLevelComponentElementKey().ToPythonString();
 
 	FString ContentStr;
 	Component->GetScriptStruct()->ExportText(ContentStr, Component, Component, nullptr, PPF_None, nullptr);
diff --git a/Engine/Plugins/Animation/ControlRig/Source/ControlRig/Private/Rigs/RigHierarchyDefines.cpp b/Engine/Plugins/Animation/ControlRig/Source/ControlRig/Private/Rigs/RigHierarchyDefines.cpp
--- a/Engine/Plugins/Animation/ControlRig/Source/ControlRig/Private/Rigs/RigHierarchyDefines.cpp
+++ b/Engine/Plugins/Animation/ControlRig/Source/ControlRig/Private/Rigs/RigHierarchyDefines.cpp
@@ -196,6 +196,60 @@ bool FRigComponentKey::IsTopLevel() const
 	return ElementKey == URigHierarchy::GetTopLevelComponentElementKey();
 }
 
+////////////////////////////////////////////////////////////////////////////////
+// FRigHierarchyKey
+////////////////////////////////////////////////////////////////////////////////
+
+void FRigHierarchyKey::Serialize(FArchive& Ar)
+{
+	if(Ar.IsLoading())
+	{
+		bool bIsElement = false;
+		Ar << bIsElement;
+		if(bIsElement)
+		{
+			FRigElementKey Key;
+			Ar << Key;
+			Element = Key;
+		}
+		else
+		{
+			Element.Reset();
+		}
+
+		bool bIsComponent = false;
+		Ar << bIsComponent;
+		if(bIsComponent)
+		{
+			FRigComponentKey Key;
+			Ar << Key;
+			Component = Key;
+		}
+		else
+		{
+			Component.Reset();
+		}
+	}
+	else if(Ar.IsSaving())
+	{
+		bool bIsElement = IsElement();
+		Ar << bIsElement;
+		if(bIsElement)
+		{
+			FRigElementKey Key = GetElement();
+			Ar << Key;
+		}
+
+		bool bIsComponent = IsComponent();
+		Ar << bIsComponent;
+		if(bIsComponent)
+		{
+			FRigComponentKey Key = GetComponent();
+			Ar << Key;
+		}
+	}
+}
+
 ////////////////////////////////////////////////////////////////////////////////
 // FRigElementKeyAndIndex
 ////////////////////////////////////////////////////////////////////////////////
diff --git a/Engine/Plugins/Animation/ControlRig/Source/ControlRig/Public/Rigs/RigHierarchy.h b/Engine/Plugins/Animation/ControlRig/Source/ControlRig/Public/Rigs/RigHierarchy.h
--- a/Engine/Plugins/Animation/ControlRig/Source/ControlRig/Public/Rigs/RigHierarchy.h
+++ b/Engine/Plugins/Animation/ControlRig/Source/ControlRig/Public/Rigs/RigHierarchy.h
@@ -413,6 +413,27 @@ public:
 	 */
 	bool IsProcedural(const FRigBaseElement* InElement) const;
 
+	/**
+	 * Returns true if the provided component is procedural.
+	 * @param InKey The key to validate
+	 * @return Returns true if the component is procedural
+	 */
+	bool IsProcedural(const FRigComponentKey& InKey) const;
+
+	/**
+	 * Returns true if the provided component is procedural.
+	 * @param InComponent The component to check
+	 * @return Returns true if the component is procedural
+	 */
+	bool IsProcedural(const FRigBaseComponent* InComponent) const;
+
+	/**
+	 * Returns true if the provided component or element is procedural.
+	 * @param InKey The key to validate
+	 * @return Returns true if the component or element is procedural
+	 */
+	bool IsProcedural(const FRigHierarchyKey& InKey) const;
+
 	/**
 	 * Returns the index of an element given its key
 	 * @param InKey The key of the element to retrieve the index for
@@ -1886,6 +1907,12 @@ public:
 	 */
 	TArray<const FRigBaseElement*> GetSelectedElements(ERigElementType InTypeFilter = ERigElementType::All) const;
 
+	/**
+	 * Returns the selected components
+	 * @return An array of the currently selected components
+	 */
+	TArray<const FRigBaseComponent*> GetSelectedComponents() const;
+
 	/**
 	 * Returns the keys of selected elements
 	 * @InTypeFilter The types to retrieve the selection for
@@ -1894,6 +1921,13 @@ public:
 	UFUNCTION(BlueprintCallable, Category = URigHierarchy)
 	TArray<FRigElementKey> GetSelectedKeys(ERigElementType InTypeFilter = ERigElementType::All) const;
 
+	/**
+	 * Returns the keys of selected elements and components
+	 * @return An array of the currently selected elements
+	 */
+	UFUNCTION(BlueprintCallable, Category = URigHierarchy)
+	TArray<FRigHierarchyKey> GetSelectedHierarchyKeys() const;
+
 	/**
 	 * Returns true if any element is selected and respects the provided predicate 
 	 * @InPredicate A function that should return true when the element is desirable
@@ -1919,6 +1953,30 @@ public:
 		return IsSelected(Find(InKey));
 	}
 
+	/**
+	 * Returns true if a given component is selected
+	 * @param InKey The key to check
+	 * @return true if a given component is selected
+	 */
+	UFUNCTION(BlueprintCallable, Category = URigHierarchy)
+	bool IsComponentSelected(FRigComponentKey InKey) const
+	{
+		return IsComponentSelected(FindComponent(InKey));
+	}
+
+	bool IsHierarchyKeySelected(FRigHierarchyKey InKey) const
+	{
+		if(InKey.IsElement())
+		{
+			return IsSelected(InKey.GetElement());
+		}
+		if(InKey.IsComponent())
+		{
+			return IsComponentSelected(InKey.GetComponent());
+		}
+		return false;
+	}
+
 	/**
 	 * Returns true if a given element is selected
 	 * @param InIndex The index to check
@@ -4443,6 +4501,13 @@ private:
 	 */
     bool IsSelected(const FRigBaseElement* InElement) const;
 
+	/**
+	 * Returns true if a given component is selected
+	 * @param InComponent The component to check
+	 * @return true if a given component is selected
+	 */
+	bool IsComponentSelected(const FRigBaseComponent* InComponent) const;
+
 	/**
 	 * Updates the transient cached children table if the topology version is out of date with the one
 	 * stored with the cached table.
@@ -4786,7 +4851,7 @@ private:
 
 	TWeakObjectPtr<URigHierarchy> HierarchyForSelectionPtr;
 	TWeakObjectPtr<URigHierarchy> DefaultHierarchyPtr;
-	TArray<FRigElementKey> OrderedSelection;
+	TArray<FRigHierarchyKey> OrderedSelection;
 
 	UPROPERTY(Transient)
 	TObjectPtr<URigHierarchyController> HierarchyController;
diff --git a/Engine/Plugins/Animation/ControlRig/Source/ControlRig/Public/Rigs/RigHierarchyComponents.h b/Engine/Plugins/Animation/ControlRig/Source/ControlRig/Public/Rigs/RigHierarchyComponents.h
--- a/Engine/Plugins/Animation/ControlRig/Source/ControlRig/Public/Rigs/RigHierarchyComponents.h
+++ b/Engine/Plugins/Animation/ControlRig/Source/ControlRig/Public/Rigs/RigHierarchyComponents.h
@@ -75,10 +75,10 @@ protected:
 		return true;
 	}
 
+	FRigComponentKey Key;
+
 	// ReSharper disable once CppUE4ProbableMemoryIssuesWithUObject
 	FRigBaseElement* Element = nullptr;
-	
-	FRigComponentKey Key;
 	int32 IndexInHierarchy = INDEX_NONE;
 	int32 IndexInElement = INDEX_NONE;
 	int32 CreatedAtInstructionIndex = INDEX_NONE;
@@ -90,6 +90,8 @@ public:
 
 	virtual UScriptStruct* GetScriptStruct() const { return FRigBaseComponent::StaticStruct(); }
 
+	static TArray<UScriptStruct*> GetAllComponentScriptStructs(bool bSorted = false);
+
 	const FRigComponentKey& GetKey() const { return Key; }
 	const FName& GetFName() const { return Key.Name; }
 	const FString& GetName() const
@@ -119,6 +121,8 @@ public:
 	virtual void Save(FArchive& Ar);
 	virtual void Load(FArchive& Ar);
 
+	FString GetContentAsText() const;
+
 	bool IsA(const UScriptStruct* InScriptStruct) const
 	{
 		return GetScriptStruct()->IsChildOf(InScriptStruct);
@@ -175,6 +179,7 @@ protected:
 	friend class URigHierarchy;
 	friend class URigHierarchyController;
 	friend struct FControlRigReplayTracks;
+	friend class UControlRigWrapperObject;
 };
 
 template<>
diff --git a/Engine/Plugins/Animation/ControlRig/Source/ControlRig/Public/Rigs/RigHierarchyController.h b/Engine/Plugins/Animation/ControlRig/Source/ControlRig/Public/Rigs/RigHierarchyController.h
--- a/Engine/Plugins/Animation/ControlRig/Source/ControlRig/Public/Rigs/RigHierarchyController.h
+++ b/Engine/Plugins/Animation/ControlRig/Source/ControlRig/Public/Rigs/RigHierarchyController.h
@@ -41,6 +41,7 @@ public:
 	 * Selects or deselects an element in the hierarchy
 	 * @param InKey The key of the element to select
 	 * @param bSelect If set to false the element will be deselected
+	 * @param bClearSelection If this is true the selection will be cleared
 	 * @return Returns true if the selection was applied
 	 */
 	UFUNCTION(BlueprintCallable, Category = URigHierarchyController)
@@ -58,12 +59,70 @@ public:
 	}
 
 	/**
-	 * Sets the selection based on a list of keys
+	 * Selects or deselects a component in the hierarchy
+	 * @param InKey The key of the component to select
+	 * @param bSelect If set to false the component will be deselected
+	 * @param bClearSelection If this is true the selection will be cleared
+	 * @return Returns true if the selection was applied
+	 */
+	UFUNCTION(BlueprintCallable, Category = URigHierarchyController)
+	bool SelectComponent(FRigComponentKey InKey, bool bSelect = true, bool bClearSelection = false);
+
+	/**
+	 * Deselects or deselects a component in the hierarchy
+	 * @param InKey The key of the component to deselect
+	 * @return Returns true if the selection was applied
+	 */
+	UFUNCTION(BlueprintCallable, Category = URigHierarchyController)
+	bool DeselectComponent(FRigComponentKey InKey)
+	{
+		return SelectComponent(InKey, false);
+	}
+
+	/**
+	 * Selects or deselects a component or an element in the hierarchy
+	 * @param InKey The key of the component or an element to select
+	 * @param bSelect If set to false the component or an element will be deselected
+	 * @param bClearSelection If this is true the selection will be cleared
+	 * @return Returns true if the selection was applied
+	 */
+	UFUNCTION(BlueprintCallable, Category = URigHierarchyController)
+	bool SelectHierarchyKey(FRigHierarchyKey InKey, bool bSelect = true, bool bClearSelection = false);
+
+	/**
+	 * Deselects or deselects a component or an element in the hierarchy
+	 * @param InKey The key of the component or element to deselect
+	 * @return Returns true if the selection was applied
+	 */
+	UFUNCTION(BlueprintCallable, Category = URigHierarchyController)
+	bool DeselectHierarchyKey(FRigHierarchyKey InKey)
+	{
+		return SelectHierarchyKey(InKey, false);
+	}
+
+	/**
+	 * Sets the selection based on a list of element keys
 	 * @param InKeys The array of keys of the elements to select
 	 * @return Returns true if the selection was applied
 	 */
 	UFUNCTION(BlueprintCallable, Category = URigHierarchyController)
-    bool SetSelection(const TArray<FRigElementKey>& InKeys, bool bPrintPythonCommand = false);
+	bool SetSelection(const TArray<FRigElementKey>& InKeys, bool bPrintPythonCommand = false);
+
+	/**
+	 * Sets the selection based on a list of component keys
+	 * @param InKeys The array of keys of the component to select
+	 * @return Returns true if the selection was applied
+	 */
+	UFUNCTION(BlueprintCallable, Category = URigHierarchyController)
+	bool SetComponentSelection(const TArray<FRigComponentKey>& InKeys, bool bPrintPythonCommand = false);
+
+	/**
+	 * Sets the selection based on a list of component keys
+	 * @param InKeys The array of keys of the component to select
+	 * @return Returns true if the selection was applied
+	 */
+	UFUNCTION(BlueprintCallable, Category = URigHierarchyController)
+	bool SetHierarchySelection(const TArray<FRigHierarchyKey>& InKeys, bool bPrintPythonCommand = false);
 
 	/**
 	 * Clears the selection
diff --git a/Engine/Plugins/Animation/ControlRig/Source/ControlRig/Public/Rigs/RigHierarchyDefines.h b/Engine/Plugins/Animation/ControlRig/Source/ControlRig/Public/Rigs/RigHierarchyDefines.h
--- a/Engine/Plugins/Animation/ControlRig/Source/ControlRig/Public/Rigs/RigHierarchyDefines.h
+++ b/Engine/Plugins/Animation/ControlRig/Source/ControlRig/Public/Rigs/RigHierarchyDefines.h
@@ -104,6 +104,8 @@ enum class ERigHierarchyNotification : uint8
 	ComponentAdded,
 	ComponentRemoved,
 	ComponentContentChanged,
+	ComponentSelected,
+	ComponentDeselected,
 
 	/** MAX - invalid */
 	Max UMETA(Hidden),
@@ -1765,6 +1767,182 @@ public:
 	bool IsTopLevel() const;
 };
 
+USTRUCT(BlueprintType)
+struct CONTROLRIG_API FRigHierarchyKey
+{
+	GENERATED_BODY()
+
+public:
+	FRigHierarchyKey()
+	{
+	}
+
+	FRigHierarchyKey(const FRigElementKey& InElementKey)
+	{
+		if(InElementKey.IsValid())
+		{
+			Element = InElementKey;
+		}
+	}
+	
+	FRigHierarchyKey(const FRigComponentKey& InComponentKey)
+	{
+		if(InComponentKey.IsValid())
+		{
+			Component = InComponentKey;
+		}
+	}
+
+	bool IsElement() const { return Element.IsSet(); }
+	bool IsComponent() const { return Component.IsSet(); }
+	bool IsValid() const
+	{
+		return IsElement() || IsComponent();
+	}
+
+	const FName& GetFName() const
+	{
+		if(IsElement())
+		{
+			return Element.GetValue().Name;
+		}
+		if(IsComponent())
+		{
+			return Component.GetValue().Name;
+		}
+		static const FName NoneName(NAME_None);
+		return NoneName;;
+	}
+
+	FString GetName() const
+	{
+		return GetFName().ToString();
+	}
+
+	void SetName(const FName& InName)
+	{
+		if(IsElement())
+		{
+			Element.GetValue().Name = InName;
+		}
+		if(IsComponent())
+		{
+			Component.GetValue().Name = InName;
+		}
+	}
+
+	const FRigElementKey& GetElement() const
+	{
+		if(IsElement())
+		{
+			return Element.GetValue();
+		}
+		if(IsComponent())
+		{
+			return Component.GetValue().ElementKey;
+		}
+		static const FRigElementKey InvalidKey;
+		return InvalidKey;
+	}
+
+	const FRigComponentKey& GetComponent() const
+	{
+		if(IsComponent())
+		{
+			return Component.GetValue();
+		}
+		static const FRigComponentKey InvalidComponent;
+		return InvalidComponent;
+	}
+
+	bool operator ==(const FRigHierarchyKey& InOther) const
+	{
+		if(IsElement() != InOther.IsElement() ||
+			IsComponent() != InOther.IsComponent())
+		{
+			return false;
+		}
+		if(IsElement())
+		{
+			return Element == InOther.Element;
+		}
+		if(IsComponent())
+		{
+			return Component == InOther.Component;
+		}
+		return true;
+	}
+
+	bool operator ==(const FRigElementKey& InOther) const
+	{
+		if(IsElement())
+		{
+			return Element == InOther;
+		}
+		return false;
+	}
+
+	bool operator ==(const FRigComponentKey& InOther) const
+	{
+		if(IsComponent())
+		{
+			return Component == InOther;
+		}
+		return false;
+	}
+	
+	bool operator <(const FRigHierarchyKey& Other) const
+	{
+		if(IsElement() && Other.IsComponent())
+		{
+			return false;
+		}
+		if(IsComponent() && Other.IsElement())
+		{
+			return true;
+		}
+		if(IsComponent() && Other.IsComponent())
+		{
+			return GetComponent() < Other.GetComponent();
+		}
+		if(IsElement() && Other.IsElement())
+		{
+			return GetElement() < Other.GetElement();
+		}
+		return false;
+	}
+
+	bool operator >(const FRigHierarchyKey& Other) const
+	{
+		return !(*this < Other);
+	}
+	
+	void Serialize(FArchive& Ar);
+	friend FArchive& operator<<(FArchive& Ar, FRigHierarchyKey& P)
+	{
+		P.Serialize(Ar);
+		return Ar;
+	}
+
+	friend uint32 GetTypeHash(const FRigHierarchyKey& InKey)
+	{
+		if(InKey.IsElement())
+		{
+			return GetTypeHash(InKey.GetElement());
+		}
+		if(InKey.IsComponent())
+		{
+			return GetTypeHash(InKey.GetComponent());
+		}
+		return 0;
+	}
+
+private:
+	
+	TOptional<FRigElementKey> Element;
+	TOptional<FRigComponentKey> Component;
+};
+
 struct CONTROLRIG_API FRigElementKeyAndIndex
 {
 public:
diff --git a/Engine/Plugins/Animation/ControlRig/Source/ControlRig/Public/Rigs/RigHierarchyElements.h b/Engine/Plugins/Animation/ControlRig/Source/ControlRig/Public/Rigs/RigHierarchyElements.h
--- a/Engine/Plugins/Animation/ControlRig/Source/ControlRig/Public/Rigs/RigHierarchyElements.h
+++ b/Engine/Plugins/Animation/ControlRig/Source/ControlRig/Public/Rigs/RigHierarchyElements.h
@@ -1029,6 +1029,7 @@ protected:
 	friend class URigHierarchy;
 	friend class URigHierarchyController;
 	friend struct FRigElementKeyAndIndex;
+	friend class UControlRigWrapperObject;
 };
 
 USTRUCT(BlueprintType)
diff --git a/Engine/Plugins/Animation/ControlRig/Source/ControlRigDeveloper/Private/ControlRigBlueprint.cpp b/Engine/Plugins/Animation/ControlRig/Source/ControlRigDeveloper/Private/ControlRigBlueprint.cpp
--- a/Engine/Plugins/Animation/ControlRig/Source/ControlRigDeveloper/Private/ControlRigBlueprint.cpp
+++ b/Engine/Plugins/Animation/ControlRig/Source/ControlRigDeveloper/Private/ControlRigBlueprint.cpp
@@ -2760,6 +2760,9 @@ void UControlRigBlueprint::HandleHierarchyModified(ERigHierarchyNotification InN
 		case ERigHierarchyNotification::ParentChanged:
 		case ERigHierarchyNotification::ElementReordered:
 		case ERigHierarchyNotification::HierarchyReset:
+		case ERigHierarchyNotification::ComponentAdded:
+		case ERigHierarchyNotification::ComponentRemoved:
+		case ERigHierarchyNotification::ComponentContentChanged:
 		{
 			Modify();
 			PropagateHierarchyFromBPToInstances();
diff --git a/Engine/Plugins/Animation/ControlRig/Source/ControlRigEditor/Private/ControlRigDragOps.cpp b/Engine/Plugins/Animation/ControlRig/Source/ControlRigEditor/Private/ControlRigDragOps.cpp
--- a/Engine/Plugins/Animation/ControlRig/Source/ControlRigEditor/Private/ControlRigDragOps.cpp
+++ b/Engine/Plugins/Animation/ControlRig/Source/ControlRigEditor/Private/ControlRigDragOps.cpp
@@ -7,7 +7,7 @@
 /// FRigElementHierarchyDragDropOp
 ///////////////////////////////////////////////////////////
 
-TSharedRef<FRigElementHierarchyDragDropOp> FRigElementHierarchyDragDropOp::New(const TArray<FRigElementKey>& InElements)
+TSharedRef<FRigElementHierarchyDragDropOp> FRigElementHierarchyDragDropOp::New(const TArray<FRigHierarchyKey>& InElements)
 {
 	TSharedRef<FRigElementHierarchyDragDropOp> Operation = MakeShared<FRigElementHierarchyDragDropOp>();
 	Operation->Elements = InElements;
@@ -30,9 +30,9 @@ TSharedPtr<SWidget> FRigElementHierarchyDragDropOp::GetDefaultDecorator() const
 FString FRigElementHierarchyDragDropOp::GetJoinedElementNames() const
 {
 	TArray<FString> ElementNameStrings;
-	for (const FRigElementKey& Element: Elements)
+	for (const FRigHierarchyKey& Element: Elements)
 	{
-		ElementNameStrings.Add(Element.Name.ToString());
+		ElementNameStrings.Add(Element.GetName());
 	}
 	return FString::Join(ElementNameStrings, TEXT(","));
 }
@@ -41,7 +41,10 @@ bool FRigElementHierarchyDragDropOp::IsDraggingSingleConnector() const
 {
 	if(Elements.Num() == 1)
 	{
-		return Elements[0].Type == ERigElementType::Connector;
+		if(Elements[0].IsElement())
+		{
+			return Elements[0].GetElement().Type == ERigElementType::Connector;
+		}
 	}
 	return false;
 }
@@ -50,7 +53,10 @@ bool FRigElementHierarchyDragDropOp::IsDraggingSingleSocket() const
 {
 	if(Elements.Num() == 1)
 	{
-		return Elements[0].Type == ERigElementType::Socket;
+		if(Elements[0].IsElement())
+		{
+			return Elements[0].GetElement().Type == ERigElementType::Socket;
+		}
 	}
 	return false;
 }
diff --git a/Engine/Plugins/Animation/ControlRig/Source/ControlRigEditor/Private/ControlRigDragOps.h b/Engine/Plugins/Animation/ControlRig/Source/ControlRigEditor/Private/ControlRigDragOps.h
--- a/Engine/Plugins/Animation/ControlRig/Source/ControlRigEditor/Private/ControlRigDragOps.h
+++ b/Engine/Plugins/Animation/ControlRig/Source/ControlRigEditor/Private/ControlRigDragOps.h
@@ -12,7 +12,7 @@ class FRigElementHierarchyDragDropOp : public FGraphNodeDragDropOp
 public:
 	DRAG_DROP_OPERATOR_TYPE(FRigElementHierarchyDragDropOp, FGraphNodeDragDropOp)
 
-	static TSharedRef<FRigElementHierarchyDragDropOp> New(const TArray<FRigElementKey>& InElements);
+	static TSharedRef<FRigElementHierarchyDragDropOp> New(const TArray<FRigHierarchyKey>& InElements);
 
 	virtual TSharedPtr<SWidget> GetDefaultDecorator() const override;
 
@@ -23,7 +23,7 @@ public:
 	}
 
 	/** @return The property paths from this drag operation */
-	const TArray<FRigElementKey>& GetElements() const
+	const TArray<FRigHierarchyKey>& GetElements() const
 	{
 		return Elements;
 	}
@@ -36,7 +36,7 @@ public:
 private:
 
 	/** Data for the property paths this item represents */
-	TArray<FRigElementKey> Elements;
+	TArray<FRigHierarchyKey> Elements;
 };
 
 class FRigHierarchyTagDragDropOp : public FDecoratedDragDropOp
diff --git a/Engine/Plugins/Animation/ControlRig/Source/ControlRigEditor/Private/ControlRigElementDetails.cpp b/Engine/Plugins/Animation/ControlRig/Source/ControlRigEditor/Private/ControlRigElementDetails.cpp
--- a/Engine/Plugins/Animation/ControlRig/Source/ControlRigEditor/Private/ControlRigElementDetails.cpp
+++ b/Engine/Plugins/Animation/ControlRig/Source/ControlRigEditor/Private/ControlRigElementDetails.cpp
@@ -1599,8 +1599,7 @@ void FRigBaseElementDetails::CustomizeComponents(IDetailLayoutBuilder& DetailBui
 				FSharedComponent SharedComponent;
 				SharedComponent.Handle = FRigComponentHandle(Hierarchy, Component);
 				SharedComponent.ScriptStruct = Component->GetScriptStruct();
-				// todo: we need to discuss if procedural components should be disabled for editing?
-				//SharedComponent.bIsProcedural = Component->IsProcedural();
+				SharedComponent.bIsProcedural = Component->IsProcedural();
 				SharedComponents.Add(SharedComponent);
 			}
 		}
@@ -1621,8 +1620,6 @@ void FRigBaseElementDetails::CustomizeComponents(IDetailLayoutBuilder& DetailBui
 			});
 
 			// update the procedural flag in case any component is procedural
-			// todo: we need to discuss if procedural components should be disabled for editing?
-			/*
 			for(FSharedComponent& SharedComponent : SharedComponents)
 			{
 				if(SharedComponent.bIsProcedural)
@@ -1637,7 +1634,6 @@ void FRigBaseElementDetails::CustomizeComponents(IDetailLayoutBuilder& DetailBui
 					}
 				}
 			}
-			*/
 		}
 	}
 
@@ -4225,7 +4221,7 @@ void FRigControlElementDetails::CustomizeAvailableSpaces(IDetailLayoutBuilder& D
 	DisplaySettings.bShowNulls = true;
 	DisplaySettings.bShowReferences = false;
 	DisplaySettings.bShowSockets = false;
-	DisplaySettings.bShowPhysics = false;
+	DisplaySettings.bShowComponents = false;
 	DisplaySettings.bHideParentsOnFilter = true;
 	DisplaySettings.bFlattenHierarchyOnFilter = true;
 	DisplaySettings.bShowIconColors = true;
@@ -4615,7 +4611,7 @@ TSharedRef<SWidget> FRigControlElementDetails::GetAddSpaceContent(const TSharedR
 		return PerElementInfos[0].GetHierarchy();
 	});
 	RigTreeDelegates.OnGetDisplaySettings = FOnGetRigTreeDisplaySettings::CreateSP(this, &FRigControlElementDetails::GetDisplaySettings);
-	RigTreeDelegates.OnGetSelection = FOnRigTreeGetSelection::CreateLambda([]() -> TArray<FRigElementKey> { return {}; });
+	RigTreeDelegates.OnGetSelection = FOnRigTreeGetSelection::CreateLambda([]() -> TArray<FRigHierarchyKey> { return {}; });
 	RigTreeDelegates.OnSelectionChanged = FOnRigTreeSelectionChanged::CreateSP(this, &FRigControlElementDetails::OnAddSpaceSelection, PropertyUtilities);
 
 	return SNew(SBox)
@@ -4640,7 +4636,7 @@ void FRigControlElementDetails::OnAddSpaceSelection(TSharedPtr<FRigTreeElement>
 	if(Selection)
 	{
 		const FRigElementKey ChildKey = PerElementInfos[0].GetElement()->GetKey();
-		const FRigElementKey NewParentKey = Selection->Key;
+		const FRigElementKey NewParentKey = Selection->Key.GetElement();
 		URigHierarchy* HierarchyToChange = PerElementInfos[0].GetDefaultHierarchy();
 		AddSpaceMenuAnchor->SetIsOpen(false);
 
@@ -6703,4 +6699,51 @@ void FRigConnectionRuleDetails::OnRuleContentChanged()
 	StructPropertyHandle->GetParentHandle()->NotifyPostChange(EPropertyChangeType::ValueSet);
 }
 
+void FRigBaseComponentDetails::CustomizeDetails(IDetailLayoutBuilder& DetailBuilder)
+{
+	TArray<TWeakObjectPtr<UObject>> DetailObjects;
+	DetailBuilder.GetObjectsBeingCustomized(DetailObjects);
+
+	bool bIsReadOnly = false;
+	TArray<URigVMDetailsViewWrapperObject*> WrapperObjects;
+	for(TWeakObjectPtr<UObject> DetailObject : DetailObjects)
+	{
+		URigVMDetailsViewWrapperObject* WrapperObject = CastChecked<URigVMDetailsViewWrapperObject>(DetailObject.Get());
+		if(WrapperObject->GetWrappedStruct()->IsChildOf(FRigBaseComponent::StaticStruct()))
+		{
+			WrapperObjects.Add(WrapperObject);
+			
+			if (const URigHierarchy* Hierarchy = Cast<URigHierarchy>(WrapperObject->GetSubject()))
+			{
+				const FRigBaseComponent WrappedComponent = WrapperObject->GetContent<FRigBaseComponent>();
+				if(const FRigBaseComponent* Component = Hierarchy->FindComponent(WrappedComponent.GetKey()))
+				{
+					if(Component->IsProcedural())
+					{
+						bIsReadOnly = true;
+						break;
+					}
+				}
+			}
+		}
+	}
+
+	if(bIsReadOnly)
+	{
+		TArray<FName> CategoryNames;
+		DetailBuilder.GetCategoryNames(CategoryNames);
+		for(const FName& CategoryName : CategoryNames)
+		{
+			TArray<TSharedRef<IPropertyHandle>> Properties;
+			IDetailCategoryBuilder& Category = DetailBuilder.EditCategory(CategoryName);
+			Category.GetDefaultProperties(Properties);
+			
+			for(TSharedRef<IPropertyHandle>& Property : Properties)
+			{
+				DetailBuilder.EditDefaultProperty(Property)->IsEnabled(false);
+			}
+		}
+	}
+}
+
 #undef LOCTEXT_NAMESPACE
diff --git a/Engine/Plugins/Animation/ControlRig/Source/ControlRigEditor/Private/ControlRigElementDetails.h b/Engine/Plugins/Animation/ControlRig/Source/ControlRigEditor/Private/ControlRigElementDetails.h
--- a/Engine/Plugins/Animation/ControlRig/Source/ControlRigEditor/Private/ControlRigElementDetails.h
+++ b/Engine/Plugins/Animation/ControlRig/Source/ControlRigEditor/Private/ControlRigElementDetails.h
@@ -933,3 +933,16 @@ protected:
 	TAttribute<bool> EnabledAttribute;
 };
 
+class FRigBaseComponentDetails : public IDetailCustomization
+{
+public:
+
+	/** Makes a new instance of this detail layout class for a specific detail view requesting it */
+	static TSharedRef<IDetailCustomization> MakeInstance()
+	{
+		return MakeShareable(new FRigBaseComponentDetails);
+	}
+
+	/** IDetailCustomization interface */
+	virtual void CustomizeDetails(IDetailLayoutBuilder& DetailBuilder) override;
+};
diff --git a/Engine/Plugins/Animation/ControlRig/Source/ControlRigEditor/Private/ControlRigHierarchyCommands.cpp b/Engine/Plugins/Animation/ControlRig/Source/ControlRigEditor/Private/ControlRigHierarchyCommands.cpp
--- a/Engine/Plugins/Animation/ControlRig/Source/ControlRigEditor/Private/ControlRigHierarchyCommands.cpp
+++ b/Engine/Plugins/Animation/ControlRig/Source/ControlRigEditor/Private/ControlRigHierarchyCommands.cpp
@@ -38,9 +38,9 @@ void FControlRigHierarchyCommands::RegisterCommands()
 	UI_COMMAND(ShowBones, "Show Bones", "Whether to show or hide bones", EUserInterfaceActionType::ToggleButton, FInputChord());
 	UI_COMMAND(ShowControls, "Show Controls", "Whether to show or hide controls", EUserInterfaceActionType::ToggleButton, FInputChord());
 	UI_COMMAND(ShowNulls, "Show Nulls", "Whether to show or hide nulls", EUserInterfaceActionType::ToggleButton, FInputChord());
-	UI_COMMAND(ShowPhysics, "Show Physics", "Whether to show or hide physics", EUserInterfaceActionType::ToggleButton, FInputChord());
 	UI_COMMAND(ShowReferences, "Show References", "Whether to show or hide references", EUserInterfaceActionType::ToggleButton, FInputChord());
 	UI_COMMAND(ShowSockets, "Show Sockets", "Whether to show or hide sockets", EUserInterfaceActionType::ToggleButton, FInputChord());
+	UI_COMMAND(ShowComponents, "Show Components", "Whether to show or hide components", EUserInterfaceActionType::ToggleButton, FInputChord());
 	UI_COMMAND(ToggleControlShapeTransformEdit, "Toggle Shape Transform Edit", "Toggle Editing Selected Control's Shape Transform", EUserInterfaceActionType::ToggleButton, FInputChord(EKeys::Period, EModifierKey::Control)); 
 	UI_COMMAND(SpaceSwitching, "Space Switching", "Space switching on the control", EUserInterfaceActionType::Button, FInputChord(EKeys::Tab)); 
 	UI_COMMAND(ShowIconColors, "Show Icon Colors", "Whether to tint the icons with the element color", EUserInterfaceActionType::ToggleButton, FInputChord());
diff --git a/Engine/Plugins/Animation/ControlRig/Source/ControlRigEditor/Private/ControlRigHierarchyCommands.h b/Engine/Plugins/Animation/ControlRig/Source/ControlRigEditor/Private/ControlRigHierarchyCommands.h
--- a/Engine/Plugins/Animation/ControlRig/Source/ControlRigEditor/Private/ControlRigHierarchyCommands.h
+++ b/Engine/Plugins/Animation/ControlRig/Source/ControlRigEditor/Private/ControlRigHierarchyCommands.h
@@ -114,15 +114,15 @@ public:
 	/* Show spaces */
 	TSharedPtr< FUICommandInfo > ShowNulls;
 
-	/* Show physics */
-	TSharedPtr< FUICommandInfo > ShowPhysics;
-
 	/* Show references */
 	TSharedPtr< FUICommandInfo > ShowReferences;
 
 	/* Show sockets */
 	TSharedPtr< FUICommandInfo > ShowSockets;
 
+	/* Show components */
+	TSharedPtr< FUICommandInfo > ShowComponents;
+
 	/** Toggle Shape Transform Edit*/
 	TSharedPtr< FUICommandInfo > ToggleControlShapeTransformEdit;
 
diff --git a/Engine/Plugins/Animation/ControlRig/Source/ControlRigEditor/Private/EditMode/SControlRigEditModeTools.cpp b/Engine/Plugins/Animation/ControlRig/Source/ControlRigEditor/Private/EditMode/SControlRigEditModeTools.cpp
--- a/Engine/Plugins/Animation/ControlRig/Source/ControlRigEditor/Private/EditMode/SControlRigEditModeTools.cpp
+++ b/Engine/Plugins/Animation/ControlRig/Source/ControlRigEditor/Private/EditMode/SControlRigEditModeTools.cpp
@@ -217,7 +217,6 @@ void SControlRigEditModeTools::Construct(const FArguments& InArgs, TSharedPtr<FC
 	DisplaySettings.bShowNulls = false;
 	DisplaySettings.bShowReferences = false;
 	DisplaySettings.bShowSockets = false;
-	DisplaySettings.bShowPhysics = false;
 	DisplaySettings.bHideParentsOnFilter = true;
 	DisplaySettings.bFlattenHierarchyOnFilter = true;
 	DisplaySettings.bShowIconColors = true;
diff --git a/Engine/Plugins/Animation/ControlRig/Source/ControlRigEditor/Private/EditMode/SControlRigOutliner.cpp b/Engine/Plugins/Animation/ControlRig/Source/ControlRigEditor/Private/EditMode/SControlRigOutliner.cpp
--- a/Engine/Plugins/Animation/ControlRig/Source/ControlRigEditor/Private/EditMode/SControlRigOutliner.cpp
+++ b/Engine/Plugins/Animation/ControlRig/Source/ControlRigEditor/Private/EditMode/SControlRigOutliner.cpp
@@ -1039,7 +1039,7 @@ void SControlRigOutliner::Construct(const FArguments& InArgs, FControlRigEditMod
 	DisplaySettings.bShowNulls = false;
 	DisplaySettings.bShowReferences = false;
 	DisplaySettings.bShowSockets = false;
-	DisplaySettings.bShowPhysics = false;
+	DisplaySettings.bShowComponents = false;
 	DisplaySettings.bHideParentsOnFilter = true;
 	DisplaySettings.bFlattenHierarchyOnFilter = true;
 	DisplaySettings.bShowConnectors = false;
diff --git a/Engine/Plugins/Animation/ControlRig/Source/ControlRigEditor/Private/Editor/ControlRigContextMenuContext.cpp b/Engine/Plugins/Animation/ControlRig/Source/ControlRigEditor/Private/Editor/ControlRigContextMenuContext.cpp
--- a/Engine/Plugins/Animation/ControlRig/Source/ControlRigEditor/Private/Editor/ControlRigContextMenuContext.cpp
+++ b/Engine/Plugins/Animation/ControlRig/Source/ControlRigEditor/Private/Editor/ControlRigContextMenuContext.cpp
@@ -10,12 +10,12 @@
 
 FString FControlRigRigHierarchyToGraphDragAndDropContext::GetSectionTitle() const
 {
-	TArray<FString> ElementNameStrings;
-	for (const FRigElementKey& Element: DraggedElementKeys)
+	TArray<FString> NameStrings;
+	for (const FRigHierarchyKey& HierarchyKey: DraggedHierarchyKeys)
 	{
-		ElementNameStrings.Add(Element.Name.ToString());
+		NameStrings.Add(HierarchyKey.GetName());
 	}
-	FString SectionTitle = FString::Join(ElementNameStrings, TEXT(","));
+	FString SectionTitle = FString::Join(NameStrings, TEXT(","));
 	if(SectionTitle.Len() > 64)
 	{
 		SectionTitle = SectionTitle.Left(61) + TEXT("...");
diff --git a/Engine/Plugins/Animation/ControlRig/Source/ControlRigEditor/Private/Editor/ControlRigContextMenuContext.h b/Engine/Plugins/Animation/ControlRig/Source/ControlRigEditor/Private/Editor/ControlRigContextMenuContext.h
--- a/Engine/Plugins/Animation/ControlRig/Source/ControlRigEditor/Private/Editor/ControlRigContextMenuContext.h
+++ b/Engine/Plugins/Animation/ControlRig/Source/ControlRigEditor/Private/Editor/ControlRigContextMenuContext.h
@@ -24,17 +24,17 @@ struct FControlRigRigHierarchyDragAndDropContext
 	
 	FControlRigRigHierarchyDragAndDropContext() = default;
 	
-	FControlRigRigHierarchyDragAndDropContext(const TArray<FRigElementKey> InDraggedElementKeys, FRigElementKey InTargetElementKey)
-		: DraggedElementKeys(InDraggedElementKeys)
-		, TargetElementKey(InTargetElementKey)
+	FControlRigRigHierarchyDragAndDropContext(const TArray<FRigHierarchyKey> InDraggedHierarchyKeys, FRigHierarchyKey InTargetHierarchyKey)
+		: DraggedHierarchyKeys(InDraggedHierarchyKeys)
+		, TargetHierarchyKey(InTargetHierarchyKey)
 	{
 	}
 	
 	UPROPERTY(BlueprintReadOnly, Category = ControlRigEditorExtensions)
-	TArray<FRigElementKey> DraggedElementKeys;
+	TArray<FRigHierarchyKey> DraggedHierarchyKeys;
 
 	UPROPERTY(BlueprintReadOnly, Category = ControlRigEditorExtensions)
-	FRigElementKey TargetElementKey;
+	FRigHierarchyKey TargetHierarchyKey;
 };
 
 USTRUCT(BlueprintType)
@@ -76,14 +76,14 @@ struct FControlRigRigHierarchyToGraphDragAndDropContext
 
 	FControlRigRigHierarchyToGraphDragAndDropContext() = default;
 	
-	FControlRigRigHierarchyToGraphDragAndDropContext(const TArray<FRigElementKey>& InDraggedElementKeys, UEdGraph* InGraph, const FVector2D& InNodePosition)
-		:DraggedElementKeys(InDraggedElementKeys)
+	FControlRigRigHierarchyToGraphDragAndDropContext(const TArray<FRigHierarchyKey>& InDraggedHierarchyKeys, UEdGraph* InGraph, const FVector2D& InNodePosition)
+		:DraggedHierarchyKeys(InDraggedHierarchyKeys)
 		,Graph(InGraph)
 		,NodePosition(InNodePosition)
 	{ };
 
 	UPROPERTY(BlueprintReadOnly, Category = ControlRigEditorExtensions)
-	TArray<FRigElementKey> DraggedElementKeys;
+	TArray<FRigHierarchyKey> DraggedHierarchyKeys;
 
 	TWeakObjectPtr<UEdGraph> Graph;
 
diff --git a/Engine/Plugins/Animation/ControlRig/Source/ControlRigEditor/Private/Editor/ControlRigEditor.cpp b/Engine/Plugins/Animation/ControlRig/Source/ControlRigEditor/Private/Editor/ControlRigEditor.cpp
--- a/Engine/Plugins/Animation/ControlRig/Source/ControlRigEditor/Private/Editor/ControlRigEditor.cpp
+++ b/Engine/Plugins/Animation/ControlRig/Source/ControlRigEditor/Private/Editor/ControlRigEditor.cpp
@@ -1480,17 +1480,17 @@ void FControlRigEditor::HandleSetObjectBeingDebugged(UObject* InObject)
 void FControlRigEditor::SetDetailViewForRigElements()
 {
 	URigHierarchy* HierarchyBeingDebugged = GetHierarchyBeingDebugged();
-	SetDetailViewForRigElements(HierarchyBeingDebugged->GetSelectedKeys());
+	SetDetailViewForRigElements(HierarchyBeingDebugged->GetSelectedHierarchyKeys());
 }
 
-void FControlRigEditor::SetDetailViewForRigElements(const TArray<FRigElementKey>& InKeys)
+void FControlRigEditor::SetDetailViewForRigElements(const TArray<FRigHierarchyKey>& InKeys)
 {
 	if(IsDetailsPanelRefreshSuspended())
 	{
 		return;
 	}
 
-	TArray<FRigElementKey> Keys = InKeys;
+	TArray<FRigHierarchyKey> Keys = InKeys;
 	if(Keys.IsEmpty())
 	{
 		TArray< TWeakObjectPtr<UObject> > SelectedObjects = GetSelectedObjects();
@@ -1506,6 +1506,10 @@ void FControlRigEditor::SetDetailViewForRigElements(const TArray<FRigElementKey>
 						{
 							Keys.Add(WrapperObject->GetContent<FRigBaseElement>().Key);
 						}
+						if (WrappedStruct->IsChildOf(FRigBaseComponent::StaticStruct()))
+						{
+							Keys.Add(WrapperObject->GetContent<FRigBaseComponent>().Key);
+						}
 					}
 				}
 			}
@@ -1518,19 +1522,36 @@ void FControlRigEditor::SetDetailViewForRigElements(const TArray<FRigElementKey>
 	URigHierarchy* HierarchyBeingDebugged = GetHierarchyBeingDebugged();
 	TArray<UObject*> Objects;
 
-	for(const FRigElementKey& Key : Keys)
+	for(const FRigHierarchyKey& Key : Keys)
 	{
-		FRigBaseElement* Element = HierarchyBeingDebugged->Find(Key);
-		if (Element == nullptr)
+		if(Key.IsElement())
 		{
-			continue;
+			FRigBaseElement* Element = HierarchyBeingDebugged->Find(Key.GetElement());
+			if (Element == nullptr)
+			{
+				continue;
+			}
+
+			URigVMDetailsViewWrapperObject* WrapperObject = URigVMDetailsViewWrapperObject::MakeInstance(GetDetailWrapperClass(), FRigVMEditorBase::GetBlueprintObj(), Element->GetScriptStruct(), (uint8*)Element, HierarchyBeingDebugged);
+			WrapperObject->GetWrappedPropertyChangedChainEvent().AddSP(this, &FControlRigEditor::OnWrappedPropertyChangedChainEvent);
+			WrapperObject->AddToRoot();
+
+			Objects.Add(WrapperObject);
 		}
+		if(Key.IsComponent())
+		{
+			FRigBaseComponent* Component = HierarchyBeingDebugged->FindComponent(Key.GetComponent());
+			if (Component == nullptr)
+			{
+				continue;
+			}
 
-		URigVMDetailsViewWrapperObject* WrapperObject = URigVMDetailsViewWrapperObject::MakeInstance(GetDetailWrapperClass(), FRigVMEditorBase::GetBlueprintObj(), Element->GetScriptStruct(), (uint8*)Element, HierarchyBeingDebugged);
-		WrapperObject->GetWrappedPropertyChangedChainEvent().AddSP(this, &FControlRigEditor::OnWrappedPropertyChangedChainEvent);
-		WrapperObject->AddToRoot();
+			URigVMDetailsViewWrapperObject* WrapperObject = URigVMDetailsViewWrapperObject::MakeInstance(GetDetailWrapperClass(), FRigVMEditorBase::GetBlueprintObj(), Component->GetScriptStruct(), (uint8*)Component, HierarchyBeingDebugged);
+			WrapperObject->GetWrappedPropertyChangedChainEvent().AddSP(this, &FControlRigEditor::OnWrappedPropertyChangedChainEvent);
+			WrapperObject->AddToRoot();
 
-		Objects.Add(WrapperObject);
+			Objects.Add(WrapperObject);
+		}
 	}
 	
 	SetDetailObjects(Objects);
@@ -1571,7 +1592,7 @@ bool FControlRigEditor::DetailViewShowsAnyRigElement() const
 	return DetailViewShowsStruct(FRigBaseElement::StaticStruct());
 }
 
-bool FControlRigEditor::DetailViewShowsRigElement(FRigElementKey InKey) const
+bool FControlRigEditor::DetailViewShowsRigElement(FRigHierarchyKey InKey) const
 {
 	TArray< TWeakObjectPtr<UObject> > SelectedObjects = GetSelectedObjectsFromDetailView();
 	for (TWeakObjectPtr<UObject> SelectedObject : SelectedObjects)
@@ -1589,6 +1610,13 @@ bool FControlRigEditor::DetailViewShowsRigElement(FRigElementKey InKey) const
 							return true;
 						}
 					}
+					if (WrappedStruct->IsChildOf(FRigBaseComponent::StaticStruct()))
+					{
+						if(WrapperObject->GetContent<FRigBaseComponent>().GetKey() == InKey)
+						{
+							return true;
+						}
+					}
 				}
 			}
 		}
@@ -1596,9 +1624,9 @@ bool FControlRigEditor::DetailViewShowsRigElement(FRigElementKey InKey) const
 	return false;
 }
 
-TArray<FRigElementKey> FControlRigEditor::GetSelectedRigElementsFromDetailView() const
+TArray<FRigHierarchyKey> FControlRigEditor::GetSelectedRigElementsFromDetailView() const
 {
-	TArray<FRigElementKey> Elements;
+	TArray<FRigHierarchyKey> Keys;
 	
 	TArray< TWeakObjectPtr<UObject> > SelectedObjects = GetSelectedObjectsFromDetailView();
 	for (TWeakObjectPtr<UObject> SelectedObject : SelectedObjects)
@@ -1611,14 +1639,18 @@ TArray<FRigElementKey> FControlRigEditor::GetSelectedRigElementsFromDetailView()
 				{
 					if (WrappedStruct->IsChildOf(FRigBaseElement::StaticStruct()))
 					{
-						Elements.Add(WrapperObject->GetContent<FRigBaseElement>().GetKey());
+						Keys.Add(WrapperObject->GetContent<FRigBaseElement>().GetKey());
+					}
+					if (WrappedStruct->IsChildOf(FRigBaseComponent::StaticStruct()))
+					{
+						Keys.Add(WrapperObject->GetContent<FRigBaseComponent>().GetKey());
 					}
 				}
 			}
 		}
 	}
 
-	return Elements;
+	return Keys;
 }
 
 TArray<TWeakObjectPtr<UObject>> FControlRigEditor::GetSelectedObjectsFromDetailView() const
@@ -3221,13 +3253,22 @@ FVector2D FControlRigEditor::ComputePersonaProjectedScreenPos(const FVector& InW
 	return FVector2D::ZeroVector;
 }
 
-void FControlRigEditor::FindReferencesOfItem(const FRigElementKey& InKey)
+void FControlRigEditor::FindReferencesOfItem(const FRigHierarchyKey& InKey)
 {
-	static constexpr TCHAR Format[] = TEXT("Type,%s,Name,%s");
-	static const UEnum* TypeEnum = StaticEnum<ERigElementType>();
-	const FText TypeText = TypeEnum->GetDisplayNameTextByValue((int64)InKey.Type);
-	const FString Query = FString::Printf(Format, *TypeText.ToString(), *InKey.Name.ToString()); 
-	SummonSearchUI(true, Query, true);
+	if(InKey.IsElement())
+	{
+		static constexpr TCHAR Format[] = TEXT("Type,%s,Name,%s");
+		static const UEnum* TypeEnum = StaticEnum<ERigElementType>();
+		const FText TypeText = TypeEnum->GetDisplayNameTextByValue((int64)InKey.GetElement().Type);
+		const FString Query = FString::Printf(Format, *TypeText.ToString(), *InKey.GetName()); 
+		SummonSearchUI(true, Query, true);
+	}
+	if(InKey.IsComponent())
+	{
+		static constexpr TCHAR Format[] = TEXT("Name,%s");
+		const FString Query = FString::Printf(Format, *InKey.GetName()); 
+		SummonSearchUI(true, Query, true);
+	}
 }
 
 void FControlRigEditor::HandlePreviewMeshChanged(USkeletalMesh* InOldSkeletalMesh, USkeletalMesh* InNewSkeletalMesh)
@@ -3845,6 +3886,15 @@ void FControlRigEditor::OnWrappedPropertyChangedChainEvent(URigVMDetailsViewWrap
 				ControlRigBP->MarkPackageDirty();
 			}
 		}
+		else if(WrappedStruct->IsChildOf(FRigBaseComponent::StaticStruct()))
+		{
+			FStructOnScope Content(WrappedStruct);
+			InWrapperObject->GetContent(Content.GetStructMemory(), Content.GetStruct());
+
+			const FRigBaseComponent* WrappedComponent = reinterpret_cast<const FRigBaseComponent*>(Content.GetStructMemory());
+			FString TextContent = WrappedComponent->GetContentAsText();
+			(void)ControlRigBP->GetHierarchyController()->SetComponentContent(WrappedComponent->GetKey(), TextContent, true, true);
+		}
 	}
 }
 
@@ -4970,8 +5020,19 @@ void FControlRigEditor::CreateRigHierarchyToGraphDragAndDropMenu() const
 					const FControlRigRigHierarchyToGraphDragAndDropContext& DragDropContext = MainContext->GetRigHierarchyToGraphDragAndDropContext();
 
 					URigHierarchy* Hierarchy = ControlRigEditor->GetHierarchyBeingDebugged();
-					TArray<FRigElementKey> DraggedKeys = DragDropContext.DraggedElementKeys;
-					ControlRigEditor->FilterDraggedKeys(DraggedKeys, true);
+					TArray<FRigElementKey> DraggedElements;
+					DraggedElements.Reserve(DragDropContext.DraggedHierarchyKeys.Num());
+					for(const FRigHierarchyKey& DraggedHierarchyKey : DragDropContext.DraggedHierarchyKeys)
+					{
+						if(DraggedHierarchyKey.IsElement())
+						{
+							DraggedElements.Add(DraggedHierarchyKey.GetElement());
+						}
+					}
+
+					// todo: dropping components won't work for now.
+					
+					ControlRigEditor->FilterDraggedKeys(DraggedElements, true);
 					
 					UEdGraph* Graph = DragDropContext.Graph.Get();
 					const FVector2D& NodePosition = DragDropContext.NodePosition;
@@ -4981,14 +5042,14 @@ void FControlRigEditor::CreateRigHierarchyToGraphDragAndDropMenu() const
 
 					ERigElementType LastType = ERigElementType::None;
 			
-					if (DraggedKeys.Num() > 0)
+					if (DraggedElements.Num() > 0)
 					{
-						LastType = DraggedKeys[0].Type;
+						LastType = DraggedElements[0].Type;
 					}
 			
 					uint8 DraggedTypes = 0;
 					uint8 DraggedAnimationTypes = 2;
-					for (const FRigElementKey& DraggedKey : DragDropContext.DraggedElementKeys)
+					for (const FRigElementKey& DraggedKey : DraggedElements)
 					{
 						if (DraggedKey.Type != LastType)
 						{
@@ -5080,7 +5141,7 @@ void FControlRigEditor::CreateRigHierarchyToGraphDragAndDropMenu() const
 						GetterTooltip,
 						FSlateIcon(),
 						FUIAction(
-							FExecuteAction::CreateSP(ControlRigEditor, &FControlRigEditor::HandleMakeElementGetterSetter, ERigElementGetterSetterType_Transform, true, DraggedKeys, Graph, NodePosition),
+							FExecuteAction::CreateSP(ControlRigEditor, &FControlRigEditor::HandleMakeElementGetterSetter, ERigElementGetterSetterType_Transform, true, DraggedElements, Graph, NodePosition),
 							FCanExecuteAction()
 						)
 					);
@@ -5096,7 +5157,7 @@ void FControlRigEditor::CreateRigHierarchyToGraphDragAndDropMenu() const
 						SetterTooltip,
 						FSlateIcon(),
 						FUIAction(
-							FExecuteAction::CreateSP(ControlRigEditor, &FControlRigEditor::HandleMakeElementGetterSetter, ERigElementGetterSetterType_Transform, false, DraggedKeys, Graph, NodePosition),
+							FExecuteAction::CreateSP(ControlRigEditor, &FControlRigEditor::HandleMakeElementGetterSetter, ERigElementGetterSetterType_Transform, false, DraggedElements, Graph, NodePosition),
 							FCanExecuteAction()
 						)
 					);
@@ -5120,7 +5181,7 @@ void FControlRigEditor::CreateRigHierarchyToGraphDragAndDropMenu() const
 							LOCTEXT("SetRotation_ToolTip","Setter for Rotation"),
 							FSlateIcon(),
 							FUIAction(
-								FExecuteAction::CreateSP(ControlRigEditor, &FControlRigEditor::HandleMakeElementGetterSetter, ERigElementGetterSetterType_Rotation, false, DraggedKeys, Graph, NodePosition),
+								FExecuteAction::CreateSP(ControlRigEditor, &FControlRigEditor::HandleMakeElementGetterSetter, ERigElementGetterSetterType_Rotation, false, DraggedElements, Graph, NodePosition),
 								FCanExecuteAction()
 							)
 						);
@@ -5135,7 +5196,7 @@ void FControlRigEditor::CreateRigHierarchyToGraphDragAndDropMenu() const
 							LOCTEXT("SetTranslation_ToolTip","Setter for Translation"),
 							FSlateIcon(),
 							FUIAction(
-								FExecuteAction::CreateSP(ControlRigEditor, &FControlRigEditor::HandleMakeElementGetterSetter, ERigElementGetterSetterType_Translation, false, DraggedKeys, Graph, NodePosition),
+								FExecuteAction::CreateSP(ControlRigEditor, &FControlRigEditor::HandleMakeElementGetterSetter, ERigElementGetterSetterType_Translation, false, DraggedElements, Graph, NodePosition),
 								FCanExecuteAction()
 							)
 						);
@@ -5150,7 +5211,7 @@ void FControlRigEditor::CreateRigHierarchyToGraphDragAndDropMenu() const
 							LOCTEXT("AddOffset_ToolTip","Setter for Offset"),
 							FSlateIcon(),
 							FUIAction(
-								FExecuteAction::CreateSP(ControlRigEditor, &FControlRigEditor::HandleMakeElementGetterSetter, ERigElementGetterSetterType_Offset, false, DraggedKeys, Graph, NodePosition),
+								FExecuteAction::CreateSP(ControlRigEditor, &FControlRigEditor::HandleMakeElementGetterSetter, ERigElementGetterSetterType_Offset, false, DraggedElements, Graph, NodePosition),
 								FCanExecuteAction()
 							)
 						);
@@ -5169,7 +5230,7 @@ void FControlRigEditor::CreateRigHierarchyToGraphDragAndDropMenu() const
 							LOCTEXT("GetRelativeTransform_ToolTip", "Getter for Relative Transform"),
 							FSlateIcon(),
 							FUIAction(
-								FExecuteAction::CreateSP(ControlRigEditor, &FControlRigEditor::HandleMakeElementGetterSetter, ERigElementGetterSetterType_Relative, true, DraggedKeys, Graph, NodePosition),
+								FExecuteAction::CreateSP(ControlRigEditor, &FControlRigEditor::HandleMakeElementGetterSetter, ERigElementGetterSetterType_Relative, true, DraggedElements, Graph, NodePosition),
 								FCanExecuteAction()
 							)
 						);
@@ -5184,7 +5245,7 @@ void FControlRigEditor::CreateRigHierarchyToGraphDragAndDropMenu() const
 							LOCTEXT("SetRelativeTransform_ToolTip", "Setter for Relative Transform"),
 							FSlateIcon(),
 							FUIAction(
-								FExecuteAction::CreateSP(ControlRigEditor, &FControlRigEditor::HandleMakeElementGetterSetter, ERigElementGetterSetterType_Relative, false, DraggedKeys, Graph, NodePosition),
+								FExecuteAction::CreateSP(ControlRigEditor, &FControlRigEditor::HandleMakeElementGetterSetter, ERigElementGetterSetterType_Relative, false, DraggedElements, Graph, NodePosition),
 								FCanExecuteAction()
 							)
 						);
@@ -5193,7 +5254,7 @@ void FControlRigEditor::CreateRigHierarchyToGraphDragAndDropMenu() const
 						Section.AddEntry(SetRelativeTransformEntry);
 					}
 
-					if (DraggedKeys.Num() > 0 && Hierarchy != nullptr)
+					if (DraggedElements.Num() > 0 && Hierarchy != nullptr)
 					{
 						FToolMenuEntry& ItemArraySeparator = Section.AddSeparator(TEXT("ItemArraySeparator"));
 						ItemArraySeparator.InsertPosition.Name = TEXT("SetRelativeTransformEntry"),
@@ -5205,7 +5266,7 @@ void FControlRigEditor::CreateRigHierarchyToGraphDragAndDropMenu() const
 							LOCTEXT("CreateItemArray_ToolTip", "Creates an item array from the selected elements in the hierarchy"),
 							FSlateIcon(),
 							FUIAction(
-								FExecuteAction::CreateLambda([ControlRigEditor, DraggedKeys, NodePosition]()
+								FExecuteAction::CreateLambda([ControlRigEditor, DraggedElements, NodePosition]()
 									{
 										if (URigVMController* Controller = ControlRigEditor->GetFocusedController())
 										{
@@ -5215,15 +5276,15 @@ void FControlRigEditor::CreateRigHierarchyToGraphDragAndDropMenu() const
 											{
 												if (URigVMPin* ItemsPin = ItemsNode->FindPin(TEXT("Items")))
 												{
-													Controller->SetArrayPinSize(ItemsPin->GetPinPath(), DraggedKeys.Num());
+													Controller->SetArrayPinSize(ItemsPin->GetPinPath(), DraggedElements.Num());
 
 													TArray<URigVMPin*> ItemPins = ItemsPin->GetSubPins();
-													ensure(ItemPins.Num() == DraggedKeys.Num());
+													ensure(ItemPins.Num() == DraggedElements.Num());
 
-													for (int32 ItemIndex = 0; ItemIndex < DraggedKeys.Num(); ItemIndex++)
+													for (int32 ItemIndex = 0; ItemIndex < DraggedElements.Num(); ItemIndex++)
 													{
 														FString DefaultValue;
-														FRigElementKey::StaticStruct()->ExportText(DefaultValue, &DraggedKeys[ItemIndex], nullptr, nullptr, PPF_None, nullptr);
+														FRigElementKey::StaticStruct()->ExportText(DefaultValue, &DraggedElements[ItemIndex], nullptr, nullptr, PPF_None, nullptr);
 														Controller->SetPinDefaultValue(ItemPins[ItemIndex]->GetPinPath(), DefaultValue, true, true, false, true);
 														Controller->SetPinExpansion(ItemPins[ItemIndex]->GetPinPath(), true, true, true);
 													}
@@ -6101,7 +6162,7 @@ void FControlRigEditor::OnPostConstruction_AnyThread(UControlRig* InRig, const F
 			RebindToSkeletalMeshComponent();
 			if(DetailViewShowsAnyRigElement())
 			{
-				const TArray<FRigElementKey> Keys = GetSelectedRigElementsFromDetailView();
+				const TArray<FRigHierarchyKey> Keys = GetSelectedRigElementsFromDetailView();
 				SetDetailViewForRigElements(Keys);
 			}
 			
diff --git a/Engine/Plugins/Animation/ControlRig/Source/ControlRigEditor/Private/Editor/ControlRigEditor.h b/Engine/Plugins/Animation/ControlRig/Source/ControlRigEditor/Private/Editor/ControlRigEditor.h
--- a/Engine/Plugins/Animation/ControlRig/Source/ControlRigEditor/Private/Editor/ControlRigEditor.h
+++ b/Engine/Plugins/Animation/ControlRig/Source/ControlRigEditor/Private/Editor/ControlRigEditor.h
@@ -135,10 +135,10 @@ public:
 	URigHierarchy* GetHierarchyBeingDebugged() const;
 
 	void SetDetailViewForRigElements();
-	void SetDetailViewForRigElements(const TArray<FRigElementKey>& InKeys);
+	void SetDetailViewForRigElements(const TArray<FRigHierarchyKey>& InKeys);
 	bool DetailViewShowsAnyRigElement() const;
-	bool DetailViewShowsRigElement(FRigElementKey InKey) const;
-	TArray<FRigElementKey> GetSelectedRigElementsFromDetailView() const;
+	bool DetailViewShowsRigElement(FRigHierarchyKey InKey) const;
+	TArray<FRigHierarchyKey> GetSelectedRigElementsFromDetailView() const;
 	TArray< TWeakObjectPtr<UObject> > GetSelectedObjectsFromDetailView() const;
 
 	void SetDetailViewForRigModules();
@@ -207,7 +207,7 @@ public:
 
 	FVector2D ComputePersonaProjectedScreenPos(const FVector& InWorldPos, bool bClampToScreenRectangle = false);
 
-	void FindReferencesOfItem(const FRigElementKey& InKey);
+	void FindReferencesOfItem(const FRigHierarchyKey& InKey);
 	
 protected:
 
diff --git a/Engine/Plugins/Animation/ControlRig/Source/ControlRigEditor/Private/Editor/ControlRigSchematicModel.cpp b/Engine/Plugins/Animation/ControlRig/Source/ControlRigEditor/Private/Editor/ControlRigSchematicModel.cpp
--- a/Engine/Plugins/Animation/ControlRig/Source/ControlRigEditor/Private/Editor/ControlRigSchematicModel.cpp
+++ b/Engine/Plugins/Animation/ControlRig/Source/ControlRigEditor/Private/Editor/ControlRigSchematicModel.cpp
@@ -1205,7 +1205,17 @@ TArray<FRigElementKey> FControlRigSchematicModel::GetElementKeysFromDragDropEven
 	if(InDragDropOperation.IsOfType<FRigElementHierarchyDragDropOp>())
 	{
 		const FRigElementHierarchyDragDropOp* RigDragDropOp = StaticCast<const FRigElementHierarchyDragDropOp*>(&InDragDropOperation);
-		return RigDragDropOp->GetElements();
+		const TArray<FRigHierarchyKey>& HierarchyKeys = RigDragDropOp->GetElements();
+		TArray<FRigElementKey> Elements;
+		Elements.Reserve(HierarchyKeys.Num());
+		for(const FRigHierarchyKey& HierarchyKey : HierarchyKeys)
+		{
+			if(HierarchyKey.IsElement())
+			{
+				Elements.Add(HierarchyKey.GetElement());
+			}
+		}
+		return Elements;
 	}
 
 	if(InDragDropOperation.IsOfType<FRigHierarchyTagDragDropOp>())
diff --git a/Engine/Plugins/Animation/ControlRig/Source/ControlRigEditor/Private/Editor/ControlRigWrapperObject.cpp b/Engine/Plugins/Animation/ControlRig/Source/ControlRigEditor/Private/Editor/ControlRigWrapperObject.cpp
--- a/Engine/Plugins/Animation/ControlRig/Source/ControlRigEditor/Private/Editor/ControlRigWrapperObject.cpp
+++ b/Engine/Plugins/Animation/ControlRig/Source/ControlRigEditor/Private/Editor/ControlRigWrapperObject.cpp
@@ -47,7 +47,47 @@ UClass* UControlRigWrapperObject::GetClassForStruct(UScriptStruct* InStruct, boo
 			}
 		}
 	}
+	if(InStruct->IsChildOf(FRigBaseComponent::StaticStruct()))
+	{
+		FPropertyEditorModule& PropertyEditorModule = FModuleManager::LoadModuleChecked<FPropertyEditorModule>("PropertyEditor");
+		if (!PropertyEditorModule.GetClassNameToDetailLayoutNameMap().Contains(WrapperClassName))
+		{
+			PropertyEditorModule.RegisterCustomClassLayout(WrapperClassName, FOnGetDetailCustomizationInstance::CreateStatic(&FRigBaseComponentDetails::MakeInstance));
+		}
+	}
 #endif
 
 	return Class;
 }
+
+void UControlRigWrapperObject::SetContent(const uint8* InStructMemory, const UStruct* InStruct)
+{
+	Super::SetContent(InStructMemory, InStruct);
+
+	if(InStruct->IsChildOf((FRigBaseElement::StaticStruct())))
+	{
+		const FRigBaseElement* SourceElement = reinterpret_cast<const FRigBaseElement*>(InStructMemory);
+		HierarchyKey = SourceElement->GetKey();
+	}
+	else if(InStruct->IsChildOf((FRigBaseComponent::StaticStruct())))
+	{
+		const FRigBaseComponent* SourceComponent = reinterpret_cast<const FRigBaseComponent*>(InStructMemory);
+		HierarchyKey = SourceComponent->GetKey();
+	}
+}
+
+void UControlRigWrapperObject::GetContent(uint8* OutStructMemory, const UStruct* InStruct) const
+{
+	Super::GetContent(OutStructMemory, InStruct);
+
+	if(InStruct->IsChildOf((FRigBaseElement::StaticStruct())))
+	{
+		FRigBaseElement* TargetElement = reinterpret_cast<FRigBaseElement*>(OutStructMemory);
+		TargetElement->Key = HierarchyKey.GetElement();
+	}
+	else if(InStruct->IsChildOf((FRigBaseComponent::StaticStruct())))
+	{
+		FRigBaseComponent* TargetComponent = reinterpret_cast<FRigBaseComponent*>(OutStructMemory);
+		TargetComponent->Key = HierarchyKey.GetComponent();
+	}
+}
diff --git a/Engine/Plugins/Animation/ControlRig/Source/ControlRigEditor/Private/Editor/ControlRigWrapperObject.h b/Engine/Plugins/Animation/ControlRig/Source/ControlRigEditor/Private/Editor/ControlRigWrapperObject.h
--- a/Engine/Plugins/Animation/ControlRig/Source/ControlRigEditor/Private/Editor/ControlRigWrapperObject.h
+++ b/Engine/Plugins/Animation/ControlRig/Source/ControlRigEditor/Private/Editor/ControlRigWrapperObject.h
@@ -3,6 +3,7 @@
 #pragma once
 
 #include "Editor/RigVMDetailsViewWrapperObject.h"
+#include "Rigs/RigHierarchyDefines.h"
 #include "ControlRigWrapperObject.generated.h"
 
 UCLASS()
@@ -12,4 +13,9 @@ public:
 	GENERATED_BODY()
 
 	virtual UClass* GetClassForStruct(UScriptStruct* InStruct, bool bCreateIfNeeded) const override;
+
+	virtual void SetContent(const uint8* InStructMemory, const UStruct* InStruct) override;
+	virtual void GetContent(uint8* OutStructMemory, const UStruct* InStruct) const override;;
+
+	FRigHierarchyKey HierarchyKey;
 };
diff --git a/Engine/Plugins/Animation/ControlRig/Source/ControlRigEditor/Private/Editor/SRigConnectorTargetWidget.cpp b/Engine/Plugins/Animation/ControlRig/Source/ControlRigEditor/Private/Editor/SRigConnectorTargetWidget.cpp
--- a/Engine/Plugins/Animation/ControlRig/Source/ControlRigEditor/Private/Editor/SRigConnectorTargetWidget.cpp
+++ b/Engine/Plugins/Animation/ControlRig/Source/ControlRigEditor/Private/Editor/SRigConnectorTargetWidget.cpp
@@ -30,7 +30,7 @@ void SRigConnectorTargetComboButton::Construct(const FArguments& InArgs)
 	RigTreeDelegates = InArgs._RigTreeDelegates;
 	OnSetTarget = InArgs._OnSetTarget;
 
-	RigTreeDelegates.OnGetSelection.BindLambda([this]() -> TArray<FRigElementKey>
+	RigTreeDelegates.OnGetSelection.BindLambda([this]() -> TArray<FRigHierarchyKey>
 	{
 		return {TargetKey.Get()};
 	});
@@ -179,7 +179,7 @@ void SRigConnectorTargetComboButton::OnConnectorTargetPicked(TSharedPtr<FRigTree
 	{
 		if(Selection->Key != TargetKey.Get())
 		{
-			(void)OnSetTarget.Execute(Selection->Key);
+			(void)OnSetTarget.Execute(Selection->Key.GetElement());
 		}
 	}
 }
@@ -223,9 +223,9 @@ void SRigConnectorTargetWidget::Construct(const FArguments& InArgs)
 			}
 		}
 
-		RigTreeDelegates.OnRigTreeIsItemVisible.BindLambda([PotentialTargets](const FRigElementKey& InItem)
+		RigTreeDelegates.OnRigTreeIsItemVisible.BindLambda([PotentialTargets](const FRigHierarchyKey& InItem)
 		{
-			return InItem.IsValid() && (PotentialTargets.IsEmpty() || PotentialTargets.Contains(InItem)); 
+			return InItem.IsValid() && InItem.IsElement() && (PotentialTargets.IsEmpty() || PotentialTargets.Contains(InItem.GetElement())); 
 		});
 	}
 
diff --git a/Engine/Plugins/Animation/ControlRig/Source/ControlRigEditor/Private/Editor/SRigHierarchy.cpp b/Engine/Plugins/Animation/ControlRig/Source/ControlRigEditor/Private/Editor/SRigHierarchy.cpp
--- a/Engine/Plugins/Animation/ControlRig/Source/ControlRigEditor/Private/Editor/SRigHierarchy.cpp
+++ b/Engine/Plugins/Animation/ControlRig/Source/ControlRigEditor/Private/Editor/SRigHierarchy.cpp
@@ -49,6 +49,7 @@
 #include "Sequencer/ControlRigLayerInstance.h"
 #include "Algo/MinElement.h"
 #include "Algo/MaxElement.h"
+#include "Algo/Sort.h"
 #include "DragAndDrop/AssetDragDropOp.h"
 #include "RigVMFunctions/Math/RigVMMathLibrary.h"
 #include "Preferences/PersonaOptions.h"
@@ -117,7 +118,7 @@ void SRigHierarchy::Construct(const FArguments& InArgs, TSharedRef<FControlRigEd
 	Delegates.OnDragDetected = FOnDragDetected::CreateSP(this, &SRigHierarchy::OnDragDetected);
 	Delegates.OnGetResolvedKey = FOnRigTreeGetResolvedKey::CreateSP(this, &SRigHierarchy::OnGetResolvedKey);
 	Delegates.OnRequestDetailsInspection = FOnRigTreeRequestDetailsInspection::CreateSP(this, &SRigHierarchy::OnRequestDetailsInspection);
-	Delegates.OnRigTreeElementKeyTagDragDetected = FOnRigTreeRequestDetailsInspection::CreateSP(this, &SRigHierarchy::OnElementKeyTagDragDetected);
+	Delegates.OnRigTreeElementKeyTagDragDetected = FOnRigTreeElementKeyTagDragDetected::CreateSP(this, &SRigHierarchy::OnElementKeyTagDragDetected);
 	Delegates.OnRigTreeGetItemToolTip = FOnRigTreeItemGetToolTip::CreateSP(this, &SRigHierarchy::OnGetItemTooltip);
 
 	ChildSlot
@@ -430,12 +431,6 @@ void SRigHierarchy::BindCommands()
 		FCanExecuteAction(),
 		FIsActionChecked::CreateLambda([this]() { return DisplaySettings.bShowNulls; }));
 
-	CommandList->MapAction(
-		Commands.ShowPhysics,
-		FExecuteAction::CreateLambda([this]() { DisplaySettings.bShowPhysics = !DisplaySettings.bShowPhysics; RefreshTreeView(); }),
-		FCanExecuteAction(),
-		FIsActionChecked::CreateLambda([this]() { return DisplaySettings.bShowPhysics; }));
-
 	CommandList->MapAction(
 		Commands.ShowReferences,
 		FExecuteAction::CreateLambda([this]() { DisplaySettings.bShowReferences = !DisplaySettings.bShowReferences; RefreshTreeView(); }),
@@ -460,6 +455,12 @@ void SRigHierarchy::BindCommands()
 		FExecuteAction::CreateSP(this, &SRigHierarchy::HandleTestSpaceSwitching),
 		FCanExecuteAction::CreateSP(this, &SRigHierarchy::IsControlSelected, true));
 
+	CommandList->MapAction(
+		Commands.ShowComponents,
+		FExecuteAction::CreateLambda([this]() { DisplaySettings.bShowComponents = !DisplaySettings.bShowComponents; RefreshTreeView(); }),
+		FCanExecuteAction(),
+		FIsActionChecked::CreateLambda([this]() { return DisplaySettings.bShowComponents; }));
+
 	CommandList->MapAction(
 		Commands.ShowIconColors,
 		FExecuteAction::CreateLambda([this]() { DisplaySettings.bShowIconColors = !DisplaySettings.bShowIconColors; RefreshTreeView(); }),
@@ -492,13 +493,20 @@ FReply SRigHierarchy::OnMouseButtonUp(const FGeometry& MyGeometry, const FPointe
 			{
 				if (URigHierarchy* Hierarchy = GetHierarchy())
 				{
-					TArray<FRigElementKey> KeysToSelect = {Item->Key};
-					KeysToSelect.Append(Hierarchy->GetChildren(Item->Key, true));
+					if(Item->Key.IsElement())
+					{
+						TArray<FRigElementKey> KeysToSelect = {Item->Key.GetElement()};
+						KeysToSelect.Append(Hierarchy->GetChildren(Item->Key.GetElement(), true));
 
-					URigHierarchyController* Controller = Hierarchy->GetController(true);
-					check(Controller);
-		
-					Controller->SetSelection(KeysToSelect);
+						URigHierarchyController* Controller = Hierarchy->GetController(true);
+						check(Controller);
+			
+						Controller->SetSelection(KeysToSelect);
+					}
+					else
+					{
+						// todo: component selection
+					}
 				}
 			}
 		}
@@ -615,11 +623,11 @@ void SRigHierarchy::RefreshTreeView(bool bRebuildContent)
 	TreeView->RefreshTreeView(bRebuildContent);
 }
 
-TArray<FRigElementKey> SRigHierarchy::GetSelectedKeys() const
+TArray<FRigHierarchyKey> SRigHierarchy::GetSelectedKeys() const
 {
 	TArray<TSharedPtr<FRigTreeElement>> SelectedItems = TreeView->GetSelectedItems();
 	
-	TArray<FRigElementKey> SelectedKeys;
+	TArray<FRigHierarchyKey> SelectedKeys;
 	for (const TSharedPtr<FRigTreeElement>& SelectedItem : SelectedItems)
 	{
 		if(SelectedItem->Key.IsValid())
@@ -631,6 +639,21 @@ TArray<FRigElementKey> SRigHierarchy::GetSelectedKeys() const
 	return SelectedKeys;
 }
 
+TArray<FRigElementKey> SRigHierarchy::GetSelectedElementKeys() const
+{
+	TArray<FRigHierarchyKey> SelectedKeys = GetSelectedKeys();
+	TArray<FRigElementKey> ElementKeys;
+	ElementKeys.Reserve(SelectedKeys.Num());
+	for(const FRigHierarchyKey& SelectedKey : SelectedKeys)
+	{
+		if(SelectedKey.IsElement())
+		{
+			ElementKeys.Add(SelectedKey.GetElement());
+		}
+	}
+	return ElementKeys;
+}
+
 void SRigHierarchy::OnSelectionChanged(TSharedPtr<FRigTreeElement> Selection, ESelectInfo::Type SelectInfo)
 {
 	if (bIsChangingRigHierarchy)
@@ -641,7 +664,7 @@ void SRigHierarchy::OnSelectionChanged(TSharedPtr<FRigTreeElement> Selection, ES
 	TreeView->ClearHighlightedItems();
 
 	// an element to use for the control rig editor's detail panel
-	FRigElementKey LastSelectedElement;
+	FRigHierarchyKey LastSelectedElement;
 
 	URigHierarchy* Hierarchy = GetHierarchy();
 	if (Hierarchy)
@@ -662,8 +685,8 @@ void SRigHierarchy::OnSelectionChanged(TSharedPtr<FRigTreeElement> Selection, ES
 
 		TGuardValue<bool> SuspendDetailsPanelRefreshGuard(*SuspensionFlagPtr, true);
 		
-		const TArray<FRigElementKey> NewSelection = GetSelectedKeys();
-		if(!Controller->SetSelection(NewSelection, true))
+		const TArray<FRigHierarchyKey> NewSelection = GetSelectedKeys();
+		if(!Controller->SetHierarchySelection(NewSelection, true))
 		{
 			return;
 		}
@@ -747,6 +770,28 @@ void SRigHierarchy::OnHierarchyModified(ERigHierarchyNotification InNotif, URigH
 			}
 			break;
 		}
+		case ERigHierarchyNotification::ComponentAdded:
+		{
+			if(InComponent)
+			{
+				if(TreeView->AddComponent(InComponent))
+				{
+					RefreshTreeView(false);
+				}
+			}
+			break;
+		}
+		case ERigHierarchyNotification::ComponentRemoved:
+		{
+			if(InComponent)
+			{
+				if(TreeView->RemoveElement(InComponent->GetKey()))
+				{
+					RefreshTreeView(false);
+				}
+			}
+			break;
+		}
 		case ERigHierarchyNotification::ParentChanged:
 		{
 			check(InHierarchy);
@@ -798,14 +843,23 @@ void SRigHierarchy::OnHierarchyModified(ERigHierarchyNotification InNotif, URigH
 		}
 		case ERigHierarchyNotification::ElementSelected:
 		case ERigHierarchyNotification::ElementDeselected:
+		case ERigHierarchyNotification::ComponentSelected:
+		case ERigHierarchyNotification::ComponentDeselected:
 		{
-			if(InElement)
+			if(InElement || InComponent)
 			{
-				const bool bSelected = (InNotif == ERigHierarchyNotification::ElementSelected); 
+				const bool bSelected =
+					(InNotif == ERigHierarchyNotification::ElementSelected) || 
+					(InNotif == ERigHierarchyNotification::ComponentSelected);;
+
+				const FRigHierarchyKey Key =
+					InElement ?
+					FRigHierarchyKey(InElement->GetKey()) :
+					FRigHierarchyKey(InComponent->GetKey()) ;
 					
 				for (int32 RootIndex = 0; RootIndex < TreeView->RootElements.Num(); ++RootIndex)
 				{
-					TSharedPtr<FRigTreeElement> Found = TreeView->FindElement(InElement->GetKey(), TreeView->RootElements[RootIndex]);
+					TSharedPtr<FRigTreeElement> Found = TreeView->FindElement(Key, TreeView->RootElements[RootIndex]);
 					if (Found.IsValid())
 					{
 						TreeView->SetItemSelection(Found, bSelected, ESelectInfo::OnNavigation);
@@ -1016,7 +1070,7 @@ void SRigHierarchy::OnPreConstruction_AnyThread(UControlRig* InRig, const FName&
 		return;
 	}
 	bIsConstructionEventRunning = true;
-	SelectionBeforeConstruction = InRig->GetHierarchy()->GetSelectedKeys();
+	SelectionBeforeConstruction = InRig->GetHierarchy()->GetSelectedHierarchyKeys();
 }
 
 void SRigHierarchy::OnPostConstruction_AnyThread(UControlRig* InRig, const FName& InEventName)
@@ -1048,7 +1102,7 @@ void SRigHierarchy::OnPostConstruction_AnyThread(UControlRig* InRig, const FName
 			{
 				for (int32 RootIndex = 0; RootIndex < TreeView->RootElements.Num(); ++RootIndex)
 				{
-					for(const FRigElementKey& Key : SelectionBeforeConstruction)
+					for(const FRigHierarchyKey& Key : SelectionBeforeConstruction)
 					{
 						TSharedPtr<FRigTreeElement> Found = TreeView->FindElement(Key, TreeView->RootElements[RootIndex]);
 						if (Found.IsValid())
@@ -1136,6 +1190,7 @@ TSharedRef< SWidget > SRigHierarchy::CreateFilterMenu()
 	{
 		MenuBuilder.AddMenuEntry(Actions.ShowControls);
 		MenuBuilder.AddMenuEntry(Actions.ShowNulls);
+		MenuBuilder.AddMenuEntry(Actions.ShowComponents);
 		MenuBuilder.AddMenuEntry(Actions.ShowIconColors);
 	}
 	MenuBuilder.EndSection();
@@ -1160,20 +1215,23 @@ void SRigHierarchy::OnItemClicked(TSharedPtr<FRigTreeElement> InItem)
 	URigHierarchy* Hierarchy = GetHierarchy();
 	check(Hierarchy);
 
-	if (Hierarchy->IsSelected(InItem->Key))
+	if (Hierarchy->IsHierarchyKeySelected(InItem->Key))
 	{
 		if (ControlRigEditor.IsValid())
 		{
 			ControlRigEditor.Pin()->SetDetailViewForRigElements();
 		}
 
-		if (InItem->Key.Type == ERigElementType::Bone)
+		if(InItem->Key.IsElement())
 		{
-			if(FRigBoneElement* BoneElement = Hierarchy->Find<FRigBoneElement>(InItem->Key))
+			if (InItem->Key.GetElement().Type == ERigElementType::Bone)
 			{
-				if (BoneElement->BoneType == ERigBoneType::Imported)
+				if(FRigBoneElement* BoneElement = Hierarchy->Find<FRigBoneElement>(InItem->Key.GetElement()))
 				{
-					return;
+					if (BoneElement->BoneType == ERigBoneType::Imported)
+					{
+						return;
+					}
 				}
 			}
 		}
@@ -1209,20 +1267,23 @@ void SRigHierarchy::OnSetExpansionRecursive(TSharedPtr<FRigTreeElement> InItem,
 	TreeView->SetExpansionRecursive(InItem, false, bShouldBeExpanded);
 }
 
-TOptional<FText> SRigHierarchy::OnGetItemTooltip(const FRigElementKey& InKey) const
+TOptional<FText> SRigHierarchy::OnGetItemTooltip(const FRigHierarchyKey& InKey) const
 {
 	if(!DragRigResolveResults.IsEmpty() && FSlateApplication::Get().IsDragDropping())
 	{
 		FString Message;
-		for(const TPair<FRigElementKey, FModularRigResolveResult>& DragRigResolveResult : DragRigResolveResults)
+		if(InKey.IsElement())
 		{
-			if(DragRigResolveResult.Key != InKey)
+			for(const TPair<FRigElementKey, FModularRigResolveResult>& DragRigResolveResult : DragRigResolveResults)
 			{
-				if(!DragRigResolveResult.Value.ContainsMatch(InKey, &Message))
+				if(DragRigResolveResult.Key != InKey.GetElement())
 				{
-					if(!Message.IsEmpty())
+					if(!DragRigResolveResult.Value.ContainsMatch(InKey.GetElement(), &Message))
 					{
-						return FText::FromString(Message);
+						if(!Message.IsEmpty())
+						{
+							return FText::FromString(Message);
+						}
 					}
 				}
 			}
@@ -1301,7 +1362,7 @@ void SRigHierarchy::CreateDragDropMenu()
 	}
 }
 
-UToolMenu* SRigHierarchy::GetDragDropMenu(const TArray<FRigElementKey>& DraggedKeys, FRigElementKey TargetKey)
+UToolMenu* SRigHierarchy::GetDragDropMenu(const TArray<FRigHierarchyKey>& DraggedKeys, FRigElementKey TargetKey)
 {
 	UToolMenus* ToolMenus = UToolMenus::Get();
 	if (!ensure(ToolMenus))
@@ -1350,36 +1411,102 @@ void SRigHierarchy::CreateContextMenu()
 					const FControlRigHierarchyCommands& Commands = FControlRigHierarchyCommands::Get(); 
 				
 					FToolMenuSection& ElementsSection = InMenu->AddSection(TEXT("Elements"), LOCTEXT("ElementsHeader", "Elements"));
-					ElementsSection.AddSubMenu(TEXT("New"), LOCTEXT("New", "New"), LOCTEXT("New_ToolTip", "Create New Elements"),
+					ElementsSection.AddSubMenu(TEXT("NewElement"), LOCTEXT("NewElement", "New Element"), LOCTEXT("NewElement_ToolTip", "Create New Elements"),
 						FNewToolMenuDelegate::CreateLambda([Commands, RigHierarchyPanel](UToolMenu* InSubMenu)
 						{
 							FToolMenuSection& DefaultSection = InSubMenu->AddSection(NAME_None);
-							FRigElementKey SelectedKey;
+							FRigHierarchyKey SelectedKey;
 							TArray<TSharedPtr<FRigTreeElement>> SelectedItems = RigHierarchyPanel->TreeView->GetSelectedItems();
 							if (SelectedItems.Num() > 0)
 							{
 								SelectedKey = SelectedItems[0]->Key;
 							}
-							
-							if (!SelectedKey || SelectedKey.Type == ERigElementType::Bone || SelectedKey.Type == ERigElementType::Connector)
-							{
-								DefaultSection.AddMenuEntry(Commands.AddBoneItem);
-							}
-							DefaultSection.AddMenuEntry(Commands.AddControlItem);
-							if(SelectedKey.Type == ERigElementType::Control)
-							{
-								DefaultSection.AddMenuEntry(Commands.AddAnimationChannelItem);
-							}
-							DefaultSection.AddMenuEntry(Commands.AddNullItem);
 
-							if(CVarControlRigHierarchyEnableModules.GetValueOnAnyThread())
+							if(SelectedKey.IsElement())
 							{
-								DefaultSection.AddMenuEntry(Commands.AddConnectorItem);
-								DefaultSection.AddMenuEntry(Commands.AddSocketItem);
+								static const FSlateIcon ControlIcon = FSlateIcon(FControlRigEditorStyle::Get().GetStyleSetName(), TEXT("ControlRig.Tree.Control"));
+								static const FSlateIcon NullIcon = FSlateIcon(FControlRigEditorStyle::Get().GetStyleSetName(), TEXT("ControlRig.Tree.Null"));
+								static const FSlateIcon BoneIcon = FSlateIcon(FControlRigEditorStyle::Get().GetStyleSetName(), TEXT("ControlRig.Tree.BoneImported"));
+								static const FSlateIcon SocketIcon = FSlateIcon(FControlRigEditorStyle::Get().GetStyleSetName(), TEXT("ControlRig.Tree.Socket_Open"));
+								static const FSlateIcon ConnectorIcon = FSlateIcon(FControlRigEditorStyle::Get().GetStyleSetName(), TEXT("ControlRig.ConnectorPrimary"));
+								static const FSlateIcon AnimationChannelIcon = FSlateIcon(FAppStyle::GetAppStyleSetName(), TEXT("Kismet.VariableList.TypeIcon"));
+
+								const FRigElementKey& SelectedElementKey = SelectedKey.GetElement();
+								if (!SelectedElementKey || SelectedElementKey.Type == ERigElementType::Bone || SelectedElementKey.Type == ERigElementType::Connector)
+								{
+									DefaultSection.AddMenuEntry(Commands.AddBoneItem, {}, {}, BoneIcon);
+								}
+								DefaultSection.AddMenuEntry(Commands.AddControlItem, {}, {}, ControlIcon);
+								if(SelectedElementKey.Type == ERigElementType::Control)
+								{
+									DefaultSection.AddMenuEntry(Commands.AddAnimationChannelItem, {}, {}, AnimationChannelIcon);
+								}
+								DefaultSection.AddMenuEntry(Commands.AddNullItem, {}, {}, NullIcon);
+
+								if(CVarControlRigHierarchyEnableModules.GetValueOnAnyThread())
+								{
+									DefaultSection.AddMenuEntry(Commands.AddConnectorItem, {}, {}, ConnectorIcon);
+									DefaultSection.AddMenuEntry(Commands.AddSocketItem, {}, {}, SocketIcon);
+								}
 							}
 						})
 					);
-					
+
+					const TArray<UScriptStruct*> RigComponentStructs = FRigBaseComponent::GetAllComponentScriptStructs();
+					if(!RigComponentStructs.IsEmpty())
+					{
+						ElementsSection.AddSubMenu(TEXT("NewComponent"), LOCTEXT("NewComponent", "New Component"), LOCTEXT("NewComponent_ToolTip", "Create New Component"),
+							FNewToolMenuDelegate::CreateLambda([Commands, RigHierarchyPanel, RigComponentStructs](UToolMenu* InSubMenu)
+							{
+								FToolMenuSection& DefaultSection = InSubMenu->AddSection(NAME_None);
+								TArray<FRigElementKey> SelectedElements;
+								TArray<TSharedPtr<FRigTreeElement>> SelectedItems = RigHierarchyPanel->TreeView->GetSelectedItems();
+								for(const TSharedPtr<FRigTreeElement>& SelectedItem : SelectedItems)
+								{
+									if(SelectedItem->Key.IsElement())
+									{
+										SelectedElements.Add(SelectedItem->Key.GetElement());
+									}
+								}
+								if(SelectedElements.IsEmpty())
+								{
+									SelectedElements.Add(URigHierarchy::GetTopLevelComponentElementKey());
+								}
+
+								for(UScriptStruct* RigComponentStruct : RigComponentStructs)
+								{
+									const FString ComponentName = RigComponentStruct->GetName();
+									const FText ComponentLabel = FText::Format(LOCTEXT("AddComponentLabelFormat", "Add {0}"), RigComponentStruct->GetDisplayNameText());
+									
+									DefaultSection.AddMenuEntry(
+										*FString::Printf(TEXT("Add%s"), *ComponentName),
+										ComponentLabel,
+										FText(),
+										FSlateIcon(), // todo: we may want to see the icon here for each component
+										FUIAction(FExecuteAction::CreateLambda([RigHierarchyPanel, RigComponentStruct, SelectedElements, ComponentLabel]() {
+											if(!RigHierarchyPanel->ControlRigBlueprint.IsValid())
+											{
+												return;
+											}
+
+											RigHierarchyPanel->DisplaySettings.bShowComponents = true;
+
+											FScopedTransaction Transaction(ComponentLabel);
+											UControlRigBlueprint* CurrentControlRigBlueprint = RigHierarchyPanel->ControlRigBlueprint.Get();
+											URigHierarchyController* Controller = CurrentControlRigBlueprint->GetHierarchyController();
+											TArray<FRigComponentKey> ComponentKeys;
+											for(const FRigElementKey& SelectedElementKey : SelectedElements)
+											{
+												ComponentKeys.Add(Controller->AddComponent(RigComponentStruct, *RigComponentStruct->GetDisplayNameText().ToString(), SelectedElementKey, FString(), true, true));
+											}
+											Controller->SetComponentSelection(ComponentKeys);
+										})
+									));
+								}
+							})
+						);
+					}
+
 					ElementsSection.AddMenuEntry(Commands.DeleteItem);
 					ElementsSection.AddMenuEntry(Commands.DuplicateItem);
 					ElementsSection.AddMenuEntry(Commands.FindReferencesOfItem);
@@ -1460,7 +1587,10 @@ void SRigHierarchy::CreateContextMenu()
 						{
 							if(SelectedItem.IsValid())
 							{
-								SelectedTypes.AddUnique(SelectedItem->Key.Type);
+								if(SelectedItem->Key.IsElement())
+								{
+									SelectedTypes.AddUnique(SelectedItem->Key.GetElement().Type);
+								}
 							}
 						}
 						if(SelectedTypes.Contains(ERigElementType::Control))
@@ -1748,20 +1878,24 @@ void SRigHierarchy::CreateImportMenu(FMenuBuilder& MenuBuilder)
 	);
 }
 
-FRigElementKey SRigHierarchy::OnGetResolvedKey(const FRigElementKey& InKey)
+FRigHierarchyKey SRigHierarchy::OnGetResolvedKey(const FRigHierarchyKey& InKey)
 {
 	if (const UControlRigBlueprint* Blueprint = ControlRigEditor.Pin()->GetControlRigBlueprint())
 	{
-		const FRigElementKey ResolvedKey = Blueprint->ModularRigModel.Connections.FindTargetFromConnector(InKey);
+		const FRigElementKey ResolvedKey = Blueprint->ModularRigModel.Connections.FindTargetFromConnector(InKey.GetElement());
 		if(ResolvedKey.IsValid())
 		{
-			return ResolvedKey;
+			if(InKey.IsElement())
+			{
+				return ResolvedKey;
+			}
+			return FRigComponentKey(ResolvedKey, InKey.GetFName());
 		}
 	}
 	return InKey;
 }
 
-void SRigHierarchy::OnRequestDetailsInspection(const FRigElementKey& InKey)
+void SRigHierarchy::OnRequestDetailsInspection(const FRigHierarchyKey& InKey)
 {
 	if(!ControlRigEditor.IsValid())
 	{
@@ -1885,7 +2019,11 @@ bool SRigHierarchy::IsSingleBoneSelected(bool bIncludeProcedural) const
 	{
 		return false;
 	}
-	return GetSelectedKeys()[0].Type == ERigElementType::Bone;
+	if(!GetSelectedKeys()[0].IsElement())
+	{
+		return false;
+	}
+	return GetSelectedKeys()[0].GetElement().Type == ERigElementType::Bone;
 }
 
 bool SRigHierarchy::IsSingleNullSelected(bool bIncludeProcedural) const
@@ -1894,7 +2032,11 @@ bool SRigHierarchy::IsSingleNullSelected(bool bIncludeProcedural) const
 	{
 		return false;
 	}
-	return GetSelectedKeys()[0].Type == ERigElementType::Null;
+	if(!GetSelectedKeys()[0].IsElement())
+	{
+		return false;
+	}
+	return GetSelectedKeys()[0].GetElement().Type == ERigElementType::Null;
 }
 
 bool SRigHierarchy::IsControlSelected(bool bIncludeProcedural) const
@@ -1904,12 +2046,15 @@ bool SRigHierarchy::IsControlSelected(bool bIncludeProcedural) const
 		return false;
 	}
 	
-	TArray<FRigElementKey> SelectedKeys = GetSelectedKeys();
-	for (const FRigElementKey& SelectedKey : SelectedKeys)
+	TArray<FRigHierarchyKey> SelectedKeys = GetSelectedKeys();
+	for (const FRigHierarchyKey& SelectedKey : SelectedKeys)
 	{
-		if (SelectedKey.Type == ERigElementType::Control)
+		if(SelectedKey.IsElement())
 		{
-			return true;
+			if (SelectedKey.GetElement().Type == ERigElementType::Control)
+			{
+				return true;
+			}
 		}
 	}
 	return false;
@@ -1922,16 +2067,19 @@ bool SRigHierarchy::IsControlOrNullSelected(bool bIncludeProcedural) const
 		return false;
 	}
 
-	TArray<FRigElementKey> SelectedKeys = GetSelectedKeys();
-	for (const FRigElementKey& SelectedKey : SelectedKeys)
+	TArray<FRigHierarchyKey> SelectedKeys = GetSelectedKeys();
+	for (const FRigHierarchyKey& SelectedKey : SelectedKeys)
 	{
-		if (SelectedKey.Type == ERigElementType::Control)
-		{
-			return true;
-		}
-		if (SelectedKey.Type == ERigElementType::Null)
+		if(SelectedKey.IsElement())
 		{
-			return true;
+			if (SelectedKey.GetElement().Type == ERigElementType::Control)
+			{
+				return true;
+			}
+			if (SelectedKey.GetElement().Type == ERigElementType::Null)
+			{
+				return true;
+			}
 		}
 	}
 	return false;
@@ -1939,15 +2087,19 @@ bool SRigHierarchy::IsControlOrNullSelected(bool bIncludeProcedural) const
 
 bool SRigHierarchy::IsProceduralElementSelected() const
 {
-	TArray<FRigElementKey> SelectedKeys = GetSelectedKeys();
+	TArray<FRigHierarchyKey> SelectedKeys = GetSelectedKeys();
 	if (SelectedKeys.IsEmpty())
 	{
 		return false;
 	}
 	
-	for (const FRigElementKey& SelectedKey : SelectedKeys)
+	for (const FRigHierarchyKey& SelectedKey : SelectedKeys)
 	{
-		if(!GetHierarchy()->IsProcedural(SelectedKey))
+		if(SelectedKey.IsElement() && !GetHierarchy()->IsProcedural(SelectedKey.GetElement()))
+		{
+			return false;
+		}
+		if(SelectedKey.IsComponent() && !GetHierarchy()->IsProcedural(SelectedKey.GetComponent()))
 		{
 			return false;
 		}
@@ -1957,15 +2109,19 @@ bool SRigHierarchy::IsProceduralElementSelected() const
 
 bool SRigHierarchy::IsNonProceduralElementSelected() const
 {
-	TArray<FRigElementKey> SelectedKeys = GetSelectedKeys();
+	TArray<FRigHierarchyKey> SelectedKeys = GetSelectedKeys();
 	if (SelectedKeys.IsEmpty())
 	{
 		return false;
 	}
 	
-	for (const FRigElementKey& SelectedKey : SelectedKeys)
+	for (const FRigHierarchyKey& SelectedKey : SelectedKeys)
 	{
-		if(GetHierarchy()->IsProcedural(SelectedKey))
+		if(SelectedKey.IsElement() && GetHierarchy()->IsProcedural(SelectedKey.GetElement()))
+		{
+			return false;
+		}
+		if(SelectedKey.IsComponent() && GetHierarchy()->IsProcedural(SelectedKey.GetComponent()))
 		{
 			return false;
 		}
@@ -2012,8 +2168,6 @@ void SRigHierarchy::HandleDeleteItem()
 	URigHierarchy* Hierarchy = GetDefaultHierarchy();
  	if (Hierarchy)
  	{
-		TArray<FRigElementKey> RemovedItems;
-
 		ClearDetailPanel();
 		FScopedTransaction Transaction(LOCTEXT("HierarchyTreeDeleteSelected", "Delete selected items from hierarchy"));
 
@@ -2026,24 +2180,27 @@ void SRigHierarchy::HandleDeleteItem()
  		URigHierarchyController* Controller = Hierarchy->GetController(true);
  		check(Controller);
 
- 		TArray<FRigElementKey> SelectedKeys = GetSelectedKeys();
+ 		TArray<FRigHierarchyKey> SelectedKeys = GetSelectedKeys();
 
  		if (ControlRigBlueprint.IsValid() && ControlRigBlueprint->IsControlRigModule())
  		{
- 			SelectedKeys.RemoveAll([Hierarchy, Controller](const FRigElementKey& Selected)
+ 			SelectedKeys.RemoveAll([Hierarchy, Controller](const FRigHierarchyKey& Selected)
 			{
-				if (const FRigBaseElement* Element = Hierarchy->Find(Selected))
-				{
-				   if (const FRigConnectorElement* Connector = Cast<FRigConnectorElement>(Element))
-				   {
-					   if (Connector->IsPrimary())
-					   {
-						   static constexpr TCHAR Format[] = TEXT("Cannot delete primary connector: %s");
-						   Controller->ReportAndNotifyErrorf(Format, *Connector->GetName());
-						   return true;
-					   }
-				   }
-			   }
+ 				if(Selected.IsElement())
+ 				{
+ 					if (const FRigBaseElement* Element = Hierarchy->Find(Selected.GetElement()))
+ 					{
+						if (const FRigConnectorElement* Connector = Cast<FRigConnectorElement>(Element))
+						{
+							if (Connector->IsPrimary())
+							{
+								static constexpr TCHAR Format[] = TEXT("Cannot delete primary connector: %s");
+								Controller->ReportAndNotifyErrorf(Format, *Connector->GetName());
+								return true;
+							}
+						}
+					}
+ 				}
 				return false;
 			});
  		}
@@ -2054,40 +2211,46 @@ void SRigHierarchy::HandleDeleteItem()
 
  		FRigHierarchyInteractionBracket InteractionBracket(Hierarchy);
  		
-		for (const FRigElementKey& SelectedKey : SelectedKeys)
+		for (const FRigHierarchyKey& SelectedKey : SelectedKeys)
 		{
 			TGuardValue<bool> GuardRigHierarchyChanges(bIsChangingRigHierarchy, true);
 			TGuardValue<bool> SuspendBlueprintNotifs(ControlRigBlueprint->bSuspendAllNotifications, true);
 
-			if (SelectedKey.Type == ERigElementType::Bone)
+			if(SelectedKey.IsElement())
 			{
-				if (FRigBoneElement* BoneElement = Hierarchy->Find<FRigBoneElement>(SelectedKey))
+				if (SelectedKey.GetElement().Type == ERigElementType::Bone)
 				{
-					if (BoneElement->BoneType == ERigBoneType::Imported && BoneElement->ParentElement != nullptr)
+					if (FRigBoneElement* BoneElement = Hierarchy->Find<FRigBoneElement>(SelectedKey.GetElement()))
 					{
-						if (!bConfirmedByUser)
+						if (BoneElement->BoneType == ERigBoneType::Imported && BoneElement->ParentElement != nullptr)
 						{
-							FText ConfirmDelete = LOCTEXT("ConfirmDeleteBoneHierarchy", "Deleting imported(white) bones can cause issues with animation - are you sure ?");
+							if (!bConfirmedByUser)
+							{
+								FText ConfirmDelete = LOCTEXT("ConfirmDeleteBoneHierarchy", "Deleting imported(white) bones can cause issues with animation - are you sure ?");
 
-							FSuppressableWarningDialog::FSetupInfo Info(ConfirmDelete, LOCTEXT("DeleteImportedBone", "Delete Imported Bone"), "DeleteImportedBoneHierarchy_Warning");
-							Info.ConfirmText = LOCTEXT("DeleteImportedBoneHierarchy_Yes", "Yes");
-							Info.CancelText = LOCTEXT("DeleteImportedBoneHierarchy_No", "No");
+								FSuppressableWarningDialog::FSetupInfo Info(ConfirmDelete, LOCTEXT("DeleteImportedBone", "Delete Imported Bone"), "DeleteImportedBoneHierarchy_Warning");
+								Info.ConfirmText = LOCTEXT("DeleteImportedBoneHierarchy_Yes", "Yes");
+								Info.CancelText = LOCTEXT("DeleteImportedBoneHierarchy_No", "No");
 
-							FSuppressableWarningDialog DeleteImportedBonesInHierarchy(Info);
-							bDeleteImportedBones = DeleteImportedBonesInHierarchy.ShowModal() != FSuppressableWarningDialog::Cancel;
-							bConfirmedByUser = true;
-						}
+								FSuppressableWarningDialog DeleteImportedBonesInHierarchy(Info);
+								bDeleteImportedBones = DeleteImportedBonesInHierarchy.ShowModal() != FSuppressableWarningDialog::Cancel;
+								bConfirmedByUser = true;
+							}
 
-						if (!bDeleteImportedBones)
-						{
-							break;
+							if (!bDeleteImportedBones)
+							{
+								break;
+							}
 						}
 					}
 				}
-			}
 
-			Controller->RemoveElement(SelectedKey, true, true);
-			RemovedItems.Add(SelectedKey);
+				Controller->RemoveElement(SelectedKey.GetElement(), true, true);
+			}
+			else if(SelectedKey.IsComponent())
+			{
+				Controller->RemoveComponent(SelectedKey.GetComponent(), true, true);
+			}
 		}
 	}
 
@@ -2128,18 +2291,27 @@ void SRigHierarchy::HandleNewItem(ERigElementType InElementType, bool bIsAnimati
 
 		FScopedTransaction Transaction(LOCTEXT("HierarchyTreeAdded", "Add new item to hierarchy"));
 
-		TArray<FRigElementKey> SelectedKeys = GetSelectedKeys();
-		if (SelectedKeys.Num() > 1 && !bAllowMultipleItems)
+		TArray<FRigHierarchyKey> SelectedKeys = GetSelectedKeys();
+		TArray<FRigElementKey> SelectedElementKeys;
+		SelectedElementKeys.Reserve(SelectedKeys.Num());
+		for(const FRigHierarchyKey& SelectedKey : SelectedKeys)
+		{
+			if(SelectedKey.IsElement())
+			{
+				SelectedElementKeys.Add(SelectedKey.GetElement());
+			}
+		}
+		if (SelectedElementKeys.Num() > 1 && !bAllowMultipleItems)
 		{
-			SelectedKeys = {SelectedKeys[0]};
+			SelectedElementKeys = {SelectedElementKeys[0]};
 		}
-		else if(SelectedKeys.IsEmpty())
+		else if(SelectedElementKeys.IsEmpty())
 		{
-			SelectedKeys = {FRigElementKey()};
+			SelectedElementKeys = {FRigElementKey()};
 		}
 
 		TMap<FRigElementKey, FRigElementKey> SelectedToCreated;
-		for(const FRigElementKey& SelectedKey : SelectedKeys)
+		for(const FRigElementKey& SelectedKey : SelectedElementKeys)
 		{
 			FRigElementKey ParentKey;
 			FTransform ParentTransform = FTransform::Identity;
@@ -2340,7 +2512,7 @@ bool SRigHierarchy::CanDuplicateItem() const
 
 	if (ControlRigBlueprint->IsControlRigModule())
 	{
-		bool bAnyNonConnector = GetSelectedKeys().ContainsByPredicate([](const FRigElementKey& Key)
+		bool bAnyNonConnector = GetSelectedElementKeys().ContainsByPredicate([](const FRigElementKey& Key)
 		{
 			return Key.Type != ERigElementType::Connector;
 		});
@@ -2372,7 +2544,7 @@ void SRigHierarchy::HandleDuplicateItem()
 			URigHierarchyController* Controller = Hierarchy->GetController(true);
 			check(Controller);
 
-			const TArray<FRigElementKey> KeysToDuplicate = GetSelectedKeys();
+			const TArray<FRigElementKey> KeysToDuplicate = GetSelectedElementKeys();
 			Controller->DuplicateElements(KeysToDuplicate, true, true, true);
 		}
 
@@ -2431,9 +2603,8 @@ void SRigHierarchy::HandleMirrorItem()
 
 				FScopedTransaction Transaction(LOCTEXT("HierarchyTreeMirrorSelected", "Mirror selected items from hierarchy"));
 
-				const TArray<FRigElementKey> KeysToMirror = GetSelectedKeys();
-				const TArray<FRigElementKey> KeysToDuplicate = GetSelectedKeys();
-				Controller->MirrorElements(KeysToDuplicate, Settings, true, true, true);
+				const TArray<FRigElementKey> KeysToMirror = GetSelectedElementKeys();
+				Controller->MirrorElements(KeysToMirror, Settings, true, true, true);
 			}
 			ControlRigBlueprint->PropagateHierarchyFromBPToInstances();
 		}
@@ -2449,26 +2620,33 @@ bool SRigHierarchy::CanRenameItem() const
 {
 	if(IsSingleSelected(false))
 	{
-		const FRigElementKey Key = GetSelectedKeys()[0];
-		if(Key.Type == ERigElementType::Physics ||
-			Key.Type == ERigElementType::Reference)
+		const FRigHierarchyKey Key = GetSelectedKeys()[0];
+		if(Key.IsElement())
 		{
-			return false;
-		}
-		if(Key.Type == ERigElementType::Control)
-		{
-			if(URigHierarchy* DebuggedHierarchy = GetHierarchy())
+			if(Key.GetElement().Type == ERigElementType::Physics ||
+				Key.GetElement().Type == ERigElementType::Reference)
 			{
-				if(FRigControlElement* ControlElement = DebuggedHierarchy->Find<FRigControlElement>(Key))
+				return false;
+			}
+			if(Key.GetElement().Type == ERigElementType::Control)
+			{
+				if(URigHierarchy* DebuggedHierarchy = GetHierarchy())
 				{
-					if(ControlElement->Settings.bIsTransientControl)
+					if(FRigControlElement* ControlElement = DebuggedHierarchy->Find<FRigControlElement>(Key.GetElement()))
 					{
-						return false;
+						if(ControlElement->Settings.bIsTransientControl)
+						{
+							return false;
+						}
 					}
 				}
 			}
+			return true;
+		}
+		else
+		{
+			// todo: not going to allow renaming of components for now
 		}
-		return true;
 	}
 	return false;
 }
@@ -2494,22 +2672,25 @@ void SRigHierarchy::HandleRenameItem()
 		TArray<TSharedPtr<FRigTreeElement>> SelectedItems = TreeView->GetSelectedItems();
 		if (SelectedItems.Num() == 1)
 		{
-			if (SelectedItems[0]->Key.Type == ERigElementType::Bone)
+			if(SelectedItems[0]->Key.IsElement())
 			{
-				if(FRigBoneElement* BoneElement = Hierarchy->Find<FRigBoneElement>(SelectedItems[0]->Key))
+				if (SelectedItems[0]->Key.GetElement().Type == ERigElementType::Bone)
 				{
-					if (BoneElement->BoneType == ERigBoneType::Imported)
+					if(FRigBoneElement* BoneElement = Hierarchy->Find<FRigBoneElement>(SelectedItems[0]->Key.GetElement()))
 					{
-						FText ConfirmRename = LOCTEXT("RenameDeleteBoneHierarchy", "Renaming imported(white) bones can cause issues with animation - are you sure ?");
+						if (BoneElement->BoneType == ERigBoneType::Imported)
+						{
+							FText ConfirmRename = LOCTEXT("RenameDeleteBoneHierarchy", "Renaming imported(white) bones can cause issues with animation - are you sure ?");
 
-						FSuppressableWarningDialog::FSetupInfo Info(ConfirmRename, LOCTEXT("RenameImportedBone", "Rename Imported Bone"), "RenameImportedBoneHierarchy_Warning");
-						Info.ConfirmText = LOCTEXT("RenameImportedBoneHierarchy_Yes", "Yes");
-						Info.CancelText = LOCTEXT("RenameImportedBoneHierarchy_No", "No");
+							FSuppressableWarningDialog::FSetupInfo Info(ConfirmRename, LOCTEXT("RenameImportedBone", "Rename Imported Bone"), "RenameImportedBoneHierarchy_Warning");
+							Info.ConfirmText = LOCTEXT("RenameImportedBoneHierarchy_Yes", "Yes");
+							Info.CancelText = LOCTEXT("RenameImportedBoneHierarchy_No", "No");
 
-						FSuppressableWarningDialog RenameImportedBonesInHierarchy(Info);
-						if (RenameImportedBonesInHierarchy.ShowModal() == FSuppressableWarningDialog::Cancel)
-						{
-							return;
+							FSuppressableWarningDialog RenameImportedBonesInHierarchy(Info);
+							if (RenameImportedBonesInHierarchy.ShowModal() == FSuppressableWarningDialog::Cancel)
+							{
+								return;
+							}
 						}
 					}
 				}
@@ -2701,7 +2882,7 @@ void SRigHierarchy::PostUndo(bool bSuccess)
 
 FReply SRigHierarchy::OnDragDetected(const FGeometry& MyGeometry, const FPointerEvent& MouseEvent)
 {
-	TArray<FRigElementKey> DraggedElements = GetSelectedKeys();
+	TArray<FRigHierarchyKey> DraggedElements = GetSelectedKeys();
 	if (MouseEvent.IsMouseButtonDown(EKeys::LeftMouseButton) && DraggedElements.Num() > 0)
 	{
 		if (ControlRigEditor.IsValid())
@@ -2725,7 +2906,7 @@ TOptional<EItemDropZone> SRigHierarchy::OnCanAcceptDrop(const FDragDropEvent& Dr
 	const TSharedPtr<FRigElementHierarchyDragDropOp> RigDragDropOp = DragDropEvent.GetOperationAs<FRigElementHierarchyDragDropOp>();
 	if (RigDragDropOp.IsValid())
 	{
-		FRigElementKey TargetKey;
+		FRigHierarchyKey TargetKey;
 		if (const URigHierarchy* Hierarchy = GetHierarchy())
 		{
 			switch(DropZone)
@@ -2733,17 +2914,30 @@ TOptional<EItemDropZone> SRigHierarchy::OnCanAcceptDrop(const FDragDropEvent& Dr
 				case EItemDropZone::AboveItem:
 				case EItemDropZone::BelowItem:
 				{
-					TargetKey = Hierarchy->GetFirstParent(TargetItem->Key);
+					if(TargetItem->Key.IsElement())
+					{
+						TargetKey = Hierarchy->GetFirstParent(TargetItem->Key.GetElement());
+					}
+					else if(TargetItem->Key.IsComponent())
+					{
+						// the element the component belongs to
+						TargetKey = TargetItem->Key.GetElement();
+					}
 					break;
 				}
 				case EItemDropZone::OntoItem:
 				{
-					TargetKey = TargetItem->Key;
+					TargetKey = TargetItem->Key.GetElement();
 					break;
 				}
 			}
 
-			for (const FRigElementKey& DraggedKey : RigDragDropOp->GetElements())
+			if(TargetKey.IsComponent())
+			{
+				return InvalidDropZone;
+			}
+
+			for (const FRigHierarchyKey& DraggedKey : RigDragDropOp->GetElements())
 			{
 				if(Hierarchy->IsProcedural(DraggedKey) && !RigDragDropOp->IsDraggingSingleConnector())
 				{
@@ -2758,9 +2952,9 @@ TOptional<EItemDropZone> SRigHierarchy::OnCanAcceptDrop(const FDragDropEvent& Dr
 
 				if(RigDragDropOp->IsDraggingSingleConnector() || RigDragDropOp->IsDraggingSingleSocket())
 				{
-					if(const FModularRigResolveResult* ResolveResult = DragRigResolveResults.Find(DraggedKey))
+					if(const FModularRigResolveResult* ResolveResult = DragRigResolveResults.Find(DraggedKey.GetElement()))
 					{
-						if(!ResolveResult->ContainsMatch(TargetKey))
+						if(!ResolveResult->ContainsMatch(TargetKey.GetElement()))
 						{
 							return InvalidDropZone;
 						}
@@ -2772,8 +2966,17 @@ TOptional<EItemDropZone> SRigHierarchy::OnCanAcceptDrop(const FDragDropEvent& Dr
 				}
 				else if(DropZone == EItemDropZone::OntoItem)
 				{
-					if(Hierarchy->IsParentedTo(TargetKey, DraggedKey))
+					if(DraggedKey.IsElement())
 					{
+						if(Hierarchy->IsParentedTo(TargetKey.GetElement(), DraggedKey.GetElement()))
+						{
+							return InvalidDropZone;
+						}
+					}
+					if(DraggedKey.IsComponent())
+					{
+						// todo: perform more validation here - as to where we are allowed to drag components
+						// return DraggedKey.GetElement() != TargetKey;
 						return InvalidDropZone;
 					}
 				}
@@ -2781,13 +2984,13 @@ TOptional<EItemDropZone> SRigHierarchy::OnCanAcceptDrop(const FDragDropEvent& Dr
 		}
 
 		// don't allow dragging onto procedural items (except for connectors + sockets)
-		if(TargetKey.IsValid() && !GetDefaultHierarchy()->Contains(TargetKey) &&
+		if(TargetKey.IsValid() && !GetDefaultHierarchy()->Contains(TargetKey.GetElement()) &&
 			!(RigDragDropOp->IsDraggingSingleConnector() || RigDragDropOp->IsDraggingSingleSocket()))
 		{
 			return InvalidDropZone;
 		}
 
-		switch (TargetKey.Type)
+		switch (TargetKey.GetElement().Type)
 		{
 			case ERigElementType::Bone:
 			{
@@ -2800,22 +3003,25 @@ TOptional<EItemDropZone> SRigHierarchy::OnCanAcceptDrop(const FDragDropEvent& Dr
 			case ERigElementType::Physics:
 			case ERigElementType::Reference:
 			{
-				for (const FRigElementKey& DraggedKey : RigDragDropOp->GetElements())
+				for (const FRigHierarchyKey& DraggedKey : RigDragDropOp->GetElements())
 				{
-					switch (DraggedKey.Type)
+					if(DraggedKey.IsElement())
 					{
-						case ERigElementType::Control:
-						case ERigElementType::Null:
-						case ERigElementType::Physics:
-						case ERigElementType::Reference:
-						case ERigElementType::Connector:
-						case ERigElementType::Socket:
+						switch (DraggedKey.GetElement().Type)
 						{
-							break;
-						}
-						default:
-						{
-							return InvalidDropZone;
+							case ERigElementType::Control:
+							case ERigElementType::Null:
+							case ERigElementType::Physics:
+							case ERigElementType::Reference:
+							case ERigElementType::Connector:
+							case ERigElementType::Socket:
+							{
+								break;
+							}
+							default:
+							{
+								return InvalidDropZone;
+							}
 						}
 					}
 				}
@@ -2868,7 +3074,7 @@ TOptional<EItemDropZone> SRigHierarchy::OnCanAcceptDrop(const FDragDropEvent& Dr
 				{
 					if(const FModularRigResolveResult* ResolveResult = DragRigResolveResults.Find(DraggedKey))
 					{
-						if(!ResolveResult->ContainsMatch(TargetItem->Key))
+						if(!ResolveResult->ContainsMatch(TargetItem->Key.GetElement()))
 						{
 							return InvalidDropZone;
 						}
@@ -2897,7 +3103,7 @@ TOptional<EItemDropZone> SRigHierarchy::OnCanAcceptDrop(const FDragDropEvent& Dr
 			return InvalidDropZone;
 		}
 
-		const FRigElementKey TargetKey = TargetItem->Key;
+		const FRigElementKey TargetKey = TargetItem->Key.GetElement();
 		const TArray<FRigElementKey> DraggedKeys = FControlRigSchematicModel::GetElementKeysFromDragDropEvent(*ModuleDropOp.Get(), ControlRig);
 		for(const FRigElementKey& DraggedKey : DraggedKeys)
 		{
@@ -2960,7 +3166,7 @@ TOptional<EItemDropZone> SRigHierarchy::OnCanAcceptDrop(const FDragDropEvent& Dr
 					FRigElementKey TargetKey;
 					if (TargetItem.IsValid())
 					{
-						TargetKey = TargetItem->Key;
+						TargetKey = TargetItem->Key.GetElement();
 					}
 
 					/*
@@ -2998,7 +3204,7 @@ FReply SRigHierarchy::OnAcceptDrop(const FDragDropEvent& DragDropEvent, EItemDro
 			const FVector2D& SummonLocation = DragDropEvent.GetScreenSpacePosition();
 
 			// Get the context menu content. If NULL, don't open a menu.
-			UToolMenu* DragDropMenu = GetDragDropMenu(RigDragDropOp->GetElements(), TargetItem->Key);
+			UToolMenu* DragDropMenu = GetDragDropMenu(RigDragDropOp->GetElements(), TargetItem->Key.GetElement());
 			const TSharedPtr<SWidget> MenuContent = UToolMenus::Get()->GenerateWidget(DragDropMenu);
 
 			if (MenuContent.IsValid())
@@ -3023,19 +3229,19 @@ FReply SRigHierarchy::OnAcceptDrop(const FDragDropEvent& DragDropEvent, EItemDro
 				{
 					case EItemDropZone::AboveItem:
 					{
-						TargetKey = Hierarchy->GetFirstParent(TargetItem->Key);
-						LocalIndex = Hierarchy->GetLocalIndex(TargetItem->Key);
+						TargetKey = Hierarchy->GetFirstParent(TargetItem->Key.GetElement());
+						LocalIndex = Hierarchy->GetLocalIndex(TargetItem->Key.GetElement());
 						break;
 					}
 					case EItemDropZone::BelowItem:
 					{
-						TargetKey = Hierarchy->GetFirstParent(TargetItem->Key);
-						LocalIndex = Hierarchy->GetLocalIndex(TargetItem->Key) + 1;
+						TargetKey = Hierarchy->GetFirstParent(TargetItem->Key.GetElement());
+						LocalIndex = Hierarchy->GetLocalIndex(TargetItem->Key.GetElement()) + 1;
 						break;
 					}
 					case EItemDropZone::OntoItem:
 					{
-						TargetKey = TargetItem->Key;
+						TargetKey = TargetItem->Key.GetElement();
 						break;
 					}
 				}
@@ -3043,9 +3249,9 @@ FReply SRigHierarchy::OnAcceptDrop(const FDragDropEvent& DragDropEvent, EItemDro
 
 			if(RigDragDropOp->IsDraggingSingleConnector())
 			{
-				return ResolveConnector(RigDragDropOp->GetElements()[0], TargetKey);
+				return ResolveConnector(RigDragDropOp->GetElements()[0].GetElement(), TargetKey);
 			}
-			
+
 			return ReparentOrMatchTransform(RigDragDropOp->GetElements(), TargetKey, bReparentItems, LocalIndex);			
 		}
 
@@ -3058,7 +3264,7 @@ FReply SRigHierarchy::OnAcceptDrop(const FDragDropEvent& DragDropEvent, EItemDro
 		FRigElementKey::StaticStruct()->ImportText(*TagDragDropOp->GetIdentifier(), &DraggedKey, nullptr, EPropertyPortFlags::PPF_None, nullptr, FRigElementKey::StaticStruct()->GetName(), true);
 		if(TargetItem.IsValid())
 		{
-			TArray<FRigElementKey> TargetsForConnector = {TargetItem->Key};
+			TArray<FRigElementKey> TargetsForConnector = {TargetItem->Key.GetElement()};
 			
 			// do we want to add this target to an array connector?
 			if(DragDropEvent.GetModifierKeys().IsShiftDown())
@@ -3075,7 +3281,7 @@ FReply SRigHierarchy::OnAcceptDrop(const FDragDropEvent& DragDropEvent, EItemDro
 								{
 									TargetsForConnector.Reset();
 									TargetsForConnector.Append(FRigElementKeyRedirector::Convert(*Cache));
-									TargetsForConnector.AddUnique(TargetItem->Key);
+									TargetsForConnector.AddUnique(TargetItem->Key.GetElement());
 								}
 							}
 						}
@@ -3096,7 +3302,7 @@ FReply SRigHierarchy::OnAcceptDrop(const FDragDropEvent& DragDropEvent, EItemDro
 			return FReply::Handled();
 		}
 
-		const FRigElementKey TargetKey = TargetItem->Key;
+		const FRigElementKey TargetKey = TargetItem->Key.GetElement();
 		const TArray<FRigElementKey> DraggedKeys = FControlRigSchematicModel::GetElementKeysFromDragDropEvent(*ModuleDropOp.Get(), ControlRig);
 
 		bool bSuccess = false;
@@ -3146,7 +3352,7 @@ FReply SRigHierarchy::OnAcceptDrop(const FDragDropEvent& DragDropEvent, EItemDro
 								}
 							}
 						}
-						return ResolveConnector(PrimaryConnectorKey, TargetItem->Key);
+						return ResolveConnector(PrimaryConnectorKey, TargetItem->Key.GetElement());
 					}
 				}
 				return FReply::Handled();
@@ -3162,12 +3368,12 @@ void SRigHierarchy::OnElementKeyTagDragDetected(const FRigElementKey& InDraggedT
 	UpdateConnectorMatchesOnDrag({InDraggedTag});
 }
 
-void SRigHierarchy::UpdateConnectorMatchesOnDrag(const TArray<FRigElementKey>& InDraggedKeys)
+void SRigHierarchy::UpdateConnectorMatchesOnDrag(const TArray<FRigHierarchyKey>& InDraggedKeys)
 {
 	DragRigResolveResults.Reset();
 
 	// fade in all items
-	for(const TPair<FRigElementKey, TSharedPtr<FRigTreeElement>>& Pair : TreeView->ElementMap)
+	for(const TPair<FRigHierarchyKey, TSharedPtr<FRigTreeElement>>& Pair : TreeView->ElementMap)
 	{
 		Pair.Value->bFadedOutDuringDragDrop = false;
 	}
@@ -3180,19 +3386,22 @@ void SRigHierarchy::UpdateConnectorMatchesOnDrag(const TArray<FRigElementKey>& I
 			{
 				if(const UModularRigRuleManager* RuleManager = Hierarchy->GetRuleManager())
 				{
-					for(const FRigElementKey& DraggedElement : InDraggedKeys)
+					for(const FRigHierarchyKey& DraggedElement : InDraggedKeys)
 					{
-						if(DraggedElement.Type == ERigElementType::Connector)
+						if(DraggedElement.IsElement())
 						{
-							if(const FRigConnectorElement* Connector = Hierarchy->Find<FRigConnectorElement>(DraggedElement))
+							if(DraggedElement.GetElement().Type == ERigElementType::Connector)
 							{
-								const FName NameSpace = Hierarchy->GetNameSpaceFName(Connector->GetKey());
-								if(!NameSpace.IsNone())
+								if(const FRigConnectorElement* Connector = Hierarchy->Find<FRigConnectorElement>(DraggedElement.GetElement()))
 								{
-									if(const FRigModuleInstance* Module = ControlRig->FindModule(NameSpace.ToString()))
+									const FName NameSpace = Hierarchy->GetNameSpaceFName(Connector->GetKey());
+									if(!NameSpace.IsNone())
 									{
-										const FModularRigResolveResult ResolveResult = RuleManager->FindMatches(Connector, Module, ControlRig->ElementKeyRedirector); 
-										DragRigResolveResults.Add(DraggedElement, ResolveResult);
+										if(const FRigModuleInstance* Module = ControlRig->FindModule(NameSpace.ToString()))
+										{
+											const FModularRigResolveResult ResolveResult = RuleManager->FindMatches(Connector, Module, ControlRig->ElementKeyRedirector); 
+											DragRigResolveResults.Add(DraggedElement.GetElement(), ResolveResult);
+										}
 									}
 								}
 							}
@@ -3216,7 +3425,7 @@ void SRigHierarchy::UpdateConnectorMatchesOnDrag(const TArray<FRigElementKey>& I
 	}
 }
 
-FName SRigHierarchy::HandleRenameElement(const FRigElementKey& OldKey, const FString& NewName)
+FName SRigHierarchy::HandleRenameElement(const FRigHierarchyKey& OldKey, const FString& NewName)
 {
 	ClearDetailPanel();
 
@@ -3234,21 +3443,24 @@ FName SRigHierarchy::HandleRenameElement(const FRigElementKey& OldKey, const FSt
 		FName ResultingName = NAME_None;
 
 		bool bUseDisplayName = false;
-		if(const FRigControlElement* ControlElement = Hierarchy->Find<FRigControlElement>(OldKey))
+		if(OldKey.IsElement())
 		{
-			if(ControlElement->IsAnimationChannel())
+			if(const FRigControlElement* ControlElement = Hierarchy->Find<FRigControlElement>(OldKey.GetElement()))
 			{
-				bUseDisplayName = true;
+				if(ControlElement->IsAnimationChannel())
+				{
+					bUseDisplayName = true;
+				}
 			}
-		}
 
-		if(bUseDisplayName)
-		{
-			ResultingName = Controller->SetDisplayName(OldKey, SanitizedName, true, true, true);
-		}
-		else
-		{
-			ResultingName = Controller->RenameElement(OldKey, SanitizedName, true, true, false).Name;
+			if(bUseDisplayName)
+			{
+				ResultingName = Controller->SetDisplayName(OldKey.GetElement(), SanitizedName, true, true, true);
+			}
+			else
+			{
+				ResultingName = Controller->RenameElement(OldKey.GetElement(), SanitizedName, true, true, false).Name;
+			}
 		}
 		ControlRigBlueprint->PropagateHierarchyFromBPToInstances();
 		return ResultingName;
@@ -3257,21 +3469,24 @@ FName SRigHierarchy::HandleRenameElement(const FRigElementKey& OldKey, const FSt
 	return NAME_None;
 }
 
-bool SRigHierarchy::HandleVerifyNameChanged(const FRigElementKey& OldKey, const FString& NewName, FText& OutErrorMessage)
+bool SRigHierarchy::HandleVerifyNameChanged(const FRigHierarchyKey& OldKey, const FString& NewName, FText& OutErrorMessage)
 {
 	bool bIsAnimationChannel = false;
 	if (ControlRigBlueprint.IsValid())
 	{
 		URigHierarchy* Hierarchy = GetHierarchy();
-		if(const FRigControlElement* ControlElement = Hierarchy->Find<FRigControlElement>(OldKey))
+		if(OldKey.IsElement())
 		{
-			if(ControlElement->IsAnimationChannel())
+			if(const FRigControlElement* ControlElement = Hierarchy->Find<FRigControlElement>(OldKey.GetElement()))
 			{
-				bIsAnimationChannel = true;
-
-				if(ControlElement->GetDisplayName().ToString() == NewName)
+				if(ControlElement->IsAnimationChannel())
 				{
-					return true;
+					bIsAnimationChannel = true;
+
+					if(ControlElement->GetDisplayName().ToString() == NewName)
+					{
+						return true;
+					}
 				}
 			}
 		}
@@ -3279,7 +3494,7 @@ bool SRigHierarchy::HandleVerifyNameChanged(const FRigElementKey& OldKey, const
 
 	if(!bIsAnimationChannel)
 	{
-		if (OldKey.Name.ToString() == NewName)
+		if (OldKey.GetName() == NewName)
 		{
 			return true;
 		}
@@ -3298,7 +3513,7 @@ bool SRigHierarchy::HandleVerifyNameChanged(const FRigElementKey& OldKey, const
 
 		if(bIsAnimationChannel)
 		{
-			if(const FRigControlElement* ControlElement = Hierarchy->Find<FRigControlElement>(OldKey))
+			if(const FRigControlElement* ControlElement = Hierarchy->Find<FRigControlElement>(OldKey.GetElement()))
 			{
 				if(const FRigBaseElement* ParentElement = Hierarchy->GetFirstParent(ControlElement))
 				{
@@ -3311,15 +3526,20 @@ bool SRigHierarchy::HandleVerifyNameChanged(const FRigElementKey& OldKey, const
 				}
 			}
 		}
-		else
+		else if(OldKey.IsElement())
 		{
 			FString OutErrorString;
-			if (!Hierarchy->IsNameAvailable(FRigName(NewName), OldKey.Type, &OutErrorString))
+			if (!Hierarchy->IsNameAvailable(FRigName(NewName), OldKey.GetElement().Type, &OutErrorString))
 			{
 				OutErrorMessage = FText::FromString(OutErrorString);
 				return false;
 			}
 		}
+		else if(OldKey.IsComponent())
+		{
+			// todo: not allowing renaming of components for now
+			return false;
+		}
 	}
 	return true;
 }
@@ -3346,7 +3566,7 @@ void SRigHierarchy::HandleResetTransform(bool bSelectionOnly)
 			{
 				FScopedTransaction Transaction(LOCTEXT("HierarchyResetTransforms", "Reset Transforms"));
 
-				TArray<FRigElementKey> KeysToReset = GetSelectedKeys();
+				TArray<FRigElementKey> KeysToReset = GetSelectedElementKeys();
 				if (!bSelectionOnly)
 				{
 					KeysToReset = DebuggedHierarchy->GetAllKeys(true, ERigElementType::Control);
@@ -3382,7 +3602,7 @@ void SRigHierarchy::HandleSetInitialTransformFromCurrentTransform()
 			{
 				FScopedTransaction Transaction(LOCTEXT("HierarchySetInitialTransforms", "Set Initial Transforms"));
 
-				TArray<FRigElementKey> SelectedKeys = GetSelectedKeys();
+				TArray<FRigElementKey> SelectedKeys = GetSelectedElementKeys();
 				TMap<FRigElementKey, FTransform> GlobalTransforms;
 				TMap<FRigElementKey, FTransform> ParentGlobalTransforms;
 
@@ -3484,7 +3704,7 @@ void SRigHierarchy::HandleControlBoneOrSpaceTransform()
 		return;
 	}
 
-	TArray<FRigElementKey> SelectedKeys = GetSelectedKeys();
+	TArray<FRigElementKey> SelectedKeys = GetSelectedElementKeys();
 	if (SelectedKeys.Num() == 1)
 	{
 		if (SelectedKeys[0].Type == ERigElementType::Bone ||
@@ -3512,7 +3732,7 @@ void SRigHierarchy::HandleUnparent()
 	bool bUnparentImportedBones = false;
 	bool bConfirmedByUser = false;
 
-	TArray<FRigElementKey> SelectedKeys = GetSelectedKeys();
+	TArray<FRigElementKey> SelectedKeys = GetSelectedElementKeys();
 	TMap<FRigElementKey, FTransform> InitialTransforms;
 	TMap<FRigElementKey, FTransform> GlobalTransforms;
 
@@ -3634,7 +3854,7 @@ void SRigHierarchy::HandleParent(const FToolMenuContext& Context)
 	if (UControlRigContextMenuContext* MenuContext = Cast<UControlRigContextMenuContext>(Context.FindByClass(UControlRigContextMenuContext::StaticClass())))
 	{
 		const FControlRigRigHierarchyDragAndDropContext DragAndDropContext = MenuContext->GetRigHierarchyDragAndDropContext();
-		ReparentOrMatchTransform(DragAndDropContext.DraggedElementKeys, DragAndDropContext.TargetElementKey, true);
+		ReparentOrMatchTransform(DragAndDropContext.DraggedHierarchyKeys, DragAndDropContext.TargetHierarchyKey, true);
 	}
 }
 
@@ -3643,11 +3863,11 @@ void SRigHierarchy::HandleAlign(const FToolMenuContext& Context)
 	if (UControlRigContextMenuContext* MenuContext = Cast<UControlRigContextMenuContext>(Context.FindByClass(UControlRigContextMenuContext::StaticClass())))
 	{
 		const FControlRigRigHierarchyDragAndDropContext DragAndDropContext = MenuContext->GetRigHierarchyDragAndDropContext();
-		ReparentOrMatchTransform(DragAndDropContext.DraggedElementKeys, DragAndDropContext.TargetElementKey, false);
+		ReparentOrMatchTransform(DragAndDropContext.DraggedHierarchyKeys, DragAndDropContext.TargetHierarchyKey, false);
 	}
 }
 
-FReply SRigHierarchy::ReparentOrMatchTransform(const TArray<FRigElementKey>& DraggedKeys, FRigElementKey TargetKey, bool bReparentItems, int32 LocalIndex)
+FReply SRigHierarchy::ReparentOrMatchTransform(const TArray<FRigHierarchyKey>& DraggedKeys, FRigHierarchyKey TargetKey, bool bReparentItems, int32 LocalIndex)
 {
 	bool bMatchTransforms = !bReparentItems;
 
@@ -3663,24 +3883,33 @@ FReply SRigHierarchy::ReparentOrMatchTransform(const TArray<FRigElementKey>& Dra
 		{
 			return FReply::Unhandled();
 		}
+		if(!TargetKey.IsElement())
+		{
+			return FReply::Unhandled();
+		}
 
 		TGuardValue<bool> GuardRigHierarchyChanges(bIsChangingRigHierarchy, true);
 		TGuardValue<bool> SuspendBlueprintNotifs(ControlRigBlueprint->bSuspendAllNotifications, true);
 		FScopedTransaction Transaction(LOCTEXT("HierarchyDragAndDrop", "Drag & Drop"));
 		FRigHierarchyInteractionBracket InteractionBracket(Hierarchy);
 
-		FTransform TargetGlobalTransform = DebuggedHierarchy->GetGlobalTransform(TargetKey);
+		FTransform TargetGlobalTransform = DebuggedHierarchy->GetGlobalTransform(TargetKey.GetElement());
 
-		for (const FRigElementKey& DraggedKey : DraggedKeys)
+		for (const FRigHierarchyKey& DraggedKey : DraggedKeys)
 		{
 			if (DraggedKey == TargetKey)
 			{
 				return FReply::Unhandled();
 			}
 
+			if(!DraggedKey.IsElement())
+			{
+				continue;
+			}
+
 			if (bReparentItems)
 			{
-				if(Hierarchy->IsParentedTo(TargetKey, DraggedKey))
+				if(Hierarchy->IsParentedTo(TargetKey.GetElement(), DraggedKey.GetElement()))
 				{
 					if(LocalIndex == INDEX_NONE)
 					{
@@ -3689,9 +3918,9 @@ FReply SRigHierarchy::ReparentOrMatchTransform(const TArray<FRigElementKey>& Dra
 				}
 			}
 
-			if (DraggedKey.Type == ERigElementType::Bone)
+			if (DraggedKey.GetElement().Type == ERigElementType::Bone)
 			{
-				if(FRigBoneElement* BoneElement = Hierarchy->Find<FRigBoneElement>(DraggedKey))
+				if(FRigBoneElement* BoneElement = Hierarchy->Find<FRigBoneElement>(DraggedKey.GetElement()))
 				{
 					if (BoneElement->BoneType == ERigBoneType::Imported && BoneElement->ParentElement != nullptr)
 					{
@@ -3717,13 +3946,13 @@ FReply SRigHierarchy::ReparentOrMatchTransform(const TArray<FRigElementKey>& Dra
 			}
 		}
 
-		for (const FRigElementKey& DraggedKey : DraggedKeys)
+		for (const FRigHierarchyKey& DraggedKey : DraggedKeys)
 		{
-			if (bMatchTransforms)
+			if (bMatchTransforms && DraggedKey.IsElement())
 			{
-				if (DraggedKey.Type == ERigElementType::Control)
+				if (DraggedKey.GetElement().Type == ERigElementType::Control)
 				{
-					int32 ControlIndex = DebuggedHierarchy->GetIndex(DraggedKey);
+					int32 ControlIndex = DebuggedHierarchy->GetIndex(DraggedKey.GetElement());
 					if (ControlIndex == INDEX_NONE)
 					{
 						continue;
@@ -3734,51 +3963,57 @@ FReply SRigHierarchy::ReparentOrMatchTransform(const TArray<FRigElementKey>& Dra
 
 					Hierarchy->SetControlOffsetTransformByIndex(ControlIndex, OffsetTransform, ERigTransformType::InitialLocal, true, true, true);
 					Hierarchy->SetControlOffsetTransformByIndex(ControlIndex, OffsetTransform, ERigTransformType::CurrentLocal, true, true, true);
-					Hierarchy->SetLocalTransform(DraggedKey, FTransform::Identity, true, true, true, true);
-					Hierarchy->SetInitialLocalTransform(DraggedKey, FTransform::Identity, true, true, true);
+					Hierarchy->SetLocalTransform(DraggedKey.GetElement(), FTransform::Identity, true, true, true, true);
+					Hierarchy->SetInitialLocalTransform(DraggedKey.GetElement(), FTransform::Identity, true, true, true);
 					DebuggedHierarchy->SetControlOffsetTransformByIndex(ControlIndex, OffsetTransform, ERigTransformType::InitialLocal, true, true);
 					DebuggedHierarchy->SetControlOffsetTransformByIndex(ControlIndex, OffsetTransform, ERigTransformType::CurrentLocal, true, true);
-					DebuggedHierarchy->SetLocalTransform(DraggedKey, FTransform::Identity, true, true, true);
-					DebuggedHierarchy->SetInitialLocalTransform(DraggedKey, FTransform::Identity, true, true);
+					DebuggedHierarchy->SetLocalTransform(DraggedKey.GetElement(), FTransform::Identity, true, true, true);
+					DebuggedHierarchy->SetInitialLocalTransform(DraggedKey.GetElement(), FTransform::Identity, true, true);
 				}
 				else
 				{
-					Hierarchy->SetInitialGlobalTransform(DraggedKey, TargetGlobalTransform, true, true);
-					Hierarchy->SetGlobalTransform(DraggedKey, TargetGlobalTransform, false, true, true);
-					DebuggedHierarchy->SetInitialGlobalTransform(DraggedKey, TargetGlobalTransform, true, true);
-					DebuggedHierarchy->SetGlobalTransform(DraggedKey, TargetGlobalTransform, false, true, true);
+					Hierarchy->SetInitialGlobalTransform(DraggedKey.GetElement(), TargetGlobalTransform, true, true);
+					Hierarchy->SetGlobalTransform(DraggedKey.GetElement(), TargetGlobalTransform, false, true, true);
+					DebuggedHierarchy->SetInitialGlobalTransform(DraggedKey.GetElement(), TargetGlobalTransform, true, true);
+					DebuggedHierarchy->SetGlobalTransform(DraggedKey.GetElement(), TargetGlobalTransform, false, true, true);
 				}
 				continue;
 			}
 
-			FRigElementKey ParentKey = TargetKey;
+			FRigElementKey ParentKey = TargetKey.GetElement();
 
-			const FTransform InitialGlobalTransform = DebuggedHierarchy->GetInitialGlobalTransform(DraggedKey);
-			const FTransform CurrentGlobalTransform = DebuggedHierarchy->GetGlobalTransform(DraggedKey);
-			const FTransform InitialLocalTransform = DebuggedHierarchy->GetInitialLocalTransform(DraggedKey);
-			const FTransform CurrentLocalTransform = DebuggedHierarchy->GetLocalTransform(DraggedKey);
-			const FTransform CurrentGlobalOffsetTransform = DebuggedHierarchy->GetGlobalControlOffsetTransform(DraggedKey, false);
+			if(DraggedKey.IsComponent())
+			{
+				// todo: reparent the component
+				continue;
+			}
+
+			const FTransform InitialGlobalTransform = DebuggedHierarchy->GetInitialGlobalTransform(DraggedKey.GetElement());
+			const FTransform CurrentGlobalTransform = DebuggedHierarchy->GetGlobalTransform(DraggedKey.GetElement());
+			const FTransform InitialLocalTransform = DebuggedHierarchy->GetInitialLocalTransform(DraggedKey.GetElement());
+			const FTransform CurrentLocalTransform = DebuggedHierarchy->GetLocalTransform(DraggedKey.GetElement());
+			const FTransform CurrentGlobalOffsetTransform = DebuggedHierarchy->GetGlobalControlOffsetTransform(DraggedKey.GetElement(), false);
 
 			bool bElementWasReparented = false;
-			if(ParentKey.IsValid() && (Hierarchy->GetFirstParent(DraggedKey) != ParentKey))
+			if(ParentKey.IsValid() && (Hierarchy->GetFirstParent(DraggedKey.GetElement()) != ParentKey))
 			{
-				bElementWasReparented = Controller->SetParent(DraggedKey, ParentKey, true, true, true);
+				bElementWasReparented = Controller->SetParent(DraggedKey.GetElement(), ParentKey, true, true, true);
 			}
-			else if(!ParentKey.IsValid() && (Hierarchy->GetNumberOfParents(DraggedKey) > 0))
+			else if(!ParentKey.IsValid() && (Hierarchy->GetNumberOfParents(DraggedKey.GetElement()) > 0))
 			{
-				bElementWasReparented = Controller->RemoveAllParents(DraggedKey, true, true, true);
+				bElementWasReparented = Controller->RemoveAllParents(DraggedKey.GetElement(), true, true, true);
 			}
 
 			if(LocalIndex != INDEX_NONE)
 			{
-				Controller->ReorderElement(DraggedKey, LocalIndex, true, true);
+				Controller->ReorderElement(DraggedKey.GetElement(), LocalIndex, true, true);
 			}
 
 			if(bElementWasReparented)
 			{
-				if (DraggedKey.Type == ERigElementType::Control)
+				if (DraggedKey.GetElement().Type == ERigElementType::Control)
 				{
-					int32 ControlIndex = DebuggedHierarchy->GetIndex(DraggedKey);
+					int32 ControlIndex = DebuggedHierarchy->GetIndex(DraggedKey.GetElement());
 					if (ControlIndex == INDEX_NONE)
 					{
 						continue;
@@ -3789,19 +4024,19 @@ FReply SRigHierarchy::ReparentOrMatchTransform(const TArray<FRigElementKey>& Dra
 
 					Hierarchy->SetControlOffsetTransformByIndex(ControlIndex, LocalOffsetTransform, ERigTransformType::InitialLocal, true, true, true);
 					Hierarchy->SetControlOffsetTransformByIndex(ControlIndex, LocalOffsetTransform, ERigTransformType::CurrentLocal, true, true, true);
-					Hierarchy->SetLocalTransform(DraggedKey, CurrentLocalTransform, true, true, true, true);
-					Hierarchy->SetInitialLocalTransform(DraggedKey, InitialLocalTransform, true, true, true);
+					Hierarchy->SetLocalTransform(DraggedKey.GetElement(), CurrentLocalTransform, true, true, true, true);
+					Hierarchy->SetInitialLocalTransform(DraggedKey.GetElement(), InitialLocalTransform, true, true, true);
 					DebuggedHierarchy->SetControlOffsetTransformByIndex(ControlIndex, LocalOffsetTransform, ERigTransformType::InitialLocal, true, true);
 					DebuggedHierarchy->SetControlOffsetTransformByIndex(ControlIndex, LocalOffsetTransform, ERigTransformType::CurrentLocal, true, true);
-					DebuggedHierarchy->SetLocalTransform(DraggedKey, CurrentLocalTransform, true, true, true);
-					DebuggedHierarchy->SetInitialLocalTransform(DraggedKey, InitialLocalTransform, true, true);
+					DebuggedHierarchy->SetLocalTransform(DraggedKey.GetElement(), CurrentLocalTransform, true, true, true);
+					DebuggedHierarchy->SetInitialLocalTransform(DraggedKey.GetElement(), InitialLocalTransform, true, true);
 				}
 				else
 				{
-					DebuggedHierarchy->SetInitialGlobalTransform(DraggedKey, InitialGlobalTransform, true, true);
-					DebuggedHierarchy->SetGlobalTransform(DraggedKey, CurrentGlobalTransform, false, true, true);
-					Hierarchy->SetInitialGlobalTransform(DraggedKey, InitialGlobalTransform, true, true);
-					Hierarchy->SetGlobalTransform(DraggedKey, CurrentGlobalTransform, false, true, true);
+					DebuggedHierarchy->SetInitialGlobalTransform(DraggedKey.GetElement(), InitialGlobalTransform, true, true);
+					DebuggedHierarchy->SetGlobalTransform(DraggedKey.GetElement(), CurrentGlobalTransform, false, true, true);
+					Hierarchy->SetInitialGlobalTransform(DraggedKey.GetElement(), InitialGlobalTransform, true, true);
+					Hierarchy->SetGlobalTransform(DraggedKey.GetElement(), CurrentGlobalTransform, false, true, true);
 				}
 			}
 		}
@@ -3862,7 +4097,7 @@ void SRigHierarchy::HandleSetInitialTransformFromClosestBone()
 
 				FScopedTransaction Transaction(LOCTEXT("HierarchySetInitialTransforms", "Set Initial Transforms"));
 
-				TArray<FRigElementKey> SelectedKeys = GetSelectedKeys();
+				TArray<FRigElementKey> SelectedKeys = GetSelectedElementKeys();
 				TMap<FRigElementKey, FTransform> ClosestTransforms;
 				TMap<FRigElementKey, FTransform> ParentGlobalTransforms;
 
@@ -3948,7 +4183,11 @@ void SRigHierarchy::HandleSetShapeTransformFromCurrent()
 				TArray<TSharedPtr<FRigTreeElement>> SelectedItems = TreeView->GetSelectedItems();
 				for (const TSharedPtr<FRigTreeElement>& SelectedItem : SelectedItems)
 				{
-					if(FRigControlElement* ControlElement = DebuggedHierarchy->Find<FRigControlElement>(SelectedItem->Key))
+					if(!SelectedItem->Key.IsElement())
+					{
+						continue;
+					}
+					if(FRigControlElement* ControlElement = DebuggedHierarchy->Find<FRigControlElement>(SelectedItem->Key.GetElement()))
 					{
 						const FRigElementKey Key = ControlElement->GetKey();
 						
diff --git a/Engine/Plugins/Animation/ControlRig/Source/ControlRigEditor/Private/Editor/SRigHierarchy.h b/Engine/Plugins/Animation/ControlRig/Source/ControlRigEditor/Private/Editor/SRigHierarchy.h
--- a/Engine/Plugins/Animation/ControlRig/Source/ControlRigEditor/Private/Editor/SRigHierarchy.h
+++ b/Engine/Plugins/Animation/ControlRig/Source/ControlRigEditor/Private/Editor/SRigHierarchy.h
@@ -72,7 +72,10 @@ private:
 	void RefreshTreeView(bool bRebuildContent = true);
 
 	/** Return all selected keys */
-	TArray<FRigElementKey> GetSelectedKeys() const;
+	TArray<FRigHierarchyKey> GetSelectedKeys() const;
+
+	/** Return all selected element keys */
+	TArray<FRigElementKey> GetSelectedElementKeys() const;
 
 	/** Check whether we can deleting the selected item(s) */
 	bool CanDeleteItem() const;
@@ -120,7 +123,7 @@ private:
 	void OnItemClicked(TSharedPtr<FRigTreeElement> InItem);
 	void OnItemDoubleClicked(TSharedPtr<FRigTreeElement> InItem);
 	void OnSetExpansionRecursive(TSharedPtr<FRigTreeElement> InItem, bool bShouldBeExpanded);
-	TOptional<FText> OnGetItemTooltip(const FRigElementKey& InKey) const;
+	TOptional<FText> OnGetItemTooltip(const FRigHierarchyKey& InKey) const;
 
 	// FEditorUndoClient
 	virtual void PostUndo(bool bSuccess) override;
@@ -131,7 +134,7 @@ private:
 	TOptional<EItemDropZone> OnCanAcceptDrop(const FDragDropEvent& DragDropEvent, EItemDropZone DropZone, TSharedPtr<FRigTreeElement> TargetItem);
 	FReply OnAcceptDrop(const FDragDropEvent& DragDropEvent, EItemDropZone DropZone, TSharedPtr<FRigTreeElement> TargetItem);
 	void OnElementKeyTagDragDetected(const FRigElementKey& InDraggedTag);
-	void UpdateConnectorMatchesOnDrag(const TArray<FRigElementKey>& InDraggedKeys);
+	void UpdateConnectorMatchesOnDrag(const TArray<FRigHierarchyKey>& InDraggedKeys);
 
 	static const FName ContextMenuName;
 	static void CreateContextMenu();
@@ -140,7 +143,7 @@ private:
 	
 	static const FName DragDropMenuName;
 	static void CreateDragDropMenu();
-	UToolMenu* GetDragDropMenu(const TArray<FRigElementKey>& DraggedKeys, FRigElementKey TargetKey);
+	UToolMenu* GetDragDropMenu(const TArray<FRigHierarchyKey>& DraggedKeys, FRigElementKey TargetKey);
 
 	/** Our owning control rig editor */
 	TWeakPtr<FControlRigEditor> ControlRigEditor;
@@ -184,8 +187,8 @@ private:
 	URigHierarchy* GetHierarchy() const;
 	URigHierarchy* GetDefaultHierarchy() const;
 	const URigHierarchy* GetHierarchyForTreeView() const { return GetHierarchy(); }
-	FRigElementKey OnGetResolvedKey(const FRigElementKey& InKey);
-	void OnRequestDetailsInspection(const FRigElementKey& InKey);
+	FRigHierarchyKey OnGetResolvedKey(const FRigHierarchyKey& InKey);
+	void OnRequestDetailsInspection(const FRigHierarchyKey& InKey);
 	
 	void ImportHierarchy(const FAssetData& InAssetData);
 	void CreateImportMenu(FMenuBuilder& MenuBuilder);
@@ -208,7 +211,7 @@ private:
 
 	void HandleParent(const FToolMenuContext& Context);
 	void HandleAlign(const FToolMenuContext& Context);
-	FReply ReparentOrMatchTransform(const TArray<FRigElementKey>& DraggedKeys, FRigElementKey TargetKey, bool bReparentItems, int32 LocalIndex = INDEX_NONE);
+	FReply ReparentOrMatchTransform(const TArray<FRigHierarchyKey>& DraggedKeys, FRigHierarchyKey TargetKey, bool bReparentItems, int32 LocalIndex = INDEX_NONE);
 	FReply ResolveConnector(const FRigElementKey& DraggedKey, const FRigElementKey& TargetKey);
 	FReply ResolveConnectorToArray(const FRigElementKey& DraggedKey, const TArray<FRigElementKey>& TargetKeys);
 
@@ -229,12 +232,12 @@ private:
 	bool bIsConstructionEventRunning;
 	bool bRestrictRefreshToMeshBones = false;
 	uint32 LastHierarchyHash;
-	TArray<FRigElementKey> SelectionBeforeConstruction;
+	TArray<FRigHierarchyKey> SelectionBeforeConstruction;
 	TMap<FRigElementKey, FModularRigResolveResult> DragRigResolveResults;
 
 public:
-	FName HandleRenameElement(const FRigElementKey& OldKey, const FString& NewName);
-	bool HandleVerifyNameChanged(const FRigElementKey& OldKey, const FString& NewName, FText& OutErrorMessage);
+	FName HandleRenameElement(const FRigHierarchyKey& OldKey, const FString& NewName);
+	bool HandleVerifyNameChanged(const FRigHierarchyKey& OldKey, const FString& NewName, FText& OutErrorMessage);
 
 	// SWidget Overrides
 	virtual FReply OnDrop(const FGeometry& MyGeometry, const FDragDropEvent& DragDropEvent) override;
diff --git a/Engine/Plugins/Animation/ControlRig/Source/ControlRigEditor/Private/Editor/SRigHierarchyTreeView.cpp b/Engine/Plugins/Animation/ControlRig/Source/ControlRigEditor/Private/Editor/SRigHierarchyTreeView.cpp
--- a/Engine/Plugins/Animation/ControlRig/Source/ControlRigEditor/Private/Editor/SRigHierarchyTreeView.cpp
+++ b/Engine/Plugins/Animation/ControlRig/Source/ControlRigEditor/Private/Editor/SRigHierarchyTreeView.cpp
@@ -33,10 +33,10 @@ FRigTreeDisplaySettings FRigTreeDelegates::DefaultDisplaySettings;
 //////////////////////////////////////////////////////////////
 /// FRigTreeElement
 ///////////////////////////////////////////////////////////
-FRigTreeElement::FRigTreeElement(const FRigElementKey& InKey, TWeakPtr<SRigHierarchyTreeView> InTreeView, bool InSupportsRename, ERigTreeFilterResult InFilterResult)
+FRigTreeElement::FRigTreeElement(const FRigHierarchyKey& InKey, TWeakPtr<SRigHierarchyTreeView> InTreeView, bool InSupportsRename, ERigTreeFilterResult InFilterResult)
 {
 	Key = InKey;
-	ShortName = InKey.Name;
+	ShortName = InKey.GetFName();
 	ChannelName = NAME_None;
 	bIsTransient = false;
 	bIsAnimationChannel = false;
@@ -49,7 +49,10 @@ FRigTreeElement::FRigTreeElement(const FRigElementKey& InKey, TWeakPtr<SRigHiera
 	{
 		if(const URigHierarchy* Hierarchy = InTreeView.Pin()->GetRigTreeDelegates().GetHierarchy())
 		{
-			ShortName = *Hierarchy->GetDisplayNameForUI(InKey, false).ToString();
+			if(InKey.IsElement())
+			{
+				ShortName = *Hierarchy->GetDisplayNameForUI(InKey.GetElement(), false).ToString();
+			}
 			
 			const FRigTreeDisplaySettings& Settings = InTreeView.Pin()->GetRigTreeDelegates().GetDisplaySettings();
 			RefreshDisplaySettings(Hierarchy, Settings);
@@ -75,20 +78,28 @@ void FRigTreeElement::RefreshDisplaySettings(const URigHierarchy* InHierarchy, c
 {
 	const TPair<const FSlateBrush*, FSlateColor> Result = SRigHierarchyItem::GetBrushForElementType(InHierarchy, Key);
 
-	if(const FRigBaseElement* Element = InHierarchy->Find(Key))
+	bIsProcedural = false;
+	if(Key.IsElement())
 	{
-		bIsProcedural = Element->IsProcedural();
-				
-		if(const FRigControlElement* ControlElement = Cast<FRigControlElement>(Element))
+		if(const FRigBaseElement* Element = InHierarchy->Find(Key.GetElement()))
 		{
-			bIsTransient = ControlElement->Settings.bIsTransientControl;
-			bIsAnimationChannel = ControlElement->IsAnimationChannel();
-			if(bIsAnimationChannel)
+			bIsProcedural = Element->IsProcedural();
+				
+			if(const FRigControlElement* ControlElement = Cast<FRigControlElement>(Element))
 			{
-				ChannelName = ControlElement->GetDisplayName();
+				bIsTransient = ControlElement->Settings.bIsTransientControl;
+				bIsAnimationChannel = ControlElement->IsAnimationChannel();
+				if(bIsAnimationChannel)
+				{
+					ChannelName = ControlElement->GetDisplayName();
+				}
 			}
 		}
 	}
+	else
+	{
+		bIsProcedural = InHierarchy->IsProcedural(Key.GetComponent());
+	}
 
 	IconBrush = Result.Key;
 	IconColor = Result.Value;
@@ -102,8 +113,8 @@ void FRigTreeElement::RefreshDisplaySettings(const URigHierarchy* InHierarchy, c
 	}
 
 	TextColor = FilterResult == ERigTreeFilterResult::Shown ?
-		(InHierarchy->IsProcedural(Key) ? FSlateColor(FLinearColor(0.9f, 0.8f, 0.4f)) : FSlateColor::UseForeground()) :
-		(InHierarchy->IsProcedural(Key) ? FSlateColor(FLinearColor(0.9f, 0.8f, 0.4f) * 0.5f) : FSlateColor(FLinearColor::Gray * 0.5f));
+		(bIsProcedural ? FSlateColor(FLinearColor(0.9f, 0.8f, 0.4f)) : FSlateColor::UseForeground()) :
+		(bIsProcedural ? FSlateColor(FLinearColor(0.9f, 0.8f, 0.4f) * 0.5f) : FSlateColor(FLinearColor::Gray * 0.5f));
 }
 
 FSlateColor FRigTreeElement::GetIconColor() const
@@ -262,7 +273,7 @@ FText SRigHierarchyItem::GetName(bool bUseShortName) const
 	{
 		return (FText::FromName(WeakRigTreeElement.Pin()->ShortName));
 	}
-	return (FText::FromName(WeakRigTreeElement.Pin()->Key.Name));
+	return (FText::FromName(WeakRigTreeElement.Pin()->Key.GetFName()));
 }
 
 FText SRigHierarchyItem::GetItemTooltip() const
@@ -380,7 +391,7 @@ void SRigHierarchyTreeView::Tick(const FGeometry& AllottedGeometry, const double
 	}
 }
 
-TSharedPtr<FRigTreeElement> SRigHierarchyTreeView::FindElement(const FRigElementKey& InElementKey) const
+TSharedPtr<FRigTreeElement> SRigHierarchyTreeView::FindElement(const FRigHierarchyKey& InElementKey) const
 {
 	for (int32 RootIndex = 0; RootIndex < RootElements.Num(); ++RootIndex)
 	{
@@ -393,7 +404,7 @@ TSharedPtr<FRigTreeElement> SRigHierarchyTreeView::FindElement(const FRigElement
 	return TSharedPtr<FRigTreeElement>();
 }
 
-TSharedPtr<FRigTreeElement> SRigHierarchyTreeView::FindElement(const FRigElementKey& InElementKey, TSharedPtr<FRigTreeElement> CurrentItem)
+TSharedPtr<FRigTreeElement> SRigHierarchyTreeView::FindElement(const FRigHierarchyKey& InElementKey, TSharedPtr<FRigTreeElement> CurrentItem)
 {
 	if (CurrentItem->Key == InElementKey)
 	{
@@ -412,7 +423,7 @@ TSharedPtr<FRigTreeElement> SRigHierarchyTreeView::FindElement(const FRigElement
 	return TSharedPtr<FRigTreeElement>();
 }
 
-bool SRigHierarchyTreeView::AddElement(FRigElementKey InKey, FRigElementKey InParentKey)
+bool SRigHierarchyTreeView::AddElement(FRigHierarchyKey InKey, FRigHierarchyKey InParentKey)
 {
 	if(ElementMap.Contains(InKey))
 	{
@@ -420,13 +431,16 @@ bool SRigHierarchyTreeView::AddElement(FRigElementKey InKey, FRigElementKey InPa
 	}
 
 	// skip transient controls
-	if(const URigHierarchy* Hierarchy = Delegates.GetHierarchy())
+	if(InKey.IsElement())
 	{
-		if(const FRigControlElement* ControlElement = Hierarchy->Find<FRigControlElement>(InKey))
+		if(const URigHierarchy* Hierarchy = Delegates.GetHierarchy())
 		{
-			if(ControlElement->Settings.bIsTransientControl)
+			if(const FRigControlElement* ControlElement = Hierarchy->Find<FRigControlElement>(InKey.GetElement()))
 			{
-				return false;
+				if(ControlElement->Settings.bIsTransientControl)
+				{
+					return false;
+				}
 			}
 		}
 	}
@@ -448,12 +462,12 @@ bool SRigHierarchyTreeView::AddElement(FRigElementKey InKey, FRigElementKey InPa
 		if (InKey.IsValid())
 		{
 			ElementMap.Add(InKey, NewItem);
-			if (InParentKey)
+			if (InParentKey.IsValid())
 			{
 				ParentMap.Add(InKey, InParentKey);
 			}
 
-			if (InParentKey)
+			if (InParentKey.IsValid())
 			{
 				TSharedPtr<FRigTreeElement>* FoundItem = ElementMap.Find(InParentKey);
 				check(FoundItem);
@@ -476,9 +490,23 @@ bool SRigHierarchyTreeView::AddElement(FRigElementKey InKey, FRigElementKey InPa
 		{
 			bIsFilteredOut = !Delegates.OnRigTreeIsItemVisible.Execute(InKey);
 		}
+
+		const URigHierarchy* Hierarchy = Delegates.GetHierarchy();
+		auto GetFirstParent = [Hierarchy](const FRigHierarchyKey& InKey) -> FRigHierarchyKey
+		{
+			if(InKey.IsElement())
+			{
+				return Hierarchy->GetFirstParent(InKey.GetElement());;
+			}
+			if(InKey.IsComponent())
+			{
+				return InKey.GetComponent().ElementKey;
+			}
+			return FRigHierarchyKey();
+		};
 		
 		FString FilteredStringUnderScores = FilteredString.Replace(TEXT(" "), TEXT("_"));
-		if (!bIsFilteredOut && (InKey.Name.ToString().Contains(FilteredString) || InKey.Name.ToString().Contains(FilteredStringUnderScores)))
+		if (!bIsFilteredOut && (InKey.GetName().Contains(FilteredString) || InKey.GetName().Contains(FilteredStringUnderScores)))
 		{
 			TSharedPtr<FRigTreeElement> NewItem = MakeShared<FRigTreeElement>(InKey, SharedThis(this), bSupportsRename, ERigTreeFilterResult::Shown);
 			ElementMap.Add(InKey, NewItem);
@@ -486,10 +514,10 @@ bool SRigHierarchyTreeView::AddElement(FRigElementKey InKey, FRigElementKey InPa
 
 			if (!Settings.bFlattenHierarchyOnFilter && !Settings.bHideParentsOnFilter)
 			{
-				if(const URigHierarchy* Hierarchy = Delegates.GetHierarchy())
+				if(Hierarchy)
 				{
 					TSharedPtr<FRigTreeElement> ChildItem = NewItem;
-					FRigElementKey ParentKey = Hierarchy->GetFirstParent(InKey);
+					FRigHierarchyKey ParentKey = GetFirstParent(InKey);
 					while (ParentKey.IsValid())
 					{
 						if (!ElementMap.Contains(ParentKey))
@@ -501,7 +529,7 @@ bool SRigHierarchyTreeView::AddElement(FRigElementKey InKey, FRigElementKey InPa
 							ReparentElement(ChildItem->Key, ParentKey);
 
 							ChildItem = ParentItem;
-							ParentKey = Hierarchy->GetFirstParent(ParentKey);
+							ParentKey = GetFirstParent(ParentKey);
 						}
 						else
 						{
@@ -653,7 +681,7 @@ bool SRigHierarchyTreeView::AddElement(const FRigBaseElement* InElement)
 							{
 								TagArguments.OnRenamed_Lambda([ConnectorKey, this](const FText& InNewName, ETextCommit::Type InCommitType)
 								{
-									Delegates.HandleRenameElement(ConnectorKey, InNewName.ToString());
+									(void)Delegates.HandleRenameElement(ConnectorKey, InNewName.ToString());
 								});
 								TagArguments.OnVerifyRename_Lambda([ConnectorKey, this](const FText& InText, FText& OutError)
 								{
@@ -686,7 +714,7 @@ bool SRigHierarchyTreeView::AddElement(const FRigBaseElement* InElement)
 		}
 	}
 
-	if(!AddElement(InElement->GetKey(), FRigElementKey()))
+	if(!AddElement(InElement->GetKey(), FRigHierarchyKey()))
 	{
 		return false;
 	}
@@ -703,7 +731,7 @@ bool SRigHierarchyTreeView::AddElement(const FRigBaseElement* InElement)
 			FRigElementKey ParentKey = Hierarchy->GetFirstParent(InElement->GetKey());
 			if(InElement->GetType() == ERigElementType::Connector)
 			{
-				ParentKey = Delegates.GetResolvedKey(InElement->GetKey());
+				ParentKey = Delegates.GetResolvedKey(InElement->GetKey()).GetElement();
 				if(ParentKey == InElement->GetKey())
 				{
 					ParentKey.Reset();
@@ -741,24 +769,55 @@ bool SRigHierarchyTreeView::AddElement(const FRigBaseElement* InElement)
 		}
 	}
 
+	for(int32 ComponentIndex = 0; ComponentIndex < InElement->NumComponents(); ComponentIndex++)
+	{
+		AddComponent(InElement->GetComponent(ComponentIndex));
+	}
+
+	return true;
+}
+
+bool SRigHierarchyTreeView::AddComponent(const FRigBaseComponent* InComponent)
+{
+	check(InComponent);
+	
+	if (ElementMap.Contains(InComponent->GetKey()))
+	{
+		return false;
+	}
+
+	const FRigTreeDisplaySettings& Settings = Delegates.GetDisplaySettings();
+	if(!Settings.bShowComponents)
+	{
+		return false;
+	}
+
+	if(!AddElement(InComponent->GetKey(), InComponent->GetElementKey()))
+	{
+		return false;
+	}
+
 	return true;
 }
 
 void SRigHierarchyTreeView::AddSpacerElement()
 {
-	AddElement(FRigElementKey(), FRigElementKey());
+	AddElement(FRigElementKey(), FRigHierarchyKey());
 }
 
-bool SRigHierarchyTreeView::ReparentElement(FRigElementKey InKey, FRigElementKey InParentKey)
+bool SRigHierarchyTreeView::ReparentElement(FRigHierarchyKey InKey, FRigHierarchyKey InParentKey)
 {
 	if (!InKey.IsValid() || InKey == InParentKey)
 	{
 		return false;
 	}
 
-	if(InKey.Type == ERigElementType::Connector)
+	if(InKey.IsElement())
 	{
-		return false;
+		if(InKey.GetElement().Type == ERigElementType::Connector)
+		{
+			return false;
+		}
 	}
 
 	const FRigTreeDisplaySettings& Settings = Delegates.GetDisplaySettings();
@@ -774,7 +833,7 @@ bool SRigHierarchyTreeView::ReparentElement(FRigElementKey InKey, FRigElementKey
 		return false;
 	}
 
-	if (const FRigElementKey* ExistingParentKey = ParentMap.Find(InKey))
+	if (const FRigHierarchyKey* ExistingParentKey = ParentMap.Find(InKey))
 	{
 		if (*ExistingParentKey == InParentKey)
 		{
@@ -798,7 +857,7 @@ bool SRigHierarchyTreeView::ReparentElement(FRigElementKey InKey, FRigElementKey
 		RootElements.Remove(*FoundItem);
 	}
 
-	if (InParentKey)
+	if (InParentKey.IsValid())
 	{
 		ParentMap.Add(InKey, InParentKey);
 
@@ -816,7 +875,7 @@ bool SRigHierarchyTreeView::ReparentElement(FRigElementKey InKey, FRigElementKey
 	return true;
 }
 
-bool SRigHierarchyTreeView::RemoveElement(FRigElementKey InKey)
+bool SRigHierarchyTreeView::RemoveElement(FRigHierarchyKey InKey)
 {
 	TSharedPtr<FRigTreeElement>* FoundItem = ElementMap.Find(InKey);
 	if (FoundItem == nullptr)
@@ -824,7 +883,7 @@ bool SRigHierarchyTreeView::RemoveElement(FRigElementKey InKey)
 		return false;
 	}
 
-	ReparentElement(InKey, FRigElementKey());
+	ReparentElement(InKey, FRigHierarchyKey());
 
 	RootElements.Remove(*FoundItem);
 	return ElementMap.Remove(InKey) > 0;
@@ -832,11 +891,11 @@ bool SRigHierarchyTreeView::RemoveElement(FRigElementKey InKey)
 
 void SRigHierarchyTreeView::RefreshTreeView(bool bRebuildContent)
 {
-		TMap<FRigElementKey, bool> ExpansionState;
+		TMap<FRigHierarchyKey, bool> ExpansionState;
 
 	if(bRebuildContent)
 	{
-		for (TPair<FRigElementKey, TSharedPtr<FRigTreeElement>> Pair : ElementMap)
+		for (TPair<FRigHierarchyKey, TSharedPtr<FRigTreeElement>> Pair : ElementMap)
 		{
 			ExpansionState.FindOrAdd(Pair.Key) = IsItemExpanded(Pair.Value);
 		}
@@ -888,6 +947,12 @@ void SRigHierarchyTreeView::RefreshTreeView(bool bRebuildContent)
 				bContinue = true;
 			});
 
+			// first add the root components
+			for(int32 Index = 0; Index < Hierarchy->NumTopLevelComponents(); Index++)
+			{
+				AddComponent(Hierarchy->GetTopLevelComponent(Index));
+			}
+
 			// sort the sockets by depth
 			Algo::SortBy(Sockets, [ElementDepth](const FRigSocketElement* Socket) -> int32
 			{
@@ -920,7 +985,7 @@ void SRigHierarchyTreeView::RefreshTreeView(bool bRebuildContent)
 			}
 			else if (ExpansionState.Num() < ElementMap.Num())
 			{
-				for (const TPair<FRigElementKey, TSharedPtr<FRigTreeElement>>& Element : ElementMap)
+				for (const TPair<FRigHierarchyKey, TSharedPtr<FRigTreeElement>>& Element : ElementMap)
 				{
 					if (!ExpansionState.Contains(Element.Key))
 					{
@@ -966,18 +1031,15 @@ void SRigHierarchyTreeView::RefreshTreeView(bool bRebuildContent)
 	{
 		ClearSelection();
 
-		if(const URigHierarchy* Hierarchy = Delegates.GetHierarchy())
+		TArray<FRigHierarchyKey> Selection = Delegates.GetSelection();
+		for (const FRigHierarchyKey& Key : Selection)
 		{
-			TArray<FRigElementKey> Selection = Delegates.GetSelection();
-			for (const FRigElementKey& Key : Selection)
+			for (int32 RootIndex = 0; RootIndex < RootElements.Num(); ++RootIndex)
 			{
-				for (int32 RootIndex = 0; RootIndex < RootElements.Num(); ++RootIndex)
+				TSharedPtr<FRigTreeElement> Found = FindElement(Key, RootElements[RootIndex]);
+				if (Found.IsValid())
 				{
-					TSharedPtr<FRigTreeElement> Found = FindElement(Key, RootElements[RootIndex]);
-					if (Found.IsValid())
-					{
-						SetItemSelection(Found, true, ESelectInfo::OnNavigation);
-					}
+					SetItemSelection(Found, true, ESelectInfo::OnNavigation);
 				}
 			}
 		}
@@ -991,7 +1053,7 @@ void SRigHierarchyTreeView::SetExpansionRecursive(TSharedPtr<FRigTreeElement> In
 
 	if (bTowardsParent)
 	{
-		if (const FRigElementKey* ParentKey = ParentMap.Find(InElement->Key))
+		if (const FRigHierarchyKey* ParentKey = ParentMap.Find(InElement->Key))
 		{
 			if (TSharedPtr<FRigTreeElement>* ParentItem = ElementMap.Find(*ParentKey))
 			{
@@ -1026,9 +1088,9 @@ void SRigHierarchyTreeView::OnElementKeyTagDragDetected(const FRigElementKey& In
 	(void)Delegates.OnRigTreeElementKeyTagDragDetected.ExecuteIfBound(InDraggedTag);
 }
 
-TArray<FRigElementKey> SRigHierarchyTreeView::GetSelectedKeys() const
+TArray<FRigHierarchyKey> SRigHierarchyTreeView::GetSelectedKeys() const
 {
-	TArray<FRigElementKey> Keys;
+	TArray<FRigHierarchyKey> Keys;
 	TArray<TSharedPtr<FRigTreeElement>> SelectedElements = GetSelectedItems();
 	for(const TSharedPtr<FRigTreeElement>& SelectedElement : SelectedElements)
 	{
@@ -1048,7 +1110,7 @@ const TSharedPtr<FRigTreeElement>* SRigHierarchyTreeView::FindItemAtPosition(FVe
 			TSharedRef<SRigHierarchyItem> ItemWidget = StaticCastSharedRef<SRigHierarchyItem>(ArrangedChildren[Index].Widget);
 			if (ItemWidget->WeakRigTreeElement.IsValid())
 			{
-				const FRigElementKey Key = ItemWidget->WeakRigTreeElement.Pin()->Key;
+				const FRigHierarchyKey Key = ItemWidget->WeakRigTreeElement.Pin()->Key;
 				const TSharedPtr<FRigTreeElement>* ResultPtr = SListView<TSharedPtr<FRigTreeElement>>::GetItems().FindByPredicate([Key](const TSharedPtr<FRigTreeElement>& Item) -> bool
 					{
 						return Item->Key == Key;
@@ -1134,11 +1196,11 @@ FText SRigHierarchyTreeView::GetConnectorWarningMessage(TSharedPtr<FRigTreeEleme
 bool SRigHierarchyItem::OnVerifyNameChanged(const FText& InText, FText& OutErrorMessage)
 {
 	const FString NewName = InText.ToString();
-	const FRigElementKey OldKey = WeakRigTreeElement.Pin()->Key;
+	const FRigHierarchyKey OldKey = WeakRigTreeElement.Pin()->Key;
 	return Delegates.HandleVerifyElementNameChanged(OldKey, NewName, OutErrorMessage);
 }
 
-TPair<const FSlateBrush*, FSlateColor> SRigHierarchyItem::GetBrushForElementType(const URigHierarchy* InHierarchy, const FRigElementKey& InKey)
+TPair<const FSlateBrush*, FSlateColor> SRigHierarchyItem::GetBrushForElementType(const URigHierarchy* InHierarchy, const FRigHierarchyKey& InKey)
 {
 	static const FSlateBrush* ProxyControlBrush = FControlRigEditorStyle::Get().GetBrush("ControlRig.Tree.ProxyControl"); 
 	static const FSlateBrush* ControlBrush = FControlRigEditorStyle::Get().GetBrush("ControlRig.Tree.Control");
@@ -1154,122 +1216,131 @@ TPair<const FSlateBrush*, FSlateColor> SRigHierarchyItem::GetBrushForElementType
 
 	const FSlateBrush* Brush = nullptr;
 	FSlateColor Color = FSlateColor::UseForeground();
-	switch (InKey.Type)
+
+	if(InKey.IsElement())
 	{
-		case ERigElementType::Control:
+		switch (InKey.GetElement().Type)
 		{
-			if(const FRigControlElement* Control = InHierarchy->Find<FRigControlElement>(InKey))
+			case ERigElementType::Control:
 			{
-				FLinearColor ShapeColor = FLinearColor::White;
-				
-				if(Control->Settings.SupportsShape())
+				if(const FRigControlElement* Control = InHierarchy->Find<FRigControlElement>(InKey.GetElement()))
 				{
-					if(Control->Settings.AnimationType == ERigControlAnimationType::ProxyControl)
+					FLinearColor ShapeColor = FLinearColor::White;
+					
+					if(Control->Settings.SupportsShape())
 					{
-						Brush = ProxyControlBrush;
+						if(Control->Settings.AnimationType == ERigControlAnimationType::ProxyControl)
+						{
+							Brush = ProxyControlBrush;
+						}
+						else
+						{
+							Brush = ControlBrush;
+						}
+						ShapeColor = Control->Settings.ShapeColor;
 					}
 					else
 					{
-						Brush = ControlBrush;
+						static const FLazyName TypeIcon(TEXT("Kismet.VariableList.TypeIcon"));
+						Brush = FAppStyle::GetBrush(TypeIcon);
+						ShapeColor = GetColorForControlType(Control->Settings.ControlType, Control->Settings.ControlEnum);
 					}
-					ShapeColor = Control->Settings.ShapeColor;
+					
+					// ensure the alpha is always visible
+					ShapeColor.A = 1.f;
+					Color = FSlateColor(ShapeColor);
 				}
 				else
 				{
-					static const FLazyName TypeIcon(TEXT("Kismet.VariableList.TypeIcon"));
-					Brush = FAppStyle::GetBrush(TypeIcon);
-					ShapeColor = GetColorForControlType(Control->Settings.ControlType, Control->Settings.ControlEnum);
+					Brush = ControlBrush;
 				}
-				
-				// ensure the alpha is always visible
-				ShapeColor.A = 1.f;
-				Color = FSlateColor(ShapeColor);
+				break;
 			}
-			else
+			case ERigElementType::Null:
 			{
-				Brush = ControlBrush;
+				Brush = NullBrush;
+				break;
 			}
-			break;
-		}
-		case ERigElementType::Null:
-		{
-			Brush = NullBrush;
-			break;
-		}
-		case ERigElementType::Bone:
-		{
-			ERigBoneType BoneType = ERigBoneType::User;
-
-			if(InHierarchy)
+			case ERigElementType::Bone:
 			{
-				const FRigBoneElement* BoneElement = InHierarchy->Find<FRigBoneElement>(InKey);
-				if(BoneElement)
-				{
-					BoneType = BoneElement->BoneType;
-				}
-			}
+				ERigBoneType BoneType = ERigBoneType::User;
 
-			switch (BoneType)
-			{
-				case ERigBoneType::Imported:
+				if(InHierarchy)
 				{
-					Brush = BoneImportedBrush;
-					break;
+					const FRigBoneElement* BoneElement = InHierarchy->Find<FRigBoneElement>(InKey.GetElement());
+					if(BoneElement)
+					{
+						BoneType = BoneElement->BoneType;
+					}
 				}
-				case ERigBoneType::User:
-				default:
+
+				switch (BoneType)
 				{
-					Brush = BoneUserBrush;
-					break;
+					case ERigBoneType::Imported:
+					{
+						Brush = BoneImportedBrush;
+						break;
+					}
+					case ERigBoneType::User:
+					default:
+					{
+						Brush = BoneUserBrush;
+						break;
+					}
 				}
-			}
-
-			break;
-		}
-		case ERigElementType::Physics:
-		{
-			Brush = PhysicsBrush;
-			break;
-		}
-		case ERigElementType::Reference:
-		case ERigElementType::Socket:
-		{
-			Brush = SocketOpenBrush;
 
-			if(UControlRig* ControlRig = Cast<UControlRig>(InHierarchy->GetOuter()))
+				break;
+			}
+			case ERigElementType::Physics:
+			{
+				Brush = PhysicsBrush;
+				break;
+			}
+			case ERigElementType::Reference:
+			case ERigElementType::Socket:
 			{
-				if(const FRigElementKey* ConnectorKey = ControlRig->GetElementKeyRedirector().FindReverse(InKey))
+				Brush = SocketOpenBrush;
+
+				if(UControlRig* ControlRig = Cast<UControlRig>(InHierarchy->GetOuter()))
 				{
-					if(ConnectorKey->Type == ERigElementType::Connector)
+					if(const FRigElementKey* ConnectorKey = ControlRig->GetElementKeyRedirector().FindReverse(InKey.GetElement()))
 					{
-						Brush = SocketClosedBrush;
+						if(ConnectorKey->Type == ERigElementType::Connector)
+						{
+							Brush = SocketClosedBrush;
+						}
 					}
 				}
-			}
 
-			if(const FRigSocketElement* Socket = InHierarchy->Find<FRigSocketElement>(InKey))
-			{
-				Color = Socket->GetColor(InHierarchy);
+				if(const FRigSocketElement* Socket = InHierarchy->Find<FRigSocketElement>(InKey.GetElement()))
+				{
+					Color = Socket->GetColor(InHierarchy);
+				}
+				break;
 			}
-			break;
-		}
-		case ERigElementType::Connector:
-		{
-			Brush = PrimaryConnectorBrush;
-			if(const FRigConnectorElement* Connector = InHierarchy->Find<FRigConnectorElement>(InKey))
+			case ERigElementType::Connector:
 			{
-				if(!Connector->IsPrimary())
+				Brush = PrimaryConnectorBrush;
+				if(const FRigConnectorElement* Connector = InHierarchy->Find<FRigConnectorElement>(InKey.GetElement()))
 				{
-					Brush = Connector->IsOptional() ? OptionalConnectorBrush : SecondaryConnectorBrush;
+					if(!Connector->IsPrimary())
+					{
+						Brush = Connector->IsOptional() ? OptionalConnectorBrush : SecondaryConnectorBrush;
+					}
 				}
+				break;
+			}
+			default:
+			{
+				break;
 			}
-			break;
-		}
-		default:
-		{
-			break;
 		}
 	}
+	else if(InKey.IsComponent())
+	{
+		// todo: offer virtuals on the components to customize their look.
+		Brush = PhysicsBrush; // for now
+	}
 
 	return TPair<const FSlateBrush*, FSlateColor>(Brush, Color);
 }
@@ -1342,7 +1413,7 @@ void SRigHierarchyItem::OnNameCommitted(const FText& InText, ETextCommit::Type I
 	if (InCommitType == ETextCommit::OnEnter)
 	{
 		FString NewName = InText.ToString();
-		const FRigElementKey OldKey = WeakRigTreeElement.Pin()->Key;
+		const FRigHierarchyKey OldKey = WeakRigTreeElement.Pin()->Key;
 
 		const FName NewSanitizedName = Delegates.HandleRenameElement(OldKey, NewName);
 		if (NewSanitizedName.IsNone())
@@ -1353,7 +1424,7 @@ void SRigHierarchyItem::OnNameCommitted(const FText& InText, ETextCommit::Type I
 
 		if (WeakRigTreeElement.IsValid())
 		{
-			WeakRigTreeElement.Pin()->Key.Name = *NewName;
+			WeakRigTreeElement.Pin()->Key.SetName(*NewName);
 		}
 	}
 }
diff --git a/Engine/Plugins/Animation/ControlRig/Source/ControlRigEditor/Private/Editor/SRigHierarchyTreeView.h b/Engine/Plugins/Animation/ControlRig/Source/ControlRigEditor/Private/Editor/SRigHierarchyTreeView.h
--- a/Engine/Plugins/Animation/ControlRig/Source/ControlRigEditor/Private/Editor/SRigHierarchyTreeView.h
+++ b/Engine/Plugins/Animation/ControlRig/Source/ControlRigEditor/Private/Editor/SRigHierarchyTreeView.h
@@ -25,10 +25,10 @@ struct CONTROLRIGEDITOR_API FRigTreeDisplaySettings
 		bShowBones = true;
 		bShowControls = true;
 		bShowNulls = true;
-		bShowPhysics = true;
 		bShowReferences = true;
 		bShowSockets = true;
 		bShowConnectors = true;
+		bShowComponents = false;
 		bShowIconColors = true;
 	}
 	
@@ -55,9 +55,6 @@ struct CONTROLRIGEDITOR_API FRigTreeDisplaySettings
 	/** Whether or not to show spaces in the hierarchy */
 	bool bShowNulls;
 
-	/** Whether or not to show physics elements in the hierarchy */
-	bool bShowPhysics;
-
 	/** Whether or not to show references in the hierarchy */
 	bool bShowReferences;
 
@@ -66,21 +63,24 @@ struct CONTROLRIGEDITOR_API FRigTreeDisplaySettings
 
 	/** Whether or not to show connectors in the hierarchy */
 	bool bShowConnectors;
-	
+
+	/** Whether or not to show components in the hierarchy */
+	bool bShowComponents;
+
 	/** Whether to tint the icons with the element color */
 	bool bShowIconColors;
 };
 
 DECLARE_DELEGATE_RetVal(const URigHierarchy*, FOnGetRigTreeHierarchy);
 DECLARE_DELEGATE_RetVal(const FRigTreeDisplaySettings&, FOnGetRigTreeDisplaySettings);
-DECLARE_DELEGATE_RetVal(const TArray<FRigElementKey>, FOnRigTreeGetSelection);
-DECLARE_DELEGATE_RetVal_TwoParams(FName, FOnRigTreeRenameElement, const FRigElementKey& /*OldKey*/, const FString& /*NewName*/);
-DECLARE_DELEGATE_RetVal_ThreeParams(bool, FOnRigTreeVerifyElementNameChanged, const FRigElementKey& /*OldKey*/, const FString& /*NewName*/, FText& /*OutErrorMessage*/);
-DECLARE_DELEGATE_RetVal_TwoParams(bool, FOnRigTreeCompareKeys, const FRigElementKey& /*A*/, const FRigElementKey& /*B*/);
-DECLARE_DELEGATE_RetVal_OneParam(FRigElementKey, FOnRigTreeGetResolvedKey, const FRigElementKey&);
-DECLARE_DELEGATE_OneParam(FOnRigTreeRequestDetailsInspection, const FRigElementKey&);
-DECLARE_DELEGATE_RetVal_OneParam(TOptional<FText>, FOnRigTreeItemGetToolTip, const FRigElementKey&);
-DECLARE_DELEGATE_RetVal_OneParam(bool, FOnRigTreeIsItemVisible, const FRigElementKey&);
+DECLARE_DELEGATE_RetVal(const TArray<FRigHierarchyKey>, FOnRigTreeGetSelection);
+DECLARE_DELEGATE_RetVal_TwoParams(FName, FOnRigTreeRenameElement, const FRigHierarchyKey& /*OldKey*/, const FString& /*NewName*/);
+DECLARE_DELEGATE_RetVal_ThreeParams(bool, FOnRigTreeVerifyElementNameChanged, const FRigHierarchyKey& /*OldKey*/, const FString& /*NewName*/, FText& /*OutErrorMessage*/);
+DECLARE_DELEGATE_RetVal_TwoParams(bool, FOnRigTreeCompareKeys, const FRigHierarchyKey& /*A*/, const FRigHierarchyKey& /*B*/);
+DECLARE_DELEGATE_RetVal_OneParam(FRigHierarchyKey, FOnRigTreeGetResolvedKey, const FRigHierarchyKey&);
+DECLARE_DELEGATE_OneParam(FOnRigTreeRequestDetailsInspection, const FRigHierarchyKey&);
+DECLARE_DELEGATE_RetVal_OneParam(TOptional<FText>, FOnRigTreeItemGetToolTip, const FRigHierarchyKey&);
+DECLARE_DELEGATE_RetVal_OneParam(bool, FOnRigTreeIsItemVisible, const FRigHierarchyKey&);
 
 typedef STableRow<TSharedPtr<FRigTreeElement>>::FOnCanAcceptDrop FOnRigTreeCanAcceptDrop;
 typedef STableRow<TSharedPtr<FRigTreeElement>>::FOnAcceptDrop FOnRigTreeAcceptDrop;
@@ -143,7 +143,7 @@ struct CONTROLRIGEDITOR_API FRigTreeDelegates
 		return DefaultDisplaySettings;
 	}
 
-	TArray<FRigElementKey> GetSelection() const
+	TArray<FRigHierarchyKey> GetSelection() const
 	{
 		if (OnGetSelection.IsBound())
 		{
@@ -151,21 +151,28 @@ struct CONTROLRIGEDITOR_API FRigTreeDelegates
 		}
 		if (const URigHierarchy* Hierarchy = GetHierarchy())
 		{
-			return Hierarchy->GetSelectedKeys();
+			const TArray<FRigElementKey> SelectedElementKeys = Hierarchy->GetSelectedKeys(); 
+			TArray<FRigHierarchyKey> Selection;
+			Selection.Reserve(SelectedElementKeys.Num());
+			for(const FRigElementKey& SelectedElementKey : SelectedElementKeys)
+			{
+				Selection.Emplace(SelectedElementKey);
+			}
+			return Selection;
 		}
 		return {};
 	}
 	
-	FName HandleRenameElement(const FRigElementKey& OldKey, const FString& NewName) const
+	FName HandleRenameElement(const FRigHierarchyKey& OldKey, const FString& NewName) const
 	{
 		if(OnRenameElement.IsBound())
 		{
 			return OnRenameElement.Execute(OldKey, NewName);
 		}
-		return OldKey.Name;
+		return OldKey.GetFName();
 	}
 	
-	bool HandleVerifyElementNameChanged(const FRigElementKey& OldKey, const FString& NewName, FText& OutErrorMessage) const
+	bool HandleVerifyElementNameChanged(const FRigHierarchyKey& OldKey, const FString& NewName, FText& OutErrorMessage) const
 	{
 		if(OnVerifyElementNameChanged.IsBound())
 		{
@@ -181,10 +188,10 @@ struct CONTROLRIGEDITOR_API FRigTreeDelegates
 			return;
 		}
 		TGuardValue<bool> Guard(bIsChangingRigHierarchy, true);
-		OnSelectionChanged.ExecuteIfBound(Selection, SelectInfo);
+		(void)OnSelectionChanged.ExecuteIfBound(Selection, SelectInfo);
 	}
 
-	FRigElementKey GetResolvedKey(const FRigElementKey& InKey)
+	FRigHierarchyKey GetResolvedKey(const FRigHierarchyKey& InKey)
 	{
 		if(OnGetResolvedKey.IsBound())
 		{
@@ -193,7 +200,7 @@ struct CONTROLRIGEDITOR_API FRigTreeDelegates
 		return InKey;
 	}
 
-	void RequestDetailsInspection(const FRigElementKey& InKey)
+	void RequestDetailsInspection(const FRigHierarchyKey& InKey)
 	{
 		if(OnRequestDetailsInspection.IsBound())
 		{
@@ -226,10 +233,10 @@ enum class ERigTreeFilterResult : int32
 class FRigTreeElement : public TSharedFromThis<FRigTreeElement>
 {
 public:
-	FRigTreeElement(const FRigElementKey& InKey, TWeakPtr<SRigHierarchyTreeView> InTreeView, bool InSupportsRename, ERigTreeFilterResult InFilterResult);
+	FRigTreeElement(const FRigHierarchyKey& InKey, TWeakPtr<SRigHierarchyTreeView> InTreeView, bool InSupportsRename, ERigTreeFilterResult InFilterResult);
 public:
 	/** Element Data to display */
-	FRigElementKey Key;
+	FRigHierarchyKey Key;
 	FName ShortName;
 	FName ChannelName;
 	bool bIsTransient;
@@ -277,7 +284,7 @@ public:
 	void Construct(const FArguments& InArgs, const TSharedRef<STableViewBase>& OwnerTable, TSharedRef<FRigTreeElement> InRigTreeElement, TSharedPtr<SRigHierarchyTreeView> InTreeView, const FRigTreeDisplaySettings& InSettings, bool bPinned);
  	void OnNameCommitted(const FText& InText, ETextCommit::Type InCommitType) const;
 	bool OnVerifyNameChanged(const FText& InText, FText& OutErrorMessage);
-	static TPair<const FSlateBrush*, FSlateColor> GetBrushForElementType(const URigHierarchy* InHierarchy, const FRigElementKey& InKey);
+	static TPair<const FSlateBrush*, FSlateColor> GetBrushForElementType(const URigHierarchy* InHierarchy, const FRigHierarchyKey& InKey);
 	static FLinearColor GetColorForControlType(ERigControlType InControlType, UEnum* InControlEnum);
 
 private:
@@ -346,20 +353,21 @@ public:
 		}
 	}
 
-	TSharedPtr<FRigTreeElement> FindElement(const FRigElementKey& InElementKey) const;
-	static TSharedPtr<FRigTreeElement> FindElement(const FRigElementKey& InElementKey, TSharedPtr<FRigTreeElement> CurrentItem);
-	bool AddElement(FRigElementKey InKey, FRigElementKey InParentKey = FRigElementKey());
+	TSharedPtr<FRigTreeElement> FindElement(const FRigHierarchyKey& InElementKey) const;
+	static TSharedPtr<FRigTreeElement> FindElement(const FRigHierarchyKey& InElementKey, TSharedPtr<FRigTreeElement> CurrentItem);
+	bool AddElement(FRigHierarchyKey InKey, FRigHierarchyKey InParentKey = FRigHierarchyKey());
 	bool AddElement(const FRigBaseElement* InElement);
+	bool AddComponent(const FRigBaseComponent* InComponent);
 	void AddSpacerElement();
-	bool ReparentElement(FRigElementKey InKey, FRigElementKey InParentKey);
-	bool RemoveElement(FRigElementKey InKey);
+	bool ReparentElement(FRigHierarchyKey InKey, FRigHierarchyKey InParentKey);
+	bool RemoveElement(FRigHierarchyKey InKey);
 	void RefreshTreeView(bool bRebuildContent = true);
 	void SetExpansionRecursive(TSharedPtr<FRigTreeElement> InElement, bool bTowardsParent, bool bShouldBeExpanded);
 	TSharedRef<ITableRow> MakeTableRowWidget(TSharedPtr<FRigTreeElement> InItem, const TSharedRef<STableViewBase>& OwnerTable, bool bPinned);
 	void HandleGetChildrenForTree(TSharedPtr<FRigTreeElement> InItem, TArray<TSharedPtr<FRigTreeElement>>& OutChildren);
 	void OnElementKeyTagDragDetected(const FRigElementKey& InDraggedTag);
 
-	TArray<FRigElementKey> GetSelectedKeys() const;
+	TArray<FRigHierarchyKey> GetSelectedKeys() const;
 	const TArray<TSharedPtr<FRigTreeElement>>& GetRootElements() const { return RootElements; }
 	FRigTreeDelegates& GetRigTreeDelegates() { return Delegates; }
 
@@ -378,10 +386,10 @@ private:
 	TArray<TSharedPtr<FRigTreeElement>> RootElements;
 	
 	/** A map for looking up items based on their key */
-	TMap<FRigElementKey, TSharedPtr<FRigTreeElement>> ElementMap;
+	TMap<FRigHierarchyKey, TSharedPtr<FRigTreeElement>> ElementMap;
 
 	/** A map for looking up a parent based on their key */
-	TMap<FRigElementKey, FRigElementKey> ParentMap;
+	TMap<FRigHierarchyKey, FRigHierarchyKey> ParentMap;
 
 	FRigTreeDelegates Delegates;
 
diff --git a/Engine/Plugins/Animation/ControlRig/Source/ControlRigEditor/Private/Editor/SRigSpacePickerWidget.cpp b/Engine/Plugins/Animation/ControlRig/Source/ControlRigEditor/Private/Editor/SRigSpacePickerWidget.cpp
--- a/Engine/Plugins/Animation/ControlRig/Source/ControlRigEditor/Private/Editor/SRigSpacePickerWidget.cpp
+++ b/Engine/Plugins/Animation/ControlRig/Source/ControlRigEditor/Private/Editor/SRigSpacePickerWidget.cpp
@@ -704,7 +704,7 @@ FReply SRigSpacePickerWidget::HandleAddElementClicked()
 {
 	HierarchyDisplaySettings.bShowConnectors = false;
 	HierarchyDisplaySettings.bShowSockets = false;
-	HierarchyDisplaySettings.bShowPhysics = false;
+	HierarchyDisplaySettings.bShowComponents = false;
 	
 	FRigTreeDelegates TreeDelegates;
 	TreeDelegates.OnGetDisplaySettings = FOnGetRigTreeDisplaySettings::CreateSP(this, &SRigSpacePickerWidget::GetHierarchyDisplaySettings); 
@@ -713,7 +713,7 @@ FReply SRigSpacePickerWidget::HandleAddElementClicked()
 	{
 		if(InItem.IsValid())
 		{
-			const FRigElementKey Key = InItem->Key;
+			const FRigElementKey Key = InItem->Key.GetElement();
 			if(!IsDefaultSpace(Key) && IsValidKey(Key))
 			{
 				if (Hierarchy.IsValid())
@@ -778,17 +778,29 @@ FReply SRigSpacePickerWidget::HandleAddElementClicked()
 		}
 	});
 
-	TreeDelegates.OnCompareKeys = FOnRigTreeCompareKeys::CreateLambda([](const FRigElementKey& A, const FRigElementKey& B) -> bool
+	TreeDelegates.OnCompareKeys = FOnRigTreeCompareKeys::CreateLambda([](const FRigHierarchyKey& A, const FRigHierarchyKey& B) -> bool
 	{
-		// controls should always show up first - so we'll sort them to the start of the list
-		if(A.Type == ERigElementType::Control && B.Type != ERigElementType::Control)
+		if(A.IsElement() && B.IsComponent())
 		{
 			return true;
 		}
-		if(B.Type == ERigElementType::Control && A.Type != ERigElementType::Control)
+		if(B.IsElement() && A.IsComponent())
 		{
 			return false;
 		}
+
+		// controls should always show up first - so we'll sort them to the start of the list
+		if(A.IsElement() && B.IsElement())
+		{
+			if(A.GetElement().Type == ERigElementType::Control && B.GetElement().Type != ERigElementType::Control)
+			{
+				return true;
+			}
+			if(B.GetElement().Type == ERigElementType::Control && A.GetElement().Type != ERigElementType::Control)
+			{
+				return false;
+			}
+		}
 		return A < B;
 	});
 
diff --git a/Engine/Plugins/Runtime/RigVM/Source/RigVMEditor/Public/Editor/RigVMDetailsViewWrapperObject.h b/Engine/Plugins/Runtime/RigVM/Source/RigVMEditor/Public/Editor/RigVMDetailsViewWrapperObject.h
--- a/Engine/Plugins/Runtime/RigVM/Source/RigVMEditor/Public/Editor/RigVMDetailsViewWrapperObject.h
+++ b/Engine/Plugins/Runtime/RigVM/Source/RigVMEditor/Public/Editor/RigVMDetailsViewWrapperObject.h
@@ -45,8 +45,8 @@ public:
 		return IsChildOf(T::StaticStruct());
 	}
 
-	void SetContent(const uint8* InStructMemory, const UStruct* InStruct);
-	void GetContent(uint8* OutStructMemory, const UStruct* InStruct) const;
+	virtual void SetContent(const uint8* InStructMemory, const UStruct* InStruct);
+	virtual void GetContent(uint8* OutStructMemory, const UStruct* InStruct) const;
 	void SetContent(URigVMNode* InNode);
 
 	template<typename T>
