Control Rig: Fix build break due to use of NoneName

#rnx

diff --git a/Engine/Plugins/Animation/ControlRig/Source/ControlRig/Private/Rigs/RigHierarchyDefines.cpp b/Engine/Plugins/Animation/ControlRig/Source/ControlRig/Private/Rigs/RigHierarchyDefines.cpp
--- a/Engine/Plugins/Animation/ControlRig/Source/ControlRig/Private/Rigs/RigHierarchyDefines.cpp
+++ b/Engine/Plugins/Animation/ControlRig/Source/ControlRig/Private/Rigs/RigHierarchyDefines.cpp
@@ -200,6 +200,20 @@ bool FRigComponentKey::IsTopLevel() const
 // FRigHierarchyKey
 ////////////////////////////////////////////////////////////////////////////////
 
+const FName& FRigHierarchyKey::GetFName() const
+{
+	if(IsElement())
+	{
+		return Element.GetValue().Name;
+	}
+	if(IsComponent())
+	{
+		return Component.GetValue().Name;
+	}
+	static const FName InvalidName(NAME_None);
+	return InvalidName;;
+}
+
 void FRigHierarchyKey::Serialize(FArchive& Ar)
 {
 	if(Ar.IsLoading())
diff --git a/Engine/Plugins/Animation/ControlRig/Source/ControlRig/Public/Rigs/RigHierarchyDefines.h b/Engine/Plugins/Animation/ControlRig/Source/ControlRig/Public/Rigs/RigHierarchyDefines.h
--- a/Engine/Plugins/Animation/ControlRig/Source/ControlRig/Public/Rigs/RigHierarchyDefines.h
+++ b/Engine/Plugins/Animation/ControlRig/Source/ControlRig/Public/Rigs/RigHierarchyDefines.h
@@ -1800,19 +1800,7 @@ public:
 		return IsElement() || IsComponent();
 	}
 
-	const FName& GetFName() const
-	{
-		if(IsElement())
-		{
-			return Element.GetValue().Name;
-		}
-		if(IsComponent())
-		{
-			return Component.GetValue().Name;
-		}
-		static const FName NoneName(NAME_None);
-		return NoneName;;
-	}
+	const FName& GetFName() const;
 
 	FString GetName() const
 	{
