Control Rig: New dispatch node to spawn top level component

#rnx

diff --git a/Engine/Plugins/Animation/ControlRig/Source/ControlRig/Private/Units/Hierarchy/RigUnit_Component.cpp b/Engine/Plugins/Animation/ControlRig/Source/ControlRig/Private/Units/Hierarchy/RigUnit_Component.cpp
--- a/Engine/Plugins/Animation/ControlRig/Source/ControlRig/Private/Units/Hierarchy/RigUnit_Component.cpp
+++ b/Engine/Plugins/Animation/ControlRig/Source/ControlRig/Private/Units/Hierarchy/RigUnit_Component.cpp
@@ -132,6 +132,14 @@ const TArray<FRigVMTemplateArgumentInfo>& FRigDispatch_SpawnComponent::GetArgume
 	return Infos;
 }
 
+#if WITH_EDITOR
+
+FString FRigDispatch_SpawnComponent::GetKeywords() const
+{
+	return FRigDispatch_ComponentBase::GetKeywords() + TEXT(",AddComponent,CreateComponent");
+}
+
+#endif
 
 void FRigDispatch_SpawnComponent::Execute(FRigVMExtendedExecuteContext& InContext, FRigVMMemoryHandleArray Handles, FRigVMPredicateBranchArray /* not needed */)
 {
@@ -182,6 +190,71 @@ void FRigDispatch_SpawnComponent::Execute(FRigVMExtendedExecuteContext& InContex
 	}
 }
 
+//////////////////////////////////////////////////////////////////////////////////////////////////////
+// FRigDispatch_SpawnTopLevelComponent
+//////////////////////////////////////////////////////////////////////////////////////////////////////
+
+const TArray<FRigVMTemplateArgumentInfo>& FRigDispatch_SpawnTopLevelComponent::GetArgumentInfos() const
+{
+	if (Infos.IsEmpty())
+	{
+		NameArgIndex = Infos.Emplace(NameArgName, ERigVMPinDirection::Input, RigVMTypeUtils::TypeIndex::FName);
+		ComponentArgIndex = Infos.Add(FRigVMTemplateArgumentInfo(ComponentArgName, ERigVMPinDirection::Input, {FRigVMTemplateArgument::ETypeCategory_SingleScriptStructValue}, &FRigDispatch_SpawnTopLevelComponent::IsTypeSupported));
+		KeyArgIndex = Infos.Emplace(KeyArgName, ERigVMPinDirection::Output, GetComponentKeyType());
+		SuccessArgIndex = Infos.Emplace(SuccessArgName, ERigVMPinDirection::Output, RigVMTypeUtils::TypeIndex::Bool);
+	}
+	return Infos;
+}
+
+#if WITH_EDITOR
+
+FString FRigDispatch_SpawnTopLevelComponent::GetKeywords() const
+{
+	return FRigDispatch_ComponentBase::GetKeywords() + TEXT(",AddComponent,CreateComponent,TopLevel,RootComponent");
+}
+
+#endif
+
+void FRigDispatch_SpawnTopLevelComponent::Execute(FRigVMExtendedExecuteContext& InContext, FRigVMMemoryHandleArray Handles, FRigVMPredicateBranchArray /* not needed */)
+{
+	const FRigDispatch_SpawnTopLevelComponent* Factory = static_cast<const FRigDispatch_SpawnTopLevelComponent*>(InContext.Factory);
+		
+	// unpack the memory
+	const FName& Name = *(const FName*)Handles[Factory->NameArgIndex].GetData();
+	const uint8* SourceComponent = Handles[Factory->ComponentArgIndex].GetData();
+	FRigComponentKey& ComponentKey = *(FRigComponentKey*)Handles[Factory->KeyArgIndex].GetData();
+	bool& bSuccess = *(bool*)Handles[Factory->SuccessArgIndex].GetData();
+
+	const FStructProperty* StructProperty = CastField<FStructProperty>(Handles[Factory->ComponentArgIndex].GetResolvedProperty());
+	check(StructProperty);
+	UScriptStruct* ComponentStruct = StructProperty->Struct;
+	check(ComponentStruct && ComponentStruct->IsChildOf(FRigBaseComponent::StaticStruct()));
+
+	ComponentKey = FRigComponentKey();
+	bSuccess = false;
+
+	// extract the component from the hierarchy
+	const FControlRigExecuteContext& Context = InContext.GetPublicDataSafe<FControlRigExecuteContext>();
+	if(URigHierarchyController* Controller = Context.Hierarchy->GetController())
+	{
+		FRigHierarchyControllerInstructionBracket InstructionBracket(Controller, Context.GetInstructionIndex());
+		ComponentKey = Controller->AddTopLevelComponent(ComponentStruct, Name);
+		if(ComponentKey.IsValid())
+		{
+			FRigBaseComponent* TargetComponent = Context.Hierarchy->FindComponent(ComponentKey);
+			check(TargetComponent);
+			
+			// copy the public content over (members which are not UProperty will be skipped) 
+			for (TFieldIterator<FProperty> It(ComponentStruct); It; ++It)
+			{
+				It->CopyCompleteValue_InContainer(TargetComponent, SourceComponent);
+			}
+
+			bSuccess = true;
+		}
+	}
+}
+
 //////////////////////////////////////////////////////////////////////////////////////////////////////
 // FRigDispatch_GetComponentContent
 //////////////////////////////////////////////////////////////////////////////////////////////////////
@@ -292,6 +365,144 @@ void FRigDispatch_SetComponentContent::Execute(FRigVMExtendedExecuteContext& InC
 		return;
 	}
 
+	// copy the public content over (members which are not UProperty will be skipped) 
+	for (TFieldIterator<FProperty> It(ComponentStruct); It; ++It)
+	{
+		It->CopyCompleteValue_InContainer(TargetComponent, SourceComponent);
+	}
+	bSuccess = true;
+}
+
+//////////////////////////////////////////////////////////////////////////////////////////////////////
+// FRigDispatch_GetTopLevelComponentContent
+//////////////////////////////////////////////////////////////////////////////////////////////////////
+
+const TArray<FRigVMTemplateArgumentInfo>& FRigDispatch_GetTopLevelComponentContent::GetArgumentInfos() const
+{
+	if (Infos.IsEmpty())
+	{
+		NameArgIndex = Infos.Emplace(NameArgName, ERigVMPinDirection::Input, RigVMTypeUtils::TypeIndex::FName);
+		ComponentArgIndex = Infos.Add(FRigVMTemplateArgumentInfo(ComponentArgName, ERigVMPinDirection::Output, {FRigVMTemplateArgument::ETypeCategory_SingleScriptStructValue}, &FRigDispatch_GetTopLevelComponentContent::IsTypeSupported));
+		SuccessArgIndex = Infos.Emplace(SuccessArgName, ERigVMPinDirection::Output, RigVMTypeUtils::TypeIndex::Bool);
+	}
+	
+	return Infos;
+}
+
+#if WITH_EDITOR
+
+FString FRigDispatch_GetTopLevelComponentContent::GetKeywords() const
+{
+	return FRigDispatch_ComponentBase::GetKeywords() + TEXT(",GetComponent,ReadComponent");
+}
+
+#endif
+
+void FRigDispatch_GetTopLevelComponentContent::Execute(FRigVMExtendedExecuteContext& InContext, FRigVMMemoryHandleArray Handles, FRigVMPredicateBranchArray)
+{
+	const FRigDispatch_GetTopLevelComponentContent* Factory = static_cast<const FRigDispatch_GetTopLevelComponentContent*>(InContext.Factory);
+			
+	// unpack the memory
+	const FName& Name = *(const FName*)Handles[Factory->NameArgIndex].GetData();
+	const FRigComponentKey ComponentKey(URigHierarchy::GetTopLevelComponentElementKey(), Name);
+	uint8* TargetComponent = Handles[Factory->ComponentArgIndex].GetData();
+	bool& bSuccess = *(bool*)Handles[Factory->SuccessArgIndex].GetData();
+
+	const FStructProperty* StructProperty = CastField<FStructProperty>(Handles[Factory->ComponentArgIndex].GetResolvedProperty());
+	check(StructProperty);
+	const UScriptStruct* ComponentStruct = StructProperty->Struct;
+	check(ComponentStruct && ComponentStruct->IsChildOf(FRigBaseComponent::StaticStruct()));
+
+	// extract the component from the hierarchy
+	const FControlRigExecuteContext& Context = InContext.GetPublicDataSafe<FControlRigExecuteContext>();
+	const FRigBaseComponent* SourceComponent = Context.Hierarchy->FindComponent(ComponentKey);
+	if(SourceComponent == nullptr)
+	{
+		ComponentStruct->InitializeDefaultValue(TargetComponent);
+#if WITH_EDITOR
+		Context.Report(FRigVMLogSettings(EMessageSeverity::Warning), Context.GetFunctionName(), Context.GetInstructionIndex(), FString::Printf(TEXT("%s not found."), *ComponentKey.ToString()));
+#endif
+		bSuccess = false;
+		return;
+	}
+
+	if(SourceComponent->GetScriptStruct() != ComponentStruct)
+	{
+		ComponentStruct->InitializeDefaultValue(TargetComponent);
+#if WITH_EDITOR
+		const FString Message = FString::Printf(TEXT("Component pin type (%s) doesn't match component in hierarchy (%s)."), *ComponentStruct->GetName(), *SourceComponent->GetScriptStruct()->GetName());
+		Context.Report(FRigVMLogSettings(EMessageSeverity::Warning), Context.GetFunctionName(), Context.GetInstructionIndex(), Message);
+#endif
+		bSuccess = false;
+		return;
+	}
+
+	ComponentStruct->CopyScriptStruct(TargetComponent, SourceComponent, 1);
+	bSuccess = true;
+}
+
+//////////////////////////////////////////////////////////////////////////////////////////////////////
+// FRigDispatch_SetTopLevelComponentContent
+//////////////////////////////////////////////////////////////////////////////////////////////////////
+
+const TArray<FRigVMTemplateArgumentInfo>& FRigDispatch_SetTopLevelComponentContent::GetArgumentInfos() const
+{
+	if (Infos.IsEmpty())
+	{
+		NameArgIndex = Infos.Emplace(NameArgName, ERigVMPinDirection::Input, RigVMTypeUtils::TypeIndex::FName);
+		ComponentArgIndex = Infos.Add(FRigVMTemplateArgumentInfo(ComponentArgName, ERigVMPinDirection::Input, {FRigVMTemplateArgument::ETypeCategory_SingleScriptStructValue}, &FRigDispatch_SetTopLevelComponentContent::IsTypeSupported));
+		SuccessArgIndex = Infos.Emplace(SuccessArgName, ERigVMPinDirection::Output, RigVMTypeUtils::TypeIndex::Bool);
+	}
+	
+	return Infos;
+}
+
+#if WITH_EDITOR
+
+FString FRigDispatch_SetTopLevelComponentContent::GetKeywords() const
+{
+	return FRigDispatch_ComponentBase::GetKeywords() + TEXT(",SetComponent,WriteComponent");
+}
+
+#endif
+
+void FRigDispatch_SetTopLevelComponentContent::Execute(FRigVMExtendedExecuteContext& InContext, FRigVMMemoryHandleArray Handles, FRigVMPredicateBranchArray)
+{
+	const FRigDispatch_SetTopLevelComponentContent* Factory = static_cast<const FRigDispatch_SetTopLevelComponentContent*>(InContext.Factory);
+			
+	// unpack the memory
+	const FName& Name = *(const FName*)Handles[Factory->NameArgIndex].GetData();
+	const FRigComponentKey ComponentKey(URigHierarchy::GetTopLevelComponentElementKey(), Name);
+	const  uint8* SourceComponent = Handles[Factory->ComponentArgIndex].GetData();
+	bool& bSuccess = *(bool*)Handles[Factory->SuccessArgIndex].GetData();
+
+	const FStructProperty* StructProperty = CastField<FStructProperty>(Handles[Factory->ComponentArgIndex].GetResolvedProperty());
+	check(StructProperty);
+	const UScriptStruct* ComponentStruct = StructProperty->Struct;
+	check(ComponentStruct && ComponentStruct->IsChildOf(FRigBaseComponent::StaticStruct()));
+
+	// extract the component from the hierarchy
+	const FControlRigExecuteContext& Context = InContext.GetPublicDataSafe<FControlRigExecuteContext>();
+	FRigBaseComponent* TargetComponent = Context.Hierarchy->FindComponent(ComponentKey);
+	if(TargetComponent == nullptr)
+	{
+#if WITH_EDITOR
+		Context.Report(FRigVMLogSettings(EMessageSeverity::Warning), Context.GetFunctionName(), Context.GetInstructionIndex(), FString::Printf(TEXT("%s not found."), *ComponentKey.ToString()));
+#endif
+		bSuccess = false;
+		return;
+	}
+
+	if(TargetComponent->GetScriptStruct() != ComponentStruct)
+	{
+#if WITH_EDITOR
+		const FString Message = FString::Printf(TEXT("Component pin type (%s) doesn't match component in hierarchy (%s)."), *ComponentStruct->GetName(), *TargetComponent->GetScriptStruct()->GetName());
+		Context.Report(FRigVMLogSettings(EMessageSeverity::Warning), Context.GetFunctionName(), Context.GetInstructionIndex(), Message);
+#endif
+		bSuccess = false;
+		return;
+	}
+
 	// copy the public content over (members which are not UProperty will be skipped) 
 	for (TFieldIterator<FProperty> It(ComponentStruct); It; ++It)
 	{
diff --git a/Engine/Plugins/Animation/ControlRig/Source/ControlRig/Public/Units/Hierarchy/RigUnit_Component.h b/Engine/Plugins/Animation/ControlRig/Source/ControlRig/Public/Units/Hierarchy/RigUnit_Component.h
--- a/Engine/Plugins/Animation/ControlRig/Source/ControlRig/Public/Units/Hierarchy/RigUnit_Component.h
+++ b/Engine/Plugins/Animation/ControlRig/Source/ControlRig/Public/Units/Hierarchy/RigUnit_Component.h
@@ -41,10 +41,12 @@ protected:
 	static inline const FLazyName KeyArgName = FLazyName(TEXT("Key"));
 	static inline const FLazyName ComponentArgName = FLazyName(TEXT("Component"));
 	static inline const FLazyName SuccessArgName = FLazyName(TEXT("Success"));
+
+	friend class FControlRigEditor;
 };
 
 /**
- * Break / Decompose node for a hierarchy component
+ * Adds a component under an element in the hierarchy
  */
 USTRUCT(meta=(DisplayName="Spawn Component"))
 struct CONTROLRIG_API FRigDispatch_SpawnComponent : public FRigDispatch_ComponentBase
@@ -58,6 +60,39 @@ struct CONTROLRIG_API FRigDispatch_SpawnComponent : public FRigDispatch_Componen
 	}
 
 	virtual const TArray<FRigVMTemplateArgumentInfo>& GetArgumentInfos() const override;
+#if WITH_EDITOR
+	virtual FString GetKeywords() const override;
+#endif
+
+protected:
+
+	virtual bool IsMutable() const override { return true; }
+
+	virtual FRigVMFunctionPtr GetDispatchFunctionImpl(const FRigVMTemplateTypeMap& InTypes) const override
+	{
+		return &Execute;
+	}
+	static void Execute(FRigVMExtendedExecuteContext& InContext, FRigVMMemoryHandleArray Handles, FRigVMPredicateBranchArray /* not needed */);
+};
+
+/**
+ * Break / Decompose node for a hierarchy component
+ */
+USTRUCT(meta=(DisplayName="Spawn Top Level Component"))
+struct CONTROLRIG_API FRigDispatch_SpawnTopLevelComponent : public FRigDispatch_ComponentBase
+{
+	GENERATED_BODY()
+
+	FRigDispatch_SpawnTopLevelComponent()
+	{
+		// this is only needed for engine test
+		FactoryScriptStruct = FRigDispatch_SpawnTopLevelComponent::StaticStruct();
+	}
+
+	virtual const TArray<FRigVMTemplateArgumentInfo>& GetArgumentInfos() const override;
+#if WITH_EDITOR
+	virtual FString GetKeywords() const override;
+#endif
 
 protected:
 
@@ -111,6 +146,64 @@ struct CONTROLRIG_API FRigDispatch_SetComponentContent : public FRigDispatch_Com
 
 	virtual const TArray<FRigVMTemplateArgumentInfo>& GetArgumentInfos() const override;
 
+protected:
+
+	virtual bool IsMutable() const override { return true; }
+
+	virtual FRigVMFunctionPtr GetDispatchFunctionImpl(const FRigVMTemplateTypeMap& InTypes) const override
+	{
+		return &Execute;
+	}
+	static void Execute(FRigVMExtendedExecuteContext& InContext, FRigVMMemoryHandleArray Handles, FRigVMPredicateBranchArray /* not needed */);
+};
+
+/**
+ * Break / Decompose node for a hierarchy component
+ */
+USTRUCT(meta=(DisplayName="Get Top Level Component"))
+struct CONTROLRIG_API FRigDispatch_GetTopLevelComponentContent : public FRigDispatch_ComponentBase
+{
+	GENERATED_BODY()
+
+	FRigDispatch_GetTopLevelComponentContent()
+	{
+		// this is only needed for engine test
+		FactoryScriptStruct = FRigDispatch_GetTopLevelComponentContent::StaticStruct();
+	}
+
+	virtual const TArray<FRigVMTemplateArgumentInfo>& GetArgumentInfos() const override;
+#if WITH_EDITOR
+	virtual FString GetKeywords() const override;
+#endif
+
+protected:
+
+	virtual FRigVMFunctionPtr GetDispatchFunctionImpl(const FRigVMTemplateTypeMap& InTypes) const override
+	{
+		return &Execute;
+	}
+	static void Execute(FRigVMExtendedExecuteContext& InContext, FRigVMMemoryHandleArray Handles, FRigVMPredicateBranchArray /* not needed */);
+};
+
+/**
+ * Set the content of a component
+ */
+USTRUCT(meta=(DisplayName="Set Top Level Component"))
+struct CONTROLRIG_API FRigDispatch_SetTopLevelComponentContent : public FRigDispatch_ComponentBase
+{
+	GENERATED_BODY()
+
+	FRigDispatch_SetTopLevelComponentContent()
+	{
+		// this is only needed for engine test
+		FactoryScriptStruct = FRigDispatch_SetTopLevelComponentContent::StaticStruct();
+	}
+
+	virtual const TArray<FRigVMTemplateArgumentInfo>& GetArgumentInfos() const override;
+#if WITH_EDITOR
+	virtual FString GetKeywords() const override;
+#endif
+
 protected:
 
 	virtual bool IsMutable() const override { return true; }
