Control Rig: Component renaming and drag and drop support

#rnx

diff --git a/Engine/Plugins/Animation/ControlRig/Source/ControlRig/Private/Rigs/RigHierarchy.cpp b/Engine/Plugins/Animation/ControlRig/Source/ControlRig/Private/Rigs/RigHierarchy.cpp
--- a/Engine/Plugins/Animation/ControlRig/Source/ControlRig/Private/Rigs/RigHierarchy.cpp
+++ b/Engine/Plugins/Animation/ControlRig/Source/ControlRig/Private/Rigs/RigHierarchy.cpp
@@ -1727,6 +1727,10 @@ FRigBaseComponent* URigHierarchy::GetComponent(int32 InIndex)
 
 int32 URigHierarchy::NumComponents(FRigElementKey InElement) const
 {
+	if(InElement == GetTopLevelComponentElementKey())
+	{
+		return NumTopLevelComponents();
+	}
 	if(const FRigBaseElement* Element = Find(InElement))
 	{
 		return Element->NumComponents();
@@ -1743,6 +1747,10 @@ TArray<FRigComponentKey> URigHierarchy::GetAllComponentKeys() const
 
 TArray<FRigComponentKey> URigHierarchy::GetComponentKeys(FRigElementKey InElement) const
 {
+	if(InElement == GetTopLevelComponentElementKey())
+	{
+		return GetTopLevelComponentKeys();
+	}
 	if(const FRigBaseElement* Element = Find(InElement))
 	{
 		return Element->GetComponentKeys();
@@ -1813,18 +1821,14 @@ FString URigHierarchy::GetComponentContent(FRigElementKey InElement, int32 InCom
 	{
 		if(const FRigBaseComponent* Component = GetTopLevelComponent(InComponentIndex))
 		{
-			FString Content;
-			Component->GetScriptStruct()->ExportText(Content, Component, Component, nullptr, PPF_None, nullptr);
-			return Content;
+			return Component->GetContentAsText();
 		}
 	}
 	if(const FRigBaseElement* Element = Find(InElement))
 	{
 		if(const FRigBaseComponent* Component = Element->GetComponent(InComponentIndex))
 		{
-			FString Content;
-			Component->GetScriptStruct()->ExportText(Content, Component, Component, nullptr, PPF_None, nullptr);
-			return Content;
+			return Component->GetContentAsText();
 		}
 	}
 	return FString();
@@ -1866,6 +1870,103 @@ FString URigHierarchy::GetTopLevelComponentContent(int32 InTopLevelComponentInde
 	return GetComponentContent(GetTopLevelComponentElementKey(), InTopLevelComponentIndex);
 }
 
+bool URigHierarchy::CanAddComponent(FRigElementKey InElementKey, const UScriptStruct* InComponentStruct, FString* OutFailureReason) const
+{
+	check(InComponentStruct);
+
+	if(InElementKey != GetTopLevelComponentElementKey() && !InElementKey.IsValid())
+	{
+		if(OutFailureReason)
+		{
+			*OutFailureReason = FString::Printf(TEXT("InElementKey %s is not valid."), *InElementKey.ToString());
+		}
+		return false;
+	}
+
+	if(!InComponentStruct->IsChildOf(FRigBaseComponent::StaticStruct()))
+	{
+		if(OutFailureReason)
+		{
+			*OutFailureReason = FString::Printf(TEXT("The provided structure '%s' is not a component."), *InComponentStruct->GetName());
+		}
+		return false;
+	}
+
+	const FRigBaseElement* Element = Find(InElementKey);
+	if(Element == nullptr && InElementKey != GetTopLevelComponentElementKey())
+	{
+		if(OutFailureReason)
+		{
+			*OutFailureReason = FString::Printf(TEXT("The element %s cannot be found."), *InElementKey.ToString());
+		}
+		return false;
+	}
+
+	const FStructOnScope StructOnScope(InComponentStruct);
+	if(const FRigBaseComponent* StructMemory = reinterpret_cast<const FRigBaseComponent*>(StructOnScope.GetStructMemory()))
+	{
+		if(StructMemory->IsSingleton())
+		{
+			if(Element)
+			{
+				for(int32 Index = 0; Index < Element->NumComponents(); Index++)
+				{
+					if(const FRigBaseComponent* ExistingComponent = Element->GetComponent(Index))
+					{
+						if(ExistingComponent->GetScriptStruct() == InComponentStruct)
+						{
+							if(OutFailureReason)
+							{
+								*OutFailureReason = FString::Printf(TEXT("A component of type '%s' already exists on element '%s'."), *InComponentStruct->GetName(), *InElementKey.ToString());
+							}
+							return false;
+						}
+					}
+				}
+			}
+			else
+			{
+				for(int32 Index = 0; Index < NumTopLevelComponents(); Index++)
+				{
+					if(const FRigBaseComponent* ExistingComponent = GetTopLevelComponent(Index))
+					{
+						if(ExistingComponent->GetScriptStruct() == InComponentStruct)
+						{
+							if(OutFailureReason)
+							{
+								*OutFailureReason = FString::Printf(TEXT("A component of type '%s' already exists at the top level."), *InComponentStruct->GetName());
+							}
+							return false;
+						}
+					}
+				}
+			}
+		}
+		if(!StructMemory->CanBeAddedTo(InElementKey, this, OutFailureReason))
+		{
+			return false;
+		}
+	}
+	return true;
+}
+
+bool URigHierarchy::CanAddComponent(FRigElementKey InElementKey, const FRigBaseComponent* InComponent, FString* OutFailureReason) const
+{
+	if(!CanAddComponent(InElementKey, InComponent->GetScriptStruct(), OutFailureReason))
+	{
+		return false;
+	}
+	if(InComponent->GetElementKey() == InElementKey)
+	{
+		if(OutFailureReason)
+		{
+			*OutFailureReason = TEXT("Component is already under target element.");
+		}
+		return false;
+	}
+	return true;
+}
+
 const FRigElementKey& URigHierarchy::GetTopLevelComponentElementKey()
 {
 	static const FRigElementKey TopLevelComponentElementKey = FRigElementKey(TEXT("TopLevel"), ERigElementType::None);
@@ -2440,6 +2541,19 @@ bool URigHierarchy::IsDisplayNameAvailable(const FRigElementKey& InParentElement
 	return true;
 }
 
+bool URigHierarchy::IsComponentNameAvailable(const FRigElementKey& InElementKey, const FRigName& InPotentialNewName, FString* OutErrorMessage) const
+{
+	const int32 Count = NumComponents(InElementKey);
+	for(int32 Index = 0; Index < Count; Index++)
+	{
+		if(InPotentialNewName.GetFName() == GetComponentName(InElementKey, Index))
+		{
+			return false;
+		}
+	}
+	return true;
+}
+
 FRigName URigHierarchy::GetSafeNewName(const FRigName& InPotentialNewName, ERigElementType InType, bool bAllowNameSpace) const
 {
 	FRigName SanitizedName = InPotentialNewName;
@@ -2568,6 +2682,29 @@ FRigName URigHierarchy::GetSafeNewDisplayName(const FRigElementKey& InParentElem
 	return UniqueName;
 }
 
+FRigName URigHierarchy::GetSafeNewComponentName(const FRigElementKey& InElementKey, const FRigName& InPotentialNewName) const
+{
+	if(InPotentialNewName.IsNone())
+	{
+		return FRigName();
+	}
+
+	FRigName SanitizedName = InPotentialNewName;
+	SanitizeName(SanitizedName);
+
+	if (SanitizedName.Len() > GetMaxNameLength() - 4)
+	{
+		SanitizedName.SetName(SanitizedName.GetName().LeftChop(SanitizedName.Len() - (GetMaxNameLength() - 4)));
+	}
+
+	const FName UniqueName = FRigName(GetUniqueName(*SanitizedName.GetName(), [this, InElementKey](const FName& InName) -> bool
+	{
+		return IsComponentNameAvailable(InElementKey, InName);
+	}));
+
+	return FRigName(UniqueName);
+}
+
 FText URigHierarchy::GetDisplayNameForUI(const FRigBaseElement* InElement, bool bIncludeNameSpace) const
 {
 	check(InElement);
@@ -4350,6 +4487,8 @@ void URigHierarchy::Notify(ERigHierarchyNotification InNotifType, const FRigNoti
 		}
 		case ERigHierarchyNotification::ComponentAdded:
 		case ERigHierarchyNotification::ComponentRemoved:
+		case ERigHierarchyNotification::ComponentRenamed:
+		case ERigHierarchyNotification::ComponentReparented:
 		{
 			if (ensure(InComponent != nullptr))
 			{
@@ -4443,6 +4582,8 @@ void URigHierarchy::SendQueuedNotifications()
 			case ERigHierarchyNotification::ElementRemoved:
 			case ERigHierarchyNotification::ElementRenamed:
 			case ERigHierarchyNotification::ComponentRemoved:
+			case ERigHierarchyNotification::ComponentRenamed:
+			case ERigHierarchyNotification::ComponentReparented:
 			{
 				// we don't allow these to happen during the run of a VM
 				if(const URigHierarchyController* Controller = GetController())
@@ -6526,42 +6667,7 @@ FRigBaseComponent* URigHierarchy::MakeComponent(const UScriptStruct* InComponent
 		ComponentIndex = ElementComponents.Emplace(InComponentStruct);
 	}
 
-	FRigName SanitizedName = InName;
-	SanitizeName(SanitizedName, false);
-
-	FRigName UniqueName = SanitizedName;
-
-	TArray<FName> ExistingNames;
-	if(InElement)
-	{
-		ExistingNames.Reserve(InElement->NumComponents());
-		for(int32 ExistingComponentIndex = 0; ExistingComponentIndex < InElement->NumComponents(); ExistingComponentIndex++)
-		{
-			ExistingNames.Add(InElement->GetComponent(ExistingComponentIndex)->GetFName());
-		}
-	}
-	else
-	{
-		ExistingNames.Reserve(NumTopLevelComponents());
-		for(int32 ExistingComponentIndex = 0; ExistingComponentIndex < NumTopLevelComponents(); ExistingComponentIndex++)
-		{
-			ExistingNames.Add(GetTopLevelComponentName(ExistingComponentIndex));
-		}
-	}
-
-	int32 Suffix = 1;
-	while (ExistingNames.ContainsByPredicate([UniqueName](const FName& ExistingName) -> bool
-	{
-		return ExistingName == UniqueName.GetFName();
-	}))
-	{
-		FString BaseString = SanitizedName.GetName();
-		if (BaseString.Len() > GetMaxNameLength() - 4)
-		{
-			BaseString.LeftChopInline(BaseString.Len() - (GetMaxNameLength() - 4));
-		}
-		UniqueName.SetName(FString::Printf(TEXT("%s_%d"), *BaseString, ++Suffix));
-	}
+	const FRigName UniqueName = GetSafeNewComponentName(InElement ? InElement->GetKey() : GetTopLevelComponentElementKey(), InName);
 	
 	FRigBaseComponent* Component = ElementComponents[ComponentIndex].GetMutablePtr<FRigBaseComponent>();
 	Component->IndexInHierarchy = ComponentIndex;
diff --git a/Engine/Plugins/Animation/ControlRig/Source/ControlRig/Private/Rigs/RigHierarchyComponents.cpp b/Engine/Plugins/Animation/ControlRig/Source/ControlRig/Private/Rigs/RigHierarchyComponents.cpp
--- a/Engine/Plugins/Animation/ControlRig/Source/ControlRig/Private/Rigs/RigHierarchyComponents.cpp
+++ b/Engine/Plugins/Animation/ControlRig/Source/ControlRig/Private/Rigs/RigHierarchyComponents.cpp
@@ -1,11 +1,29 @@
 // Copyright Epic Games, Inc. All Rights Reserved.
 
 #include "Rigs/RigHierarchyComponents.h"
+
+#include "Rigs/RigHierarchy.h"
 #include "Rigs/RigHierarchyElements.h"
 #include "UObject/UObjectIterator.h"
+#include "Serialization/MemoryReader.h"
+#include "Serialization/MemoryWriter.h"
 
 #include UE_INLINE_GENERATED_CPP_BY_NAME(RigHierarchyComponents)
 
+////////////////////////////////////////////////////////////////////////////////
+// FRigComponentState
+////////////////////////////////////////////////////////////////////////////////
+
+bool FRigComponentState::IsValid() const
+{
+	return ComponentStruct != nullptr && !Data.IsEmpty();
+}
+
+const UScriptStruct* FRigComponentState::GetComponentStruct() const
+{
+	return ComponentStruct;
+}
+
 ////////////////////////////////////////////////////////////////////////////////
 // FRigBaseComponent
 ////////////////////////////////////////////////////////////////////////////////
@@ -14,6 +32,27 @@ FRigBaseComponent::~FRigBaseComponent()
 {
 }
 
+FName FRigBaseComponent::GetDefaultComponentName() const
+{
+#if WITH_EDITORONLY_DATA
+	return *GetScriptStruct()->GetDisplayNameText().ToString();
+#else
+	return TEXT("Component");
+#endif
+}
+
+const FSlateIcon& FRigBaseComponent::GetIconForUI() const
+{
+	// todo
+	static const FSlateIcon ComponentIcon = FSlateIcon(TEXT("ControlRigEditorStyle"), TEXT("ControlRig.Tree.RigidBody"));
+	return ComponentIcon;
+}
+
+FSlateColor FRigBaseComponent::GetColorForUI() const
+{
+	return FSlateColor::UseForeground();
+}
+
 TArray<UScriptStruct*> FRigBaseComponent::GetAllComponentScriptStructs(bool bSorted)
 {
 	TArray<UScriptStruct*> RigComponentStructs;
@@ -63,6 +102,26 @@ void FRigBaseComponent::Load(FArchive& Ar)
 	CachedNameString.Reset();
 }
 
+FRigComponentState FRigBaseComponent::GetState() const
+{
+	FRigComponentState State;
+	State.ComponentStruct = GetScriptStruct();
+	FMemoryWriter Writer(State.Data);
+	const_cast<FRigBaseComponent*>(this)->Save(Writer);
+	return State;
+}
+
+bool FRigBaseComponent::SetState(const FRigComponentState& InState)
+{
+	if(InState.IsValid() && InState.GetComponentStruct() == GetScriptStruct())
+	{
+		FMemoryReader Reader(InState.Data);
+		Load(Reader);
+		return true;
+	}
+	return false;
+}
+
 FString FRigBaseComponent::GetContentAsText() const
 {
 	FString Content;
diff --git a/Engine/Plugins/Animation/ControlRig/Source/ControlRig/Private/Rigs/RigHierarchyController.cpp b/Engine/Plugins/Animation/ControlRig/Source/ControlRig/Private/Rigs/RigHierarchyController.cpp
--- a/Engine/Plugins/Animation/ControlRig/Source/ControlRig/Private/Rigs/RigHierarchyController.cpp
+++ b/Engine/Plugins/Animation/ControlRig/Source/ControlRig/Private/Rigs/RigHierarchyController.cpp
@@ -943,12 +943,6 @@ FRigComponentKey URigHierarchyController::AddComponent(UScriptStruct* InComponen
 		return FRigComponentKey();
 	}
 
-	if(InName.IsNone())
-	{
-		ReportError(TEXT("The passed name is None."));
-		return FRigComponentKey();
-	}
-
 	URigHierarchy* Hierarchy = GetHierarchy();
 	check(Hierarchy);
 
@@ -963,45 +957,11 @@ FRigComponentKey URigHierarchyController::AddComponent(UScriptStruct* InComponen
 		}
 	}
 
-	const FStructOnScope StructOnScope(InComponentStruct);
-	if(const FRigBaseComponent* StructMemory = reinterpret_cast<const FRigBaseComponent*>(StructOnScope.GetStructMemory()))
+	FString FailureReason;
+	if(!Hierarchy->CanAddComponent(InElement, InComponentStruct, &FailureReason))
 	{
-		if(StructMemory->IsSingleton())
-		{
-			if(Element)
-			{
-				for(int32 Index = 0; Index < Element->NumComponents(); Index++)
-				{
-					if(const FRigBaseComponent* ExistingComponent = Element->GetComponent(Index))
-					{
-						if(ExistingComponent->GetScriptStruct() == InComponentStruct)
-						{
-							ReportErrorf(TEXT("A component of type '%s' already exists on element '%s'."), *InComponentStruct->GetName(), *InElement.ToString());
-							return FRigComponentKey();
-						}
-					}
-				}
-			}
-			else
-			{
-				for(int32 Index = 0; Index < Hierarchy->NumTopLevelComponents(); Index++)
-				{
-					if(const FRigBaseComponent* ExistingComponent = Hierarchy->GetTopLevelComponent(Index))
-					{
-						if(ExistingComponent->GetScriptStruct() == InComponentStruct)
-						{
-							ReportErrorf(TEXT("A component of type '%s' already exists at the top level."), *InComponentStruct->GetName());
-							return FRigComponentKey();
-						}
-					}
-				}
-			}
-		}
-		if(!StructMemory->CanBeAddedTo(InElement, Hierarchy))
-		{
-			ReportErrorf(TEXT("The component of type '%s' cannot be added to element '%s'."), *InComponentStruct->GetName(), *InElement.ToString());
-			return FRigComponentKey();
-		}
+		ReportError(FailureReason);
+		return FRigComponentKey();
 	}
 
 #if WITH_EDITOR
@@ -1013,6 +973,16 @@ FRigComponentKey URigHierarchyController::AddComponent(UScriptStruct* InComponen
 	}
 #endif
 
+	if(InName.IsNone())
+	{
+		const FStructOnScope StructOnScope(InComponentStruct);
+		if(const FRigBaseComponent* StructMemory = reinterpret_cast<const FRigBaseComponent*>(StructOnScope.GetStructMemory()))
+		{
+			InName = StructMemory->GetDefaultComponentName();
+			check(!InName.IsNone());
+		}
+	}
+
 	FRigBaseComponent* NewComponent = Hierarchy->MakeComponent(InComponentStruct, InName, Element);
 	NewComponent->CreatedAtInstructionIndex = CurrentInstructionIndex;
 	if(!InContent.IsEmpty())
@@ -1100,6 +1070,121 @@ bool URigHierarchyController::RemoveComponent(FRigComponentKey InComponent, bool
 	return true;
 }
 
+FRigComponentKey URigHierarchyController::RenameComponent(FRigComponentKey InComponent, FName InName, bool bSetupUndo, bool bPrintPythonCommand,
+	bool bClearSelection)
+{
+	if(!IsValid())
+	{
+		return FRigComponentKey();
+	}
+
+	URigHierarchy* Hierarchy = GetHierarchy();
+
+	FRigBaseComponent* Component = Hierarchy->FindComponent(InComponent);
+	if(Component == nullptr)
+	{
+		ReportWarningf(TEXT("Cannot Rename Component: '%s' not found."), *InComponent.ToString());
+		return FRigComponentKey();
+	}
+
+#if WITH_EDITOR
+	TSharedPtr<FScopedTransaction> TransactionPtr;
+	if(bSetupUndo)
+	{
+		TransactionPtr = MakeShared<FScopedTransaction>(NSLOCTEXT("RigHierarchyController", "Rename Component", "Rename Component"));
+		Hierarchy->Modify();
+	}
+#endif
+
+	const bool bRenamed = RenameComponent(Component, InName, bClearSelection);
+
+#if WITH_EDITOR
+	if(!bRenamed && TransactionPtr.IsValid())
+	{
+		TransactionPtr->Cancel();
+	}
+	TransactionPtr.Reset();
+
+	if (bRenamed && bClearSelection)
+	{
+		ClearSelection();
+	}
+
+	if (bRenamed && bPrintPythonCommand && !bSuspendPythonPrinting)
+	{
+		UBlueprint* Blueprint = GetTypedOuter<UBlueprint>();
+		if (Blueprint)
+		{
+			RigVMPythonUtils::Print(Blueprint->GetFName().ToString(), 
+				FString::Printf(TEXT("hierarchy_controller.rename_component(%s, '%s')"),
+				*InComponent.ToPythonString(),
+				*InName.ToString()));
+		}
+	}
+#endif
+
+	return bRenamed ? Component->GetKey() : FRigComponentKey();
+}
+
+FRigComponentKey URigHierarchyController::ReparentComponent(FRigComponentKey InComponentKey, FRigElementKey InParentElementKey, bool bSetupUndo,
+	bool bPrintPythonCommand, bool bClearSelection)
+{
+	if(!IsValid())
+	{
+		return FRigComponentKey();
+	}
+
+	URigHierarchy* Hierarchy = GetHierarchy();
+
+	FRigBaseComponent* Component = Hierarchy->FindComponent(InComponentKey);
+	if(Component == nullptr)
+	{
+		ReportWarningf(TEXT("Cannot Reparent Component: '%s' not found."), *InComponentKey.ToString());
+		return FRigComponentKey();
+	}
+
+	// it's ok if this is nullptr
+	FRigBaseElement* ParentElement = Hierarchy->Find(InParentElementKey);
+
+#if WITH_EDITOR
+	TSharedPtr<FScopedTransaction> TransactionPtr;
+	if(bSetupUndo)
+	{
+		TransactionPtr = MakeShared<FScopedTransaction>(NSLOCTEXT("RigHierarchyController", "Reparent Component", "Reparent Component"));
+		Hierarchy->Modify();
+	}
+#endif
+
+	const bool bReparented = ReparentComponent(Component, ParentElement, bClearSelection);
+
+#if WITH_EDITOR
+	if(!bReparented && TransactionPtr.IsValid())
+	{
+		TransactionPtr->Cancel();
+	}
+	TransactionPtr.Reset();
+
+	if (bReparented && bClearSelection)
+	{
+		ClearSelection();
+	}
+
+	if (bReparented && bPrintPythonCommand && !bSuspendPythonPrinting)
+	{
+		UBlueprint* Blueprint = GetTypedOuter<UBlueprint>();
+		if (Blueprint)
+		{
+			RigVMPythonUtils::Print(Blueprint->GetFName().ToString(), 
+				FString::Printf(TEXT("hierarchy_controller.reparent_component(%s, '%s')"),
+				*InComponentKey.ToPythonString(),
+				*InParentElementKey.ToPythonString()));
+		}
+	}
+#endif
+
+	return bReparented ? Component->GetKey() : FRigComponentKey();
+}
+
 bool URigHierarchyController::SetComponentContent(FRigComponentKey InComponent, const FString& InContent, bool bSetupUndo, bool bPrintPythonCommand)
 {
 	if(!IsValid())
@@ -1169,6 +1254,69 @@ bool URigHierarchyController::SetComponentContent(FRigComponentKey InComponent,
 	return true;
 }
 
+bool URigHierarchyController::SetComponentState(FRigComponentKey InComponent, const FRigComponentState& InState, bool bSetupUndo)
+{
+	if(!IsValid())
+	{
+		return false;
+	}
+	
+	if(!InComponent.IsValid())
+	{
+		ReportError(TEXT("The passed component key is valid."));
+		return false;
+	}
+
+	if(!InState.IsValid())
+	{
+		ReportError(TEXT("The passed content is not valid."));
+		return false;
+	}
+
+	URigHierarchy* Hierarchy = GetHierarchy();
+	check(Hierarchy);
+	
+	FRigBaseComponent* Component = Hierarchy->FindComponent(InComponent);
+	if(Component == nullptr)
+	{
+		ReportErrorf(TEXT("The component '%s' cannot be found."), *InComponent.ToString());
+		return false;
+	}
+
+	if(InState.GetComponentStruct() != Component->GetScriptStruct())
+	{
+		ReportErrorf(TEXT("The passed content(%s) does not match the component(%s)."), *InState.GetComponentStruct()->GetName(), *Component->GetScriptStruct()->GetName());
+		return false;
+	}
+
+#if WITH_EDITOR
+	TSharedPtr<FScopedTransaction> TransactionPtr;
+	if(bSetupUndo)
+	{
+		TransactionPtr = MakeShared<FScopedTransaction>(NSLOCTEXT("RigHierarchyController", "Set Component State", "Set Component State"));
+		Hierarchy->Modify();
+	}
+#endif
+
+	if(!Component->SetState(InState))
+	{
+#if WITH_EDITOR
+		if(TransactionPtr)
+		{
+			TransactionPtr->Cancel();
+		}
+#endif
+		return false;
+	}
+	Notify(ERigHierarchyNotification::ComponentContentChanged, Component);
+
+#if WITH_EDITOR
+	TransactionPtr.Reset();
+#endif
+
+	return true;
+}
+
 TArray<FRigElementKey> URigHierarchyController::ImportBones(
 	USkeletalMesh* InSkeletalMesh,
 	const FName& InNameSpace,
@@ -3038,7 +3186,7 @@ bool URigHierarchyController::RenameElement(FRigBaseElement* InElement, const FN
 	// deselect the key that no longer exists
 	// no need to trigger a reselect since we always clear selection after rename
 	const bool bWasSelected = Hierarchy->IsSelected(InElement); 
-	if (Hierarchy->IsSelected(InElement))
+	if (bWasSelected)
 	{
 		DeselectElement(OldKey);
 	}
@@ -3118,6 +3266,145 @@ bool URigHierarchyController::RenameElement(FRigBaseElement* InElement, const FN
 	return true;
 }
 
+bool URigHierarchyController::RenameComponent(FRigBaseComponent* InComponent, const FName& InName, bool bClearSelection)
+{
+	if(InComponent == nullptr)
+	{
+		return false;
+	}
+
+	if (InComponent->GetFName().IsEqual(InName, ENameCase::CaseSensitive))
+	{
+		return false;
+	}
+
+	const FRigComponentKey OldKey = InComponent->GetKey();
+
+	URigHierarchy* Hierarchy = GetHierarchy();
+
+	// deselect the key that no longer exists
+	// no need to trigger a reselect since we always clear selection after rename
+	const bool bWasSelected = Hierarchy->IsComponentSelected(InComponent); 
+	if (bWasSelected)
+	{
+		DeselectComponent(OldKey);
+	}
+
+	{
+		// create a temp copy of the map and remove the current item's key
+		TMap<FRigComponentKey, int32> TemporaryMap = Hierarchy->ComponentIndexLookup;
+		TemporaryMap.Remove(OldKey);
+   
+		TGuardValue<TMap<FRigComponentKey, int32>> MapGuard(Hierarchy->ComponentIndexLookup, TemporaryMap);
+		InComponent->Key.Name = Hierarchy->GetSafeNewComponentName(OldKey.ElementKey, InName);
+		InComponent->CachedNameString.Reset();
+	}
+	
+	const FRigComponentKey NewKey = InComponent->GetKey();
+
+	Hierarchy->ComponentIndexLookup.Remove(OldKey);
+	Hierarchy->ComponentIndexLookup.Add(NewKey, InComponent->IndexInHierarchy);
+
+	Hierarchy->IncrementTopologyVersion();
+	Notify(ERigHierarchyNotification::ComponentRenamed, InComponent);
+
+	if (!bClearSelection && bWasSelected)
+	{
+		SelectComponent(InComponent->GetKey(), true);
+	}
+
+	return true;
+}
+
+bool URigHierarchyController::ReparentComponent(FRigBaseComponent* InComponent, FRigBaseElement* InParentElement, bool bClearSelection)
+{
+	if(InComponent == nullptr)
+	{
+		return false;
+	}
+
+	if(InComponent->Element == InParentElement)
+	{
+		return false;
+	}
+
+	const FRigElementKey NewParentKey = InParentElement ? InParentElement->GetKey() : URigHierarchy::GetTopLevelComponentElementKey();
+
+	URigHierarchy* Hierarchy = GetHierarchy();
+	if(!Hierarchy->CanAddComponent(NewParentKey, InComponent))
+	{
+		return false;
+	}
+
+	const FRigComponentKey OldKey = InComponent->GetKey();
+
+	// deselect the key that no longer exists
+	// no need to trigger a reselect since we always clear selection after rename
+	const bool bWasSelected = Hierarchy->IsComponentSelected(InComponent); 
+	if (bWasSelected)
+	{
+		DeselectComponent(OldKey);
+	}
+
+	if(OldKey.ElementKey == URigHierarchy::GetTopLevelComponentElementKey())
+	{
+		verify(Hierarchy->TopLevelComponentIndices.Remove(InComponent->IndexInHierarchy) > 0);
+	}
+	else if(FRigBaseElement* OldParentElement = Hierarchy->Find(OldKey.ElementKey))
+	{
+		OldParentElement->ComponentIndices.Remove(InComponent->IndexInHierarchy);
+	}
+
+	InComponent->Key.Name = Hierarchy->GetSafeNewComponentName(NewParentKey, InComponent->Key.Name);
+	InComponent->Key.ElementKey = NewParentKey;
+	InComponent->CachedNameString.Reset();
+
+	if(NewParentKey == URigHierarchy::GetTopLevelComponentElementKey())
+	{
+		check(!Hierarchy->TopLevelComponentIndices.Contains(InComponent->IndexInHierarchy));
+		Hierarchy->TopLevelComponentIndices.Add(InComponent->IndexInHierarchy);
+		InComponent->Element = nullptr;
+		
+		for(int32 IndexInElement = 0; IndexInElement < Hierarchy->TopLevelComponentIndices.Num(); IndexInElement++)
+		{
+			if(FRigBaseComponent* RemainingComponent = Hierarchy->GetComponent(Hierarchy->TopLevelComponentIndices[IndexInElement]))
+			{
+				RemainingComponent->IndexInElement = IndexInElement;
+			}
+		}
+	}
+	else
+	{
+		check(InParentElement);
+		check(!InParentElement->ComponentIndices.Contains(InComponent->IndexInHierarchy));
+		InParentElement->ComponentIndices.Add(InComponent->IndexInHierarchy);
+		InComponent->Element = InParentElement;
+
+		for(int32 IndexInElement = 0; IndexInElement < InParentElement->ComponentIndices.Num(); IndexInElement++)
+		{
+			if(FRigBaseComponent* RemainingComponent = Hierarchy->GetComponent(InParentElement->ComponentIndices[IndexInElement]))
+			{
+				RemainingComponent->IndexInElement = IndexInElement;
+			}
+		}
+	}
+	
+	const FRigComponentKey NewKey = InComponent->GetKey();
+
+	Hierarchy->ComponentIndexLookup.Remove(OldKey);
+	Hierarchy->ComponentIndexLookup.Add(NewKey, InComponent->IndexInHierarchy);
+
+	Hierarchy->IncrementTopologyVersion();
+	Notify(ERigHierarchyNotification::ComponentReparented, InComponent);
+
+	if (!bClearSelection && bWasSelected)
+	{
+		SelectComponent(InComponent->GetKey(), true);
+	}
+
+	return true;
+}
+
 bool URigHierarchyController::ReorderElement(FRigBaseElement* InElement, int32 InIndex)
 {
 	if(InElement == nullptr)
diff --git a/Engine/Plugins/Animation/ControlRig/Source/ControlRig/Private/Units/Hierarchy/RigUnit_Component.cpp b/Engine/Plugins/Animation/ControlRig/Source/ControlRig/Private/Units/Hierarchy/RigUnit_Component.cpp
--- a/Engine/Plugins/Animation/ControlRig/Source/ControlRig/Private/Units/Hierarchy/RigUnit_Component.cpp
+++ b/Engine/Plugins/Animation/ControlRig/Source/ControlRig/Private/Units/Hierarchy/RigUnit_Component.cpp
@@ -75,7 +75,7 @@ FText FRigDispatch_ComponentBase::GetArgumentTooltip(const FName& InArgumentName
 {
 	if(InArgumentName == NameArgName)
 	{
-		return NSLOCTEXT("FRigDispatch_ComponentBase", "NameArgTooltip", "The name of the component");
+		return NSLOCTEXT("FRigDispatch_ComponentBase", "NameArgTooltip", "The name of the component (can be empty)");
 	}
 	if(InArgumentName == ItemArgName)
 	{
@@ -100,7 +100,7 @@ FString FRigDispatch_ComponentBase::GetArgumentDefaultValue(const FName& InArgum
 {
 	if(InArgumentName == NameArgName)
 	{
-		return TEXT("Component");
+		return FString();
 	}
 	if(InArgumentName == ItemArgName)
 	{
diff --git a/Engine/Plugins/Animation/ControlRig/Source/ControlRig/Public/Rigs/RigHierarchy.h b/Engine/Plugins/Animation/ControlRig/Source/ControlRig/Public/Rigs/RigHierarchy.h
--- a/Engine/Plugins/Animation/ControlRig/Source/ControlRig/Public/Rigs/RigHierarchy.h
+++ b/Engine/Plugins/Animation/ControlRig/Source/ControlRig/Public/Rigs/RigHierarchy.h
@@ -936,6 +936,9 @@ public:
 	UFUNCTION(BlueprintPure, Category = URigHierarchy)
 	FString GetTopLevelComponentContent(int32 InTopLevelComponentIndex) const;
 
+	bool CanAddComponent(FRigElementKey InElementKey, const UScriptStruct* InComponentStruct, FString* OutFailureReason = nullptr) const;
+	bool CanAddComponent(FRigElementKey InElementKey, const FRigBaseComponent* InComponent, FString* OutFailureReason = nullptr) const;
+
 	static const FRigElementKey& GetTopLevelComponentElementKey();
 
 private:
@@ -2064,6 +2067,15 @@ public:
 	 */
 	bool IsDisplayNameAvailable(const FRigElementKey& InParentElement, const FRigName& InPotentialNewDisplayName, FString* OutErrorMessage = nullptr) const;
 
+	/**
+	 * Returns true if a given component name is available.
+	 * @param InElementKey The element to check the component name for
+	 * @param InPotentialNewName The name to test for availability
+	 * @param OutErrorMessage An optional pointer to return a potential error message 
+	 * @return Returns true if the name is available.
+	 */
+	bool IsComponentNameAvailable(const FRigElementKey& InElementKey, const FRigName& InPotentialNewName, FString* OutErrorMessage = nullptr) const;
+
 	/**
 	 * Returns a valid new name for a to-be-added element.
 	 * @param InPotentialNewName The name to be sanitized and adjusted for availability
@@ -2081,6 +2093,14 @@ public:
 	 */
 	FRigName GetSafeNewDisplayName(const FRigElementKey& InParentElement, const FRigName& InPotentialNewDisplayName) const;
 
+	/**
+	 * Returns a valid new name for a to-be-added element.
+	 * @param InElementKey The element to check the component name for
+	 * @param InPotentialNewName The name to be sanitized and adjusted for availability
+	 * @return Returns the name to use for the to-be-added element.
+	 */
+	FRigName GetSafeNewComponentName(const FRigElementKey& InElementKey, const FRigName& InPotentialNewName) const;
+
 	/**
 	 * Returns the display label for an element to be used for the UI
 	 */
diff --git a/Engine/Plugins/Animation/ControlRig/Source/ControlRig/Public/Rigs/RigHierarchyComponents.h b/Engine/Plugins/Animation/ControlRig/Source/ControlRig/Public/Rigs/RigHierarchyComponents.h
--- a/Engine/Plugins/Animation/ControlRig/Source/ControlRig/Public/Rigs/RigHierarchyComponents.h
+++ b/Engine/Plugins/Animation/ControlRig/Source/ControlRig/Public/Rigs/RigHierarchyComponents.h
@@ -4,6 +4,8 @@
 
 #include "CoreMinimal.h"
 #include "RigHierarchyDefines.h"
+#include "Textures/SlateIcon.h"
+#include "Styling/SlateColor.h"
 #include "RigHierarchyComponents.generated.h"
 
 #define DECLARE_RIG_COMPONENT_METHODS(ComponentType) \
@@ -33,6 +35,21 @@ static bool IsClassOf(const FRigBaseComponent* InComponent) \
 	return InComponent->GetScriptStruct()->IsChildOf(StaticStruct()); \
 }
 
+struct CONTROLRIG_API FRigComponentState
+{
+public:
+
+	bool IsValid() const;
+	const UScriptStruct* GetComponentStruct() const;
+
+private:
+
+	const UScriptStruct* ComponentStruct = nullptr;
+	TArray<uint8> Data;
+
+	friend struct FRigBaseComponent;
+};
+
 USTRUCT(BlueprintType)
 struct CONTROLRIG_API FRigBaseComponent
 {
@@ -60,13 +77,25 @@ public:
 		return *this;
 	}
 
+	// returns the default name to use when instantiating the component.
+	virtual FName GetDefaultComponentName() const;
+
+	// returns true if this component can be renamed
+	virtual bool CanBeRenamed() const { return true; }
+
 	// returns true if this component can only be added once
 	virtual bool IsSingleton() const { return false; }
 
 	// returns true if this component can be added to a given key.
 	// if you want to determine if a component can be added as a top level component
 	// you need to pass the URigHierarchy::GetTopLevelComponentElementKey() as InElementKey
-	virtual bool CanBeAddedTo(const FRigElementKey& InElementKey, const URigHierarchy* InHierarchy) const { return true; }
+	virtual bool CanBeAddedTo(const FRigElementKey& InElementKey, const URigHierarchy* InHierarchy, FString* OutFailureReason = nullptr) const { return true; }
+
+	// return the icon to use for this component in the UI
+	virtual const FSlateIcon& GetIconForUI() const;
+
+	// returns the color to use for this component in the UI
+	virtual FSlateColor GetColorForUI() const;
 	
 protected:
 
@@ -121,6 +150,9 @@ public:
 	virtual void Save(FArchive& Ar);
 	virtual void Load(FArchive& Ar);
 
+	FRigComponentState GetState() const;
+	bool SetState(const FRigComponentState& InState);
+
 	FString GetContentAsText() const;
 
 	bool IsA(const UScriptStruct* InScriptStruct) const
diff --git a/Engine/Plugins/Animation/ControlRig/Source/ControlRig/Public/Rigs/RigHierarchyController.h b/Engine/Plugins/Animation/ControlRig/Source/ControlRig/Public/Rigs/RigHierarchyController.h
--- a/Engine/Plugins/Animation/ControlRig/Source/ControlRig/Public/Rigs/RigHierarchyController.h
+++ b/Engine/Plugins/Animation/ControlRig/Source/ControlRig/Public/Rigs/RigHierarchyController.h
@@ -356,6 +356,28 @@ public:
     UFUNCTION(BlueprintCallable, Category = URigHierarchyController)
     bool RemoveComponent(FRigComponentKey InComponent, bool bSetupUndo = false, bool bPrintPythonCommand = false);
 
+	/**
+	 * Renames an existing component in the hierarchy
+	 * @param InComponent The key of the component to rename
+	 * @param InName The new name to set for the component
+	 * @param bSetupUndo If set to true the stack will record the change for undo / redo
+	 * @param bPrintPythonCommand If set to true a python command equivalent to this call will be printed out
+	 * @param bClearSelection True if the selection should be cleared after a rename
+	 * @return Returns the new key used for the component
+	 */
+	UFUNCTION(BlueprintCallable, Category = URigHierarchyController)
+	FRigComponentKey RenameComponent(FRigComponentKey InComponent, FName InName, bool bSetupUndo = false, bool bPrintPythonCommand = false, bool bClearSelection = true);
+
+	/**
+	 * Reparents an existing component in the hierarchy
+	 * @param InComponentKey The component key to reparent
+	 * @param InParentElementKey The new element key to reparent to
+	 * @param bClearSelection True if the selection should be cleared after a reparenting
+	 * @return Returns the new component key if successful or an invalid key if unsuccessful
+	 */
+	UFUNCTION(BlueprintCallable, Category = URigHierarchyController)
+	FRigComponentKey ReparentComponent(FRigComponentKey InComponentKey, FRigElementKey InParentElementKey, bool bSetupUndo = false, bool bPrintPythonCommand = false, bool bClearSelection = true);
+
 	/**
 	 * Updates the content of a component in the hierarchy
 	 * @param InComponent The component to change the content for
@@ -367,6 +389,15 @@ public:
 	UFUNCTION(BlueprintCallable, Category = URigHierarchyController)
 	bool SetComponentContent(FRigComponentKey InComponent, const FString& InContent, bool bSetupUndo = false, bool bPrintPythonCommand = false);
 
+	/**
+	 * Updates the state of a component in the hierarchy
+	 * @param InComponent The component to change the content for
+	 * @param InState The serialized state of the component
+	 * @param bSetupUndo If set to true the stack will record the change for undo / redo
+	 * @return True if the component was updated correctly
+	 */
+	bool SetComponentState(FRigComponentKey InComponent, const FRigComponentState& InState, bool bSetupUndo = false);
+
 	/**
 	 * Imports an existing skeleton to the hierarchy
 	 * @param InSkeleton The skeleton to import
@@ -943,6 +974,24 @@ private:
 	 */
     bool RenameElement(FRigBaseElement* InElement, const FName &InName, bool bClearSelection = true);
 
+	/**
+	 * Renames an existing component in the hierarchy
+	 * @param InComponent The component to rename
+	 * @param InName The new name to set for the component
+	 * @param bClearSelection True if the selection should be cleared after a rename
+	 * @return Returns true if successful.
+	 */
+	bool RenameComponent(FRigBaseComponent* InComponent, const FName &InName, bool bClearSelection = true);
+
+	/**
+	 * Reparents an existing component in the hierarchy
+	 * @param InComponent The component to reparent
+	 * @param InParentElement The new element to reparent to
+	 * @param bClearSelection True if the selection should be cleared after a reparenting
+	 * @return Returns true if successful.
+	 */
+	bool ReparentComponent(FRigBaseComponent* InComponent, FRigBaseElement* InParentElement, bool bClearSelection = true);
+
 	/**
 	 * Changes the element's index within its default parent (or the top level)
 	 * @param InElement The key of the element to rename
diff --git a/Engine/Plugins/Animation/ControlRig/Source/ControlRig/Public/Rigs/RigHierarchyDefines.h b/Engine/Plugins/Animation/ControlRig/Source/ControlRig/Public/Rigs/RigHierarchyDefines.h
--- a/Engine/Plugins/Animation/ControlRig/Source/ControlRig/Public/Rigs/RigHierarchyDefines.h
+++ b/Engine/Plugins/Animation/ControlRig/Source/ControlRig/Public/Rigs/RigHierarchyDefines.h
@@ -106,6 +106,8 @@ enum class ERigHierarchyNotification : uint8
 	ComponentContentChanged,
 	ComponentSelected,
 	ComponentDeselected,
+	ComponentRenamed,
+	ComponentReparented,
 
 	/** MAX - invalid */
 	Max UMETA(Hidden),
@@ -1777,17 +1779,17 @@ public:
 	{
 	}
 
-	FRigHierarchyKey(const FRigElementKey& InElementKey)
+	FRigHierarchyKey(const FRigElementKey& InElementKey, bool bForce = false)
 	{
-		if(InElementKey.IsValid())
+		if(InElementKey.IsValid() || bForce)
 		{
 			Element = InElementKey;
 		}
 	}
 	
-	FRigHierarchyKey(const FRigComponentKey& InComponentKey)
+	FRigHierarchyKey(const FRigComponentKey& InComponentKey, bool bForce = false)
 	{
-		if(InComponentKey.IsValid())
+		if(InComponentKey.IsValid() || bForce)
 		{
 			Component = InComponentKey;
 		}
diff --git a/Engine/Plugins/Animation/ControlRig/Source/ControlRigDeveloper/Private/ControlRigBlueprint.cpp b/Engine/Plugins/Animation/ControlRig/Source/ControlRigDeveloper/Private/ControlRigBlueprint.cpp
--- a/Engine/Plugins/Animation/ControlRig/Source/ControlRigDeveloper/Private/ControlRigBlueprint.cpp
+++ b/Engine/Plugins/Animation/ControlRig/Source/ControlRigDeveloper/Private/ControlRigBlueprint.cpp
@@ -2762,6 +2762,8 @@ void UControlRigBlueprint::HandleHierarchyModified(ERigHierarchyNotification InN
 		case ERigHierarchyNotification::HierarchyReset:
 		case ERigHierarchyNotification::ComponentAdded:
 		case ERigHierarchyNotification::ComponentRemoved:
+		case ERigHierarchyNotification::ComponentRenamed:
+		case ERigHierarchyNotification::ComponentReparented:
 		case ERigHierarchyNotification::ComponentContentChanged:
 		{
 			Modify();
diff --git a/Engine/Plugins/Animation/ControlRig/Source/ControlRigEditor/Private/ControlRigElementDetails.cpp b/Engine/Plugins/Animation/ControlRig/Source/ControlRigEditor/Private/ControlRigElementDetails.cpp
--- a/Engine/Plugins/Animation/ControlRig/Source/ControlRigEditor/Private/ControlRigElementDetails.cpp
+++ b/Engine/Plugins/Animation/ControlRig/Source/ControlRigEditor/Private/ControlRigElementDetails.cpp
@@ -6740,7 +6740,10 @@ void FRigBaseComponentDetails::CustomizeDetails(IDetailLayoutBuilder& DetailBuil
 			
 			for(TSharedRef<IPropertyHandle>& Property : Properties)
 			{
-				DetailBuilder.EditDefaultProperty(Property)->IsEnabled(false);
+				if(IDetailPropertyRow* Row = DetailBuilder.EditDefaultProperty(Property))
+				{
+					Row->IsEnabled(false);
+				}
 			}
 		}
 	}
diff --git a/Engine/Plugins/Animation/ControlRig/Source/ControlRigEditor/Private/Editor/ControlRigEditor.cpp b/Engine/Plugins/Animation/ControlRig/Source/ControlRigEditor/Private/Editor/ControlRigEditor.cpp
--- a/Engine/Plugins/Animation/ControlRig/Source/ControlRigEditor/Private/Editor/ControlRigEditor.cpp
+++ b/Engine/Plugins/Animation/ControlRig/Source/ControlRigEditor/Private/Editor/ControlRigEditor.cpp
@@ -63,6 +63,7 @@
 #include "Units/Hierarchy/RigUnit_GetControlTransform.h"
 #include "Units/Hierarchy/RigUnit_SetCurveValue.h"
 #include "Units/Hierarchy/RigUnit_AddBoneTransform.h"
+#include "Units/Hierarchy/RigUnit_Component.h"
 #include "EdGraph/NodeSpawners/RigVMEdGraphUnitNodeSpawner.h"
 #include "Graph/ControlRigGraphSchema.h"
 #include "ControlRigObjectVersion.h"
@@ -3892,8 +3893,8 @@ void FControlRigEditor::OnWrappedPropertyChangedChainEvent(URigVMDetailsViewWrap
 			InWrapperObject->GetContent(Content.GetStructMemory(), Content.GetStruct());
 
 			const FRigBaseComponent* WrappedComponent = reinterpret_cast<const FRigBaseComponent*>(Content.GetStructMemory());
-			FString TextContent = WrappedComponent->GetContentAsText();
-			(void)ControlRigBP->GetHierarchyController()->SetComponentContent(WrappedComponent->GetKey(), TextContent, true, true);
+			const FRigComponentState State = WrappedComponent->GetState();
+			(void)ControlRigBP->GetHierarchyController()->SetComponentState(WrappedComponent->GetKey(), State, true);
 		}
 	}
 }
@@ -4481,6 +4482,8 @@ void FControlRigEditor::OnHierarchyModified(ERigHierarchyNotification InNotif, U
 		}
 		case ERigHierarchyNotification::ComponentAdded:
 		case ERigHierarchyNotification::ComponentRemoved:
+		case ERigHierarchyNotification::ComponentRenamed:
+		case ERigHierarchyNotification::ComponentReparented:
 		{
 			OnHierarchyChanged();
 			break;
@@ -5021,288 +5024,329 @@ void FControlRigEditor::CreateRigHierarchyToGraphDragAndDropMenu() const
 
 					URigHierarchy* Hierarchy = ControlRigEditor->GetHierarchyBeingDebugged();
 					TArray<FRigElementKey> DraggedElements;
-					DraggedElements.Reserve(DragDropContext.DraggedHierarchyKeys.Num());
+					TArray<FRigComponentKey> DraggedComponents;
 					for(const FRigHierarchyKey& DraggedHierarchyKey : DragDropContext.DraggedHierarchyKeys)
 					{
 						if(DraggedHierarchyKey.IsElement())
 						{
 							DraggedElements.Add(DraggedHierarchyKey.GetElement());
 						}
+						if(DraggedHierarchyKey.IsComponent())
+						{
+							DraggedComponents.Add(DraggedHierarchyKey.GetComponent());
+						}
 					}
-
-					// todo: dropping components won't work for now.
-					
-					ControlRigEditor->FilterDraggedKeys(DraggedElements, true);
 					
 					UEdGraph* Graph = DragDropContext.Graph.Get();
 					const FVector2D& NodePosition = DragDropContext.NodePosition;
-					
-					// if multiple types are selected, we show Get Elements/Set Elements
-					bool bMultipleTypeSelected = false;
 
-					ERigElementType LastType = ERigElementType::None;
-			
-					if (DraggedElements.Num() > 0)
+					if(DraggedComponents.Num() > 0)
 					{
-						LastType = DraggedElements[0].Type;
+						const FText SectionText = FText::FromString(DragDropContext.GetSectionTitle());
+						FToolMenuSection& Section = InMenu->AddSection(NAME_None, SectionText);
+
+						FText GetterLabel = LOCTEXT("GetComponent","Get Component");
+						FText GetterTooltip = LOCTEXT("GetComponent_ToolTip", "Getter For Component");
+						FText SetterLabel = LOCTEXT("SetComponent","Set Component");
+						FText SetterTooltip = LOCTEXT("SetComponent_ToolTip", "Setter For Component");
+
+						FToolMenuEntry GetComponentsEntry = FToolMenuEntry::InitMenuEntry(
+							TEXT("GetComponents"),
+							GetterLabel,
+							GetterTooltip,
+							FSlateIcon(),
+							FUIAction(
+								FExecuteAction::CreateSP(ControlRigEditor, &FControlRigEditor::HandleMakeComponentGetterSetter, true, DraggedComponents, Graph, NodePosition),
+								FCanExecuteAction()
+							)
+						);
+						GetComponentsEntry.InsertPosition.Name = NAME_None;
+						GetComponentsEntry.InsertPosition.Position = EToolMenuInsertType::First;
+
+						Section.AddEntry(GetComponentsEntry);
+
+						FToolMenuEntry SetComponentsEntry = FToolMenuEntry::InitMenuEntry(
+							TEXT("SetComponents"),
+							SetterLabel,
+							SetterTooltip,
+							FSlateIcon(),
+							FUIAction(
+								FExecuteAction::CreateSP(ControlRigEditor, &FControlRigEditor::HandleMakeComponentGetterSetter, false, DraggedComponents, Graph, NodePosition),
+								FCanExecuteAction()
+							)
+						);
+						SetComponentsEntry.InsertPosition.Name = GetComponentsEntry.Name;
+						SetComponentsEntry.InsertPosition.Position = EToolMenuInsertType::After;	
+
+						Section.AddEntry(SetComponentsEntry);
 					}
-			
-					uint8 DraggedTypes = 0;
-					uint8 DraggedAnimationTypes = 2;
-					for (const FRigElementKey& DraggedKey : DraggedElements)
+
+					if(DraggedElements.Num() > 0)
 					{
-						if (DraggedKey.Type != LastType)
-						{
-							bMultipleTypeSelected = true;
-						}
-						else if(DraggedKey.Type == ERigElementType::Control)
+						ControlRigEditor->FilterDraggedKeys(DraggedElements, true);
+					
+						// if multiple types are selected, we show Get Elements/Set Elements
+						bool bMultipleTypeSelected = false;
+
+						ERigElementType LastType = DraggedElements[0].Type;
+				
+						uint8 DraggedTypes = 0;
+						uint8 DraggedAnimationTypes = 2;
+						for (const FRigElementKey& DraggedKey : DraggedElements)
 						{
-							if(const FRigControlElement* ControlElement = Hierarchy->Find<FRigControlElement>(DraggedKey))
+							if (DraggedKey.Type != LastType)
 							{
-								const uint8 DraggedAnimationType = ControlElement->IsAnimationChannel() ? 1 : 0; 
-								if(DraggedAnimationTypes == 2)
-								{
-									DraggedAnimationTypes = DraggedAnimationType;
-								}
-								else
+								bMultipleTypeSelected = true;
+							}
+							else if(DraggedKey.Type == ERigElementType::Control)
+							{
+								if(const FRigControlElement* ControlElement = Hierarchy->Find<FRigControlElement>(DraggedKey))
 								{
-									if(DraggedAnimationTypes != DraggedAnimationType)
+									const uint8 DraggedAnimationType = ControlElement->IsAnimationChannel() ? 1 : 0; 
+									if(DraggedAnimationTypes == 2)
 									{
-										bMultipleTypeSelected = true;
+										DraggedAnimationTypes = DraggedAnimationType;
+									}
+									else
+									{
+										if(DraggedAnimationTypes != DraggedAnimationType)
+										{
+											bMultipleTypeSelected = true;
+										}
 									}
 								}
 							}
-						}
-				
-						DraggedTypes = DraggedTypes | (uint8)DraggedKey.Type;
-					}
 					
-					const FText SectionText = FText::FromString(DragDropContext.GetSectionTitle());
-					FToolMenuSection& Section = InMenu->AddSection(NAME_None, SectionText);
-
-					FText GetterLabel = LOCTEXT("GetElement","Get Element");
-					FText GetterTooltip = LOCTEXT("GetElement_ToolTip", "Getter For Element");
-					FText SetterLabel = LOCTEXT("SetElement","Set Element");
-					FText SetterTooltip = LOCTEXT("SetElement_ToolTip", "Setter For Element");
-					// if multiple types are selected, we show Get Elements/Set Elements
-					if (bMultipleTypeSelected)
-					{
-						GetterLabel = LOCTEXT("GetElements","Get Elements");
-						GetterTooltip = LOCTEXT("GetElements_ToolTip", "Getter For Elements");
-						SetterLabel = LOCTEXT("SetElements","Set Elements");
-						SetterTooltip = LOCTEXT("SetElements_ToolTip", "Setter For Elements");
-					}
-					else
-					{
-						// otherwise, we show "Get Bone/NUll/Control"
-						if ((DraggedTypes & (uint8)ERigElementType::Bone) != 0)
-						{
-							GetterLabel = LOCTEXT("GetBone","Get Bone");
-							GetterTooltip = LOCTEXT("GetBone_ToolTip", "Getter For Bone");
-							SetterLabel = LOCTEXT("SetBone","Set Bone");
-							SetterTooltip = LOCTEXT("SetBone_ToolTip", "Setter For Bone");
+							DraggedTypes = DraggedTypes | (uint8)DraggedKey.Type;
 						}
-						else if ((DraggedTypes & (uint8)ERigElementType::Null) != 0)
+						
+						const FText SectionText = FText::FromString(DragDropContext.GetSectionTitle());
+						FToolMenuSection& Section = InMenu->AddSection(NAME_None, SectionText);
+
+						FText GetterLabel = LOCTEXT("GetElement","Get Element");
+						FText GetterTooltip = LOCTEXT("GetElement_ToolTip", "Getter For Element");
+						FText SetterLabel = LOCTEXT("SetElement","Set Element");
+						FText SetterTooltip = LOCTEXT("SetElement_ToolTip", "Setter For Element");
+						// if multiple types are selected, we show Get Elements/Set Elements
+						if (bMultipleTypeSelected)
 						{
-							GetterLabel = LOCTEXT("GetNull","Get Null");
-							GetterTooltip = LOCTEXT("GetNull_ToolTip", "Getter For Null");
-							SetterLabel = LOCTEXT("SetNull","Set Null");
-							SetterTooltip = LOCTEXT("SetNull_ToolTip", "Setter For Null");
+							GetterLabel = LOCTEXT("GetElements","Get Elements");
+							GetterTooltip = LOCTEXT("GetElements_ToolTip", "Getter For Elements");
+							SetterLabel = LOCTEXT("SetElements","Set Elements");
+							SetterTooltip = LOCTEXT("SetElements_ToolTip", "Setter For Elements");
 						}
-						else if ((DraggedTypes & (uint8)ERigElementType::Control) != 0)
+						else
 						{
-							if(DraggedAnimationTypes == 0)
+							// otherwise, we show "Get Bone/NUll/Control"
+							if ((DraggedTypes & (uint8)ERigElementType::Bone) != 0)
 							{
-								GetterLabel = LOCTEXT("GetControl","Get Control");
-								GetterTooltip = LOCTEXT("GetControl_ToolTip", "Getter For Control");
-								SetterLabel = LOCTEXT("SetControl","Set Control");
-								SetterTooltip = LOCTEXT("SetControl_ToolTip", "Setter For Control");
+								GetterLabel = LOCTEXT("GetBone","Get Bone");
+								GetterTooltip = LOCTEXT("GetBone_ToolTip", "Getter For Bone");
+								SetterLabel = LOCTEXT("SetBone","Set Bone");
+								SetterTooltip = LOCTEXT("SetBone_ToolTip", "Setter For Bone");
 							}
-							else
+							else if ((DraggedTypes & (uint8)ERigElementType::Null) != 0)
 							{
-								GetterLabel = LOCTEXT("GetAnimationChannel","Get Animation Channel");
-								GetterTooltip = LOCTEXT("GetAnimationChannel_ToolTip", "Getter For Animation Channel");
-								SetterLabel = LOCTEXT("SetAnimationChannel","Set Animation Channel");
-								SetterTooltip = LOCTEXT("SetAnimationChannel_ToolTip", "Setter For Animation Channel");
+								GetterLabel = LOCTEXT("GetNull","Get Null");
+								GetterTooltip = LOCTEXT("GetNull_ToolTip", "Getter For Null");
+								SetterLabel = LOCTEXT("SetNull","Set Null");
+								SetterTooltip = LOCTEXT("SetNull_ToolTip", "Setter For Null");
+							}
+							else if ((DraggedTypes & (uint8)ERigElementType::Control) != 0)
+							{
+								if(DraggedAnimationTypes == 0)
+								{
+									GetterLabel = LOCTEXT("GetControl","Get Control");
+									GetterTooltip = LOCTEXT("GetControl_ToolTip", "Getter For Control");
+									SetterLabel = LOCTEXT("SetControl","Set Control");
+									SetterTooltip = LOCTEXT("SetControl_ToolTip", "Setter For Control");
+								}
+								else
+								{
+									GetterLabel = LOCTEXT("GetAnimationChannel","Get Animation Channel");
+									GetterTooltip = LOCTEXT("GetAnimationChannel_ToolTip", "Getter For Animation Channel");
+									SetterLabel = LOCTEXT("SetAnimationChannel","Set Animation Channel");
+									SetterTooltip = LOCTEXT("SetAnimationChannel_ToolTip", "Setter For Animation Channel");
+								}
+							}
+							else if ((DraggedTypes & (uint8)ERigElementType::Connector) != 0)
+							{
+								GetterLabel = LOCTEXT("GetConnector","Get Connector");
+								GetterTooltip = LOCTEXT("GetConnector_ToolTip", "Getter For Connector");
+								SetterLabel = LOCTEXT("SetConnector","Set Connector");
+								SetterTooltip = LOCTEXT("SetConnector_ToolTip", "Setter For Connector");
 							}
 						}
-						else if ((DraggedTypes & (uint8)ERigElementType::Connector) != 0)
-						{
-							GetterLabel = LOCTEXT("GetConnector","Get Connector");
-							GetterTooltip = LOCTEXT("GetConnector_ToolTip", "Getter For Connector");
-							SetterLabel = LOCTEXT("SetConnector","Set Connector");
-							SetterTooltip = LOCTEXT("SetConnector_ToolTip", "Setter For Connector");
-						}
-					}
-
-					FToolMenuEntry GetElementsEntry = FToolMenuEntry::InitMenuEntry(
-						TEXT("GetElements"),
-						GetterLabel,
-						GetterTooltip,
-						FSlateIcon(),
-						FUIAction(
-							FExecuteAction::CreateSP(ControlRigEditor, &FControlRigEditor::HandleMakeElementGetterSetter, ERigElementGetterSetterType_Transform, true, DraggedElements, Graph, NodePosition),
-							FCanExecuteAction()
-						)
-					);
-					GetElementsEntry.InsertPosition.Name = NAME_None;
-					GetElementsEntry.InsertPosition.Position = EToolMenuInsertType::First;
-					
 
-					Section.AddEntry(GetElementsEntry);
-
-					FToolMenuEntry SetElementsEntry = FToolMenuEntry::InitMenuEntry(
-						TEXT("SetElements"),
-						SetterLabel,
-						SetterTooltip,
-						FSlateIcon(),
-						FUIAction(
-							FExecuteAction::CreateSP(ControlRigEditor, &FControlRigEditor::HandleMakeElementGetterSetter, ERigElementGetterSetterType_Transform, false, DraggedElements, Graph, NodePosition),
-							FCanExecuteAction()
-						)
-					);
-					SetElementsEntry.InsertPosition.Name = GetElementsEntry.Name;
-					SetElementsEntry.InsertPosition.Position = EToolMenuInsertType::After;	
-
-					Section.AddEntry(SetElementsEntry);
-
-					if (((DraggedTypes & (uint8)ERigElementType::Bone) != 0) ||
-						((DraggedTypes & (uint8)ERigElementType::Control) != 0) ||
-						((DraggedTypes & (uint8)ERigElementType::Null) != 0) ||
-						((DraggedTypes & (uint8)ERigElementType::Connector) != 0))
-					{
-						FToolMenuEntry& RotationTranslationSeparator = Section.AddSeparator(TEXT("RotationTranslationSeparator"));
-						RotationTranslationSeparator.InsertPosition.Name = SetElementsEntry.Name;
-						RotationTranslationSeparator.InsertPosition.Position = EToolMenuInsertType::After;
-
-						FToolMenuEntry SetRotationEntry = FToolMenuEntry::InitMenuEntry(
-							TEXT("SetRotation"),
-							LOCTEXT("SetRotation","Set Rotation"),
-							LOCTEXT("SetRotation_ToolTip","Setter for Rotation"),
+						FToolMenuEntry GetElementsEntry = FToolMenuEntry::InitMenuEntry(
+							TEXT("GetElements"),
+							GetterLabel,
+							GetterTooltip,
 							FSlateIcon(),
 							FUIAction(
-								FExecuteAction::CreateSP(ControlRigEditor, &FControlRigEditor::HandleMakeElementGetterSetter, ERigElementGetterSetterType_Rotation, false, DraggedElements, Graph, NodePosition),
+								FExecuteAction::CreateSP(ControlRigEditor, &FControlRigEditor::HandleMakeElementGetterSetter, ERigElementGetterSetterType_Transform, true, DraggedElements, Graph, NodePosition),
 								FCanExecuteAction()
 							)
 						);
+						GetElementsEntry.InsertPosition.Name = NAME_None;
+						GetElementsEntry.InsertPosition.Position = EToolMenuInsertType::First;
 						
-						SetRotationEntry.InsertPosition.Name = RotationTranslationSeparator.Name;
-						SetRotationEntry.InsertPosition.Position = EToolMenuInsertType::After;		
-						Section.AddEntry(SetRotationEntry);
-
-						FToolMenuEntry SetTranslationEntry = FToolMenuEntry::InitMenuEntry(
-							TEXT("SetTranslation"),
-							LOCTEXT("SetTranslation","Set Translation"),
-							LOCTEXT("SetTranslation_ToolTip","Setter for Translation"),
-							FSlateIcon(),
-							FUIAction(
-								FExecuteAction::CreateSP(ControlRigEditor, &FControlRigEditor::HandleMakeElementGetterSetter, ERigElementGetterSetterType_Translation, false, DraggedElements, Graph, NodePosition),
-								FCanExecuteAction()
-							)
-						);
 
-						SetTranslationEntry.InsertPosition.Name = SetRotationEntry.Name;
-						SetTranslationEntry.InsertPosition.Position = EToolMenuInsertType::After;		
-						Section.AddEntry(SetTranslationEntry);
+						Section.AddEntry(GetElementsEntry);
 
-						FToolMenuEntry AddOffsetEntry = FToolMenuEntry::InitMenuEntry(
-							TEXT("AddOffset"),
-							LOCTEXT("AddOffset","Add Offset"),
-							LOCTEXT("AddOffset_ToolTip","Setter for Offset"),
+						FToolMenuEntry SetElementsEntry = FToolMenuEntry::InitMenuEntry(
+							TEXT("SetElements"),
+							SetterLabel,
+							SetterTooltip,
 							FSlateIcon(),
 							FUIAction(
-								FExecuteAction::CreateSP(ControlRigEditor, &FControlRigEditor::HandleMakeElementGetterSetter, ERigElementGetterSetterType_Offset, false, DraggedElements, Graph, NodePosition),
+								FExecuteAction::CreateSP(ControlRigEditor, &FControlRigEditor::HandleMakeElementGetterSetter, ERigElementGetterSetterType_Transform, false, DraggedElements, Graph, NodePosition),
 								FCanExecuteAction()
 							)
 						);
-						
-						AddOffsetEntry.InsertPosition.Name = SetTranslationEntry.Name;
-						AddOffsetEntry.InsertPosition.Position = EToolMenuInsertType::After;						
-						Section.AddEntry(AddOffsetEntry);
+						SetElementsEntry.InsertPosition.Name = GetElementsEntry.Name;
+						SetElementsEntry.InsertPosition.Position = EToolMenuInsertType::After;	
 
-						FToolMenuEntry& RelativeTransformSeparator = Section.AddSeparator(TEXT("RelativeTransformSeparator"));
-						RelativeTransformSeparator.InsertPosition.Name = AddOffsetEntry.Name;
-						RelativeTransformSeparator.InsertPosition.Position = EToolMenuInsertType::After;
-						
-						FToolMenuEntry GetRelativeTransformEntry = FToolMenuEntry::InitMenuEntry(
-							TEXT("GetRelativeTransformEntry"),
-							LOCTEXT("GetRelativeTransform", "Get Relative Transform"),
-							LOCTEXT("GetRelativeTransform_ToolTip", "Getter for Relative Transform"),
-							FSlateIcon(),
-							FUIAction(
-								FExecuteAction::CreateSP(ControlRigEditor, &FControlRigEditor::HandleMakeElementGetterSetter, ERigElementGetterSetterType_Relative, true, DraggedElements, Graph, NodePosition),
-								FCanExecuteAction()
-							)
-						);
-						
-						GetRelativeTransformEntry.InsertPosition.Name = RelativeTransformSeparator.Name;
-						GetRelativeTransformEntry.InsertPosition.Position = EToolMenuInsertType::After;	
-						Section.AddEntry(GetRelativeTransformEntry);
-						
-						FToolMenuEntry SetRelativeTransformEntry = FToolMenuEntry::InitMenuEntry(
-							TEXT("SetRelativeTransformEntry"),
-							LOCTEXT("SetRelativeTransform", "Set Relative Transform"),
-							LOCTEXT("SetRelativeTransform_ToolTip", "Setter for Relative Transform"),
-							FSlateIcon(),
-							FUIAction(
-								FExecuteAction::CreateSP(ControlRigEditor, &FControlRigEditor::HandleMakeElementGetterSetter, ERigElementGetterSetterType_Relative, false, DraggedElements, Graph, NodePosition),
-								FCanExecuteAction()
-							)
-						);
-						SetRelativeTransformEntry.InsertPosition.Name = GetRelativeTransformEntry.Name;
-						SetRelativeTransformEntry.InsertPosition.Position = EToolMenuInsertType::After;		
-						Section.AddEntry(SetRelativeTransformEntry);
-					}
+						Section.AddEntry(SetElementsEntry);
 
-					if (DraggedElements.Num() > 0 && Hierarchy != nullptr)
-					{
-						FToolMenuEntry& ItemArraySeparator = Section.AddSeparator(TEXT("ItemArraySeparator"));
-						ItemArraySeparator.InsertPosition.Name = TEXT("SetRelativeTransformEntry"),
-						ItemArraySeparator.InsertPosition.Position = EToolMenuInsertType::After;
-						
-						FToolMenuEntry CreateItemArrayEntry = FToolMenuEntry::InitMenuEntry(
-							TEXT("CreateItemArray"),
-							LOCTEXT("CreateItemArray", "Create Item Array"),
-							LOCTEXT("CreateItemArray_ToolTip", "Creates an item array from the selected elements in the hierarchy"),
-							FSlateIcon(),
-							FUIAction(
-								FExecuteAction::CreateLambda([ControlRigEditor, DraggedElements, NodePosition]()
-									{
-										if (URigVMController* Controller = ControlRigEditor->GetFocusedController())
-										{
-											Controller->OpenUndoBracket(TEXT("Create Item Array From Selection"));
+						if (((DraggedTypes & (uint8)ERigElementType::Bone) != 0) ||
+							((DraggedTypes & (uint8)ERigElementType::Control) != 0) ||
+							((DraggedTypes & (uint8)ERigElementType::Null) != 0) ||
+							((DraggedTypes & (uint8)ERigElementType::Connector) != 0))
+						{
+							FToolMenuEntry& RotationTranslationSeparator = Section.AddSeparator(TEXT("RotationTranslationSeparator"));
+							RotationTranslationSeparator.InsertPosition.Name = SetElementsEntry.Name;
+							RotationTranslationSeparator.InsertPosition.Position = EToolMenuInsertType::After;
+
+							FToolMenuEntry SetRotationEntry = FToolMenuEntry::InitMenuEntry(
+								TEXT("SetRotation"),
+								LOCTEXT("SetRotation","Set Rotation"),
+								LOCTEXT("SetRotation_ToolTip","Setter for Rotation"),
+								FSlateIcon(),
+								FUIAction(
+									FExecuteAction::CreateSP(ControlRigEditor, &FControlRigEditor::HandleMakeElementGetterSetter, ERigElementGetterSetterType_Rotation, false, DraggedElements, Graph, NodePosition),
+									FCanExecuteAction()
+								)
+							);
+							
+							SetRotationEntry.InsertPosition.Name = RotationTranslationSeparator.Name;
+							SetRotationEntry.InsertPosition.Position = EToolMenuInsertType::After;		
+							Section.AddEntry(SetRotationEntry);
+
+							FToolMenuEntry SetTranslationEntry = FToolMenuEntry::InitMenuEntry(
+								TEXT("SetTranslation"),
+								LOCTEXT("SetTranslation","Set Translation"),
+								LOCTEXT("SetTranslation_ToolTip","Setter for Translation"),
+								FSlateIcon(),
+								FUIAction(
+									FExecuteAction::CreateSP(ControlRigEditor, &FControlRigEditor::HandleMakeElementGetterSetter, ERigElementGetterSetterType_Translation, false, DraggedElements, Graph, NodePosition),
+									FCanExecuteAction()
+								)
+							);
+
+							SetTranslationEntry.InsertPosition.Name = SetRotationEntry.Name;
+							SetTranslationEntry.InsertPosition.Position = EToolMenuInsertType::After;		
+							Section.AddEntry(SetTranslationEntry);
+
+							FToolMenuEntry AddOffsetEntry = FToolMenuEntry::InitMenuEntry(
+								TEXT("AddOffset"),
+								LOCTEXT("AddOffset","Add Offset"),
+								LOCTEXT("AddOffset_ToolTip","Setter for Offset"),
+								FSlateIcon(),
+								FUIAction(
+									FExecuteAction::CreateSP(ControlRigEditor, &FControlRigEditor::HandleMakeElementGetterSetter, ERigElementGetterSetterType_Offset, false, DraggedElements, Graph, NodePosition),
+									FCanExecuteAction()
+								)
+							);
+							
+							AddOffsetEntry.InsertPosition.Name = SetTranslationEntry.Name;
+							AddOffsetEntry.InsertPosition.Position = EToolMenuInsertType::After;						
+							Section.AddEntry(AddOffsetEntry);
+
+							FToolMenuEntry& RelativeTransformSeparator = Section.AddSeparator(TEXT("RelativeTransformSeparator"));
+							RelativeTransformSeparator.InsertPosition.Name = AddOffsetEntry.Name;
+							RelativeTransformSeparator.InsertPosition.Position = EToolMenuInsertType::After;
+							
+							FToolMenuEntry GetRelativeTransformEntry = FToolMenuEntry::InitMenuEntry(
+								TEXT("GetRelativeTransformEntry"),
+								LOCTEXT("GetRelativeTransform", "Get Relative Transform"),
+								LOCTEXT("GetRelativeTransform_ToolTip", "Getter for Relative Transform"),
+								FSlateIcon(),
+								FUIAction(
+									FExecuteAction::CreateSP(ControlRigEditor, &FControlRigEditor::HandleMakeElementGetterSetter, ERigElementGetterSetterType_Relative, true, DraggedElements, Graph, NodePosition),
+									FCanExecuteAction()
+								)
+							);
+							
+							GetRelativeTransformEntry.InsertPosition.Name = RelativeTransformSeparator.Name;
+							GetRelativeTransformEntry.InsertPosition.Position = EToolMenuInsertType::After;	
+							Section.AddEntry(GetRelativeTransformEntry);
+							
+							FToolMenuEntry SetRelativeTransformEntry = FToolMenuEntry::InitMenuEntry(
+								TEXT("SetRelativeTransformEntry"),
+								LOCTEXT("SetRelativeTransform", "Set Relative Transform"),
+								LOCTEXT("SetRelativeTransform_ToolTip", "Setter for Relative Transform"),
+								FSlateIcon(),
+								FUIAction(
+									FExecuteAction::CreateSP(ControlRigEditor, &FControlRigEditor::HandleMakeElementGetterSetter, ERigElementGetterSetterType_Relative, false, DraggedElements, Graph, NodePosition),
+									FCanExecuteAction()
+								)
+							);
+							SetRelativeTransformEntry.InsertPosition.Name = GetRelativeTransformEntry.Name;
+							SetRelativeTransformEntry.InsertPosition.Position = EToolMenuInsertType::After;		
+							Section.AddEntry(SetRelativeTransformEntry);
+						}
 
-											if (URigVMNode* ItemsNode = Controller->AddUnitNode(FRigUnit_ItemArray::StaticStruct(), TEXT("Execute"), NodePosition))
+						if (Hierarchy != nullptr)
+						{
+							FToolMenuEntry& ItemArraySeparator = Section.AddSeparator(TEXT("ItemArraySeparator"));
+							ItemArraySeparator.InsertPosition.Name = TEXT("SetRelativeTransformEntry"),
+							ItemArraySeparator.InsertPosition.Position = EToolMenuInsertType::After;
+							
+							FToolMenuEntry CreateItemArrayEntry = FToolMenuEntry::InitMenuEntry(
+								TEXT("CreateItemArray"),
+								LOCTEXT("CreateItemArray", "Create Item Array"),
+								LOCTEXT("CreateItemArray_ToolTip", "Creates an item array from the selected elements in the hierarchy"),
+								FSlateIcon(),
+								FUIAction(
+									FExecuteAction::CreateLambda([ControlRigEditor, DraggedElements, NodePosition]()
+										{
+											if (URigVMController* Controller = ControlRigEditor->GetFocusedController())
 											{
-												if (URigVMPin* ItemsPin = ItemsNode->FindPin(TEXT("Items")))
+												Controller->OpenUndoBracket(TEXT("Create Item Array From Selection"));
+
+												if (URigVMNode* ItemsNode = Controller->AddUnitNode(FRigUnit_ItemArray::StaticStruct(), TEXT("Execute"), NodePosition))
 												{
-													Controller->SetArrayPinSize(ItemsPin->GetPinPath(), DraggedElements.Num());
+													if (URigVMPin* ItemsPin = ItemsNode->FindPin(TEXT("Items")))
+													{
+														Controller->SetArrayPinSize(ItemsPin->GetPinPath(), DraggedElements.Num());
 
-													TArray<URigVMPin*> ItemPins = ItemsPin->GetSubPins();
-													ensure(ItemPins.Num() == DraggedElements.Num());
+														TArray<URigVMPin*> ItemPins = ItemsPin->GetSubPins();
+														ensure(ItemPins.Num() == DraggedElements.Num());
 
-													for (int32 ItemIndex = 0; ItemIndex < DraggedElements.Num(); ItemIndex++)
-													{
-														FString DefaultValue;
-														FRigElementKey::StaticStruct()->ExportText(DefaultValue, &DraggedElements[ItemIndex], nullptr, nullptr, PPF_None, nullptr);
-														Controller->SetPinDefaultValue(ItemPins[ItemIndex]->GetPinPath(), DefaultValue, true, true, false, true);
-														Controller->SetPinExpansion(ItemPins[ItemIndex]->GetPinPath(), true, true, true);
-													}
+														for (int32 ItemIndex = 0; ItemIndex < DraggedElements.Num(); ItemIndex++)
+														{
+															FString DefaultValue;
+															FRigElementKey::StaticStruct()->ExportText(DefaultValue, &DraggedElements[ItemIndex], nullptr, nullptr, PPF_None, nullptr);
+															Controller->SetPinDefaultValue(ItemPins[ItemIndex]->GetPinPath(), DefaultValue, true, true, false, true);
+															Controller->SetPinExpansion(ItemPins[ItemIndex]->GetPinPath(), true, true, true);
+														}
 
-													Controller->SetPinExpansion(ItemsPin->GetPinPath(), true, true, true);
+														Controller->SetPinExpansion(ItemsPin->GetPinPath(), true, true, true);
+													}
 												}
-											}
 
-											Controller->CloseUndoBracket();
+												Controller->CloseUndoBracket();
+											}
 										}
-									}
+									)
 								)
-							)
-						);
-						
-						CreateItemArrayEntry.InsertPosition.Name = ItemArraySeparator.Name,
-						CreateItemArrayEntry.InsertPosition.Position = EToolMenuInsertType::After;
-						Section.AddEntry(CreateItemArrayEntry);
+							);
+							
+							CreateItemArrayEntry.InsertPosition.Name = ItemArraySeparator.Name,
+							CreateItemArrayEntry.InsertPosition.Position = EToolMenuInsertType::After;
+							Section.AddEntry(CreateItemArrayEntry);
+						}
 					}
 				}
 			})
@@ -5864,6 +5908,109 @@ void FControlRigEditor::HandleMakeElementGetterSetter(ERigElementGetterSetterTyp
 	}
 }
 
+void FControlRigEditor::HandleMakeComponentGetterSetter(bool bIsGetter, TArray<FRigComponentKey> Keys, UEdGraph* Graph, FVector2D NodePosition)
+{
+	DECLARE_SCOPE_HIERARCHICAL_COUNTER_FUNC()
+
+	if (Keys.Num() == 0)
+	{
+		return;
+	}
+
+	URigHierarchy* Hierarchy = GetHierarchyBeingDebugged();
+	if (Hierarchy == nullptr)
+	{
+		return;
+	}
+	if (GetFocusedController() == nullptr)
+	{
+		return;
+	}
+
+	GetFocusedController()->OpenUndoBracket(TEXT("Adding Nodes from Hierarchy"));
+
+	static const FRigVMDispatchFactory* GetComponentFactory = FRigVMRegistry_RWLock::Get().FindDispatchFactory(FRigDispatch_GetComponentContent().GetFactoryName());
+	static const FRigVMDispatchFactory* SetComponentFactory = FRigVMRegistry_RWLock::Get().FindDispatchFactory(FRigDispatch_SetComponentContent().GetFactoryName());
+	static const FRigVMDispatchFactory* GetTopLevelComponentFactory = FRigVMRegistry_RWLock::Get().FindDispatchFactory(FRigDispatch_GetTopLevelComponentContent().GetFactoryName());
+	static const FRigVMDispatchFactory* SetTopLevelComponentFactory = FRigVMRegistry_RWLock::Get().FindDispatchFactory(FRigDispatch_SetTopLevelComponentContent().GetFactoryName());
+
+	TArray<FName> NewNodes;
+
+	for (int32 Index = 0; Index < Keys.Num(); Index++)
+	{
+		const FRigComponentKey& Key = Keys[Index];
+		const FRigBaseComponent* Component = Hierarchy->FindComponent(Key);
+		if(Component == nullptr)
+		{
+			continue;
+		}
+
+		const TRigVMTypeIndex ComponentTypeIndex = FRigVMRegistry_RWLock::Get().GetTypeIndex(Component->GetScriptStruct());
+		if(ComponentTypeIndex == INDEX_NONE)
+		{
+			continue;
+		}
+
+		const FRigVMDispatchFactory* Factory = nullptr;
+		if(Key.IsTopLevel())
+		{
+			Factory = bIsGetter ? GetTopLevelComponentFactory : SetTopLevelComponentFactory;
+		}
+		else
+		{
+			Factory = bIsGetter ? GetComponentFactory : SetComponentFactory;
+		}
+
+		check(Factory);
+
+		FVector2D NodePositionIncrement(0.f, 120.f);
+		if (!bIsGetter)
+		{
+			NodePositionIncrement = FVector2D(380.f, 0.f);
+		}
+
+		if (URigVMTemplateNode* ModelNode = GetFocusedController()->AddTemplateNode(Factory->GetTemplateNotation(), NodePosition, FString(), true, true))
+		{
+			if(URigVMPin* ComponentPin = ModelNode->FindPin(FRigDispatch_ComponentBase::ComponentArgName.ToString()))
+			{
+				GetFocusedController()->ResolveWildCardPin(ComponentPin, ComponentTypeIndex, true, true);
+				if(!bIsGetter)
+				{
+					GetFocusedController()->SetPinDefaultValue(ComponentPin->GetPinPath(), Component->GetContentAsText(), true, true, false, true);
+				}
+			}
+
+			if(const URigVMPin* NamePin = ModelNode->FindPin(FRigDispatch_ComponentBase::NameArgName.ToString()))
+			{
+				GetFocusedController()->SetPinDefaultValue(NamePin->GetPinPath(), Component->GetName(), true, true, false, true);
+			}
+
+			if(const URigVMPin* KeyPin = ModelNode->FindPin(FRigDispatch_ComponentBase::KeyArgName.ToString()))
+			{
+				FString DefaultValue;
+				FRigComponentKey::StaticStruct()->ExportText(DefaultValue, &Key, nullptr, nullptr, PPF_None, nullptr);
+				GetFocusedController()->SetPinDefaultValue(KeyPin->GetPinPath(), DefaultValue, true, true, false, true);
+			}
+
+			URigVMEdGraphUnitNodeSpawner::HookupMutableNode(ModelNode, GetControlRigBlueprint());
+
+			NewNodes.Add(ModelNode->GetFName());
+		}
+
+		NodePosition += NodePositionIncrement;
+	}
+
+	if (NewNodes.Num() > 0)
+	{
+		GetFocusedController()->SetNodeSelection(NewNodes);
+		GetFocusedController()->CloseUndoBracket();
+	}
+	else
+	{
+		GetFocusedController()->CancelUndoBracket();
+	}
+}
+
 void FControlRigEditor::HandleOnControlModified(UControlRig* Subject, FRigControlElement* ControlElement, const FRigControlModifiedContext& Context)
 {
 	UControlRig* DebuggedControlRig = Cast<UControlRig>(FRigVMEditorBase::GetBlueprintObj()->GetObjectBeingDebugged());
diff --git a/Engine/Plugins/Animation/ControlRig/Source/ControlRigEditor/Private/Editor/ControlRigEditor.h b/Engine/Plugins/Animation/ControlRig/Source/ControlRigEditor/Private/Editor/ControlRigEditor.h
--- a/Engine/Plugins/Animation/ControlRig/Source/ControlRigEditor/Private/Editor/ControlRigEditor.h
+++ b/Engine/Plugins/Animation/ControlRig/Source/ControlRigEditor/Private/Editor/ControlRigEditor.h
@@ -293,6 +293,7 @@ private:
 
 	void FilterDraggedKeys(TArray<FRigElementKey>& Keys, bool bRemoveNameSpace);
 	void HandleMakeElementGetterSetter(ERigElementGetterSetterType Type, bool bIsGetter, TArray<FRigElementKey> Keys, UEdGraph* Graph, FVector2D NodePosition);
+	void HandleMakeComponentGetterSetter(bool bIsGetter, TArray<FRigComponentKey> Keys, UEdGraph* Graph, FVector2D NodePosition);
 
 	void HandleOnControlModified(UControlRig* Subject, FRigControlElement* ControlElement, const FRigControlModifiedContext& Context);
 
diff --git a/Engine/Plugins/Animation/ControlRig/Source/ControlRigEditor/Private/Editor/SRigHierarchy.cpp b/Engine/Plugins/Animation/ControlRig/Source/ControlRigEditor/Private/Editor/SRigHierarchy.cpp
--- a/Engine/Plugins/Animation/ControlRig/Source/ControlRigEditor/Private/Editor/SRigHierarchy.cpp
+++ b/Engine/Plugins/Animation/ControlRig/Source/ControlRigEditor/Private/Editor/SRigHierarchy.cpp
@@ -837,6 +837,8 @@ void SRigHierarchy::OnHierarchyModified(ERigHierarchyNotification InNotif, URigH
 		case ERigHierarchyNotification::ElementRenamed:
 		case ERigHierarchyNotification::ElementReordered:
 		case ERigHierarchyNotification::HierarchyReset:
+		case ERigHierarchyNotification::ComponentRenamed:
+		case ERigHierarchyNotification::ComponentReparented:
 		{
 			RefreshTreeView();
 			break;
@@ -1421,6 +1423,11 @@ void SRigHierarchy::CreateContextMenu()
 							{
 								SelectedKey = SelectedItems[0]->Key;
 							}
+							else
+							{
+								// we use an invalid key in case the user has clicked into the view but not on an element
+								SelectedKey = FRigHierarchyKey(FRigElementKey(), true);
+							}
 
 							if(SelectedKey.IsElement())
 							{
@@ -1477,31 +1484,67 @@ void SRigHierarchy::CreateContextMenu()
 								{
 									const FString ComponentName = RigComponentStruct->GetName();
 									const FText ComponentLabel = FText::Format(LOCTEXT("AddComponentLabelFormat", "Add {0}"), RigComponentStruct->GetDisplayNameText());
-									
+
+									const FStructOnScope StructOnScope(RigComponentStruct);
+									const FRigBaseComponent* DefaultComponent = reinterpret_cast<const FRigBaseComponent*>(StructOnScope.GetStructMemory());
+									const FSlateIcon& ComponentIcon = DefaultComponent->GetIconForUI();
+
+									TSharedPtr<FString> FailureReason = MakeShared<FString>();
+									FCanExecuteAction CanExecute = FCanExecuteAction::CreateLambda([RigHierarchyPanel, RigComponentStruct, SelectedElements, FailureReason]() {
+										const URigHierarchy* CurrentHierarchy = RigHierarchyPanel->GetHierarchy();
+										if(CurrentHierarchy == nullptr)
+										{
+											return false;
+										}
+
+										for(const FRigElementKey& SelectedElement : SelectedElements)
+										{
+											if(!CurrentHierarchy->CanAddComponent(SelectedElement, RigComponentStruct, FailureReason.Get()))
+											{
+												return false;
+											}
+										}
+										return true;
+									});
+
 									DefaultSection.AddMenuEntry(
 										*FString::Printf(TEXT("Add%s"), *ComponentName),
 										ComponentLabel,
-										FText(),
-										FSlateIcon(), // todo: we may want to see the icon here for each component
-										FUIAction(FExecuteAction::CreateLambda([RigHierarchyPanel, RigComponentStruct, SelectedElements, ComponentLabel]() {
-											if(!RigHierarchyPanel->ControlRigBlueprint.IsValid())
+										TAttribute<FText>::CreateLambda([RigComponentStruct, FailureReason, CanExecute]() -> FText
+										{
+											FText TooltipText = RigComponentStruct->GetToolTipText();
+											if(!CanExecute.Execute())
 											{
-												return;
+												if(!FailureReason->IsEmpty())
+												{
+													TooltipText = FText::Format(LOCTEXT("AddComponentTooltipTextFormat", "{0}\n\n{1}"), TooltipText, FText::FromString(*FailureReason.Get()));
+												}
 											}
+											return TooltipText;
+										}),
+										ComponentIcon,
+										FUIAction(
+											FExecuteAction::CreateLambda([RigHierarchyPanel, RigComponentStruct, SelectedElements, ComponentLabel]() {
+												if(!RigHierarchyPanel->ControlRigBlueprint.IsValid())
+												{
+													return;
+												}
 
-											RigHierarchyPanel->DisplaySettings.bShowComponents = true;
+												RigHierarchyPanel->DisplaySettings.bShowComponents = true;
 
-											FScopedTransaction Transaction(ComponentLabel);
-											UControlRigBlueprint* CurrentControlRigBlueprint = RigHierarchyPanel->ControlRigBlueprint.Get();
-											URigHierarchyController* Controller = CurrentControlRigBlueprint->GetHierarchyController();
-											TArray<FRigComponentKey> ComponentKeys;
-											for(const FRigElementKey& SelectedElementKey : SelectedElements)
-											{
-												ComponentKeys.Add(Controller->AddComponent(RigComponentStruct, *RigComponentStruct->GetDisplayNameText().ToString(), SelectedElementKey, FString(), true, true));
-											}
-											Controller->SetComponentSelection(ComponentKeys);
-										})
-									));
+												FScopedTransaction Transaction(ComponentLabel);
+												UControlRigBlueprint* CurrentControlRigBlueprint = RigHierarchyPanel->ControlRigBlueprint.Get();
+												URigHierarchyController* Controller = CurrentControlRigBlueprint->GetHierarchyController();
+												TArray<FRigComponentKey> ComponentKeys;
+												for(const FRigElementKey& SelectedElementKey : SelectedElements)
+												{
+													ComponentKeys.Add(Controller->AddComponent(RigComponentStruct, NAME_None, SelectedElementKey, FString(), true, true));
+												}
+												Controller->SetComponentSelection(ComponentKeys);
+											}),
+											CanExecute
+										)
+									);
 								}
 							})
 						);
@@ -1513,7 +1556,7 @@ void SRigHierarchy::CreateContextMenu()
 					ElementsSection.AddMenuEntry(Commands.RenameItem);
 					ElementsSection.AddMenuEntry(Commands.MirrorItem);
 
-					if(RigHierarchyPanel->IsProceduralElementSelected() && RigHierarchyPanel->ControlRigBlueprint.IsValid())
+					if(RigHierarchyPanel->IsProceduralSelected() && RigHierarchyPanel->ControlRigBlueprint.IsValid())
 					{
 						ElementsSection.AddMenuEntry(
 							"SelectSpawnerNode",
@@ -1991,7 +2034,7 @@ bool SRigHierarchy::IsMultiSelected(bool bIncludeProcedural) const
 {
 	if(GetSelectedKeys().Num() > 0)
 	{
-		if(!bIncludeProcedural && IsProceduralElementSelected())
+		if(!bIncludeProcedural && IsProceduralSelected())
 		{
 			return false;
 		}
@@ -2004,7 +2047,7 @@ bool SRigHierarchy::IsSingleSelected(bool bIncludeProcedural) const
 {
 	if(GetSelectedKeys().Num() == 1)
 	{
-		if(!bIncludeProcedural && IsProceduralElementSelected())
+		if(!bIncludeProcedural && IsProceduralSelected())
 		{
 			return false;
 		}
@@ -2041,7 +2084,7 @@ bool SRigHierarchy::IsSingleNullSelected(bool bIncludeProcedural) const
 
 bool SRigHierarchy::IsControlSelected(bool bIncludeProcedural) const
 {
-	if(!bIncludeProcedural && IsProceduralElementSelected())
+	if(!bIncludeProcedural && IsProceduralSelected())
 	{
 		return false;
 	}
@@ -2062,7 +2105,7 @@ bool SRigHierarchy::IsControlSelected(bool bIncludeProcedural) const
 
 bool SRigHierarchy::IsControlOrNullSelected(bool bIncludeProcedural) const
 {
-	if(!bIncludeProcedural && IsProceduralElementSelected())
+	if(!bIncludeProcedural && IsProceduralSelected())
 	{
 		return false;
 	}
@@ -2085,7 +2128,7 @@ bool SRigHierarchy::IsControlOrNullSelected(bool bIncludeProcedural) const
 	return false;
 }
 
-bool SRigHierarchy::IsProceduralElementSelected() const
+bool SRigHierarchy::IsProceduralSelected() const
 {
 	TArray<FRigHierarchyKey> SelectedKeys = GetSelectedKeys();
 	if (SelectedKeys.IsEmpty())
@@ -2107,7 +2150,7 @@ bool SRigHierarchy::IsProceduralElementSelected() const
 	return true;
 }
 
-bool SRigHierarchy::IsNonProceduralElementSelected() const
+bool SRigHierarchy::IsNonProceduralSelected() const
 {
 	TArray<FRigHierarchyKey> SelectedKeys = GetSelectedKeys();
 	if (SelectedKeys.IsEmpty())
@@ -2510,6 +2553,15 @@ bool SRigHierarchy::CanDuplicateItem() const
 		return false;
 	}
 
+	// don't allow duplication on components
+	if(GetSelectedKeys().ContainsByPredicate([](const FRigHierarchyKey& Key)
+	{
+		return Key.IsComponent();
+	}))
+	{
+		return false;
+	}
+
 	if (ControlRigBlueprint->IsControlRigModule())
 	{
 		bool bAnyNonConnector = GetSelectedElementKeys().ContainsByPredicate([](const FRigElementKey& Key)
@@ -2643,9 +2695,15 @@ bool SRigHierarchy::CanRenameItem() const
 			}
 			return true;
 		}
-		else
+		if(Key.IsComponent())
 		{
-			// todo: not going to allow renaming of components for now
+			if(URigHierarchy* DebuggedHierarchy = GetHierarchy())
+			{
+				if(const FRigBaseComponent* Component = DebuggedHierarchy->FindComponent(Key.GetComponent()))
+				{
+					return Component->CanBeRenamed();
+				}
+			}
 		}
 	}
 	return false;
@@ -2920,8 +2978,7 @@ TOptional<EItemDropZone> SRigHierarchy::OnCanAcceptDrop(const FDragDropEvent& Dr
 					}
 					else if(TargetItem->Key.IsComponent())
 					{
-						// the element the component belongs to
-						TargetKey = TargetItem->Key.GetElement();
+						return InvalidDropZone;
 					}
 					break;
 				}
@@ -2975,9 +3032,15 @@ TOptional<EItemDropZone> SRigHierarchy::OnCanAcceptDrop(const FDragDropEvent& Dr
 					}
 					if(DraggedKey.IsComponent())
 					{
-						// todo: perform more validation here - as to where we are allowed to drag components
-						// return DraggedKey.GetElement() != TargetKey;
-						return InvalidDropZone;
+						const FRigBaseComponent* Component = Hierarchy->FindComponent(DraggedKey.GetComponent());
+						if(Component == nullptr)
+						{
+							return InvalidDropZone;
+						}
+						if(!Hierarchy->CanAddComponent(TargetKey.GetElement(), Component))
+						{
+							return InvalidDropZone;
+						}
 					}
 				}
 			}
@@ -3440,28 +3503,35 @@ FName SRigHierarchy::HandleRenameElement(const FRigHierarchyKey& OldKey, const F
 
 		FRigName SanitizedName(NewName);
 		Hierarchy->SanitizeName(SanitizedName);
-		FName ResultingName = NAME_None;
 
-		bool bUseDisplayName = false;
+		FName ResultingName = NAME_None;
 		if(OldKey.IsElement())
 		{
-			if(const FRigControlElement* ControlElement = Hierarchy->Find<FRigControlElement>(OldKey.GetElement()))
+			bool bUseDisplayName = false;
+			if(OldKey.IsElement())
 			{
-				if(ControlElement->IsAnimationChannel())
+				if(const FRigControlElement* ControlElement = Hierarchy->Find<FRigControlElement>(OldKey.GetElement()))
 				{
-					bUseDisplayName = true;
+					if(ControlElement->IsAnimationChannel())
+					{
+						bUseDisplayName = true;
+					}
 				}
-			}
 
-			if(bUseDisplayName)
-			{
-				ResultingName = Controller->SetDisplayName(OldKey.GetElement(), SanitizedName, true, true, true);
-			}
-			else
-			{
-				ResultingName = Controller->RenameElement(OldKey.GetElement(), SanitizedName, true, true, false).Name;
+				if(bUseDisplayName)
+				{
+					ResultingName = Controller->SetDisplayName(OldKey.GetElement(), SanitizedName, true, true, true);
+				}
+				else
+				{
+					ResultingName = Controller->RenameElement(OldKey.GetElement(), SanitizedName, true, true, false).Name;
+				}
 			}
 		}
+		else if(OldKey.IsComponent())
+		{
+			ResultingName = Controller->RenameComponent(OldKey.GetComponent(), SanitizedName, true, true, false).Name;
+		}
 		ControlRigBlueprint->PropagateHierarchyFromBPToInstances();
 		return ResultingName;
 	}
@@ -3537,8 +3607,12 @@ bool SRigHierarchy::HandleVerifyNameChanged(const FRigHierarchyKey& OldKey, cons
 		}
 		else if(OldKey.IsComponent())
 		{
-			// todo: not allowing renaming of components for now
-			return false;
+			FString OutErrorString;
+			if (!Hierarchy->IsComponentNameAvailable(OldKey.GetElement(), FRigName(NewName), &OutErrorString))
+			{
+				OutErrorMessage = FText::FromString(OutErrorString);
+				return false;
+			}
 		}
 	}
 	return true;
@@ -3902,48 +3976,54 @@ FReply SRigHierarchy::ReparentOrMatchTransform(const TArray<FRigHierarchyKey>& D
 				return FReply::Unhandled();
 			}
 
-			if(!DraggedKey.IsElement())
+			if(DraggedKey.IsElement())
 			{
-				continue;
-			}
-
-			if (bReparentItems)
-			{
-				if(Hierarchy->IsParentedTo(TargetKey.GetElement(), DraggedKey.GetElement()))
+				if (bReparentItems)
 				{
-					if(LocalIndex == INDEX_NONE)
+					if(Hierarchy->IsParentedTo(TargetKey.GetElement(), DraggedKey.GetElement()))
 					{
-						return FReply::Unhandled();
+						if(LocalIndex == INDEX_NONE)
+						{
+							return FReply::Unhandled();
+						}
 					}
 				}
-			}
 
-			if (DraggedKey.GetElement().Type == ERigElementType::Bone)
-			{
-				if(FRigBoneElement* BoneElement = Hierarchy->Find<FRigBoneElement>(DraggedKey.GetElement()))
+				if (DraggedKey.GetElement().Type == ERigElementType::Bone)
 				{
-					if (BoneElement->BoneType == ERigBoneType::Imported && BoneElement->ParentElement != nullptr)
+					if(FRigBoneElement* BoneElement = Hierarchy->Find<FRigBoneElement>(DraggedKey.GetElement()))
 					{
-						FText ConfirmText = bMatchTransforms ?
-							LOCTEXT("ConfirmMatchTransform", "Matching transforms of imported(white) bones can cause issues with animation - are you sure ?") :
-							LOCTEXT("ConfirmReparentBoneHierarchy", "Reparenting imported(white) bones can cause issues with animation - are you sure ?");
+						if (BoneElement->BoneType == ERigBoneType::Imported && BoneElement->ParentElement != nullptr)
+						{
+							FText ConfirmText = bMatchTransforms ?
+								LOCTEXT("ConfirmMatchTransform", "Matching transforms of imported(white) bones can cause issues with animation - are you sure ?") :
+								LOCTEXT("ConfirmReparentBoneHierarchy", "Reparenting imported(white) bones can cause issues with animation - are you sure ?");
 
-						FText TitleText = bMatchTransforms ?
-							LOCTEXT("MatchTransformImportedBone", "Match Transform on Imported Bone") :
-							LOCTEXT("ReparentImportedBone", "Reparent Imported Bone");
+							FText TitleText = bMatchTransforms ?
+								LOCTEXT("MatchTransformImportedBone", "Match Transform on Imported Bone") :
+								LOCTEXT("ReparentImportedBone", "Reparent Imported Bone");
 
-						FSuppressableWarningDialog::FSetupInfo Info(ConfirmText, TitleText, "SRigHierarchy_Warning");
-						Info.ConfirmText = LOCTEXT("SRigHierarchy_Warning_Yes", "Yes");
-						Info.CancelText = LOCTEXT("SRigHierarchy_Warning_No", "No");
+							FSuppressableWarningDialog::FSetupInfo Info(ConfirmText, TitleText, "SRigHierarchy_Warning");
+							Info.ConfirmText = LOCTEXT("SRigHierarchy_Warning_Yes", "Yes");
+							Info.CancelText = LOCTEXT("SRigHierarchy_Warning_No", "No");
 
-						FSuppressableWarningDialog ChangeImportedBonesInHierarchy(Info);
-						if (ChangeImportedBonesInHierarchy.ShowModal() == FSuppressableWarningDialog::Cancel)
-						{
-							return FReply::Unhandled();
+							FSuppressableWarningDialog ChangeImportedBonesInHierarchy(Info);
+							if (ChangeImportedBonesInHierarchy.ShowModal() == FSuppressableWarningDialog::Cancel)
+							{
+								return FReply::Unhandled();
+							}
 						}
 					}
 				}
 			}
+
+			if(DraggedKey.IsComponent())
+			{
+				if(DraggedKey.GetElement() == TargetKey.GetElement())
+				{
+					return FReply::Unhandled();
+				}
+			}
 		}
 
 		for (const FRigHierarchyKey& DraggedKey : DraggedKeys)
@@ -3984,59 +4064,61 @@ FReply SRigHierarchy::ReparentOrMatchTransform(const TArray<FRigHierarchyKey>& D
 
 			if(DraggedKey.IsComponent())
 			{
-				// todo: reparent the component
-				continue;
+				Controller->ReparentComponent(DraggedKey.GetComponent(), ParentKey, true, true);
 			}
 
-			const FTransform InitialGlobalTransform = DebuggedHierarchy->GetInitialGlobalTransform(DraggedKey.GetElement());
-			const FTransform CurrentGlobalTransform = DebuggedHierarchy->GetGlobalTransform(DraggedKey.GetElement());
-			const FTransform InitialLocalTransform = DebuggedHierarchy->GetInitialLocalTransform(DraggedKey.GetElement());
-			const FTransform CurrentLocalTransform = DebuggedHierarchy->GetLocalTransform(DraggedKey.GetElement());
-			const FTransform CurrentGlobalOffsetTransform = DebuggedHierarchy->GetGlobalControlOffsetTransform(DraggedKey.GetElement(), false);
-
-			bool bElementWasReparented = false;
-			if(ParentKey.IsValid() && (Hierarchy->GetFirstParent(DraggedKey.GetElement()) != ParentKey))
-			{
-				bElementWasReparented = Controller->SetParent(DraggedKey.GetElement(), ParentKey, true, true, true);
-			}
-			else if(!ParentKey.IsValid() && (Hierarchy->GetNumberOfParents(DraggedKey.GetElement()) > 0))
+			if(DraggedKey.IsElement())
 			{
-				bElementWasReparented = Controller->RemoveAllParents(DraggedKey.GetElement(), true, true, true);
-			}
+				const FTransform InitialGlobalTransform = DebuggedHierarchy->GetInitialGlobalTransform(DraggedKey.GetElement());
+				const FTransform CurrentGlobalTransform = DebuggedHierarchy->GetGlobalTransform(DraggedKey.GetElement());
+				const FTransform InitialLocalTransform = DebuggedHierarchy->GetInitialLocalTransform(DraggedKey.GetElement());
+				const FTransform CurrentLocalTransform = DebuggedHierarchy->GetLocalTransform(DraggedKey.GetElement());
+				const FTransform CurrentGlobalOffsetTransform = DebuggedHierarchy->GetGlobalControlOffsetTransform(DraggedKey.GetElement(), false);
+
+				bool bElementWasReparented = false;
+				if(ParentKey.IsValid() && (Hierarchy->GetFirstParent(DraggedKey.GetElement()) != ParentKey))
+				{
+					bElementWasReparented = Controller->SetParent(DraggedKey.GetElement(), ParentKey, true, true, true);
+				}
+				else if(!ParentKey.IsValid() && (Hierarchy->GetNumberOfParents(DraggedKey.GetElement()) > 0))
+				{
+					bElementWasReparented = Controller->RemoveAllParents(DraggedKey.GetElement(), true, true, true);
+				}
 
-			if(LocalIndex != INDEX_NONE)
-			{
-				Controller->ReorderElement(DraggedKey.GetElement(), LocalIndex, true, true);
-			}
+				if(LocalIndex != INDEX_NONE)
+				{
+					Controller->ReorderElement(DraggedKey.GetElement(), LocalIndex, true, true);
+				}
 
-			if(bElementWasReparented)
-			{
-				if (DraggedKey.GetElement().Type == ERigElementType::Control)
+				if(bElementWasReparented)
 				{
-					int32 ControlIndex = DebuggedHierarchy->GetIndex(DraggedKey.GetElement());
-					if (ControlIndex == INDEX_NONE)
+					if (DraggedKey.GetElement().Type == ERigElementType::Control)
 					{
-						continue;
-					}
+						int32 ControlIndex = DebuggedHierarchy->GetIndex(DraggedKey.GetElement());
+						if (ControlIndex == INDEX_NONE)
+						{
+							continue;
+						}
 
-					const FTransform GlobalParentTransform = DebuggedHierarchy->GetGlobalTransform(ParentKey, false);
-					const FTransform LocalOffsetTransform = CurrentGlobalOffsetTransform.GetRelativeTransform(GlobalParentTransform);
-
-					Hierarchy->SetControlOffsetTransformByIndex(ControlIndex, LocalOffsetTransform, ERigTransformType::InitialLocal, true, true, true);
-					Hierarchy->SetControlOffsetTransformByIndex(ControlIndex, LocalOffsetTransform, ERigTransformType::CurrentLocal, true, true, true);
-					Hierarchy->SetLocalTransform(DraggedKey.GetElement(), CurrentLocalTransform, true, true, true, true);
-					Hierarchy->SetInitialLocalTransform(DraggedKey.GetElement(), InitialLocalTransform, true, true, true);
-					DebuggedHierarchy->SetControlOffsetTransformByIndex(ControlIndex, LocalOffsetTransform, ERigTransformType::InitialLocal, true, true);
-					DebuggedHierarchy->SetControlOffsetTransformByIndex(ControlIndex, LocalOffsetTransform, ERigTransformType::CurrentLocal, true, true);
-					DebuggedHierarchy->SetLocalTransform(DraggedKey.GetElement(), CurrentLocalTransform, true, true, true);
-					DebuggedHierarchy->SetInitialLocalTransform(DraggedKey.GetElement(), InitialLocalTransform, true, true);
-				}
-				else
-				{
-					DebuggedHierarchy->SetInitialGlobalTransform(DraggedKey.GetElement(), InitialGlobalTransform, true, true);
-					DebuggedHierarchy->SetGlobalTransform(DraggedKey.GetElement(), CurrentGlobalTransform, false, true, true);
-					Hierarchy->SetInitialGlobalTransform(DraggedKey.GetElement(), InitialGlobalTransform, true, true);
-					Hierarchy->SetGlobalTransform(DraggedKey.GetElement(), CurrentGlobalTransform, false, true, true);
+						const FTransform GlobalParentTransform = DebuggedHierarchy->GetGlobalTransform(ParentKey, false);
+						const FTransform LocalOffsetTransform = CurrentGlobalOffsetTransform.GetRelativeTransform(GlobalParentTransform);
+
+						Hierarchy->SetControlOffsetTransformByIndex(ControlIndex, LocalOffsetTransform, ERigTransformType::InitialLocal, true, true, true);
+						Hierarchy->SetControlOffsetTransformByIndex(ControlIndex, LocalOffsetTransform, ERigTransformType::CurrentLocal, true, true, true);
+						Hierarchy->SetLocalTransform(DraggedKey.GetElement(), CurrentLocalTransform, true, true, true, true);
+						Hierarchy->SetInitialLocalTransform(DraggedKey.GetElement(), InitialLocalTransform, true, true, true);
+						DebuggedHierarchy->SetControlOffsetTransformByIndex(ControlIndex, LocalOffsetTransform, ERigTransformType::InitialLocal, true, true);
+						DebuggedHierarchy->SetControlOffsetTransformByIndex(ControlIndex, LocalOffsetTransform, ERigTransformType::CurrentLocal, true, true);
+						DebuggedHierarchy->SetLocalTransform(DraggedKey.GetElement(), CurrentLocalTransform, true, true, true);
+						DebuggedHierarchy->SetInitialLocalTransform(DraggedKey.GetElement(), InitialLocalTransform, true, true);
+					}
+					else
+					{
+						DebuggedHierarchy->SetInitialGlobalTransform(DraggedKey.GetElement(), InitialGlobalTransform, true, true);
+						DebuggedHierarchy->SetGlobalTransform(DraggedKey.GetElement(), CurrentGlobalTransform, false, true, true);
+						Hierarchy->SetInitialGlobalTransform(DraggedKey.GetElement(), InitialGlobalTransform, true, true);
+						Hierarchy->SetGlobalTransform(DraggedKey.GetElement(), CurrentGlobalTransform, false, true, true);
+					}
 				}
 			}
 		}
diff --git a/Engine/Plugins/Animation/ControlRig/Source/ControlRigEditor/Private/Editor/SRigHierarchy.h b/Engine/Plugins/Animation/ControlRig/Source/ControlRigEditor/Private/Editor/SRigHierarchy.h
--- a/Engine/Plugins/Animation/ControlRig/Source/ControlRigEditor/Private/Editor/SRigHierarchy.h
+++ b/Engine/Plugins/Animation/ControlRig/Source/ControlRigEditor/Private/Editor/SRigHierarchy.h
@@ -179,8 +179,8 @@ private:
 	bool IsSingleNullSelected(bool bIncludeProcedural) const;
 	bool IsControlSelected(bool bIncludeProcedural) const;
 	bool IsControlOrNullSelected(bool bIncludeProcedural) const;
-	bool IsProceduralElementSelected() const;
-	bool IsNonProceduralElementSelected() const;
+	bool IsProceduralSelected() const;
+	bool IsNonProceduralSelected() const;
 	bool CanAddElement(const ERigElementType ElementType) const;
 	bool CanAddAnimationChannel() const;
 
diff --git a/Engine/Plugins/Experimental/PhysicsControl/Source/ControlRigPhysics/Private/RigPhysicsComponent.cpp b/Engine/Plugins/Experimental/PhysicsControl/Source/ControlRigPhysics/Private/RigPhysicsComponent.cpp
--- a/Engine/Plugins/Experimental/PhysicsControl/Source/ControlRigPhysics/Private/RigPhysicsComponent.cpp
+++ b/Engine/Plugins/Experimental/PhysicsControl/Source/ControlRigPhysics/Private/RigPhysicsComponent.cpp
@@ -22,3 +22,16 @@ void FRigPhysicsComponent::Load(FArchive& Ar)
 	Ar << Joint;
 }
 
+//======================================================================================================================
+bool FRigPhysicsComponent::CanBeAddedTo(const FRigElementKey& InElementKey, const URigHierarchy* InHierarchy, FString* OutFailureReason) const
+{
+	if(InElementKey.Type != ERigElementType::Bone)
+	{
+		if(OutFailureReason)
+		{
+			*OutFailureReason = TEXT("Physics components can only be added to bones.");
+		}
+		return false;
+	}
+	return true;
+}
diff --git a/Engine/Plugins/Experimental/PhysicsControl/Source/ControlRigPhysics/Private/RigPhysicsControlComponent.cpp b/Engine/Plugins/Experimental/PhysicsControl/Source/ControlRigPhysics/Private/RigPhysicsControlComponent.cpp
--- a/Engine/Plugins/Experimental/PhysicsControl/Source/ControlRigPhysics/Private/RigPhysicsControlComponent.cpp
+++ b/Engine/Plugins/Experimental/PhysicsControl/Source/ControlRigPhysics/Private/RigPhysicsControlComponent.cpp
@@ -20,3 +20,12 @@ void FRigPhysicsControlComponent::Load(FArchive& Ar)
 	Ar << ControlTarget;
 }
 
+//======================================================================================================================
+bool FRigPhysicsControlComponent::CanBeAddedTo(const FRigElementKey& InElementKey, const URigHierarchy* InHierarchy, FString* OutFailureReason) const
+{
+	if(OutFailureReason)
+	{
+		*OutFailureReason = TEXT("Physics control component is currently not supported.");
+	}
+	return false;
+}
diff --git a/Engine/Plugins/Experimental/PhysicsControl/Source/ControlRigPhysics/Private/RigPhysicsSolverComponent.cpp b/Engine/Plugins/Experimental/PhysicsControl/Source/ControlRigPhysics/Private/RigPhysicsSolverComponent.cpp
--- a/Engine/Plugins/Experimental/PhysicsControl/Source/ControlRigPhysics/Private/RigPhysicsSolverComponent.cpp
+++ b/Engine/Plugins/Experimental/PhysicsControl/Source/ControlRigPhysics/Private/RigPhysicsSolverComponent.cpp
@@ -2,6 +2,8 @@
 
 #include "RigPhysicsSolverComponent.h"
 
+#include "Rigs/RigHierarchy.h"
+
 //======================================================================================================================
 void FRigPhysicsSolverComponent::Save(FArchive& Ar)
 {
@@ -17,10 +19,3 @@ void FRigPhysicsSolverComponent::Load(FArchive& Ar)
 	Ar << SolverSettings;
 	Ar << SimulationSpaceSettings;
 }
-
-//======================================================================================================================
-bool FRigPhysicsSolverComponent::CanBeAddedTo(
-	const FRigElementKey& InElementKey, const URigHierarchy* InHierarchy) const
-{
-	return InElementKey.IsTypeOf(ERigElementType::Bone);
-}
diff --git a/Engine/Plugins/Experimental/PhysicsControl/Source/ControlRigPhysics/Public/RigPhysicsComponent.h b/Engine/Plugins/Experimental/PhysicsControl/Source/ControlRigPhysics/Public/RigPhysicsComponent.h
--- a/Engine/Plugins/Experimental/PhysicsControl/Source/ControlRigPhysics/Public/RigPhysicsComponent.h
+++ b/Engine/Plugins/Experimental/PhysicsControl/Source/ControlRigPhysics/Public/RigPhysicsComponent.h
@@ -58,5 +58,7 @@ public:
 
 	virtual void Save(FArchive& Ar) override;
 	virtual void Load(FArchive& Ar) override;
+
+	virtual bool CanBeAddedTo(const FRigElementKey& InElementKey, const URigHierarchy* InHierarchy, FString* OutFailureReason) const override;
 };
 
diff --git a/Engine/Plugins/Experimental/PhysicsControl/Source/ControlRigPhysics/Public/RigPhysicsControlComponent.h b/Engine/Plugins/Experimental/PhysicsControl/Source/ControlRigPhysics/Public/RigPhysicsControlComponent.h
--- a/Engine/Plugins/Experimental/PhysicsControl/Source/ControlRigPhysics/Public/RigPhysicsControlComponent.h
+++ b/Engine/Plugins/Experimental/PhysicsControl/Source/ControlRigPhysics/Public/RigPhysicsControlComponent.h
@@ -32,5 +32,7 @@ public:
 
 	virtual void Save(FArchive& A) override;
 	virtual void Load(FArchive& Ar) override;
+
+	virtual bool CanBeAddedTo(const FRigElementKey& InElementKey, const URigHierarchy* InHierarchy, FString* OutFailureReason) const override;
 };
 
diff --git a/Engine/Plugins/Experimental/PhysicsControl/Source/ControlRigPhysics/Public/RigPhysicsSolverComponent.h b/Engine/Plugins/Experimental/PhysicsControl/Source/ControlRigPhysics/Public/RigPhysicsSolverComponent.h
--- a/Engine/Plugins/Experimental/PhysicsControl/Source/ControlRigPhysics/Public/RigPhysicsSolverComponent.h
+++ b/Engine/Plugins/Experimental/PhysicsControl/Source/ControlRigPhysics/Public/RigPhysicsSolverComponent.h
@@ -25,8 +25,6 @@ public:
 	UPROPERTY(BlueprintReadOnly, EditAnywhere, Category = Physics, meta = (ShowOnlyInnerProperties))
 	FRigPhysicsSimulationSpaceSettings SimulationSpaceSettings;
 
-	virtual bool CanBeAddedTo(const FRigElementKey& InElementKey, const URigHierarchy* InHierarchy) const;
-
 	virtual void Save(FArchive& Ar) override;
 	virtual void Load(FArchive& Ar) override;
 };
