Anim In Engine: When selecting anim layers select controls in the new layer in sequencer, deselecting the old ones.
Fix for additive/override undo seleciton issue with combo box

diff --git a/Engine/Plugins/Animation/ControlRig/Source/ControlRigEditor/Private/Sequencer/AnimLayers/AnimLayers.cpp b/Engine/Plugins/Animation/ControlRig/Source/ControlRigEditor/Private/Sequencer/AnimLayers/AnimLayers.cpp
--- a/Engine/Plugins/Animation/ControlRig/Source/ControlRigEditor/Private/Sequencer/AnimLayers/AnimLayers.cpp
+++ b/Engine/Plugins/Animation/ControlRig/Source/ControlRigEditor/Private/Sequencer/AnimLayers/AnimLayers.cpp
@@ -630,6 +630,10 @@ void UAnimLayer::SetSelectedInList(bool bInValue)
 							{
 								for (const TPair<FName, FAnimLayerPropertyAndChannels>& ControlName : SectionItem.AnimLayerSet.Names)
 								{
+									if (bIsSelectedInList && ControlRig->IsControlSelected(ControlName.Key) == false)
+									{
+										continue; //don't select it if not selected
+									}
 									FChannelMapInfo* pChannelIndex = CRSection->ControlChannelMap.Find(ControlName.Key);
 									if (pChannelIndex != nullptr)
 									{
diff --git a/Engine/Plugins/Animation/ControlRig/Source/ControlRigEditor/Private/Sequencer/AnimLayers/SAnimLayers.cpp b/Engine/Plugins/Animation/ControlRig/Source/ControlRigEditor/Private/Sequencer/AnimLayers/SAnimLayers.cpp
--- a/Engine/Plugins/Animation/ControlRig/Source/ControlRigEditor/Private/Sequencer/AnimLayers/SAnimLayers.cpp
+++ b/Engine/Plugins/Animation/ControlRig/Source/ControlRigEditor/Private/Sequencer/AnimLayers/SAnimLayers.cpp
@@ -33,6 +33,10 @@
 #include "Modules/ModuleManager.h"
 #include "Sequencer/ControlRigParameterTrackEditor.h"
 #include "Algo/IndexOf.h"
+#include "MVVM/Selection/SequencerSelectionEventSuppressor.h"
+#include "MVVM/Selection/SequencerCoreSelection.h"
+#include "MVVM/ViewModels/SequencerEditorViewModel.h"
+#include "MVVM/Selection/Selection.h"
 
 #define LOCTEXT_NAMESPACE "AnimLayers"
 
@@ -514,6 +518,7 @@ public:
 	SLATE_ARGUMENT(FAnimLayerSourceUIEntryPtr, Entry)
 	SLATE_END_ARGS()
 
+
 	void Construct(const FArguments& Args, const TSharedRef<STableViewBase>& OwnerTableView)
 	{
 		EntryPtr = Args._Entry;
@@ -811,7 +816,23 @@ private:
 
 	FText GetLayerTypeText() const
 	{
-		return EntryPtr->GetTypeToText();
+		FText CurrentLayerTypeText = EntryPtr->GetTypeToText();
+
+		if (LayerTypeCombo.IsValid() && LayerTypeCombo->GetSelectedItem().IsValid() && LayerTypeCombo->GetSelectedItem()->IdenticalTo(CurrentLayerTypeText) == false)
+		{
+			int32 constexpr AdditiveIndex = 0;
+			int32 constexpr OverrideIndex = 1;
+			if (EntryPtr->GetType() == EAnimLayerType::Additive)
+			{
+				LayerTypeCombo->SetSelectedItem(LayerTypeTextList[AdditiveIndex]);
+			}
+			else
+			{
+				LayerTypeCombo->SetSelectedItem(LayerTypeTextList[OverrideIndex]);
+			}
+		}
+		return CurrentLayerTypeText;
+
 	}
 	void OnLayerTypeChanged(TSharedPtr<FText> InNewText, ESelectInfo::Type SelectInfo)
 	{
@@ -1205,28 +1226,27 @@ TSharedPtr<SWidget> FAnimLayerSourcesView::OnSourceConstructContextMenu()
 
 void FAnimLayerSourcesView::OnSourceListSelectionChanged(FAnimLayerSourceUIEntryPtr Entry, ESelectInfo::Type SelectionType) const
 {
-	TArray<FAnimLayerSourceUIEntryPtr> Selected;
-	SourcesListView->GetSelectedItems(Selected);
-	//first turn off all
+	using namespace UE::Sequencer;
 	if (TSharedPtr<ISequencer> SequencerPtr = UAnimLayers::GetSequencerFromAsset())
 	{
 		if (UAnimLayers* AnimLayers = UAnimLayers::GetAnimLayers(SequencerPtr.Get()))
 		{
+			TArray<FAnimLayerSourceUIEntryPtr> Selected;
+			SourcesListView->GetSelectedItems(Selected);
+			FSelectionEventSuppressor SuppressSelectionEvents = SequencerPtr->GetViewModel()->GetSelection()->SuppressEvents();
+
 			for (UAnimLayer* AnimLayer : AnimLayers->AnimLayers)
 			{
 				AnimLayer->SetSelectedInList(false);
 			}
+			for (const FAnimLayerSourceUIEntryPtr& Item : Selected)
+			{
+				Item->SetSelectedInList(true);
+			}
+			//refresh tree to rerun the filter
+			SequencerPtr->NotifyMovieSceneDataChanged(EMovieSceneDataChangeType::RefreshTree);
 		}
 	}
-	for (const FAnimLayerSourceUIEntryPtr& Item : Selected)
-	{
-		Item->SetSelectedInList(true);
-	}	
-	//refresh tree to rerun the filter
-	if (TSharedPtr<ISequencer> SequencerPtr = UAnimLayers::GetSequencerFromAsset())
-	{
-		SequencerPtr->NotifyMovieSceneDataChanged(EMovieSceneDataChangeType::RefreshTree);
-	}
 }
 
 void FAnimLayerSourcesView::RefreshSourceData(bool bRefreshUI)
