Avoid a crash in Control Rig when using detached tabs and multiple rigs open (when closing one of the not currently visible rigs)

diff --git a/Engine/Plugins/Animation/ControlRig/Source/ControlRigEditor/Private/Editor/ModularRigModelTabSummoner.cpp b/Engine/Plugins/Animation/ControlRig/Source/ControlRigEditor/Private/Editor/ModularRigModelTabSummoner.cpp
--- a/Engine/Plugins/Animation/ControlRig/Source/ControlRigEditor/Private/Editor/ModularRigModelTabSummoner.cpp
+++ b/Engine/Plugins/Animation/ControlRig/Source/ControlRigEditor/Private/Editor/ModularRigModelTabSummoner.cpp
@@ -56,6 +56,10 @@ TSharedRef<SDockTab> FModularRigModelTabSummoner::SpawnTab(const FWorkflowTabSpa
 				{
 					HierarchyTabCount = ControlRigEditorForTab->GetModularRigHierarchyTabCount();
 				}
+				else
+				{
+					return true; // if the editor has been already destroyed, allow closing the tab, so it does not stay alive and crash next frame
+				}
 			}
 		}
 		return HierarchyTabCount > 0;
diff --git a/Engine/Plugins/Animation/ControlRig/Source/ControlRigEditor/Private/Editor/RigHierarchyTabSummoner.cpp b/Engine/Plugins/Animation/ControlRig/Source/ControlRigEditor/Private/Editor/RigHierarchyTabSummoner.cpp
--- a/Engine/Plugins/Animation/ControlRig/Source/ControlRigEditor/Private/Editor/RigHierarchyTabSummoner.cpp
+++ b/Engine/Plugins/Animation/ControlRig/Source/ControlRigEditor/Private/Editor/RigHierarchyTabSummoner.cpp
@@ -57,6 +57,10 @@ TSharedRef<SDockTab> FRigHierarchyTabSummoner::SpawnTab(const FWorkflowTabSpawnI
 				{
 					HierarchyTabCount = ControlRigEditorForTab->GetRigHierarchyTabCount();
 				}
+				else
+				{
+					return true; // if the editor has been already destroyed, allow closing the tab, so it does not stay alive and crash next frame
+				}
 			}
 		}
 		return HierarchyTabCount > 0;
