Control Rig: Avoid running forward solve during interaction with physics solvers present

#rnx

diff --git a/Engine/Plugins/Animation/ControlRig/Source/ControlRig/Private/ControlRig.cpp b/Engine/Plugins/Animation/ControlRig/Source/ControlRig/Private/ControlRig.cpp
--- a/Engine/Plugins/Animation/ControlRig/Source/ControlRig/Private/ControlRig.cpp
+++ b/Engine/Plugins/Animation/ControlRig/Source/ControlRig/Private/ControlRig.cpp
@@ -91,6 +91,7 @@ UControlRig::UControlRig(const FObjectInitializer& ObjectInitializer)
 	, InterRigSyncBracket(0)
 	, ControlUndoBracketIndex(0)
 	, InteractionType((uint8)EControlRigInteractionType::None)
+	, bEvaluationTriggeredFromInteraction(false)
 	, bInteractionJustBegan(false)
 	, bIsAdditive(false)
 	, DebugBoneRadiusMultiplier(1.f)
@@ -829,7 +830,11 @@ bool UControlRig::Execute(const FName& InEventName)
 	// setup the draw interface for debug drawing
 	if(!bIsEventInQueue || bIsEventFirstInQueue)
 	{
-		DrawInterface.Reset();
+		// only reset the draw interface if we are not running the rig during interaction with simulation present. 
+		if(InEventName != FRigUnit_InteractionExecution::EventName || !bEvaluationTriggeredFromInteraction || !ContainsSimulation())
+		{
+			DrawInterface.Reset();
+		}
 	}
 
 	if (bEnableDrawInterface)
@@ -1471,7 +1476,7 @@ bool UControlRig::Execute(const FName& InEventName)
 	{
 		bInteractionJustBegan = false;
 	}
-
+	
 	if(bIsConstructionEvent)
 	{
 		RemoveRunOnceEvent(FRigUnit_PrepareForExecution::EventName);
@@ -1714,6 +1719,13 @@ void UControlRig::AdaptEventQueueForEvaluate(TArray<FName>& InOutEventQueueToRun
 				InOutEventQueueToRun.Insert(InteractionEvent, InOutEventQueueToRun.Num() - 1);
 			}
 		}
+
+		if(ContainsSimulation() && bEvaluationTriggeredFromInteraction)
+		{
+			InOutEventQueueToRun.Remove(FRigUnit_BeginExecution::EventName);
+			InOutEventQueueToRun.Remove(FRigUnit_PreBeginExecution::EventName);
+			InOutEventQueueToRun.Remove(FRigUnit_PostBeginExecution::EventName);
+		}
 	}
 }
 
@@ -3771,5 +3783,10 @@ void UControlRig::SetReplay(UControlRigReplay* InReplay)
 	Replay = InReplay;
 }
 
+bool UControlRig::ContainsSimulation() const
+{
+	return !PhysicsSimulations.IsEmpty();
+}
+
 #undef LOCTEXT_NAMESPACE
 
diff --git a/Engine/Plugins/Animation/ControlRig/Source/ControlRig/Public/ControlRig.h b/Engine/Plugins/Animation/ControlRig/Source/ControlRig/Public/ControlRig.h
--- a/Engine/Plugins/Animation/ControlRig/Source/ControlRig/Public/ControlRig.h
+++ b/Engine/Plugins/Animation/ControlRig/Source/ControlRig/Public/ControlRig.h
@@ -623,6 +623,7 @@ protected:
 	int32 InterRigSyncBracket;
 	int32 ControlUndoBracketIndex;
 	uint8 InteractionType;
+	bool bEvaluationTriggeredFromInteraction;
 	TArray<FRigElementKey> ElementsBeingInteracted;
 #if WITH_EDITOR
 	TArray<TSharedPtr<FRigDirectManipulationInfo>> RigUnitManipulationInfos;
@@ -916,6 +917,8 @@ public:
 	bool IsReplayEnabled() const;
 	void SetReplay(UControlRigReplay* InReplay);
 
+	bool ContainsSimulation() const;
+
 private:
 	int32 ReplayTimeIndex;
 	TObjectPtr<UControlRigReplay> Replay;
diff --git a/Engine/Plugins/Animation/ControlRig/Source/ControlRigEditor/Private/EditMode/ControlRigEditMode.cpp b/Engine/Plugins/Animation/ControlRig/Source/ControlRigEditor/Private/EditMode/ControlRigEditMode.cpp
--- a/Engine/Plugins/Animation/ControlRig/Source/ControlRigEditor/Private/EditMode/ControlRigEditMode.cpp
+++ b/Engine/Plugins/Animation/ControlRig/Source/ControlRigEditor/Private/EditMode/ControlRigEditMode.cpp
@@ -3558,6 +3558,7 @@ void FControlRigEditMode::ZeroTransforms(bool bSelectionOnly)
 		{
 			{
 				const TGuardValue<float> AbsoluteTimeGuard(ControlRig->AbsoluteTime, ControlRig->AbsoluteTime);
+				const TGuardValue<bool> GuardEvaluationType(ControlRig->bEvaluationTriggeredFromInteraction, true);
 				ControlRig->Evaluate_AnyThread();
 			}
 
@@ -3603,6 +3604,7 @@ void FControlRigEditMode::ZeroTransforms(bool bSelectionOnly)
 			//set global transforms in this space // do it twice since ControlsInOrder is not really always in order
 			for (int32 SetHack = 0; SetHack < 2; ++SetHack)
 			{
+				const TGuardValue<bool> GuardEvaluationType(ControlRig->bEvaluationTriggeredFromInteraction, true);
 				ControlRig->Evaluate_AnyThread();
 				for (const FRigElementKey& ElementToReset : TransformElementsToReset)
 				{
@@ -3642,6 +3644,7 @@ void FControlRigEditMode::ZeroTransforms(bool bSelectionOnly)
 			TArray<FName> NewEventQueue = {FRigUnit_InteractionExecution::EventName};
 			NewEventQueue.Append(ControlRig->EventQueue);
 			TGuardValue<TArray<FName>> EventGuard(ControlRig->EventQueue, NewEventQueue);
+			const TGuardValue<bool> GuardEvaluationType(ControlRig->bEvaluationTriggeredFromInteraction, true);
 			ControlRig->Evaluate_AnyThread();
 			for (const FRigElementKey& ControlToReset : ControlsToReset)
 			{
@@ -3697,6 +3700,7 @@ void FControlRigEditMode::InvertInputPose(bool bSelectionOnly)
 		const TArray<FRigControlElement*> ModifiedElements = ControlRig->InvertInputPose(SelectedRigElements, EControlRigSetKey::Never);
 		{
 			const TGuardValue<float> AbsoluteTimeGuard(ControlRig->AbsoluteTime, ControlRig->AbsoluteTime);
+			const TGuardValue<bool> GuardEvaluationType(ControlRig->bEvaluationTriggeredFromInteraction, true);
 			ControlRig->Evaluate_AnyThread();
 		}
 
@@ -4959,6 +4963,7 @@ void FControlRigEditMode::SetControlShapeTransform(
 		{
 			TGuardValue<bool> CompensateGuard(FMovieSceneConstraintChannelHelper::bDoNotCompensate, true);
 			const TGuardValue<float> AbsoluteTimeGuard(ControlRig->AbsoluteTime, ControlRig->AbsoluteTime);
+			const TGuardValue<bool> GuardEvaluationType(ControlRig->bEvaluationTriggeredFromInteraction, true);
 			ControlRig->Evaluate_AnyThread();
 		}
 	};
@@ -5337,6 +5342,7 @@ void FControlRigEditMode::MoveControlShape( AControlRigShapeActor* ShapeActor,
 		else
 		{
 			const TGuardValue<float> AbsoluteTimeGuard(ControlRig->AbsoluteTime, ControlRig->AbsoluteTime);
+			const TGuardValue<bool> GuardEvaluationType(ControlRig->bEvaluationTriggeredFromInteraction, true);
 			ControlRig->Evaluate_AnyThread();
 		}
 
@@ -5359,6 +5365,7 @@ void FControlRigEditMode::MoveControlShape( AControlRigShapeActor* ShapeActor,
 				else
 				{
 					const TGuardValue<float> AbsoluteTimeGuard(ControlRig->AbsoluteTime, ControlRig->AbsoluteTime);
+					const TGuardValue<bool> GuardEvaluationType(ControlRig->bEvaluationTriggeredFromInteraction, true);
 					ControlRig->Evaluate_AnyThread();
 				}
 			}
@@ -5650,6 +5657,8 @@ void FControlRigEditMode::TickManipulatableObjects(float DeltaTime)
 	{
 		if (UControlRig* ControlRig = RuntimeRigPtr.Get())
 		{
+			const TGuardValue<bool> GuardEvaluationType(ControlRig->bEvaluationTriggeredFromInteraction, true);
+			
 			// tick skeletalmeshcomponent, that's how they update their transform from rig change
 			USceneComponent* SceneComponent = GetHostingSceneComponent(ControlRig);
 			if (UControlRigComponent* ControlRigComponent = Cast<UControlRigComponent>(SceneComponent))
@@ -5658,41 +5667,45 @@ void FControlRigEditMode::TickManipulatableObjects(float DeltaTime)
 			}
 			else if (USkeletalMeshComponent* MeshComponent = Cast<USkeletalMeshComponent>(SceneComponent))
 			{
-				// NOTE: we have to update/tick ALL children skeletal mesh components here because user can attach
-				// additional skeletal meshes via the "Copy Pose from Mesh" node.
-				//
-				// If this is left up to the viewport tick(), the attached meshes will render before they get the latest
-				// parent bone transforms resulting in a visible lag on all attached components.
-
-				// get hierarchically ordered list of ALL child skeletal mesh components (recursive)
-				const AActor* ThisActor = MeshComponent->GetOwner();
-				TArray<USceneComponent*> ChildrenComponents;
-				MeshComponent->GetChildrenComponents(true, ChildrenComponents);
-				TInlineComponentArray<USkeletalMeshComponent*> SkeletalMeshesToUpdate;
-				SkeletalMeshesToUpdate.Add(MeshComponent);
-				for (USceneComponent* ChildComponent : ChildrenComponents)
+				if(!ControlRig->ContainsSimulation())
 				{
-					if (USkeletalMeshComponent* ChildMeshComponent = Cast<USkeletalMeshComponent>(ChildComponent))
+					// NOTE: we have to update/tick ALL children skeletal mesh components here because user can attach
+					// additional skeletal meshes via the "Copy Pose from Mesh" node.
+					//
+					// If this is left up to the viewport tick(), the attached meshes will render before they get the latest
+					// parent bone transforms resulting in a visible lag on all attached components.
+
+					// get hierarchically ordered list of ALL child skeletal mesh components (recursive)
+					const AActor* ThisActor = MeshComponent->GetOwner();
+					TArray<USceneComponent*> ChildrenComponents;
+					MeshComponent->GetChildrenComponents(true, ChildrenComponents);
+					TInlineComponentArray<USkeletalMeshComponent*> SkeletalMeshesToUpdate;
+					SkeletalMeshesToUpdate.Add(MeshComponent);
+					for (USceneComponent* ChildComponent : ChildrenComponents)
 					{
-						if (ThisActor == ChildMeshComponent->GetOwner())
+						if (USkeletalMeshComponent* ChildMeshComponent = Cast<USkeletalMeshComponent>(ChildComponent))
 						{
-							SkeletalMeshesToUpdate.Add(ChildMeshComponent);
+							if (ThisActor == ChildMeshComponent->GetOwner())
+							{
+								SkeletalMeshesToUpdate.Add(ChildMeshComponent);
+							}
 						}
 					}
-				}
 
-				// update pose of all children skeletal meshes in this actor
-				for (USkeletalMeshComponent* SkeletalMeshToUpdate : SkeletalMeshesToUpdate)
-				{
-					// "Copy Pose from Mesh" requires AnimInstance::PreUpdate() to copy the parent bone transforms.
-					// have to TickAnimation() to ensure that PreUpdate() is called on all anim instances
-					SkeletalMeshToUpdate->TickAnimation(0.0f, false);
-					SkeletalMeshToUpdate->RefreshBoneTransforms();
-					SkeletalMeshToUpdate->RefreshFollowerComponents	();
-					SkeletalMeshToUpdate->UpdateComponentToWorld();
-					SkeletalMeshToUpdate->FinalizeBoneTransform();
-					SkeletalMeshToUpdate->MarkRenderTransformDirty();
-					SkeletalMeshToUpdate->MarkRenderDynamicDataDirty();
+					// update pose of all children skeletal meshes in this actor
+					for (USkeletalMeshComponent* SkeletalMeshToUpdate : SkeletalMeshesToUpdate)
+					{
+						// "Copy Pose from Mesh" requires AnimInstance::PreUpdate() to copy the parent bone transforms.
+						// have to TickAnimation() to ensure that PreUpdate() is called on all anim instances
+
+						SkeletalMeshToUpdate->TickAnimation(0.0f, false);
+						SkeletalMeshToUpdate->RefreshBoneTransforms();
+						SkeletalMeshToUpdate->RefreshFollowerComponents	();
+						SkeletalMeshToUpdate->UpdateComponentToWorld();
+						SkeletalMeshToUpdate->FinalizeBoneTransform();
+						SkeletalMeshToUpdate->MarkRenderTransformDirty();
+						SkeletalMeshToUpdate->MarkRenderDynamicDataDirty();
+					}
 				}
 			}
 		}
@@ -6642,6 +6655,7 @@ FControlRigEditMode::FPendingControlRigEvaluator::~FPendingControlRigEvaluator()
 		for(UControlRig* ControlRig : EditMode->RigsToEvaluateDuringThisTick)
 		{
 			const TGuardValue<float> AbsoluteTimeGuard(ControlRig->AbsoluteTime, ControlRig->AbsoluteTime);
+			const TGuardValue<bool> GuardEvaluationType(ControlRig->bEvaluationTriggeredFromInteraction, true);
 			ControlRig->Evaluate_AnyThread();
 		}
 		EditMode->RigsToEvaluateDuringThisTick.Reset();
