Exporting FRigUnit_ParentConstraint and FRigUnit_PositionConstraint for use in game code.
Also added a few handy constructors.

diff --git a/Engine/Plugins/Animation/ControlRig/Source/ControlRig/Public/Units/Highlevel/Hierarchy/RigUnit_SphericalPoseReader.h b/Engine/Plugins/Animation/ControlRig/Source/ControlRig/Public/Units/Highlevel/Hierarchy/RigUnit_SphericalPoseReader.h
--- a/Engine/Plugins/Animation/ControlRig/Source/ControlRig/Public/Units/Highlevel/Hierarchy/RigUnit_SphericalPoseReader.h
+++ b/Engine/Plugins/Animation/ControlRig/Source/ControlRig/Public/Units/Highlevel/Hierarchy/RigUnit_SphericalPoseReader.h
@@ -27,6 +27,14 @@ struct CONTROLRIG_API FRegionScaleFactors
 	{
 	}
 
+	FRegionScaleFactors(float InPositiveWidth, float InNegativeWidth, float InPositiveHeight, float InNegativeHeight)
+		: PositiveWidth(InPositiveWidth)
+		, NegativeWidth(InNegativeWidth)
+		, PositiveHeight(InPositiveHeight)
+		, NegativeHeight(InNegativeHeight)
+	{
+	}
+
 	// Scale the region in the POSITIVE width direction. Range is 0-1. Default is 1.0.
 	UPROPERTY(meta = (Input, ClampMin = "0.0", ClampMax = "1.0", UIMin = "0.0", UIMax = "1.0"))
 	float PositiveWidth = 1.0f;
diff --git a/Engine/Plugins/Animation/ControlRig/Source/ControlRig/Public/Units/Highlevel/Hierarchy/RigUnit_TransformConstraint.h b/Engine/Plugins/Animation/ControlRig/Source/ControlRig/Public/Units/Highlevel/Hierarchy/RigUnit_TransformConstraint.h
--- a/Engine/Plugins/Animation/ControlRig/Source/ControlRig/Public/Units/Highlevel/Hierarchy/RigUnit_TransformConstraint.h
+++ b/Engine/Plugins/Animation/ControlRig/Source/ControlRig/Public/Units/Highlevel/Hierarchy/RigUnit_TransformConstraint.h
@@ -207,7 +207,7 @@ struct FRigUnit_ParentConstraint_AdvancedSettings
 * Constrains an item's transform to multiple items' transforms
 */
 USTRUCT(meta=(DisplayName="Parent Constraint", Category="Constraints", Keywords = "Parent,Orient,Scale"))
-struct FRigUnit_ParentConstraint : public FRigUnit_HighlevelBaseMutable
+struct CONTROLRIG_API FRigUnit_ParentConstraint : public FRigUnit_HighlevelBaseMutable
 {
 	GENERATED_BODY()
 
@@ -306,7 +306,7 @@ struct FRigUnit_ParentConstraintMath : public FRigUnit_HighlevelBase
 * Constrains an item's position to multiple items' positions 
 */
 USTRUCT(meta=(DisplayName="Position Constraint", Category="Constraints", Keywords = "Parent,Translation", Deprecated = "5.0"))
-struct FRigUnit_PositionConstraint : public FRigUnit_HighlevelBaseMutable
+struct CONTROLRIG_API FRigUnit_PositionConstraint : public FRigUnit_HighlevelBaseMutable
 {
 	GENERATED_BODY()
 
diff --git a/Engine/Source/Runtime/AnimationCore/Public/Constraint.h b/Engine/Source/Runtime/AnimationCore/Public/Constraint.h
--- a/Engine/Source/Runtime/AnimationCore/Public/Constraint.h
+++ b/Engine/Source/Runtime/AnimationCore/Public/Constraint.h
@@ -58,6 +58,12 @@ struct FFilterOptionPerAxis
 		, bZ(true)
 	{}
 
+	FFilterOptionPerAxis(bool X, bool Y, bool Z)
+		: bX(X)
+		, bY(Y)
+		, bZ(Z)
+	{}
+
 	void FilterVector(FVector& Input, const FVector& ResetValue = FVector::ZeroVector) const
 	{
 		if (!bX)
