[Control Rig] Fixes build break due to missing module dependency on SlateCore in ControlRig.build.cs

diff --git a/Engine/Plugins/Animation/ControlRig/Source/ControlRig/ControlRig.Build.cs b/Engine/Plugins/Animation/ControlRig/Source/ControlRig/ControlRig.Build.cs
--- a/Engine/Plugins/Animation/ControlRig/Source/ControlRig/ControlRig.Build.cs
+++ b/Engine/Plugins/Animation/ControlRig/Source/ControlRig/ControlRig.Build.cs
@@ -19,7 +19,8 @@ namespace UnrealBuildTool.Rules
                     "MovieSceneTracks",
                     "PropertyPath",
 					"TimeManagement",
-					"DeveloperSettings"
+					"DeveloperSettings",
+					"SlateCore"
 				}
             );
 
