Control Rig: Maintain selection during compile and component changes

#rnx

diff --git a/Engine/Plugins/Animation/ControlRig/Source/ControlRig/Private/ControlRig.cpp b/Engine/Plugins/Animation/ControlRig/Source/ControlRig/Private/ControlRig.cpp
--- a/Engine/Plugins/Animation/ControlRig/Source/ControlRig/Private/ControlRig.cpp
+++ b/Engine/Plugins/Animation/ControlRig/Source/ControlRig/Private/ControlRig.cpp
@@ -991,8 +991,8 @@ bool UControlRig::Execute(const FName& InEventName)
 	{
 		check(Hierarchy);
 		
-		// remember the previous selection
-		const TArray<FRigElementKey> PreviousSelection = Hierarchy->GetSelectedKeys();
+		// remember the previous selection (copy here since it may change during the code path below)
+		const TArray<FRigHierarchyKey> PreviousSelection = Hierarchy->GetSelectedHierarchyKeys();
 
 		// construction mode means that we are running the construction event
 		// constantly for testing purposes.
@@ -1089,7 +1089,7 @@ bool UControlRig::Execute(const FName& InEventName)
 			
 			// set it here to reestablish the selection. the notifications
 			// will be eaten since we still have the bSuspend flag on in the controller.
-			Hierarchy->GetController()->SetSelection(PreviousSelection);
+			Hierarchy->GetController()->SetHierarchySelection(PreviousSelection);
 			
 		} // destroy DisableSelectionNotifications
 
@@ -1113,21 +1113,44 @@ bool UControlRig::Execute(const FName& InEventName)
 		}
 
 		// synchronize the selection now with the new hierarchy after running construction
-		const TArray<const FRigBaseElement*> CurrentSelection = Hierarchy->GetSelectedElements();
-		for(const FRigBaseElement* SelectedElement : CurrentSelection)
+		const TArray<FRigHierarchyKey>& CurrentSelection = Hierarchy->GetSelectedHierarchyKeys();
+		for(const FRigHierarchyKey& SelectedHierarchyKey : CurrentSelection)
 		{
-			if(!PreviousSelection.Contains(SelectedElement->GetKey()))
+			if(!PreviousSelection.Contains(SelectedHierarchyKey))
 			{
-				Hierarchy->Notify(ERigHierarchyNotification::ElementSelected, SelectedElement);
+				if(SelectedHierarchyKey.IsElement())
+				{
+					if(const FRigBaseElement* SelectedElement = Hierarchy->Find(SelectedHierarchyKey.GetElement()))
+					{
+						Hierarchy->Notify(ERigHierarchyNotification::ElementSelected, SelectedElement);
+					}
+				}
+				if(SelectedHierarchyKey.IsComponent())
+				{
+					if(const FRigBaseComponent* SelectedComponent = Hierarchy->FindComponent(SelectedHierarchyKey.GetComponent()))
+					{
+						Hierarchy->Notify(ERigHierarchyNotification::ComponentSelected, SelectedComponent);
+					}
+				}
 			}
 		}
-		for(const FRigElementKey& PreviouslySelectedKey : PreviousSelection)
+		for(const FRigHierarchyKey& PreviouslySelectedKey : PreviousSelection)
 		{
-			if(const FRigBaseElement* PreviouslySelectedElement = Hierarchy->Find(PreviouslySelectedKey))
+			if(!CurrentSelection.Contains(PreviouslySelectedKey))
 			{
-				if(!CurrentSelection.Contains(PreviouslySelectedElement))
+				if(PreviouslySelectedKey.IsElement())
 				{
-					Hierarchy->Notify(ERigHierarchyNotification::ElementDeselected, PreviouslySelectedElement);
+					if(const FRigBaseElement* SelectedElement = Hierarchy->Find(PreviouslySelectedKey.GetElement()))
+					{
+						Hierarchy->Notify(ERigHierarchyNotification::ElementDeselected, SelectedElement);
+					}
+				}
+				if(PreviouslySelectedKey.IsComponent())
+				{
+					if(const FRigBaseComponent* SelectedComponent = Hierarchy->FindComponent(PreviouslySelectedKey.GetComponent()))
+					{
+						Hierarchy->Notify(ERigHierarchyNotification::ComponentDeselected, SelectedComponent);
+					}
 				}
 			}
 		}
diff --git a/Engine/Plugins/Animation/ControlRig/Source/ControlRig/Private/Rigs/RigHierarchy.cpp b/Engine/Plugins/Animation/ControlRig/Source/ControlRig/Private/Rigs/RigHierarchy.cpp
--- a/Engine/Plugins/Animation/ControlRig/Source/ControlRig/Private/Rigs/RigHierarchy.cpp
+++ b/Engine/Plugins/Animation/ControlRig/Source/ControlRig/Private/Rigs/RigHierarchy.cpp
@@ -878,6 +878,9 @@ void URigHierarchy::CopyHierarchy(URigHierarchy* InHierarchy)
 		return;
 	}
 
+	// remember the previous selection
+	const TArray<FRigHierarchyKey> PreviousSelection = GetSelectedHierarchyKeys();
+
 	// unlink any existing pose adapter
 	UnlinkPoseAdapter();
 
@@ -1156,6 +1159,14 @@ void URigHierarchy::CopyHierarchy(URigHierarchy* InHierarchy)
 
 	bIsCopyingHierarchy = false;
 	Notify(ERigHierarchyNotification::HierarchyCopied, {});
+
+	if(!PreviousSelection.IsEmpty())
+	{
+		if(URigHierarchyController* Controller = GetController())
+		{
+			Controller->SetHierarchySelection(PreviousSelection, false);
+		}
+	}
 }
 
 uint32 URigHierarchy::GetNameHash() const
@@ -2290,7 +2301,7 @@ TArray<FRigElementKey> URigHierarchy::GetSelectedKeys(ERigElementType InTypeFilt
 	return Selection;
 }
 
-TArray<FRigHierarchyKey> URigHierarchy::GetSelectedHierarchyKeys() const
+const TArray<FRigHierarchyKey>& URigHierarchy::GetSelectedHierarchyKeys() const
 {
 	LLM_SCOPE_BYNAME(TEXT("Animation/ControlRig"));
 	
diff --git a/Engine/Plugins/Animation/ControlRig/Source/ControlRig/Public/Rigs/RigHierarchy.h b/Engine/Plugins/Animation/ControlRig/Source/ControlRig/Public/Rigs/RigHierarchy.h
--- a/Engine/Plugins/Animation/ControlRig/Source/ControlRig/Public/Rigs/RigHierarchy.h
+++ b/Engine/Plugins/Animation/ControlRig/Source/ControlRig/Public/Rigs/RigHierarchy.h
@@ -1972,8 +1972,13 @@ public:
 	 * Returns the keys of selected elements and components
 	 * @return An array of the currently selected elements
 	 */
-	UFUNCTION(BlueprintCallable, Category = URigHierarchy)
-	TArray<FRigHierarchyKey> GetSelectedHierarchyKeys() const;
+	UFUNCTION(BlueprintCallable, Category = URigHierarchy, meta = (DisplayName = "Get Selected Hierarchy Keys", ScriptName = "GetSelectedHierarchyKeys"))
+	TArray<FRigHierarchyKey> GetSelectedHierarchyKeys_ForBlueprint() const
+	{
+		return GetSelectedHierarchyKeys();
+	}
+
+	const TArray<FRigHierarchyKey>& GetSelectedHierarchyKeys() const;
 
 	/**
 	 * Returns true if any element is selected and respects the provided predicate 
diff --git a/Engine/Plugins/Animation/ControlRig/Source/ControlRigEditor/Private/ControlRigElementDetails.cpp b/Engine/Plugins/Animation/ControlRig/Source/ControlRigEditor/Private/ControlRigElementDetails.cpp
--- a/Engine/Plugins/Animation/ControlRig/Source/ControlRigEditor/Private/ControlRigElementDetails.cpp
+++ b/Engine/Plugins/Animation/ControlRig/Source/ControlRigEditor/Private/ControlRigElementDetails.cpp
@@ -1672,20 +1672,30 @@ void FRigBaseElementDetails::CustomizeComponents(IDetailLayoutBuilder& DetailBui
 			{
 				DetailPropertyRow->IsEnabled(false);
 			}
-			DetailPropertyRow->GetPropertyHandle()->SetOnPropertyValueChanged(FSimpleDelegate::CreateLambda([this, ComponentName]()
-			{
-				for(int32 ElementIndex = 0; ElementIndex < PerElementInfos.Num(); ElementIndex++)
-				{
-					if(URigHierarchy* Hierarchy = PerElementInfos[ElementIndex].GetHierarchy())
-					{
-						const FRigComponentKey ComponentKey(PerElementInfos[ElementIndex].Element.GetKey(), ComponentName);
-						if(const FRigBaseComponent* Component = Hierarchy->FindComponent(ComponentKey))
-						{
-							Hierarchy->Notify(ERigHierarchyNotification::ComponentContentChanged, Component);
-						}
-					}
-				}
-			}));
+
+			FSimpleDelegate OnThisOrChildPropertyChanged = FSimpleDelegate::CreateLambda([this, ComponentName]()
+            {
+            	for(int32 ElementIndex = 0; ElementIndex < PerElementInfos.Num(); ElementIndex++)
+            	{
+            		if(URigHierarchy* Hierarchy = PerElementInfos[ElementIndex].GetHierarchy())
+            		{
+            			const FRigComponentKey ComponentKey(PerElementInfos[ElementIndex].Element.GetKey(), ComponentName);
+            			if(const FRigBaseComponent* Component = Hierarchy->FindComponent(ComponentKey))
+            			{
+            				if(URigHierarchy* DefaultHierarchy = PerElementInfos[ElementIndex].GetDefaultHierarchy())
+            				{
+            					if(URigHierarchyController* Controller = DefaultHierarchy->GetController())
+            					{
+            						FRigComponentState State = Component->GetState();
+            						Controller->SetComponentState(ComponentKey, State, true);
+            					}
+            				}
+            			}
+            		}
+            	}
+            });
+			DetailPropertyRow->GetPropertyHandle()->SetOnPropertyValueChanged(OnThisOrChildPropertyChanged);
+			DetailPropertyRow->GetPropertyHandle()->SetOnChildPropertyValueChanged(OnThisOrChildPropertyChanged);
 		}
 	}
 }
diff --git a/Engine/Plugins/Animation/ControlRig/Source/ControlRigEditor/Private/Editor/ControlRigEditor.cpp b/Engine/Plugins/Animation/ControlRig/Source/ControlRigEditor/Private/Editor/ControlRigEditor.cpp
--- a/Engine/Plugins/Animation/ControlRig/Source/ControlRigEditor/Private/Editor/ControlRigEditor.cpp
+++ b/Engine/Plugins/Animation/ControlRig/Source/ControlRigEditor/Private/Editor/ControlRigEditor.cpp
@@ -1590,7 +1590,7 @@ void FControlRigEditor::RefreshDetailView()
 
 bool FControlRigEditor::DetailViewShowsAnyRigElement() const
 {
-	return DetailViewShowsStruct(FRigBaseElement::StaticStruct());
+	return DetailViewShowsStruct(FRigBaseElement::StaticStruct()) || DetailViewShowsStruct(FRigBaseComponent::StaticStruct());
 }
 
 bool FControlRigEditor::DetailViewShowsRigElement(FRigHierarchyKey InKey) const
@@ -1758,7 +1758,9 @@ void FControlRigEditor::Compile()
 			return;
 		}
 
-		TArray< TWeakObjectPtr<UObject> > SelectedObjects = GetSelectedObjectsFromDetailView();
+		const TArray< TWeakObjectPtr<UObject> > SelectedObjects = GetSelectedObjectsFromDetailView();
+		const TArray<FRigHierarchyKey> SelectedHierarchyKeysInDetailsView = GetSelectedRigElementsFromDetailView();
+		const TArray<FRigHierarchyKey> SelectedHierarchyKeysInHierarchy = ControlRigBlueprint->GetHierarchy()->GetSelectedHierarchyKeys();
 
 		if(IsConstructionModeEnabled())
 		{
@@ -1792,11 +1794,6 @@ void FControlRigEditor::Compile()
 				OnAnimInitialized();
 			}
 		}
-		
-		if (SelectedObjects.Num() > 0)
-		{
-			RefreshDetailView();
-		}
 
 		if (UControlRigEditorSettings::Get()->bResetControlTransformsOnCompile)
 		{
@@ -1824,6 +1821,19 @@ void FControlRigEditor::Compile()
 		{
 			EditMode->RecreateControlShapeActors();
 		}
+		
+		if (!SelectedHierarchyKeysInDetailsView.IsEmpty())
+		{
+			SetDetailViewForRigElements(SelectedHierarchyKeysInDetailsView);
+		}
+		else if (!SelectedObjects.IsEmpty())
+		{
+			RefreshDetailView();
+		}
+		if(!SelectedHierarchyKeysInHierarchy.IsEmpty())
+		{
+			ControlRigBlueprint->GetHierarchyController()->SetHierarchySelection(SelectedHierarchyKeysInHierarchy);
+		}
 	}
 }
 
diff --git a/Engine/Plugins/Animation/ControlRig/Source/ControlRigEditor/Private/Editor/SRigHierarchyTreeView.h b/Engine/Plugins/Animation/ControlRig/Source/ControlRigEditor/Private/Editor/SRigHierarchyTreeView.h
--- a/Engine/Plugins/Animation/ControlRig/Source/ControlRigEditor/Private/Editor/SRigHierarchyTreeView.h
+++ b/Engine/Plugins/Animation/ControlRig/Source/ControlRigEditor/Private/Editor/SRigHierarchyTreeView.h
@@ -151,14 +151,7 @@ struct CONTROLRIGEDITOR_API FRigTreeDelegates
 		}
 		if (const URigHierarchy* Hierarchy = GetHierarchy())
 		{
-			const TArray<FRigElementKey> SelectedElementKeys = Hierarchy->GetSelectedKeys(); 
-			TArray<FRigHierarchyKey> Selection;
-			Selection.Reserve(SelectedElementKeys.Num());
-			for(const FRigElementKey& SelectedElementKey : SelectedElementKeys)
-			{
-				Selection.Emplace(SelectedElementKey);
-			}
-			return Selection;
+			return Hierarchy->GetSelectedHierarchyKeys(); 
 		}
 		return {};
 	}
