Control Rig: Remove ref for name argument exposed to AddParent BP functions

#rnx

diff --git a/Engine/Plugins/Animation/ControlRig/Source/ControlRig/Private/Rigs/RigHierarchyController.cpp b/Engine/Plugins/Animation/ControlRig/Source/ControlRig/Private/Rigs/RigHierarchyController.cpp
--- a/Engine/Plugins/Animation/ControlRig/Source/ControlRig/Private/Rigs/RigHierarchyController.cpp
+++ b/Engine/Plugins/Animation/ControlRig/Source/ControlRig/Private/Rigs/RigHierarchyController.cpp
@@ -3524,7 +3524,7 @@ FName URigHierarchyController::SetDisplayName(FRigControlElement* InControlEleme
 	return InControlElement->Settings.DisplayName;
 }
 
-bool URigHierarchyController::AddParent(FRigElementKey InChild, FRigElementKey InParent, float InWeight, bool bMaintainGlobalTransform, const FName& InDisplayLabel, bool bSetupUndo)
+bool URigHierarchyController::AddParent(FRigElementKey InChild, FRigElementKey InParent, float InWeight, bool bMaintainGlobalTransform, FName InDisplayLabel, bool bSetupUndo)
 {
 	if(!IsValid())
 	{
@@ -4101,7 +4101,7 @@ bool URigHierarchyController::SetParent(FRigElementKey InChild, FRigElementKey I
 	return bParentSet;
 }
 
-bool URigHierarchyController::AddAvailableSpace(FRigElementKey InControl, FRigElementKey InSpace, const FName& InDisplayLabel, bool bSetupUndo, bool bPrintPythonCommand)
+bool URigHierarchyController::AddAvailableSpace(FRigElementKey InControl, FRigElementKey InSpace, FName InDisplayLabel, bool bSetupUndo, bool bPrintPythonCommand)
 {
 	if(!IsValid())
 	{
diff --git a/Engine/Plugins/Animation/ControlRig/Source/ControlRig/Public/Rigs/RigHierarchyController.h b/Engine/Plugins/Animation/ControlRig/Source/ControlRig/Public/Rigs/RigHierarchyController.h
--- a/Engine/Plugins/Animation/ControlRig/Source/ControlRig/Public/Rigs/RigHierarchyController.h
+++ b/Engine/Plugins/Animation/ControlRig/Source/ControlRig/Public/Rigs/RigHierarchyController.h
@@ -693,7 +693,7 @@ public:
 	 * @return Returns true if successful.
 	 */
 	UFUNCTION(BlueprintCallable, Category = URigHierarchyController)
-	bool AddParent(FRigElementKey InChild, FRigElementKey InParent, float InWeight = 0.f, bool bMaintainGlobalTransform = true, const FName& InDisplayLabel = NAME_None, bool bSetupUndo = false);
+	bool AddParent(FRigElementKey InChild, FRigElementKey InParent, float InWeight = 0.f, bool bMaintainGlobalTransform = true, FName InDisplayLabel = NAME_None, bool bSetupUndo = false);
 
 	/**
 	* Adds a new parent to an element. For elements that allow only one parent the parent will be replaced (Same as ::SetParent).
@@ -752,7 +752,7 @@ public:
 	 * @return Returns true if successful.
 	 */
 	UFUNCTION(BlueprintCallable, Category = URigHierarchyController)
-	bool AddAvailableSpace(FRigElementKey InControl, FRigElementKey InSpace, const FName& InDisplayLabel = NAME_None, bool bSetupUndo = false, bool bPrintPythonCommand = false);
+	bool AddAvailableSpace(FRigElementKey InControl, FRigElementKey InSpace, FName InDisplayLabel = NAME_None, bool bSetupUndo = false, bool bPrintPythonCommand = false);
 
 	/**
 	 * Adds a new available space to the given control
