Control Rig: Move handling of hierarchy key changes to blueprint. Also allow components to react to key changes

Due to object version change here's a full cook preflight as well https://horde.devtools.epicgames.com/job/67b0b665b40c03764867a6a5

#rnx

diff --git a/Engine/Plugins/Animation/ControlRig/Source/ControlRig/Private/Rigs/RigHierarchy.cpp b/Engine/Plugins/Animation/ControlRig/Source/ControlRig/Private/Rigs/RigHierarchy.cpp
--- a/Engine/Plugins/Animation/ControlRig/Source/ControlRig/Private/Rigs/RigHierarchy.cpp
+++ b/Engine/Plugins/Animation/ControlRig/Source/ControlRig/Private/Rigs/RigHierarchy.cpp
@@ -265,8 +265,8 @@ void URigHierarchy::Save(FArchive& Ar)
 		Element->Serialize(*ArchiveForElements, SerializationSettings);
 	}
 
-	*ArchiveForElements << PreviousNameMap;
-	*ArchiveForElements << PreviousParentMap;
+	*ArchiveForElements << PreviousHierarchyNameMap;
+	*ArchiveForElements << PreviousHierarchyNameMap;
 
 	{
 		TMap<FRigElementKey, FMetadataStorage> ElementMetadataToSave;
@@ -625,13 +625,35 @@ void URigHierarchy::Load(FArchive& Ar)
 
 	if (ArchiveForElements->CustomVer(FControlRigObjectVersion::GUID) >= FControlRigObjectVersion::RigHierarchyStoringPreviousNames)
 	{
-		*ArchiveForElements << PreviousNameMap;
-		*ArchiveForElements << PreviousParentMap;
+		if (ArchiveForElements->CustomVer(FControlRigObjectVersion::GUID) >= FControlRigObjectVersion::RigHierarchyPreviousNameAndParentMapUsingHierarchyKey)
+		{
+			*ArchiveForElements << PreviousHierarchyNameMap;
+			*ArchiveForElements << PreviousHierarchyParentMap;
+		}
+		else
+		{
+			TMap<FRigElementKey, FRigElementKey> PreviousNameMap;
+			TMap<FRigElementKey, FRigElementKey> PreviousParentMap;
+			*ArchiveForElements << PreviousNameMap;
+			*ArchiveForElements << PreviousParentMap;
+
+			PreviousHierarchyNameMap.Reset();
+			for(const TPair<FRigElementKey, FRigElementKey>& Pair : PreviousNameMap)
+			{
+				PreviousHierarchyNameMap.Add(Pair.Key, Pair.Value);
+			}
+
+			PreviousHierarchyParentMap.Reset();
+			for(const TPair<FRigElementKey, FRigElementKey>& Pair : PreviousParentMap)
+			{
+				PreviousHierarchyParentMap.Add(Pair.Key, Pair.Value);
+			}
+		}
 	}
 	else
 	{
-		PreviousNameMap.Reset();
-		PreviousParentMap.Reset();
+		PreviousHierarchyNameMap.Reset();
+		PreviousHierarchyParentMap.Reset();
 	}
 
 	if (ArchiveForElements->CustomVer(FControlRigObjectVersion::GUID) >= FControlRigObjectVersion::RigHierarchyStoresElementMetadata)
@@ -1101,7 +1123,8 @@ void URigHierarchy::CopyHierarchy(URigHierarchy* InHierarchy)
 
 		(void)SortElementStorage();
 
-		PreviousNameMap.Append(InHierarchy->PreviousNameMap);
+		PreviousHierarchyNameMap.Append(InHierarchy->PreviousHierarchyNameMap);
+		PreviousHierarchyParentMap.Append(InHierarchy->PreviousHierarchyParentMap);
 
 		// make sure the curves are marked as unset after copying
 		UnsetCurveValues();
@@ -6034,20 +6057,30 @@ void URigHierarchy::UnsetCurveValue(FRigCurveElement* InCurveElement, bool bSetu
 
 FName URigHierarchy::GetPreviousName(const FRigElementKey& InKey) const
 {
-	if(const FRigElementKey* OldKeyPtr = PreviousNameMap.Find(InKey))
+	return GetPreviousHierarchyName(InKey);
+}
+
+FName URigHierarchy::GetPreviousHierarchyName(const FRigHierarchyKey& InKey) const
+{
+	if(const FRigHierarchyKey* OldKeyPtr = PreviousHierarchyNameMap.Find(InKey))
 	{
-		return OldKeyPtr->Name;
+		return OldKeyPtr->GetFName();
 	}
 	return NAME_None;
 }
 
 FRigElementKey URigHierarchy::GetPreviousParent(const FRigElementKey& InKey) const
 {
-	if(const FRigElementKey* OldParentPtr = PreviousParentMap.Find(InKey))
+	return GetPreviousHierarchyParent(InKey).GetElement();
+}
+
+FRigHierarchyKey URigHierarchy::GetPreviousHierarchyParent(const FRigHierarchyKey& InKey) const
+{
+	if(const FRigHierarchyKey* OldParentPtr = PreviousHierarchyParentMap.Find(InKey))
 	{
 		return *OldParentPtr;
 	}
-	return FRigElementKey();
+	return FRigHierarchyKey();
 }
 
 bool URigHierarchy::IsParentedTo(FRigBaseElement* InChild, FRigBaseElement* InParent, const TElementDependencyMap& InDependencyMap) const
diff --git a/Engine/Plugins/Animation/ControlRig/Source/ControlRig/Private/Rigs/RigHierarchyComponents.cpp b/Engine/Plugins/Animation/ControlRig/Source/ControlRig/Private/Rigs/RigHierarchyComponents.cpp
--- a/Engine/Plugins/Animation/ControlRig/Source/ControlRig/Private/Rigs/RigHierarchyComponents.cpp
+++ b/Engine/Plugins/Animation/ControlRig/Source/ControlRig/Private/Rigs/RigHierarchyComponents.cpp
@@ -9,6 +9,7 @@
 #include "Serialization/MemoryWriter.h"
 #include "UObject/AnimObjectVersion.h"
 #include "ControlRigObjectVersion.h"
+#include "HAL/UnrealMemory.h"
 
 #include UE_INLINE_GENERATED_CPP_BY_NAME(RigHierarchyComponents)
 
@@ -26,6 +27,34 @@ const UScriptStruct* FRigComponentState::GetComponentStruct() const
 	return ComponentStruct;
 }
 
+bool FRigComponentState::operator==(const FRigComponentState& InOther) const
+{
+	if(IsValid() != InOther.IsValid())
+	{
+		return false;
+	}
+
+	if(IsValid())
+	{
+		if(GetComponentStruct() != InOther.GetComponentStruct())
+		{
+			return false;
+		}
+
+		if(Data.Num() != InOther.Data.Num())
+		{
+			return false;
+		}
+
+		if(FMemory::Memcmp(Data.GetData(), InOther.Data.GetData(), Data.Num()) != 0)
+		{
+			return false;
+		}
+	}
+	
+	return true;
+}
+
 ////////////////////////////////////////////////////////////////////////////////
 // FRigBaseComponent
 ////////////////////////////////////////////////////////////////////////////////
diff --git a/Engine/Plugins/Animation/ControlRig/Source/ControlRig/Private/Rigs/RigHierarchyController.cpp b/Engine/Plugins/Animation/ControlRig/Source/ControlRig/Private/Rigs/RigHierarchyController.cpp
--- a/Engine/Plugins/Animation/ControlRig/Source/ControlRig/Private/Rigs/RigHierarchyController.cpp
+++ b/Engine/Plugins/Animation/ControlRig/Source/ControlRig/Private/Rigs/RigHierarchyController.cpp
@@ -1099,7 +1099,7 @@ FRigComponentKey URigHierarchyController::RenameComponent(FRigComponentKey InCom
 	}
 #endif
 
-	const bool bRenamed = RenameComponent(Component, InName, bClearSelection);
+	const bool bRenamed = RenameComponent(Component, InName, bClearSelection, bSetupUndo);
 
 #if WITH_EDITOR
 	if(!bRenamed && TransactionPtr.IsValid())
@@ -1158,7 +1158,7 @@ FRigComponentKey URigHierarchyController::ReparentComponent(FRigComponentKey InC
 	}
 #endif
 
-	const bool bReparented = ReparentComponent(Component, ParentElement, bClearSelection);
+	const bool bReparented = ReparentComponent(Component, ParentElement, bClearSelection, bSetupUndo);
 
 #if WITH_EDITOR
 	if(!bReparented && TransactionPtr.IsValid())
@@ -1292,6 +1292,12 @@ bool URigHierarchyController::SetComponentState(FRigComponentKey InComponent, co
 		return false;
 	}
 
+	const FRigComponentState CurrentState = Component->GetState();
+	if(CurrentState == InState)
+	{
+		return false;
+	}
+
 #if WITH_EDITOR
 	TSharedPtr<FScopedTransaction> TransactionPtr;
 	if(bSetupUndo)
@@ -1750,6 +1756,38 @@ USkeleton* URigHierarchyController::GetSkeletonFromAssetPath(const FString& InAs
 
 #endif
 
+void URigHierarchyController::UpdateComponentsOnHierarchyKeyChange(const TArray<TPair<FRigHierarchyKey, FRigHierarchyKey>>& InKeyMap, bool bSetupUndoRedo)
+{
+	URigHierarchy* Hierarchy = GetHierarchy();
+
+	// let all components know that their content may have changed
+	const int32 NumComponents = Hierarchy->NumComponents();
+	for(int32 ComponentIndex = 0; ComponentIndex < NumComponents; ComponentIndex++)
+	{
+		const FRigBaseComponent* CurrentComponent = Hierarchy->GetComponent(ComponentIndex);
+		if(CurrentComponent->IsProcedural())
+		{
+			continue;
+		}
+
+		const FRigComponentState OldState = CurrentComponent->GetState();
+		FStructOnScope StructOnScope(CurrentComponent->GetScriptStruct());
+		FRigBaseComponent* TempComponent = reinterpret_cast<FRigBaseComponent*>(StructOnScope.GetStructMemory());
+		TempComponent->SetState(OldState);
+
+		for(const TPair<FRigHierarchyKey, FRigHierarchyKey>& Pair : InKeyMap)
+		{
+			TempComponent->OnRigHierarchyKeyChanged(Pair.Key, Pair.Value);
+		}
+
+		const FRigComponentState NewState = TempComponent->GetState();
+		if(NewState != OldState)
+		{
+			SetComponentState(CurrentComponent->GetKey(), NewState, bSetupUndoRedo);
+		}
+	}
+}
+
 TArray<FRigElementKey> URigHierarchyController::ImportCurves(UAnimCurveMetaData* InAnimCurvesMetadata, FName InNameSpace, bool bSetupUndo)
 {
 	TArray<FRigElementKey> Keys;
@@ -3035,7 +3073,7 @@ FRigElementKey URigHierarchyController::RenameElement(FRigElementKey InElement,
 	}
 #endif
 
-	const bool bRenamed = RenameElement(Element, InName, bClearSelection);
+	const bool bRenamed = RenameElement(Element, InName, bClearSelection, bSetupUndo);
 
 #if WITH_EDITOR
 	if(!bRenamed && TransactionPtr.IsValid())
@@ -3169,7 +3207,7 @@ FName URigHierarchyController::SetDisplayName(FRigElementKey InControl, FName In
 	return NewDisplayName;
 }
 
-bool URigHierarchyController::RenameElement(FRigBaseElement* InElement, const FName &InName, bool bClearSelection)
+bool URigHierarchyController::RenameElement(FRigBaseElement* InElement, const FName &InName, bool bClearSelection, bool bSetupUndoRedo)
 {
 	if(InElement == nullptr)
 	{
@@ -3213,13 +3251,20 @@ bool URigHierarchyController::RenameElement(FRigBaseElement* InElement, const FN
 	{
 		Hierarchy->ComponentIndexLookup.Remove(ComponentKey);
 	}
+
+	TArray<TPair<FRigHierarchyKey, FRigHierarchyKey>> ChangedHierarchyKeys;
+	ChangedHierarchyKeys.Reserve(InElement->NumComponents() + 1);
+	ChangedHierarchyKeys.Emplace(OldKey, NewKey);
+	
 	for(int32 ComponentIndex = 0; ComponentIndex < InElement->NumComponents(); ComponentIndex++)
 	{
 		FRigBaseComponent* Component = InElement->GetComponent(ComponentIndex);
 		if(ensure(Component))
 		{
+			const FRigComponentKey OldComponentKey = Component->GetKey();
 			Component->Key.ElementKey = NewKey;
 			Hierarchy->ComponentIndexLookup.Add(Component->GetKey(), Component->GetIndexInHierarchy());
+			ChangedHierarchyKeys.Emplace(OldComponentKey, Component->GetKey());
 		}
 	}
 
@@ -3257,8 +3302,11 @@ bool URigHierarchyController::RenameElement(FRigBaseElement* InElement, const FN
 		}
 	}
 
-	Hierarchy->PreviousNameMap.FindOrAdd(NewKey) = OldKey;
+	Hierarchy->PreviousHierarchyNameMap.FindOrAdd(NewKey) = OldKey;
 	Hierarchy->IncrementTopologyVersion();
+
+	UpdateComponentsOnHierarchyKeyChange(ChangedHierarchyKeys, bSetupUndoRedo);
+
 	Notify(ERigHierarchyNotification::ElementRenamed, InElement);
 
 	if (!bClearSelection && bWasSelected)
@@ -3269,7 +3317,7 @@ bool URigHierarchyController::RenameElement(FRigBaseElement* InElement, const FN
 	return true;
 }
 
-bool URigHierarchyController::RenameComponent(FRigBaseComponent* InComponent, const FName& InName, bool bClearSelection)
+bool URigHierarchyController::RenameComponent(FRigBaseComponent* InComponent, const FName& InName, bool bClearSelection, bool bSetupUndoRedo)
 {
 	if(InComponent == nullptr)
 	{
@@ -3307,8 +3355,12 @@ bool URigHierarchyController::RenameComponent(FRigBaseComponent* InComponent, co
 
 	Hierarchy->ComponentIndexLookup.Remove(OldKey);
 	Hierarchy->ComponentIndexLookup.Add(NewKey, InComponent->IndexInHierarchy);
+	Hierarchy->PreviousHierarchyNameMap.FindOrAdd(NewKey) = OldKey;
 
 	Hierarchy->IncrementTopologyVersion();
+
+	UpdateComponentsOnHierarchyKeyChange({{OldKey, NewKey}}, bSetupUndoRedo);
+
 	Notify(ERigHierarchyNotification::ComponentRenamed, InComponent);
 
 	if (!bClearSelection && bWasSelected)
@@ -3319,7 +3371,7 @@ bool URigHierarchyController::RenameComponent(FRigBaseComponent* InComponent, co
 	return true;
 }
 
-bool URigHierarchyController::ReparentComponent(FRigBaseComponent* InComponent, FRigBaseElement* InParentElement, bool bClearSelection)
+bool URigHierarchyController::ReparentComponent(FRigBaseComponent* InComponent, FRigBaseElement* InParentElement, bool bClearSelection, bool bSetupUndoRedo)
 {
 	if(InComponent == nullptr)
 	{
@@ -3396,8 +3448,12 @@ bool URigHierarchyController::ReparentComponent(FRigBaseComponent* InComponent,
 
 	Hierarchy->ComponentIndexLookup.Remove(OldKey);
 	Hierarchy->ComponentIndexLookup.Add(NewKey, InComponent->IndexInHierarchy);
+	Hierarchy->PreviousHierarchyParentMap.Add(NewKey, FRigHierarchyKey(OldKey.ElementKey, true));
 
 	Hierarchy->IncrementTopologyVersion();
+
+	UpdateComponentsOnHierarchyKeyChange({{OldKey, NewKey}}, bSetupUndoRedo);
+
 	Notify(ERigHierarchyNotification::ComponentReparented, InComponent);
 
 	if (!bClearSelection && bWasSelected)
@@ -3862,7 +3918,7 @@ bool URigHierarchyController::RemoveParent(FRigBaseElement* InChild, FRigBaseEle
 			}
 
 			const FRigElementKey PreviousParentKey = SingleParentElement->ParentElement->GetKey();
-			Hierarchy->PreviousParentMap.FindOrAdd(SingleParentElement->GetKey()) = PreviousParentKey;
+			Hierarchy->PreviousHierarchyParentMap.FindOrAdd(SingleParentElement->GetKey()) = PreviousParentKey;
 			
 			// remove the previous parent
 			SingleParentElement->ParentElement = nullptr;
@@ -3924,7 +3980,7 @@ bool URigHierarchyController::RemoveParent(FRigBaseElement* InChild, FRigBaseEle
 			RemoveElementToDirty(InParent, MultiParentElement); 
 
 			const FRigElementKey PreviousParentKey = MultiParentElement->ParentConstraints[ParentIndex].ParentElement->GetKey();
-			Hierarchy->PreviousParentMap.FindOrAdd(MultiParentElement->GetKey()) = PreviousParentKey;
+			Hierarchy->PreviousHierarchyParentMap.FindOrAdd(MultiParentElement->GetKey()) = PreviousParentKey;
 
 			MultiParentElement->ParentConstraints.RemoveAt(ParentIndex);
 			MultiParentElement->IndexLookup.Remove(ParentTransformElement->GetKey());
diff --git a/Engine/Plugins/Animation/ControlRig/Source/ControlRig/Public/ControlRigObjectVersion.h b/Engine/Plugins/Animation/ControlRig/Source/ControlRig/Public/ControlRigObjectVersion.h
--- a/Engine/Plugins/Animation/ControlRig/Source/ControlRig/Public/ControlRigObjectVersion.h
+++ b/Engine/Plugins/Animation/ControlRig/Source/ControlRig/Public/ControlRigObjectVersion.h
@@ -122,6 +122,9 @@ struct CONTROLRIG_API FControlRigObjectVersion
 		// Parent constraints offering a display label
 		RigHierarchyParentContraintWithLabel,
 
+		// Previous name and parent maps serialized as FRigHierarchyKey
+		RigHierarchyPreviousNameAndParentMapUsingHierarchyKey,
+
 		// -----<new versions can be added above this line>-------------------------------------------------
 		VersionPlusOne,
 		LatestVersion = VersionPlusOne - 1
diff --git a/Engine/Plugins/Animation/ControlRig/Source/ControlRig/Public/Rigs/RigHierarchy.h b/Engine/Plugins/Animation/ControlRig/Source/ControlRig/Public/Rigs/RigHierarchy.h
--- a/Engine/Plugins/Animation/ControlRig/Source/ControlRig/Public/Rigs/RigHierarchy.h
+++ b/Engine/Plugins/Animation/ControlRig/Source/ControlRig/Public/Rigs/RigHierarchy.h
@@ -4466,6 +4466,13 @@ public:
 	UFUNCTION(BlueprintCallable, Category = URigHierarchy)
 	FName GetPreviousName(const FRigElementKey& InKey) const;
 
+	/**
+	 * Returns the previous name of an element or a component prior to a rename operation
+	 * @param InKey The key to request the old name for
+	 */
+	UFUNCTION(BlueprintCallable, Category = URigHierarchy)
+	FName GetPreviousHierarchyName(const FRigHierarchyKey& InKey) const;
+
 	/**
 	 * Returns the previous parent of an element prior to a reparent operation
 	 * @param InKey The key of the element to request the old parent  for
@@ -4473,6 +4480,13 @@ public:
 	UFUNCTION(BlueprintCallable, Category = URigHierarchy)
 	FRigElementKey GetPreviousParent(const FRigElementKey& InKey) const;
 
+	/**
+	 * Returns the previous parent of an element or a component prior to a reparent operation
+	 * @param InKey The key to request the old parent  for
+	 */
+	UFUNCTION(BlueprintCallable, Category = URigHierarchy)
+	FRigHierarchyKey GetPreviousHierarchyParent(const FRigHierarchyKey& InKey) const;
+
 	/**
 	 * Returns true if an element is parented to another element
 	 * @param InChild The child element to check for a parent
@@ -4929,11 +4943,10 @@ private:
 	UPROPERTY(Transient)
 	mutable TObjectPtr<UModularRigRuleManager> RuleManager;
 
-	TMap<FRigElementKey, FRigElementKey> PreviousParentMap;
+	TMap<FRigHierarchyKey, FRigHierarchyKey> PreviousHierarchyParentMap;
 
 	/*We save this so Sequencer can remap this after load*/
-	UPROPERTY()
-	TMap<FRigElementKey, FRigElementKey> PreviousNameMap;
+	TMap<FRigHierarchyKey, FRigHierarchyKey> PreviousHierarchyNameMap;
 
 	int32 ResetPoseHash;
 	TArray<bool> ResetPoseIsFilteredOut;
diff --git a/Engine/Plugins/Animation/ControlRig/Source/ControlRig/Public/Rigs/RigHierarchyComponents.h b/Engine/Plugins/Animation/ControlRig/Source/ControlRig/Public/Rigs/RigHierarchyComponents.h
--- a/Engine/Plugins/Animation/ControlRig/Source/ControlRig/Public/Rigs/RigHierarchyComponents.h
+++ b/Engine/Plugins/Animation/ControlRig/Source/ControlRig/Public/Rigs/RigHierarchyComponents.h
@@ -46,6 +46,8 @@ public:
 	bool IsValid() const;
 	const UScriptStruct* GetComponentStruct() const;
 
+	bool operator == (const FRigComponentState& InOther) const;
+
 private:
 
 	const UScriptStruct* ComponentStruct = nullptr;
@@ -104,6 +106,9 @@ public:
 
 	// returns the color to use for this component in the UI
 	virtual FSlateColor GetColorForUI() const;
+
+	// react to an element or component being renamed / reparented in the hierarchy
+	virtual void OnRigHierarchyKeyChanged(const FRigHierarchyKey& InOldKey, const FRigHierarchyKey& InNewKey) {}
 	
 protected:
 
diff --git a/Engine/Plugins/Animation/ControlRig/Source/ControlRig/Public/Rigs/RigHierarchyController.h b/Engine/Plugins/Animation/ControlRig/Source/ControlRig/Public/Rigs/RigHierarchyController.h
--- a/Engine/Plugins/Animation/ControlRig/Source/ControlRig/Public/Rigs/RigHierarchyController.h
+++ b/Engine/Plugins/Animation/ControlRig/Source/ControlRig/Public/Rigs/RigHierarchyController.h
@@ -972,7 +972,7 @@ private:
 	 * @param bClearSelection True if the selection should be cleared after a rename
 	 * @return Returns true if successful.
 	 */
-    bool RenameElement(FRigBaseElement* InElement, const FName &InName, bool bClearSelection = true);
+    bool RenameElement(FRigBaseElement* InElement, const FName &InName, bool bClearSelection = true, bool bSetupUndoRedo = false);
 
 	/**
 	 * Renames an existing component in the hierarchy
@@ -981,7 +981,7 @@ private:
 	 * @param bClearSelection True if the selection should be cleared after a rename
 	 * @return Returns true if successful.
 	 */
-	bool RenameComponent(FRigBaseComponent* InComponent, const FName &InName, bool bClearSelection = true);
+	bool RenameComponent(FRigBaseComponent* InComponent, const FName &InName, bool bClearSelection = true, bool bSetupUndoRedo = false);
 
 	/**
 	 * Reparents an existing component in the hierarchy
@@ -990,7 +990,7 @@ private:
 	 * @param bClearSelection True if the selection should be cleared after a reparenting
 	 * @return Returns true if successful.
 	 */
-	bool ReparentComponent(FRigBaseComponent* InComponent, FRigBaseElement* InParentElement, bool bClearSelection = true);
+	bool ReparentComponent(FRigBaseComponent* InComponent, FRigBaseElement* InParentElement, bool bClearSelection = true, bool bSetupUndoRedo = false);
 
 	/**
 	 * Changes the element's index within its default parent (or the top level)
@@ -1083,6 +1083,8 @@ private:
 	static USkeleton* GetSkeletonFromAssetPath(const FString& InAssetPath);
 #endif
 
+	void UpdateComponentsOnHierarchyKeyChange(const TArray<TPair<FRigHierarchyKey, FRigHierarchyKey>>& InKeyMap, bool bSetupUndoRedo);
+
 	/** 
 	 * If set to true all notifications coming from this hierarchy will be suspended
 	 */
diff --git a/Engine/Plugins/Animation/ControlRig/Source/ControlRigDeveloper/Private/ControlRigBlueprint.cpp b/Engine/Plugins/Animation/ControlRig/Source/ControlRigDeveloper/Private/ControlRigBlueprint.cpp
--- a/Engine/Plugins/Animation/ControlRig/Source/ControlRigDeveloper/Private/ControlRigBlueprint.cpp
+++ b/Engine/Plugins/Animation/ControlRig/Source/ControlRigDeveloper/Private/ControlRigBlueprint.cpp
@@ -2743,17 +2743,12 @@ void UControlRigBlueprint::HandleHierarchyModified(ERigHierarchyNotification InN
 		case ERigHierarchyNotification::ElementRenamed:
 		{
 			Modify();
-			const FRigElementKey OldKey(InHierarchy->GetPreviousName(InElement->GetKey()), InElement->GetType());
-			Influences.OnKeyRenamed(OldKey, InElement->GetKey());
-			if (IsControlRigModule() && InElement->IsTypeOf(ERigElementType::Connector))
+			if(InElement)
 			{
-				if (FRigElementKeyCollection* Targets = ArrayConnectionMap.Find(OldKey))
-				{
-					ArrayConnectionMap.FindOrAdd(InElement->GetKey(), *Targets);
-					ArrayConnectionMap.Remove(OldKey);
-				}
+				const FName PreviousName = InHierarchy->GetPreviousHierarchyName(InElement->GetKey());
+				const FRigElementKey OldKey(PreviousName, InElement->GetType());
+				HandleHierarchyElementKeyChanged(OldKey, InElement->GetKey());
 			}
-			PropagateHierarchyFromBPToInstances();
 			break;
 		}
 		case ERigHierarchyNotification::ElementAdded:
@@ -2762,14 +2757,37 @@ void UControlRigBlueprint::HandleHierarchyModified(ERigHierarchyNotification InN
 		case ERigHierarchyNotification::HierarchyReset:
 		case ERigHierarchyNotification::ComponentAdded:
 		case ERigHierarchyNotification::ComponentRemoved:
-		case ERigHierarchyNotification::ComponentRenamed:
-		case ERigHierarchyNotification::ComponentReparented:
 		case ERigHierarchyNotification::ComponentContentChanged:
 		{
 			Modify();
 			PropagateHierarchyFromBPToInstances();
 			break;
 		}
+		case ERigHierarchyNotification::ComponentRenamed:
+		{
+			Modify();
+			if(InComponent)
+			{
+				const FName PreviousName = InHierarchy->GetPreviousHierarchyName(InComponent->GetKey());
+				const FRigComponentKey OldKey(InComponent->GetElementKey(), PreviousName);
+				HandleHierarchyComponentKeyChanged(OldKey, InComponent->GetKey());
+			}
+			break;
+		}
+		case ERigHierarchyNotification::ComponentReparented:
+		{
+			Modify();
+			if(InComponent)
+			{
+				const FRigHierarchyKey PreviousParent = InHierarchy->GetPreviousHierarchyParent(InComponent->GetKey());
+				if(PreviousParent.IsElement())
+				{
+					const FRigComponentKey OldKey(PreviousParent.GetElement(), InComponent->GetFName());
+					HandleHierarchyComponentKeyChanged(OldKey, InComponent->GetKey());
+				}
+			}
+			break;
+		}
 		case ERigHierarchyNotification::ElementSelected:
 		{
 			bool bClearTransientControls = true;
@@ -2825,6 +2843,154 @@ void UControlRigBlueprint::HandleHierarchyModified(ERigHierarchyNotification InN
 #endif
 }
 
+void UControlRigBlueprint::HandleHierarchyElementKeyChanged(const FRigElementKey& InOldKey, const FRigElementKey& InNewKey)
+{
+	if(InOldKey == InNewKey)
+	{
+		return;
+	}
+	
+	static UEnum* RigElementTypeEnum = StaticEnum<ERigElementType>();
+	check(RigElementTypeEnum);
+
+	const FString OldNameStr = InOldKey.Name.ToString();
+	const FString NewNameStr = InNewKey.Name.ToString();
+	const ERigElementType ElementType = InNewKey.Type; 
+
+	// update all of the graphs with the new key
+	TArray<UEdGraph*> EdGraphs;
+	GetAllGraphs(EdGraphs);
+	for (UEdGraph* Graph : EdGraphs)
+	{
+		URigVMEdGraph* RigGraph = Cast<URigVMEdGraph>(Graph);
+		if (RigGraph == nullptr)
+		{
+			continue;
+		}
+
+		URigVMController* Controller = RigGraph->GetController();
+		if(Controller == nullptr)
+		{
+			continue;
+		}
+
+		{
+			FRigVMBlueprintCompileScope CompileScope(this);
+			for (UEdGraphNode* Node : RigGraph->Nodes)
+			{
+				if (URigVMEdGraphNode* RigNode = Cast<URigVMEdGraphNode>(Node))
+				{
+					if (URigVMNode* ModelNode = RigNode->GetModelNode())
+					{
+						TArray<URigVMPin*> ModelPins = ModelNode->GetAllPinsRecursively();
+						for (URigVMPin * ModelPin : ModelPins)
+						{
+							if(ModelPin->GetCPPType() == RigVMTypeUtils::FNameType)
+							{
+								if ((ModelPin->GetCustomWidgetName() == TEXT("BoneName") && ElementType == ERigElementType::Bone) ||
+									(ModelPin->GetCustomWidgetName() == TEXT("ControlName") && ElementType == ERigElementType::Control) ||
+									(ModelPin->GetCustomWidgetName() == TEXT("SpaceName") && ElementType == ERigElementType::Null) ||
+									(ModelPin->GetCustomWidgetName() == TEXT("CurveName") && ElementType == ERigElementType::Curve) ||
+									(ModelPin->GetCustomWidgetName() == TEXT("ConnectorName") && ElementType == ERigElementType::Connector))
+								{
+									if (ModelPin->GetDefaultValue() == OldNameStr)
+									{
+										Controller->SetPinDefaultValue(ModelPin->GetPinPath(), NewNameStr, false);
+									}
+								}
+							}
+							else if (ModelPin->GetCPPTypeObject() == FRigElementKey::StaticStruct())
+							{
+								const FString OldDefaultValueString = ModelPin->GetDefaultValue();
+								FRigElementKey OldDefaultKey;
+								FRigElementKey::StaticStruct()->ImportText(*OldDefaultValueString, &OldDefaultKey, nullptr, EPropertyPortFlags::PPF_None, nullptr, FRigElementKey::StaticStruct()->GetName(), true);
+
+								if(OldDefaultKey == InOldKey)
+								{
+									FString NewDefaultKeyString;
+									FRigElementKey::StaticStruct()->ExportText(NewDefaultKeyString, &InNewKey, nullptr, nullptr, PPF_ExternalEditor, nullptr);
+									Controller->SetPinDefaultValue(ModelPin->GetPinPath(), NewDefaultKeyString, false);
+								}
+							}
+						}
+					}
+				}
+			}
+		}
+	}
+
+	// update all of the influences
+	Influences.OnKeyRenamed(InOldKey, InNewKey);
+	if (IsControlRigModule() && InNewKey.Type == ERigElementType::Connector)
+	{
+		if (FRigElementKeyCollection* Targets = ArrayConnectionMap.Find(InOldKey))
+		{
+			ArrayConnectionMap.FindOrAdd(InNewKey, *Targets);
+			ArrayConnectionMap.Remove(InOldKey);
+		}
+	}
+	
+	PropagateHierarchyFromBPToInstances();
+}
+
+void UControlRigBlueprint::HandleHierarchyComponentKeyChanged(const FRigComponentKey& InOldKey, const FRigComponentKey& InNewKey)
+{
+	if(InOldKey == InNewKey)
+	{
+		return;
+	}
+	
+	// update all of the graphs with the new key
+	TArray<UEdGraph*> EdGraphs;
+	GetAllGraphs(EdGraphs);
+	for (UEdGraph* Graph : EdGraphs)
+	{
+		URigVMEdGraph* RigGraph = Cast<URigVMEdGraph>(Graph);
+		if (RigGraph == nullptr)
+		{
+			continue;
+		}
+
+		URigVMController* Controller = RigGraph->GetController();
+		if(Controller == nullptr)
+		{
+			continue;
+		}
+
+		{
+			FRigVMBlueprintCompileScope CompileScope(this);
+			for (UEdGraphNode* Node : RigGraph->Nodes)
+			{
+				if (URigVMEdGraphNode* RigNode = Cast<URigVMEdGraphNode>(Node))
+				{
+					if (URigVMNode* ModelNode = RigNode->GetModelNode())
+					{
+						TArray<URigVMPin*> ModelPins = ModelNode->GetAllPinsRecursively();
+						for (URigVMPin * ModelPin : ModelPins)
+						{
+							if (ModelPin->GetCPPTypeObject() == FRigComponentKey::StaticStruct())
+							{
+								const FString OldDefaultValueString = ModelPin->GetDefaultValue();
+								FRigComponentKey OldDefaultKey;
+								FRigComponentKey::StaticStruct()->ImportText(*OldDefaultValueString, &OldDefaultKey, nullptr, EPropertyPortFlags::PPF_None, nullptr, FRigComponentKey::StaticStruct()->GetName(), true);
+
+								if(OldDefaultKey == InOldKey)
+								{
+									FString NewDefaultKeyString;
+									FRigComponentKey::StaticStruct()->ExportText(NewDefaultKeyString, &InNewKey, nullptr, nullptr, PPF_ExternalEditor, nullptr);
+									Controller->SetPinDefaultValue(ModelPin->GetPinPath(), NewDefaultKeyString, false);
+								}
+							}
+						}
+					}
+				}
+			}
+		}
+	}
+
+	PropagateHierarchyFromBPToInstances();
+}
+
 void UControlRigBlueprint::HandleRigModulesModified(EModularRigNotification InNotification, const FRigModuleReference* InModule)
 {
 	bool bRecompile = true;
diff --git a/Engine/Plugins/Animation/ControlRig/Source/ControlRigDeveloper/Public/ControlRigBlueprint.h b/Engine/Plugins/Animation/ControlRig/Source/ControlRigDeveloper/Public/ControlRigBlueprint.h
--- a/Engine/Plugins/Animation/ControlRig/Source/ControlRigDeveloper/Public/ControlRigBlueprint.h
+++ b/Engine/Plugins/Animation/ControlRig/Source/ControlRigDeveloper/Public/ControlRigBlueprint.h
@@ -394,6 +394,8 @@ private:
 
 
 	void HandleHierarchyModified(ERigHierarchyNotification InNotification, URigHierarchy* InHierarchy, const FRigNotificationSubject& InSubject);
+	void HandleHierarchyElementKeyChanged(const FRigElementKey& InOldKey, const FRigElementKey& InNewKey);
+	void HandleHierarchyComponentKeyChanged(const FRigComponentKey& InOldKey, const FRigComponentKey& InNewKey);
 
 	void HandleRigModulesModified(EModularRigNotification InNotification, const FRigModuleReference* InModule);
 
diff --git a/Engine/Plugins/Animation/ControlRig/Source/ControlRigEditor/Private/Editor/ControlRigEditor.cpp b/Engine/Plugins/Animation/ControlRig/Source/ControlRigEditor/Private/Editor/ControlRigEditor.cpp
--- a/Engine/Plugins/Animation/ControlRig/Source/ControlRigEditor/Private/Editor/ControlRigEditor.cpp
+++ b/Engine/Plugins/Animation/ControlRig/Source/ControlRigEditor/Private/Editor/ControlRigEditor.cpp
@@ -4411,81 +4411,6 @@ void FControlRigEditor::OnHierarchyModified(ERigHierarchyNotification InNotif, U
 		}
 		case ERigHierarchyNotification::ElementRenamed:
 		{
-			UEnum* RigElementTypeEnum = StaticEnum<ERigElementType>();
-			if (RigElementTypeEnum == nullptr)
-			{
-				return;
-			}
-
-			const FString OldNameStr = InHierarchy->GetPreviousName(InElement->GetKey()).ToString();
-			const FString NewNameStr = InElement->GetName();
-			const ERigElementType ElementType = InElement->GetType(); 
-
-			TArray<UEdGraph*> EdGraphs;
-			RigBlueprint->GetAllGraphs(EdGraphs);
-
-			for (UEdGraph* Graph : EdGraphs)
-			{
-				UControlRigGraph* RigGraph = Cast<UControlRigGraph>(Graph);
-				if (RigGraph == nullptr)
-				{
-					continue;
-				}
-
-				URigVMController* Controller = RigGraph->GetController();
-				if(Controller == nullptr)
-				{
-					continue;
-				}
-
-				{
-					FRigVMBlueprintCompileScope CompileScope(RigBlueprint);
-					for (UEdGraphNode* Node : RigGraph->Nodes)
-					{
-						if (UControlRigGraphNode* RigNode = Cast<UControlRigGraphNode>(Node))
-						{
-							if (URigVMNode* ModelNode = RigNode->GetModelNode())
-							{
-								TArray<URigVMPin*> ModelPins = ModelNode->GetAllPinsRecursively();
-								for (URigVMPin * ModelPin : ModelPins)
-								{
-									if ((ModelPin->GetCPPType() == TEXT("FName") && ModelPin->GetCustomWidgetName() == TEXT("BoneName") && ElementType == ERigElementType::Bone) ||
-										(ModelPin->GetCPPType() == TEXT("FName") && ModelPin->GetCustomWidgetName() == TEXT("ControlName") && ElementType == ERigElementType::Control) ||
-										(ModelPin->GetCPPType() == TEXT("FName") && ModelPin->GetCustomWidgetName() == TEXT("SpaceName") && ElementType == ERigElementType::Null) ||
-										(ModelPin->GetCPPType() == TEXT("FName") && ModelPin->GetCustomWidgetName() == TEXT("CurveName") && ElementType == ERigElementType::Curve) ||
-										(ModelPin->GetCPPType() == TEXT("FName") && ModelPin->GetCustomWidgetName() == TEXT("ConnectorName") && ElementType == ERigElementType::Connector))
-									{
-										if (ModelPin->GetDefaultValue() == OldNameStr)
-										{
-											Controller->SetPinDefaultValue(ModelPin->GetPinPath(), NewNameStr, false);
-										}
-									}
-									else if (ModelPin->GetCPPTypeObject() == FRigElementKey::StaticStruct())
-									{
-										if (URigVMPin* TypePin = ModelPin->FindSubPin(TEXT("Type")))
-										{
-											const FString TypeStr = TypePin->GetDefaultValue();
-											const int64 TypeValue = RigElementTypeEnum->GetValueByNameString(TypeStr);
-											if (TypeValue == (int64)ElementType)
-											{
-												if (URigVMPin* NamePin = ModelPin->FindSubPin(TEXT("Name")))
-												{
-													FString NameStr = NamePin->GetDefaultValue();
-													if (NameStr == OldNameStr)
-													{
-														Controller->SetPinDefaultValue(NamePin->GetPinPath(), NewNameStr);
-													}
-												}
-											}
-										}
-									}
-								}
-							}
-						}
-					}
-				}
-			}
-				
 			OnHierarchyChanged();
 
 			break;
diff --git a/Engine/Plugins/Animation/ControlRig/Source/ControlRigEditor/Private/Editor/SRigHierarchy.cpp b/Engine/Plugins/Animation/ControlRig/Source/ControlRigEditor/Private/Editor/SRigHierarchy.cpp
--- a/Engine/Plugins/Animation/ControlRig/Source/ControlRigEditor/Private/Editor/SRigHierarchy.cpp
+++ b/Engine/Plugins/Animation/ControlRig/Source/ControlRigEditor/Private/Editor/SRigHierarchy.cpp
@@ -3978,8 +3978,18 @@ FReply SRigHierarchy::ReparentOrMatchTransform(const TArray<FRigHierarchyKey>& D
 			return FReply::Unhandled();
 		}
 
+		// only suspend blueprint notifs if we are dragging non-components
+		bool bBlueprintSuspensionFlag = ControlRigBlueprint->bSuspendAllNotifications;
+		for (const FRigHierarchyKey& DraggedKey : DraggedKeys)
+		{
+			if (!DraggedKey.IsComponent())
+			{
+				bBlueprintSuspensionFlag = true;
+			}
+		}
+
 		TGuardValue<bool> GuardRigHierarchyChanges(bIsChangingRigHierarchy, true);
-		TGuardValue<bool> SuspendBlueprintNotifs(ControlRigBlueprint->bSuspendAllNotifications, true);
+		TGuardValue<bool> SuspendBlueprintNotifs(ControlRigBlueprint->bSuspendAllNotifications, bBlueprintSuspensionFlag);
 		FScopedTransaction Transaction(LOCTEXT("HierarchyDragAndDrop", "Drag & Drop"));
 		FRigHierarchyInteractionBracket InteractionBracket(Hierarchy);
 
diff --git a/Engine/Plugins/Experimental/PhysicsControl/Source/ControlRigPhysics/Private/RigPhysicsComponent.cpp b/Engine/Plugins/Experimental/PhysicsControl/Source/ControlRigPhysics/Private/RigPhysicsComponent.cpp
--- a/Engine/Plugins/Experimental/PhysicsControl/Source/ControlRigPhysics/Private/RigPhysicsComponent.cpp
+++ b/Engine/Plugins/Experimental/PhysicsControl/Source/ControlRigPhysics/Private/RigPhysicsComponent.cpp
@@ -110,3 +110,11 @@ void FRigPhysicsComponent::OnAddedToHierarchy(URigHierarchy* InHierarchy, URigHi
 	// let the system know we changed something
 	InHierarchy->Notify(ERigHierarchyNotification::ComponentContentChanged, this);
 }
+
+//======================================================================================================================
+void FRigPhysicsComponent::OnRigHierarchyKeyChanged(const FRigHierarchyKey& InOldKey, const FRigHierarchyKey& InNewKey)
+{
+	FRigBaseComponent::OnRigHierarchyKeyChanged(InOldKey, InNewKey);
+
+	BodySolverSettings.OnRigHierarchyKeyChanged(InOldKey, InNewKey);
+}
diff --git a/Engine/Plugins/Experimental/PhysicsControl/Source/ControlRigPhysics/Private/RigPhysicsData.cpp b/Engine/Plugins/Experimental/PhysicsControl/Source/ControlRigPhysics/Private/RigPhysicsData.cpp
--- a/Engine/Plugins/Experimental/PhysicsControl/Source/ControlRigPhysics/Private/RigPhysicsData.cpp
+++ b/Engine/Plugins/Experimental/PhysicsControl/Source/ControlRigPhysics/Private/RigPhysicsData.cpp
@@ -96,3 +96,25 @@ FArchive& operator <<(FArchive& Ar, FRigPhysicsCollision& Data)
 	Ar << Data.Capsules;
 	return Ar;
 }
+
+void FRigPhysicsBodySolverSettings::OnRigHierarchyKeyChanged(const FRigHierarchyKey& InOldKey, const FRigHierarchyKey& InNewKey)
+{
+	if(InOldKey.IsComponent() && InNewKey.IsComponent())
+	{
+		if(PhysicsSolverComponentKey == InOldKey.GetComponent())
+		{
+			PhysicsSolverComponentKey = InNewKey.GetComponent();
+		}
+	}
+	if(InOldKey.IsElement() && InNewKey.IsElement())
+	{
+		if(SourceBone == InOldKey.GetElement())
+		{
+			SourceBone = InNewKey.GetElement();
+		}
+		if(TargetBone == InOldKey.GetElement())
+		{
+			TargetBone = InNewKey.GetElement();
+		}
+	}
+}
\ No newline at end of file
diff --git a/Engine/Plugins/Experimental/PhysicsControl/Source/ControlRigPhysics/Public/RigPhysicsComponent.h b/Engine/Plugins/Experimental/PhysicsControl/Source/ControlRigPhysics/Public/RigPhysicsComponent.h
--- a/Engine/Plugins/Experimental/PhysicsControl/Source/ControlRigPhysics/Public/RigPhysicsComponent.h
+++ b/Engine/Plugins/Experimental/PhysicsControl/Source/ControlRigPhysics/Public/RigPhysicsComponent.h
@@ -58,5 +58,7 @@ public:
 	virtual bool CanBeAddedTo(const FRigElementKey& InElementKey, const URigHierarchy* InHierarchy, FString* OutFailureReason) const override;
 
 	virtual void OnAddedToHierarchy(URigHierarchy* InHierarchy, URigHierarchyController* InController) override;
+
+	virtual void OnRigHierarchyKeyChanged(const FRigHierarchyKey& InOldKey, const FRigHierarchyKey& InNewKey) override;
 };
 
diff --git a/Engine/Plugins/Experimental/PhysicsControl/Source/ControlRigPhysics/Public/RigPhysicsData.h b/Engine/Plugins/Experimental/PhysicsControl/Source/ControlRigPhysics/Public/RigPhysicsData.h
--- a/Engine/Plugins/Experimental/PhysicsControl/Source/ControlRigPhysics/Public/RigPhysicsData.h
+++ b/Engine/Plugins/Experimental/PhysicsControl/Source/ControlRigPhysics/Public/RigPhysicsData.h
@@ -624,5 +624,7 @@ struct CONTROLRIGPHYSICS_API FRigPhysicsBodySolverSettings
 	/// default to the parent of our owner.
 	UPROPERTY(BlueprintReadWrite, EditAnywhere, Category = Solver)
 	FRigElementKey TargetBone;
+
+	void OnRigHierarchyKeyChanged(const FRigHierarchyKey& InOldKey, const FRigHierarchyKey& InNewKey);
 };
 
