RigVM / Control Rig: Refactor units to use empty execute context

Full cook preflight running here https://horde.devtools.epicgames.com/job/67b0ac950e3c624972a57dfc

#rnx

diff --git a/Engine/Plugins/Animation/ControlRig/Source/ControlRig/Public/Units/Debug/RigUnit_DebugHierarchy.h b/Engine/Plugins/Animation/ControlRig/Source/ControlRig/Public/Units/Debug/RigUnit_DebugHierarchy.h
--- a/Engine/Plugins/Animation/ControlRig/Source/ControlRig/Public/Units/Debug/RigUnit_DebugHierarchy.h
+++ b/Engine/Plugins/Animation/ControlRig/Source/ControlRig/Public/Units/Debug/RigUnit_DebugHierarchy.h
@@ -23,7 +23,7 @@ namespace EControlRigDrawHierarchyMode
 /**
  * Draws vectors on each bone in the viewport across the entire hierarchy
  */
-USTRUCT(meta=(DisplayName="Draw Hierarchy"))
+USTRUCT(meta=(DisplayName="Draw Hierarchy", ExecuteContext="FControlRigExecuteContext"))
 struct CONTROLRIG_API FRigUnit_DebugHierarchy : public FRigVMFunction_DebugBase
 {
 	GENERATED_BODY()
@@ -41,7 +41,7 @@ struct CONTROLRIG_API FRigUnit_DebugHierarchy : public FRigVMFunction_DebugBase
 	virtual void Execute() override;
 
 	UPROPERTY(DisplayName = "Execute", Transient, meta = (Input, Output))
-	FControlRigExecuteContext ExecuteContext;
+	FRigVMEmptyExecuteContext ExecuteContext;
 
 	// the items to draw the pose for.
 	// if this is empty we'll draw the whole hierarchy
@@ -72,7 +72,7 @@ struct CONTROLRIG_API FRigUnit_DebugHierarchy : public FRigVMFunction_DebugBase
 /**
 * Draws vectors on each bone in the viewport across the entire pose
 */
-USTRUCT(meta=(DisplayName="Draw Pose Cache"))
+USTRUCT(meta=(DisplayName="Draw Pose Cache", ExecuteContext="FControlRigExecuteContext"))
 struct CONTROLRIG_API FRigUnit_DebugPose : public FRigVMFunction_DebugBase
 {
 	GENERATED_BODY()
@@ -90,7 +90,7 @@ struct CONTROLRIG_API FRigUnit_DebugPose : public FRigVMFunction_DebugBase
 	virtual void Execute() override;
 
 	UPROPERTY(DisplayName = "Execute", Transient, meta = (Input, Output))
-	FControlRigExecuteContext ExecuteContext;
+	FRigVMEmptyExecuteContext ExecuteContext;
 
 	UPROPERTY(meta = (Input))
 	FRigPose Pose;
diff --git a/Engine/Plugins/Animation/ControlRig/Source/ControlRig/Public/Units/Execution/RigUnit_BeginExecution.h b/Engine/Plugins/Animation/ControlRig/Source/ControlRig/Public/Units/Execution/RigUnit_BeginExecution.h
--- a/Engine/Plugins/Animation/ControlRig/Source/ControlRig/Public/Units/Execution/RigUnit_BeginExecution.h
+++ b/Engine/Plugins/Animation/ControlRig/Source/ControlRig/Public/Units/Execution/RigUnit_BeginExecution.h
@@ -22,7 +22,7 @@ struct CONTROLRIG_API FRigUnit_BeginExecution : public FRigUnit
 
 	// The execution result
 	UPROPERTY(EditAnywhere, Transient, DisplayName = "Execute", Category = "BeginExecution", meta = (Output))
-	FControlRigExecuteContext ExecuteContext;
+	FRigVMEmptyExecuteContext ExecuteContext;
 
 	static inline const FLazyName EventName = FLazyName(TEXT("Forwards Solve"));
 };
@@ -43,7 +43,7 @@ struct CONTROLRIG_API FRigUnit_PreBeginExecution : public FRigUnit
 
 	// The execution result
 	UPROPERTY(EditAnywhere, Transient, DisplayName = "Execute", Category = "BeginExecution", meta = (Output))
-	FControlRigExecuteContext ExecuteContext;
+	FRigVMEmptyExecuteContext ExecuteContext;
 
 	static inline const FLazyName EventName = FLazyName(TEXT("Pre Forwards Solve"));
 };
@@ -64,7 +64,7 @@ struct CONTROLRIG_API FRigUnit_PostBeginExecution : public FRigUnit
 
 	// The execution result
 	UPROPERTY(EditAnywhere, Transient, DisplayName = "Execute", Category = "BeginExecution", meta = (Output))
-	FControlRigExecuteContext ExecuteContext;
+	FRigVMEmptyExecuteContext ExecuteContext;
 
 	static inline const FLazyName EventName = FLazyName(TEXT("Post Forwards Solve"));
 };
diff --git a/Engine/Plugins/Animation/ControlRig/Source/ControlRig/Public/Units/Execution/RigUnit_InteractionExecution.h b/Engine/Plugins/Animation/ControlRig/Source/ControlRig/Public/Units/Execution/RigUnit_InteractionExecution.h
--- a/Engine/Plugins/Animation/ControlRig/Source/ControlRig/Public/Units/Execution/RigUnit_InteractionExecution.h
+++ b/Engine/Plugins/Animation/ControlRig/Source/ControlRig/Public/Units/Execution/RigUnit_InteractionExecution.h
@@ -21,7 +21,7 @@ struct CONTROLRIG_API FRigUnit_InteractionExecution : public FRigUnit
 
 	// The execution result
 	UPROPERTY(EditAnywhere, Transient, DisplayName = "Execute", Category = "BeginExecution", meta = (Output))
-	FControlRigExecuteContext ExecuteContext;
+	FRigVMEmptyExecuteContext ExecuteContext;
 
 	static inline const FLazyName EventName = FLazyName(TEXT("Interaction"));
 };
diff --git a/Engine/Plugins/Animation/ControlRig/Source/ControlRig/Public/Units/Execution/RigUnit_InverseExecution.h b/Engine/Plugins/Animation/ControlRig/Source/ControlRig/Public/Units/Execution/RigUnit_InverseExecution.h
--- a/Engine/Plugins/Animation/ControlRig/Source/ControlRig/Public/Units/Execution/RigUnit_InverseExecution.h
+++ b/Engine/Plugins/Animation/ControlRig/Source/ControlRig/Public/Units/Execution/RigUnit_InverseExecution.h
@@ -22,7 +22,7 @@ struct CONTROLRIG_API FRigUnit_InverseExecution : public FRigUnit
 
 	// The execution result
 	UPROPERTY(EditAnywhere, Transient, DisplayName = "Execute", Category = "InverseExecution", meta = (Output))
-	FControlRigExecuteContext ExecuteContext;
+	FRigVMEmptyExecuteContext ExecuteContext;
 
 	static inline const FLazyName EventName = FLazyName(TEXT("Backwards Solve"));
 };
diff --git a/Engine/Plugins/Animation/ControlRig/Source/ControlRig/Public/Units/Execution/RigUnit_PrepareForExecution.h b/Engine/Plugins/Animation/ControlRig/Source/ControlRig/Public/Units/Execution/RigUnit_PrepareForExecution.h
--- a/Engine/Plugins/Animation/ControlRig/Source/ControlRig/Public/Units/Execution/RigUnit_PrepareForExecution.h
+++ b/Engine/Plugins/Animation/ControlRig/Source/ControlRig/Public/Units/Execution/RigUnit_PrepareForExecution.h
@@ -22,7 +22,7 @@ struct CONTROLRIG_API FRigUnit_PrepareForExecution : public FRigUnit
 
 	// The execution result
 	UPROPERTY(EditAnywhere, Transient, DisplayName = "Execute", Category = "PrepareForExecution", meta = (Output))
-	FControlRigExecuteContext ExecuteContext;
+	FRigVMEmptyExecuteContext ExecuteContext;
 
 	static inline const FLazyName EventName = FLazyName(TEXT("Construction"));
 };
@@ -43,7 +43,7 @@ struct CONTROLRIG_API FRigUnit_PostPrepareForExecution : public FRigUnit
 
 	// The execution result
 	UPROPERTY(EditAnywhere, Transient, DisplayName = "Execute", Category = "PostPrepareForExecution", meta = (Output))
-	FControlRigExecuteContext ExecuteContext;
+	FRigVMEmptyExecuteContext ExecuteContext;
 
 	static inline const FLazyName EventName = FLazyName(TEXT("PostConstruction"));
 };
diff --git a/Engine/Plugins/Animation/ControlRig/Source/ControlRig/Public/Units/Execution/RigUnit_SequenceExecution.h b/Engine/Plugins/Animation/ControlRig/Source/ControlRig/Public/Units/Execution/RigUnit_SequenceExecution.h
--- a/Engine/Plugins/Animation/ControlRig/Source/ControlRig/Public/Units/Execution/RigUnit_SequenceExecution.h
+++ b/Engine/Plugins/Animation/ControlRig/Source/ControlRig/Public/Units/Execution/RigUnit_SequenceExecution.h
@@ -19,23 +19,23 @@ struct CONTROLRIG_API FRigUnit_SequenceExecution : public FRigUnit
 
 	// The execution input
 	UPROPERTY(EditAnywhere, Transient, DisplayName = "Execute", Category = "SequenceExecution", meta = (Input))
-	FControlRigExecuteContext ExecuteContext;
+	FRigVMEmptyExecuteContext ExecuteContext;
 
 	// The execution result A
 	UPROPERTY(EditAnywhere, Transient, Category = "SequenceExecution", meta = (Output))
-	FControlRigExecuteContext A;
+	FRigVMEmptyExecuteContext A;
 
 	// The execution result B
 	UPROPERTY(EditAnywhere, Transient, Category = "SequenceExecution", meta = (Output))
-	FControlRigExecuteContext B;
+	FRigVMEmptyExecuteContext B;
 
 	// The execution result C
 	UPROPERTY(EditAnywhere, Transient, Category = "SequenceExecution", meta = (Output))
-	FControlRigExecuteContext C;
+	FRigVMEmptyExecuteContext C;
 
 	// The execution result D
 	UPROPERTY(EditAnywhere, Transient, Category = "SequenceExecution", meta = (Output))
-	FControlRigExecuteContext D;
+	FRigVMEmptyExecuteContext D;
 
 	RIGVM_METHOD()
 	virtual FRigVMStructUpgradeInfo GetUpgradeInfo() const override;
diff --git a/Engine/Plugins/Animation/ControlRig/Source/ControlRig/Public/Units/Hierarchy/RigUnit_ControlChannel.h b/Engine/Plugins/Animation/ControlRig/Source/ControlRig/Public/Units/Hierarchy/RigUnit_ControlChannel.h
--- a/Engine/Plugins/Animation/ControlRig/Source/ControlRig/Public/Units/Hierarchy/RigUnit_ControlChannel.h
+++ b/Engine/Plugins/Animation/ControlRig/Source/ControlRig/Public/Units/Hierarchy/RigUnit_ControlChannel.h
@@ -211,7 +211,7 @@ struct CONTROLRIG_API FRigUnit_SetAnimationChannelBase : public FRigUnit_GetAnim
 	}
 
 	UPROPERTY(DisplayName = "Execute", Transient, meta = (Input, Output))
-	FControlRigExecuteContext ExecuteContext;
+	FRigVMEmptyExecuteContext ExecuteContext;
 };
 
 /**
diff --git a/Engine/Plugins/Animation/ControlRig/Source/ControlRig/Public/Units/Hierarchy/RigUnit_ControlChannelFromItem.h b/Engine/Plugins/Animation/ControlRig/Source/ControlRig/Public/Units/Hierarchy/RigUnit_ControlChannelFromItem.h
--- a/Engine/Plugins/Animation/ControlRig/Source/ControlRig/Public/Units/Hierarchy/RigUnit_ControlChannelFromItem.h
+++ b/Engine/Plugins/Animation/ControlRig/Source/ControlRig/Public/Units/Hierarchy/RigUnit_ControlChannelFromItem.h
@@ -193,7 +193,7 @@ struct CONTROLRIG_API FRigUnit_SetAnimationChannelBaseFromItem : public FRigUnit
 	}
 
 	UPROPERTY(DisplayName = "Execute", Transient, meta = (Input, Output))
-	FControlRigExecuteContext ExecuteContext;
+	FRigVMEmptyExecuteContext ExecuteContext;
 };
 
 /**
diff --git a/Engine/Plugins/Animation/ControlRig/Source/ControlRig/Public/Units/RigUnit.h b/Engine/Plugins/Animation/ControlRig/Source/ControlRig/Public/Units/RigUnit.h
--- a/Engine/Plugins/Animation/ControlRig/Source/ControlRig/Public/Units/RigUnit.h
+++ b/Engine/Plugins/Animation/ControlRig/Source/ControlRig/Public/Units/RigUnit.h
@@ -119,7 +119,7 @@ struct CONTROLRIG_API FRigUnitMutable : public FRigUnit
 	 * This property is used to chain multiple mutable units together
 	 */
 	UPROPERTY(DisplayName = "Execute", Transient, meta = (Input, Output))
-	FControlRigExecuteContext ExecuteContext;
+	FRigVMEmptyExecuteContext ExecuteContext;
 };
 
 #if WITH_EDITOR
diff --git a/Engine/Plugins/Runtime/RigVM/Source/RigVM/Public/RigVMCore/RigVMStruct.h b/Engine/Plugins/Runtime/RigVM/Source/RigVM/Public/RigVMCore/RigVMStruct.h
--- a/Engine/Plugins/Runtime/RigVM/Source/RigVM/Public/RigVMCore/RigVMStruct.h
+++ b/Engine/Plugins/Runtime/RigVM/Source/RigVM/Public/RigVMCore/RigVMStruct.h
@@ -331,5 +331,5 @@ public:
 	 * This property is used to chain multiple mutable nodes together
 	 */
 	UPROPERTY(DisplayName = "Execute", Transient, meta = (Input, Output))
-	FRigVMExecuteContext ExecuteContext;
+	FRigVMEmptyExecuteContext ExecuteContext;
 };
diff --git a/Engine/Plugins/Runtime/RigVM/Source/RigVM/Public/RigVMFunctions/Execution/RigVMFunction_ForLoop.h b/Engine/Plugins/Runtime/RigVM/Source/RigVM/Public/RigVMFunctions/Execution/RigVMFunction_ForLoop.h
--- a/Engine/Plugins/Runtime/RigVM/Source/RigVM/Public/RigVMFunctions/Execution/RigVMFunction_ForLoop.h
+++ b/Engine/Plugins/Runtime/RigVM/Source/RigVM/Public/RigVMFunctions/Execution/RigVMFunction_ForLoop.h
@@ -46,5 +46,5 @@ struct RIGVM_API FRigVMFunction_ForLoopCount : public FRigVMStructMutable
 	float Ratio;
 
 	UPROPERTY(EditAnywhere, Transient, Category = "ForLoop", meta = (Output))
-	FRigVMExecuteContext Completed;
+	FRigVMEmptyExecuteContext Completed;
 };
diff --git a/Engine/Plugins/Runtime/RigVM/Source/RigVM/Public/RigVMFunctions/Execution/RigVMFunction_Sequence.h b/Engine/Plugins/Runtime/RigVM/Source/RigVM/Public/RigVMFunctions/Execution/RigVMFunction_Sequence.h
--- a/Engine/Plugins/Runtime/RigVM/Source/RigVM/Public/RigVMFunctions/Execution/RigVMFunction_Sequence.h
+++ b/Engine/Plugins/Runtime/RigVM/Source/RigVM/Public/RigVMFunctions/Execution/RigVMFunction_Sequence.h
@@ -18,13 +18,13 @@ struct RIGVM_API FRigVMFunction_Sequence : public FRigVMStruct
 
 	// The execution input
 	UPROPERTY(EditAnywhere, Transient, DisplayName = "Execute", Category = "SequenceExecution", meta = (Input, Aggregate))
-	FRigVMExecuteContext ExecuteContext;
+	FRigVMEmptyExecuteContext ExecuteContext;
 
 	// The execution result A
 	UPROPERTY(EditAnywhere, Transient, Category = "SequenceExecution", meta = (Output, Aggregate))
-	FRigVMExecuteContext A;
+	FRigVMEmptyExecuteContext A;
 
 	// The execution result B
 	UPROPERTY(EditAnywhere, Transient, Category = "SequenceExecution", meta = (Output, Aggregate))
-	FRigVMExecuteContext B;
+	FRigVMEmptyExecuteContext B;
 };
\ No newline at end of file
diff --git a/Engine/Plugins/Runtime/RigVM/Source/RigVM/Public/RigVMFunctions/Execution/RigVMFunction_UserDefinedEvent.h b/Engine/Plugins/Runtime/RigVM/Source/RigVM/Public/RigVMFunctions/Execution/RigVMFunction_UserDefinedEvent.h
--- a/Engine/Plugins/Runtime/RigVM/Source/RigVM/Public/RigVMFunctions/Execution/RigVMFunction_UserDefinedEvent.h
+++ b/Engine/Plugins/Runtime/RigVM/Source/RigVM/Public/RigVMFunctions/Execution/RigVMFunction_UserDefinedEvent.h
@@ -27,7 +27,7 @@ struct RIGVM_API FRigVMFunction_UserDefinedEvent : public FRigVMStruct
 
 	// The execution result
 	UPROPERTY(EditAnywhere, Transient, DisplayName = "Execute", Category = "UserDefinedEvent", meta = (Output))
-	FRigVMExecuteContext ExecuteContext;
+	FRigVMEmptyExecuteContext ExecuteContext;
 
 	// True if the current interaction is a rotation
 	UPROPERTY(EditAnywhere, Transient, Category = "UserDefinedEvent", meta = (Input,Constant,DetailsOnly))
