ChainInfo node
-This is a node that:
-Takes in a list of items and a param value
-Interprets those items as a chain
-Outputs helpful info about the chain at the current param value including
-Interpolated transform
-Chain Length
-Param Length
-Chain Stretch Factor
-Segment info including:
-Segment Index
-Segment Length
-Segment Stretch Value
-Local Segment Param
-Segment Start Item
-Segment End Item

diff --git a/Engine/Plugins/Animation/ControlRig/Source/ControlRig/Private/Units/Highlevel/Hierarchy/RigUnit_ChainInfo.cpp b/Engine/Plugins/Animation/ControlRig/Source/ControlRig/Private/Units/Highlevel/Hierarchy/RigUnit_ChainInfo.cpp
new file mode 100644
--- /dev/null
+++ b/Engine/Plugins/Animation/ControlRig/Source/ControlRig/Private/Units/Highlevel/Hierarchy/RigUnit_ChainInfo.cpp
@@ -0,0 +1,208 @@
+// Copyright Epic Games, Inc. All Rights Reserved.
+
+#include "Units/Highlevel/Hierarchy/RigUnit_ChainInfo.h"
+#include "Math/ControlRigMathLibrary.h"
+#include "RigVMFunctions/Debug/RigVMFunction_DebugTransform.h"
+#include "RigVMCore/RigVMDebugDrawSettings.h"
+
+#include UE_INLINE_GENERATED_CPP_BY_NAME(RigUnit_ChainInfo)
+
+FRigUnit_ChainInfo_Execute()
+{
+	URigHierarchy* Hierarchy = ExecuteContext.Hierarchy;
+	if (Hierarchy == nullptr)
+	{
+		UE_CONTROLRIG_RIGUNIT_REPORT_ERROR(TEXT("No hierarchy found in control rig."));
+		return;
+	}
+
+	if (Items.Num() < 2) 
+	{
+		UE_CONTROLRIG_RIGUNIT_REPORT_ERROR(TEXT("Must use at least two items in Items input list."));
+		return;
+	}
+
+	if(CachedElements.Num() != Items.Num())
+	{
+		CachedElements.Reset();
+		CachedElements.SetNum(Items.Num());
+	}
+	
+	// Run some initialization
+	ChainLength = 0.f;
+	float InitialChainLength = 0.f;
+	TArray<FRigUnit_ChainInfo_Segment> Segments;
+	Segments.SetNumZeroed(CachedElements.Num() - 1);
+
+    // Calculate chain length and record segment lengths
+	for (int32 Index = 0; Index < Segments.Num(); Index++) 
+	{
+		FRigUnit_ChainInfo_Segment& Segment = Segments[Index];
+
+		if (Index == 0)
+		{
+			if (!CachedElements[Index].UpdateCache(Items[Index], Hierarchy))
+			{
+				UE_CONTROLRIG_RIGUNIT_REPORT_WARNING(TEXT("Item '%s' is not valid."), *Items[Index].ToString());
+			}
+		}
+		Segment.StartItem = CachedElements[Index];
+		Segment.StartItemIndex = Index;
+
+		if (!CachedElements[Index + 1].UpdateCache(Items[Index + 1], Hierarchy))
+		{
+			UE_CONTROLRIG_RIGUNIT_REPORT_WARNING(TEXT("Item '%s' is not valid."), *Items[Index + 1].ToString());
+		}
+		Segment.EndItem = CachedElements[Index + 1];
+		Segment.EndItemIndex = Index + 1;
+
+		// Only need to calculate initials if we need stretch info or we have node set to Initial
+		if (bCalculateStretch || bInitial)
+		{
+			const FVector InitialPositionA = Hierarchy->GetInitialGlobalTransform(Segment.EndItem).GetLocation();
+        	const FVector InitialPositionB = Hierarchy->GetInitialGlobalTransform(Segment.StartItem).GetLocation();
+			const FVector InitialSegmentVector = InitialPositionA - InitialPositionB;
+			Segment.InitialLength = InitialSegmentVector.Size();
+			Segment.InitialCumLength = InitialChainLength + Segment.InitialLength;
+        	InitialChainLength = Segment.InitialCumLength;
+		}
+
+        // If node is set to Initial we use the initial calculations as the current
+		if (bInitial)
+		{
+			Segment.Length = Segment.InitialLength;
+		}
+		else
+		{
+			const FVector PositionA = Hierarchy->GetGlobalTransform(Segment.EndItem).GetLocation();
+			const FVector PositionB = Hierarchy->GetGlobalTransform(Segment.StartItem).GetLocation();
+			const FVector SegmentVector = PositionA - PositionB;
+			Segment.Length = SegmentVector.Size();
+		}
+
+		Segment.CumLength = ChainLength + Segment.Length;
+        ChainLength = Segment.CumLength;		
+	}
+
+    // Calculate param length
+	ParamLength = Param * ChainLength;
+
+	// Binary search for current segment containing param based on recorded lengths
+    SegmentIndex = 0;
+	int StartSegmentIndex = 0;
+	int EndSegmentIndex = Segments.Num() - 1;
+
+	while (StartSegmentIndex != EndSegmentIndex) 
+	{
+		int SearchIndex = StartSegmentIndex + ((EndSegmentIndex - StartSegmentIndex) / 2);
+		if (ParamLength < Segments[SearchIndex].CumLength)
+		{
+			if (SearchIndex == 0)
+			{
+				SegmentIndex = SearchIndex;
+				break;
+			}
+			EndSegmentIndex = SearchIndex;
+		}
+		else if (ParamLength >= Segments[SearchIndex].CumLength)
+		{
+			if (Segments[SearchIndex].CumLength == ParamLength || SearchIndex + 1 == EndSegmentIndex)
+			{
+				SegmentIndex = SearchIndex + 1;
+				break;
+			}
+			StartSegmentIndex = SearchIndex;
+		}
+	}
+
+	const FRigUnit_ChainInfo_Segment& CurrentSegment = Segments[SegmentIndex];
+
+	SegmentStartItem = CurrentSegment.StartItem.GetKey();
+	SegmentStartItemIndex = CurrentSegment.StartItemIndex;
+	SegmentEndItem = CurrentSegment.EndItem.GetKey();
+	SegmentEndItemIndex = CurrentSegment.EndItemIndex;
+
+    // Calculate segment length and segment param length
+	if (SegmentIndex == 0) 
+	{
+		SegmentParamLength = ParamLength;
+	}
+	else 
+	{
+		SegmentParamLength = ParamLength - Segments[SegmentIndex - 1].CumLength;
+	}
+    SegmentLength = CurrentSegment.Length;
+
+    // Calculate segment param
+	if (SegmentLength != 0.f) 
+	{
+	    SegmentParam = SegmentParamLength / SegmentLength;
+	}
+
+	if (bCalculateStretch)
+	{
+		// Calculate stretch factors
+		if (InitialChainLength != 0.f)
+		{
+			ChainStretchFactor = ChainLength / InitialChainLength;
+		}
+
+		if (CurrentSegment.InitialLength != 0.f) 
+		{
+			SegmentStretchFactor = CurrentSegment.Length / CurrentSegment.InitialLength;
+		}
+	}
+
+    // Lerp a transform from the current start and end items with segment param
+	FTransform StartTransform;
+	FTransform EndTransform;
+	if (bInitial)
+	{
+		StartTransform = Hierarchy->GetInitialGlobalTransform(SegmentStartItem);
+		EndTransform = Hierarchy->GetInitialGlobalTransform(SegmentEndItem);
+	}
+	else
+	{
+		StartTransform = Hierarchy->GetGlobalTransform(SegmentStartItem);
+		EndTransform = Hierarchy->GetGlobalTransform(SegmentEndItem);
+	}
+	OutputTransform = FControlRigMathLibrary::LerpTransform(StartTransform, EndTransform, SegmentParam);
+    
+	// Debug draw
+	if (bDebug) {
+		FRigVMDebugDrawSettings DebugDrawSettings; 
+
+		// Draw interpolated output transform
+		FRigVMFunction_DebugTransformMutableNoSpace::StaticExecute(ExecuteContext, 
+		                                OutputTransform,
+		                                ERigUnitDebugTransformMode::Axes,
+		                                FLinearColor::White,
+		                                0.f,
+		                                DebugScale,
+										FTransform::Identity,
+		                                true,
+										DebugDrawSettings);
+
+		// Draw transform of start item
+		FRigVMFunction_DebugTransformMutableNoSpace::StaticExecute(ExecuteContext, 
+		                                StartTransform,
+		                                ERigUnitDebugTransformMode::Box,
+		                                FLinearColor::Green,
+		                                0.f,
+		                                DebugScale,
+										FTransform::Identity,
+		                                true,
+										DebugDrawSettings);
+
+		// Draw transform of end item
+		FRigVMFunction_DebugTransformMutableNoSpace::StaticExecute(ExecuteContext, 
+		                                EndTransform,
+		                                ERigUnitDebugTransformMode::Box,
+		                                FLinearColor::Red,
+		                                0.f,
+		                                DebugScale,
+										FTransform::Identity,
+		                                true,
+										DebugDrawSettings);
+	}
+}
diff --git a/Engine/Plugins/Animation/ControlRig/Source/ControlRig/Public/Units/Highlevel/Hierarchy/RigUnit_ChainInfo.h b/Engine/Plugins/Animation/ControlRig/Source/ControlRig/Public/Units/Highlevel/Hierarchy/RigUnit_ChainInfo.h
new file mode 100644
--- /dev/null
+++ b/Engine/Plugins/Animation/ControlRig/Source/ControlRig/Public/Units/Highlevel/Hierarchy/RigUnit_ChainInfo.h
@@ -0,0 +1,224 @@
+// Copyright Epic Games, Inc. All Rights Reserved.
+
+#pragma once
+
+#include "Units/Highlevel/RigUnit_HighlevelBase.h"
+#include "Math/ControlRigMathLibrary.h"
+#include "RigUnit_ChainInfo.generated.h"
+
+USTRUCT()
+struct FRigUnit_ChainInfo_Segment 
+{
+    GENERATED_BODY()
+
+	FRigUnit_ChainInfo_Segment()
+		:StartItem(FCachedRigElement())
+		,StartItemIndex(0)
+		,EndItem(FCachedRigElement())
+		,EndItemIndex(0)
+		,InitialLength(0.f)
+		,InitialCumLength(0.f)
+		,Length(0.f)
+		,CumLength(0.f)
+	{}
+
+	/**
+	 * Start item of the chain segment
+	 */
+	UPROPERTY()
+	FCachedRigElement StartItem;
+
+	/**
+	 * Start item index of the chain segment
+	 */
+	UPROPERTY()
+	int StartItemIndex;
+
+	/**
+	 * Emd item of the chain segment
+	 */
+	UPROPERTY()
+	FCachedRigElement EndItem;
+
+	/**
+	 * Emd item index of the chain segment
+	 */
+	UPROPERTY()
+	int EndItemIndex;
+
+	/**
+	 * Initial length of segment
+	 */
+	UPROPERTY()
+	float InitialLength; 
+
+	/**
+	 * Inclusive initial length of all previous segments and this one
+	 */
+	UPROPERTY()
+	float InitialCumLength;
+
+	/**
+	 * Length of segment
+	 */
+	UPROPERTY()
+	float Length; 
+
+	/**
+	 * Inclusive length of all previous segments and this one
+	 */
+	UPROPERTY()
+	float CumLength;
+};
+
+/*
+ * Retrieves various pieces of info about an interpolated transform hierarchy from an rig element item list
+ */
+USTRUCT(meta=(DisplayName="Chain Info", Category="Hierarchy", Keywords="Chain"))
+struct FRigUnit_ChainInfo : public FRigUnit_HighlevelBaseMutable
+{
+	GENERATED_BODY()
+
+	FRigUnit_ChainInfo()
+		:Param(0.f)
+		,bCalculateStretch(true)
+		,bInitial(false)
+		,bDebug(false)
+		,DebugScale(1.f)
+		,OutputTransform(FTransform::Identity)
+		,ChainLength(0.f)
+		,ParamLength(0.f)
+		,SegmentIndex(0)
+		,SegmentLength(0.f)	
+		,SegmentParam(0.f)
+		,SegmentParamLength(0.f)
+		,SegmentStartItem(NAME_None, ERigElementType::Bone)
+		,SegmentStartItemIndex(0)
+		,SegmentEndItem(NAME_None, ERigElementType::Bone)
+        ,SegmentEndItemIndex(0)
+		,ChainStretchFactor(1.f)
+		,SegmentStretchFactor(1.f)
+	{}
+
+	RIGVM_METHOD()
+	virtual void Execute() override;
+
+	/** 
+	 * The items to use to interpret the chain
+	 */
+	UPROPERTY(meta = (Input))
+	TArray<FRigElementKey> Items;
+
+	/**
+	 * The parameter value down the chain of items from 0 to 1
+	 */	
+	UPROPERTY(meta = (Input, ClampMin=0.f, ClampMax=1.f))
+	float Param;
+
+	/**
+	 * If True calculate stretch factors of chain and current segment
+	 */	
+	UPROPERTY(meta = (Input))
+	bool bCalculateStretch;
+
+	/**
+	 * If True use initial transform values for chain
+	 */	
+	UPROPERTY(meta = (Input))
+	bool bInitial;
+
+	/**
+	 * Enable debug draw for node
+	 */	
+	UPROPERTY(meta = (Input))
+	bool bDebug;
+
+	/**
+	 * Debug draw scale
+	 */	
+	UPROPERTY(meta = (Input, UIMin = 0.f))
+	float DebugScale;
+
+    /* 
+	 * The interpolated transform at the chain's input parameter
+	 */
+	UPROPERTY(meta = (Output))
+	FTransform OutputTransform;
+
+    /* 
+	 * The length of the interpolated chain
+	 */
+	UPROPERTY(meta = (Output))
+	float ChainLength;
+
+    /* 
+	 * The length of the interpolated chain
+	 */
+	UPROPERTY(meta = (Output))
+	float ParamLength;
+
+    /* 
+	 * The current segment index
+	 */
+	UPROPERTY(meta = (Output))
+	int SegmentIndex;
+
+	/* 
+	 * The current segment length
+	 */
+	UPROPERTY(meta = (Output))
+	float SegmentLength;	
+
+    /* 
+	 * The current segment parameter from 0 to 1 
+	 */
+	UPROPERTY(meta = (Output))
+	float SegmentParam;
+
+    /* 
+	 * Local segment length of segment param
+	 */
+	UPROPERTY(meta = (Output))
+	float SegmentParamLength;
+
+    /* 
+	 * The item starting the current segment 
+	 */
+	UPROPERTY(meta = (Output))
+	FRigElementKey SegmentStartItem;
+
+    /* 
+	 * The item index starting the current segment 
+	 */
+	UPROPERTY(meta = (Output))
+	int SegmentStartItemIndex;
+
+    /* 
+	 * The item ending the current segment 
+	 */
+	UPROPERTY(meta = (Output))
+	FRigElementKey SegmentEndItem;
+
+    /* 
+	 * The item index ending the current segment 
+	 */
+	UPROPERTY(meta = (Output))
+	int SegmentEndItemIndex;
+
+	/* 
+	 * Stretch factor of chain
+	 */
+	UPROPERTY(meta = (Output))
+	float ChainStretchFactor;
+
+	/* 
+	 * Stretch factor of current segment
+	 */
+	UPROPERTY(meta = (Output))
+	float SegmentStretchFactor;
+
+	// Used to cache the internally used index
+	UPROPERTY()
+	TArray<FCachedRigElement> CachedElements;
+
+};
