Protected ControlRig anom nnode cache bones function with the evaluation lock, as we can receive a cache bones request on a worker thread while the main thread is executing the rig

diff --git a/Engine/Plugins/Animation/ControlRig/Source/ControlRig/Private/AnimNode_ControlRigBase.cpp b/Engine/Plugins/Animation/ControlRig/Source/ControlRig/Private/AnimNode_ControlRigBase.cpp
--- a/Engine/Plugins/Animation/ControlRig/Source/ControlRig/Private/AnimNode_ControlRigBase.cpp
+++ b/Engine/Plugins/Animation/ControlRig/Source/ControlRig/Private/AnimNode_ControlRigBase.cpp
@@ -370,45 +370,51 @@ void FAnimNode_ControlRigBase::CacheBones_AnyThread(const FAnimationCacheBonesCo
 
 		const uint16 BonesSerialNumber = RequiredBones.GetSerialNumber();
 
-		// the construction event may create a set of bones that we can map to. let's run construction now.
-		if (bExecute)
 		{
-			const bool bIsLODChange = !bControlRigRequiresInitialization && (BonesSerialNumber != LastBonesSerialNumberForCacheBones);
-		
-			if(ControlRig->IsConstructionModeEnabled() ||
-				(ControlRig->IsConstructionRequired() && (bControlRigRequiresInitialization || bIsLODChange)))
+			// We can get a call to CacheBones in the main thread while a worker thread is also executing Control Rig (i.e. a FControlRigLayerInstanceProxy)
+			// We have to avoid the concurrent execution, because we will destroy the hierarchy here and create a new one
+			UE::TScopeLock EvaluateLock(ControlRig->GetEvaluateMutex());
+
+			// the construction event may create a set of bones that we can map to. let's run construction now.
+			if (bExecute)
 			{
-				UpdateGetAssetUserDataDelegate(ControlRig);
-				ControlRig->Execute(FRigUnit_PrepareForExecution::EventName);
-				bControlRigRequiresInitialization = false;
+				const bool bIsLODChange = !bControlRigRequiresInitialization && (BonesSerialNumber != LastBonesSerialNumberForCacheBones);
+		
+				if(ControlRig->IsConstructionModeEnabled() ||
+					(ControlRig->IsConstructionRequired() && (bControlRigRequiresInitialization || bIsLODChange)))
+				{
+					UpdateGetAssetUserDataDelegate(ControlRig);
+					ControlRig->Execute(FRigUnit_PrepareForExecution::EventName);
+					bControlRigRequiresInitialization = false;
+				}
 			}
-		}
 
-		URigHierarchy* Hierarchy = ControlRig->GetHierarchy();
+			URigHierarchy* Hierarchy = ControlRig->GetHierarchy();
 
-		ControlRigHierarchyMappings.UpdateInputOutputMappingIfRequired(ControlRig
-			, Hierarchy
-			, RequiredBones
-			, InputBonesToTransfer
-			, OutputBonesToTransfer
-			, NodeMappingContainer
-			, bTransferPoseInGlobalSpace
-			, bResetInputPoseToInitial);
+			ControlRigHierarchyMappings.UpdateInputOutputMappingIfRequired(ControlRig
+				, Hierarchy
+				, RequiredBones
+				, InputBonesToTransfer
+				, OutputBonesToTransfer
+				, NodeMappingContainer
+				, bTransferPoseInGlobalSpace
+				, bResetInputPoseToInitial);
 
-		if(bControlRigRequiresInitialization)
-		{
-			if(bExecute)
+			if(bControlRigRequiresInitialization)
 			{
-				// re-init only if this is the first run
-				// and restore control values
-				ControlRig->RequestInit();
-				bControlRigRequiresInitialization = false;
+				if(bExecute)
+				{
+					// re-init only if this is the first run
+					// and restore control values
+					ControlRig->RequestInit();
+					bControlRigRequiresInitialization = false;
+				}
 			}
-		}
 		
-		LastBonesSerialNumberForCacheBones = BonesSerialNumber;
+			LastBonesSerialNumberForCacheBones = BonesSerialNumber;
 
-		ControlRigHierarchyMappings.LinkToHierarchy(Hierarchy);
+			ControlRigHierarchyMappings.LinkToHierarchy(Hierarchy);
+		}
 	}
 }
 
