Control Rig: Remove to dos because they have been done. The original CL that addressed these tickets forgot to remove the comments.

diff --git a/Engine/Plugins/Animation/ControlRig/Source/ControlRigEditor/Private/ControlRigSequencerEditorLibrary.cpp b/Engine/Plugins/Animation/ControlRig/Source/ControlRigEditor/Private/ControlRigSequencerEditorLibrary.cpp
--- a/Engine/Plugins/Animation/ControlRig/Source/ControlRigEditor/Private/ControlRigSequencerEditorLibrary.cpp
+++ b/Engine/Plugins/Animation/ControlRig/Source/ControlRigEditor/Private/ControlRigSequencerEditorLibrary.cpp
@@ -411,7 +411,6 @@ bool UControlRigSequencerEditorLibrary::TweenControlRig(ULevelSequence* LevelSeq
 	{
 		LevelSequence->GetMovieScene()->Modify();
 		
-		// TODO UE-240486: This has unnecessary dynamic memory allocation overhead until UE-240486 is addressed.
 		const FBuiltInTweenFunctionContainer TweenFunctionContainer; 
 		FControlRigTweenModel BlendTool(WeakSequencer, *TweenFunctionContainer.GetBlendFunction(EBlendFunction::ControlsToTween));
 		BlendTool.StartBlendOperationWithControlRigs({ ControlRig });
@@ -453,7 +452,6 @@ bool UControlRigSequencerEditorLibrary::BlendValuesOnSelected(ULevelSequence* Le
 			TWeakPtr<FControlRigEditMode> WeakMode = StaticCastSharedRef<FControlRigEditMode, FEdMode>(EditMode->AsShared()).ToWeakPtr();
 			LevelSequence->GetMovieScene()->Modify();
 			
-			// TODO UE-240486: This has unnecessary dynamic memory allocation overhead until UE-240486 is addressed.
 			const FBuiltInTweenFunctionContainer TweenFunctionContainer;
 			const EBlendFunction BlendFunction = SequencerLibraryDetail::ConvertBlendFunction(BlendOperation);
 			FControlRigTweenModel BlendTool(WeakSequencer, WeakMode,
