Fixed compile warning about unintentional loop variable copy in Control Rig

diff --git a/Engine/Plugins/Animation/ControlRig/Source/ControlRigEditor/Private/ControlRigElementDetails.cpp b/Engine/Plugins/Animation/ControlRig/Source/ControlRigEditor/Private/ControlRigElementDetails.cpp
--- a/Engine/Plugins/Animation/ControlRig/Source/ControlRigEditor/Private/ControlRigElementDetails.cpp
+++ b/Engine/Plugins/Animation/ControlRig/Source/ControlRigEditor/Private/ControlRigElementDetails.cpp
@@ -5755,7 +5755,7 @@ FDetailWidgetRow& FRigControlElementDetails::CreateEnumValueWidgetRow(
 		else
 		{
 			// If the key was not found for selected elements, it might be a child channel of one of the elements
-			for (const FPerElementInfo ElementInfo : PerElementInfos)
+			for (const FPerElementInfo& ElementInfo : PerElementInfos)
 			{
 				if (const FRigControlElement* ControlElement = ElementInfo.GetElement<FRigControlElement>())
 				{
