Control Rig Edit Mode: only update the manipulated rigs / components / control shapes

#rnx

diff --git a/Engine/Plugins/Animation/ControlRig/Source/ControlRigEditor/Private/EditMode/ControlRigEditMode.cpp b/Engine/Plugins/Animation/ControlRig/Source/ControlRigEditor/Private/EditMode/ControlRigEditMode.cpp
--- a/Engine/Plugins/Animation/ControlRig/Source/ControlRigEditor/Private/EditMode/ControlRigEditMode.cpp
+++ b/Engine/Plugins/Animation/ControlRig/Source/ControlRigEditor/Private/EditMode/ControlRigEditMode.cpp
@@ -1239,13 +1239,18 @@ bool FControlRigEditMode::HandleEndTransform(FEditorViewportClient* InViewportCl
 			GEditor->EndTransaction();
 		}
 
+		TArray<TWeakObjectPtr<UControlRig>> RigsToTick;
+		RigsToTick.Reserve(InteractionScopes.Num());
+		
 		for (TPair<UControlRig*, FControlRigInteractionScope*>& InteractionScope : InteractionScopes)
 		{
+			RigsToTick.Add(InteractionScope.Key);
 			if (InteractionScope.Value)
 			{
 				delete InteractionScope.Value; 
 			}
 		}
+
 		InteractionScopes.Reset();
 
 		if (bWasInteracting && !AreEditingControlRigDirectly())
@@ -1253,7 +1258,7 @@ bool FControlRigEditMode::HandleEndTransform(FEditorViewportClient* InViewportCl
 			// We invalidate the hit proxies when in level editor to ensure that the gizmo's hit proxy is up to date.
 			// The invalidation is called here to avoid useless viewport update in the FControlRigEditMode::Tick
 			// function (that does an update when not in level editor)
-			TickManipulatableObjects(0.f);
+			TickManipulatableObjects(RigsToTick);
 			
 			static constexpr bool bInvalidateChildViews = false;
 			static constexpr bool bInvalidateHitProxies = true;
@@ -2240,8 +2245,16 @@ bool FControlRigEditMode::InputDelta(FEditorViewportClient* InViewportClient, FV
 
 	if (bManipulatorMadeChange)
 	{
-		TickManipulatableObjects(0.f);
+		TArray<TWeakObjectPtr<UControlRig>> RigsToTick;
+		RigsToTick.Reserve(InteractionScopes.Num());
+		Algo::Transform(InteractionScopes, RigsToTick, [](const TPair<UControlRig*,FControlRigInteractionScope*>& InteractionData)
+		{
+			return InteractionData.Key;
+		});
+		
+		TickManipulatableObjects(RigsToTick);
 	}
+	
 	//if in level editor we want to move other things also
 	return IsInLevelEditor() ? false :bManipulatorMadeChange;
 }
@@ -4690,7 +4703,6 @@ void FControlRigEditMode::OnHierarchyModified(ERigHierarchyNotification InNotif,
 		}
 		case ERigHierarchyNotification::ParentWeightsChanged:
 		{
-			// TickManipulatableObjects(0.f);
 			break;
 		}
 		case ERigHierarchyNotification::InteractionBracketOpened:
@@ -5651,9 +5663,11 @@ void FControlRigEditMode::RemoveControlRig(UControlRig* InControlRig)
 	UpdateSelectabilityOnSkeletalMeshes(InControlRig, true);
 }
 
-void FControlRigEditMode::TickManipulatableObjects(float DeltaTime)
+void FControlRigEditMode::TickManipulatableObjects(const TArray<TWeakObjectPtr<UControlRig>>& InRigs) const
 {
-	for (TWeakObjectPtr<UControlRig> RuntimeRigPtr : RuntimeControlRigs)
+	const TArray<TWeakObjectPtr<UControlRig>>& RigsToTick = InRigs.IsEmpty() ? RuntimeControlRigs : InRigs;
+
+	for (TWeakObjectPtr<UControlRig> RuntimeRigPtr : RigsToTick)
 	{
 		if (UControlRig* ControlRig = RuntimeRigPtr.Get())
 		{
@@ -5710,7 +5724,8 @@ void FControlRigEditMode::TickManipulatableObjects(float DeltaTime)
 			}
 		}
 	}
-	PostPoseUpdate();
+	
+	PostPoseUpdate(nullptr, InRigs);
 }
 
 
@@ -5915,7 +5930,7 @@ void FControlRigEditMode::OnPoseInitialized()
 	OnAnimSystemInitializedDelegate.Broadcast();
 }
 
-void FControlRigEditMode::PostPoseUpdate(const FEditorViewportClient* InViewportClient) const
+void FControlRigEditMode::PostPoseUpdate(const FEditorViewportClient* InViewportClient, const TArray<TWeakObjectPtr<UControlRig>>& InRigs) const
 {
 	using namespace ControlRigEditMode::Shapes;
 
@@ -5927,16 +5942,37 @@ void FControlRigEditMode::PostPoseUpdate(const FEditorViewportClient* InViewport
 	}
 
 	const bool bAreEditingControlRigDirectly = AreEditingControlRigDirectly();
-	for (const TPair<UControlRig*,TArray<TObjectPtr<AControlRigShapeActor>>>& ShapeActors : ControlRigShapeActors)
+	auto UpdateShapes = [this, bAreEditingControlRigDirectly](const UControlRig* ControlRig, const TArray<TObjectPtr<AControlRigShapeActor>>& ShapeActors)
 	{
-		const UControlRig* ControlRig = ShapeActors.Key;
 		const FTransform ComponentTransform = bAreEditingControlRigDirectly ? FTransform::Identity : GetHostingSceneComponentTransform(ControlRig); 
 
 		FShapeUpdateParams Params(ControlRig, ComponentTransform, IsControlRigSkelMeshVisible(ControlRig), IsInLevelEditor());
-		for (AControlRigShapeActor* ShapeActor : ShapeActors.Value)
+		for (AControlRigShapeActor* ShapeActor : ShapeActors)
 		{
 			UpdateControlShape(ShapeActor, ControlRig->FindControl(ShapeActor->ControlName), Params);
 		}
+	};
+	
+	if (InRigs.IsEmpty())
+	{
+		// updates all control shapes properties
+		for (const TPair<UControlRig*,TArray<TObjectPtr<AControlRigShapeActor>>>& ShapeActors : ControlRigShapeActors)
+		{
+			UpdateShapes(ShapeActors.Key, ShapeActors.Value);
+		}
+	}
+	else
+	{
+		for (const TWeakObjectPtr<UControlRig>& RigPtr: InRigs)
+		{
+			if (const UControlRig* ControlRig = RigPtr.Get())
+			{
+				if (const TArray<TObjectPtr<AControlRigShapeActor>>* ShapeActors = ControlRigShapeActors.Find(ControlRig))
+				{
+					UpdateShapes(ControlRig, *ShapeActors);
+				}
+			}
+		}
 	}
 }
 
diff --git a/Engine/Plugins/Animation/ControlRig/Source/ControlRigEditor/Public/EditMode/ControlRigEditMode.h b/Engine/Plugins/Animation/ControlRig/Source/ControlRigEditor/Public/EditMode/ControlRigEditMode.h
--- a/Engine/Plugins/Animation/ControlRig/Source/ControlRigEditor/Public/EditMode/ControlRigEditMode.h
+++ b/Engine/Plugins/Animation/ControlRig/Source/ControlRigEditor/Public/EditMode/ControlRigEditMode.h
@@ -525,7 +525,12 @@ private:
 
 	/*Internal function for adding ControlRig*/
 	void AddControlRigInternal(UControlRig* InControlRig);
-	void TickManipulatableObjects(float DeltaTime);
+
+	/**
+	 * Updates the bound components (skeletal meshes or control rig component) and the control shapes so that all transform data are updated.
+	 * if InRig is empty, all rigs / control shapes will be updated.
+	 */
+	void TickManipulatableObjects(const TArray<TWeakObjectPtr<UControlRig>>& InRigs = TArray<TWeakObjectPtr<UControlRig>>()) const;
 
 	/* Check on tick to see if movie scene has changed, returns true if it has*/
 	bool CheckMovieSceneSig();
@@ -567,7 +572,13 @@ private:
 
 	// Post pose update handler
 	void OnPoseInitialized();
-	void PostPoseUpdate(const FEditorViewportClient* InViewportClient = nullptr) const;
+
+	/**
+	 * Updates the control shapes properties (transform, visibility, ...) depending on the current viewport state.
+	 * If InRig is empty, all control shapes will be updated.
+	 */
+	void PostPoseUpdate(const FEditorViewportClient* InViewportClient = nullptr, const TArray<TWeakObjectPtr<UControlRig>>& InRigs = TArray<TWeakObjectPtr<UControlRig>>()) const;
+	
 	void UpdateSelectabilityOnSkeletalMeshes(UControlRig* InControlRig, bool bEnabled);
 
 	bool IsMovingCamera(const FViewport* InViewport) const;
