Control Rig : Fixed a crash when the user generates controls in the construction event and expose these controls as a promoted variable

diff --git a/Engine/Plugins/Animation/ControlRig/Source/ControlRig/Private/AnimNode_ControlRig.cpp b/Engine/Plugins/Animation/ControlRig/Source/ControlRig/Private/AnimNode_ControlRig.cpp
--- a/Engine/Plugins/Animation/ControlRig/Source/ControlRig/Private/AnimNode_ControlRig.cpp
+++ b/Engine/Plugins/Animation/ControlRig/Source/ControlRig/Private/AnimNode_ControlRig.cpp
@@ -169,13 +169,21 @@ void FAnimNode_ControlRig::CacheBones_AnyThread(const FAnimationCacheBonesContex
 	PropagateInputProperties(Context.AnimInstanceProxy->GetAnimInstanceObject());
 
 	// update the control rig instance just in case the dynamic control rig class has changed
-	if(const UAnimInstance* AnimInstance = Cast<UAnimInstance>(Context.GetAnimInstanceObject()))
+	const UAnimInstance* AnimInstance = Cast<UAnimInstance>(Context.GetAnimInstanceObject());
+	if (AnimInstance != nullptr)
 	{
 		(void)UpdateControlRigIfNeeded(AnimInstance, Context.AnimInstanceProxy->GetRequiredBones());
 	}
 
 	FAnimNode_ControlRigBase::CacheBones_AnyThread(Context);
 
+	// The call to base cache bones might execute construction event, which will recreate the user generated controls
+	// In case one of these controls is exposed as public variable, we have to re-initialize variable mappings
+	if (AnimInstance != nullptr && ControlRigVariableMappings.RequiresInitAfterConstruction())
+	{
+		ControlRigVariableMappings.InitializeProperties(AnimInstance->GetClass(), TargetInstance, GetTargetClass(), SourcePropertyNames, DestPropertyNames);
+	}
+
 	ControlRigVariableMappings.ResetCurvesInputToControlCache();
 
 	FBoneContainer& RequiredBones = Context.AnimInstanceProxy->GetRequiredBones();
diff --git a/Engine/Plugins/Animation/ControlRig/Source/ControlRig/Private/Tools/ControlRigVariableMappings.cpp b/Engine/Plugins/Animation/ControlRig/Source/ControlRig/Private/Tools/ControlRigVariableMappings.cpp
--- a/Engine/Plugins/Animation/ControlRig/Source/ControlRig/Private/Tools/ControlRigVariableMappings.cpp
+++ b/Engine/Plugins/Animation/ControlRig/Source/ControlRig/Private/Tools/ControlRigVariableMappings.cpp
@@ -100,6 +100,13 @@ void FControlRigVariableMappings::InitializeProperties(const UClass* InSourceCla
 	}
 }
 
+bool FControlRigVariableMappings::RequiresInitAfterConstruction() const
+{
+	const bool bUseFunctions = AnimNodeLocals::bUseFunctions;
+
+	return bUseFunctions && UpdateFunctions.Num() > 0;
+}
+
 void FControlRigVariableMappings::InitializeCustomProperties(UControlRig* TargetControlRig, const FCustomPropertyMappings& CustomPropertyMapping)
 {
 	const TArray<FCustomPropertyData>& Mappings = CustomPropertyMapping.GetMappings();
diff --git a/Engine/Plugins/Animation/ControlRig/Source/ControlRig/Public/Tools/ControlRigVariableMappings.h b/Engine/Plugins/Animation/ControlRig/Source/ControlRig/Public/Tools/ControlRigVariableMappings.h
--- a/Engine/Plugins/Animation/ControlRig/Source/ControlRig/Public/Tools/ControlRigVariableMappings.h
+++ b/Engine/Plugins/Animation/ControlRig/Source/ControlRig/Public/Tools/ControlRigVariableMappings.h
@@ -101,6 +101,8 @@ struct CONTROLRIG_API FControlRigVariableMappings
 		, const TArray<FName>& SourcePropertyNames
 		, TArray<FName>& DestPropertyNames);
 
+	bool RequiresInitAfterConstruction() const;
+
 	void InitializeCustomProperties(UControlRig* TargetControlRig, const FCustomPropertyMappings& CustomPropertyMapping);
 	void ResetCustomProperties(int32 NewSize = 0);
 
