Anim In Engine: Add custom range to import animation from specified animation time range.

diff --git a/Engine/Plugins/Animation/ControlRig/Source/ControlRig/Private/Sequencer/MovieSceneControlRigParameterSection.cpp b/Engine/Plugins/Animation/ControlRig/Source/ControlRig/Private/Sequencer/MovieSceneControlRigParameterSection.cpp
--- a/Engine/Plugins/Animation/ControlRig/Source/ControlRig/Private/Sequencer/MovieSceneControlRigParameterSection.cpp
+++ b/Engine/Plugins/Animation/ControlRig/Source/ControlRig/Private/Sequencer/MovieSceneControlRigParameterSection.cpp
@@ -3439,10 +3439,28 @@ void UMovieSceneControlRigParameterSection::RecordControlRigKey(FFrameNumber Fra
 bool UMovieSceneControlRigParameterSection::LoadAnimSequenceIntoThisSection(UAnimSequence* AnimSequence, UMovieScene* MovieScene, UObject* BoundObject, bool bKeyReduce, float Tolerance, bool bResetControls, FFrameNumber InStartFrame, EMovieSceneKeyInterpolation InInterpolation)
 {
 	FFrameNumber SequenceStart = UE::MovieScene::DiscreteInclusiveLower(MovieScene->GetPlaybackRange());
-	return LoadAnimSequenceIntoThisSection(AnimSequence, SequenceStart, MovieScene, BoundObject, bKeyReduce, Tolerance, bResetControls, InStartFrame, InInterpolation);
+	UMovieSceneControlRigParameterSection::FLoadAnimSequenceData Data;
+	Data.bKeyReduce = bKeyReduce;
+	Data.Tolerance = Tolerance;
+	Data.bResetControls = bResetControls;
+	Data.StartFrame = InStartFrame;
+	return LoadAnimSequenceIntoThisSection(AnimSequence, SequenceStart, MovieScene, BoundObject, Data, InInterpolation);
 }
 
 bool UMovieSceneControlRigParameterSection::LoadAnimSequenceIntoThisSection(UAnimSequence* AnimSequence, const FFrameNumber& SequenceStart, UMovieScene* MovieScene, UObject* BoundObject, bool bKeyReduce, float Tolerance, bool bResetControls, const FFrameNumber& InStartFrame, EMovieSceneKeyInterpolation InInterpolation)
+{
+	UMovieSceneControlRigParameterSection::FLoadAnimSequenceData Data;
+	Data.bKeyReduce = bKeyReduce;
+	Data.Tolerance = Tolerance;
+	Data.bResetControls = bResetControls;
+	Data.StartFrame = InStartFrame;
+
+	return LoadAnimSequenceIntoThisSection(AnimSequence, SequenceStart, MovieScene, BoundObject, Data, InInterpolation);
+}
+
+//Function to load an Anim Sequence into this section. It will automatically resize to the section size.
+//Will return false if fails or is canceled	
+bool UMovieSceneControlRigParameterSection::LoadAnimSequenceIntoThisSection(UAnimSequence* AnimSequence, const FFrameNumber& SequenceStart, UMovieScene* MovieScene, UObject* BoundObject, const FLoadAnimSequenceData& LoadData, EMovieSceneKeyInterpolation InInterpolation)
 {
 	USkeletalMeshComponent* SkelMeshComp = Cast<USkeletalMeshComponent>(BoundObject);
 	
@@ -3474,7 +3492,7 @@ bool UMovieSceneControlRigParameterSection::LoadAnimSequenceIntoThisSection(UAni
 	float Length = AnimSequence->GetPlayLength();
 	const FFrameRate& FrameRate = AnimSequence->GetSamplingFrameRate();
 
-	FFrameNumber StartFrame = SequenceStart + InStartFrame;
+	FFrameNumber StartFrame = SequenceStart + LoadData.StartFrame;
 	FFrameNumber EndFrame = TickResolution.AsFrameNumber(Length) + StartFrame;
 
 	Modify();
@@ -3486,9 +3504,10 @@ bool UMovieSceneControlRigParameterSection::LoadAnimSequenceIntoThisSection(UAni
 	}
 	ControlRig->Modify();
 
+	
 	const int32 NumberOfKeys = AnimSequence->GetDataModel()->GetNumberOfKeys();
 	FFrameNumber FrameRateInFrameNumber = TickResolution.AsFrameNumber(FrameRate.AsInterval());
-	int32 ExtraProgress = bKeyReduce ? FloatChannels.Num() : 0;
+	int32 ExtraProgress = LoadData.bKeyReduce ? FloatChannels.Num() : 0;
 	
 	FScopedSlowTask Progress(NumberOfKeys + ExtraProgress, LOCTEXT("BakingToControlRig_SlowTask", "Baking To Control Rig..."));	
 	Progress.MakeDialog(true);
@@ -3517,7 +3536,7 @@ bool UMovieSceneControlRigParameterSection::LoadAnimSequenceIntoThisSection(UAni
 
 	// copy the hierarchy from the CDO into the target control rig.
 	// this ensures that the topology version matches in case of a dynamic hierarchy
-	if (bResetControls)
+	if (LoadData.bResetControls)
 	{
 		if(!ControlRig->GetClass()->IsNative())
 		{
@@ -3537,18 +3556,29 @@ bool UMovieSceneControlRigParameterSection::LoadAnimSequenceIntoThisSection(UAni
 	{
 		ControlRig->SetBoneInitialTransformsFromRefSkeleton(Skeleton->GetReferenceSkeleton());
 	}
-	if (bResetControls)
+	if (LoadData.bResetControls)
 	{
 		ControlRig->RequestConstruction();
 		ControlRig->Evaluate_AnyThread();
 	}
 
-	for (int32 Index = 0; Index < NumberOfKeys; ++Index)
+	int32 Index = 0;
+	int32 EndIndex = NumberOfKeys - 1; 
+	const int32 LastIndex = EndIndex;
+	if (LoadData.AnimFrameRange.IsSet())
+	{
+		Index = (int32)(LoadData.AnimFrameRange.GetValue().GetLowerBoundValue().Value);
+		Index = Index > (LastIndex) ? LastIndex : Index;
+		EndIndex = (int32)(LoadData.AnimFrameRange.GetValue().GetUpperBoundValue().Value);
+		EndIndex = EndIndex > (LastIndex) ? LastIndex : EndIndex;
+	}
+	int32 KeyIndex = 0;
+	for (; Index <= EndIndex; ++Index)
 	{
 		const float SequenceSecond = AnimSequence->GetTimeAtFrame(Index);
-		const FFrameNumber FrameNumber = StartFrame + (FMath::Max(FrameRateInFrameNumber.Value, 1) * Index);
-
-		if (bResetControls)
+		const FFrameNumber FrameNumber = StartFrame + (FMath::Max(FrameRateInFrameNumber.Value, 1) * KeyIndex);
+		++KeyIndex;
+		if (LoadData.bResetControls)
 		{
 			SourceHierarchy->ResetPoseToInitial();
 			SourceHierarchy->ResetCurveValues();
@@ -3581,14 +3611,14 @@ bool UMovieSceneControlRigParameterSection::LoadAnimSequenceIntoThisSection(UAni
 			}
 		}
 
-		if (Index == 0)
+		if (KeyIndex == 0)
 		{
 			//to make sure the first frame looks good we need to do this first. UE-100069
 			ControlRig->Execute(FRigUnit_InverseExecution::EventName);
 		}
 		ControlRig->Execute(FRigUnit_InverseExecution::EventName);
 
-		RecordControlRigKey(FrameNumber, Index == 0, InInterpolation);
+		RecordControlRigKey(FrameNumber, KeyIndex == 0, InInterpolation);
 		Progress.EnterProgressFrame(1);
 		if (Progress.ShouldCancel())
 		{
@@ -3596,11 +3626,11 @@ bool UMovieSceneControlRigParameterSection::LoadAnimSequenceIntoThisSection(UAni
 		}
 	}
 
-	if (bKeyReduce)
+	if (LoadData.bKeyReduce)
 	{
 		FKeyDataOptimizationParams Params;
 		Params.bAutoSetInterpolation = true;
-		Params.Tolerance = Tolerance;
+		Params.Tolerance = LoadData.Tolerance;
 		for (FMovieSceneFloatChannel* Channel : FloatChannels)
 		{
 			Channel->Optimize(Params); //should also auto tangent
diff --git a/Engine/Plugins/Animation/ControlRig/Source/ControlRig/Public/Sequencer/MovieSceneControlRigParameterSection.h b/Engine/Plugins/Animation/ControlRig/Source/ControlRig/Public/Sequencer/MovieSceneControlRigParameterSection.h
--- a/Engine/Plugins/Animation/ControlRig/Source/ControlRig/Public/Sequencer/MovieSceneControlRigParameterSection.h
+++ b/Engine/Plugins/Animation/ControlRig/Source/ControlRig/Public/Sequencer/MovieSceneControlRigParameterSection.h
@@ -213,7 +213,6 @@ struct CONTROLRIG_API FChannelMapInfo
 	TArray<uint32> ConstraintsIndex; //constraints data
 };
 
-
 struct FMovieSceneControlRigSpaceChannel;
 
 /**
@@ -367,11 +366,28 @@ public:
 	//Function to save control rig key when recording.
 	void RecordControlRigKey(FFrameNumber FrameNumber, bool bSetDefault, EMovieSceneKeyInterpolation InInterpMode);
 
+	struct FLoadAnimSequenceData
+	{
+		//key reduce
+		bool bKeyReduce = false;
+		// key reduction tolerance
+		float Tolerance = 0.5;
+		//Whether to rest to default control states
+		bool bResetControls = true;
+		//Frame to Insert at
+		FFrameNumber StartFrame;
+		//If set only load the animation from the specified range
+		TOptional<TRange<FFrameNumber>> AnimFrameRange;
+	};
+	
 	//Function to load an Anim Sequence into this section. It will automatically resize to the section size.
-	//Will return false if fails or is canceled
+	//Will return false if fails or is canceled	
+	virtual bool LoadAnimSequenceIntoThisSection(UAnimSequence* Sequence, const FFrameNumber& SequenceStart, UMovieScene* MovieScene, UObject* BoundObject, const FLoadAnimSequenceData& LoadData, EMovieSceneKeyInterpolation Interpolation);
+	
+	UE_DEPRECATED(5.6, "LoadAnimSequenceIntoThisSection without taking FLoadAnimSequenceData data is deprecated")
 	virtual bool LoadAnimSequenceIntoThisSection(UAnimSequence* Sequence, const FFrameNumber& SequenceStart, UMovieScene* MovieScene, UObject* BoundObject, bool bKeyReduce, float Tolerance, bool bResetControls, const FFrameNumber& InStartFrame, EMovieSceneKeyInterpolation InInterpolation);
 
-	UE_DEPRECATED(5.5, "LoadAnimSequenceIntoThisSection without taking a sequence start frame is deprecated, use version that takes start frame instead")
+	UE_DEPRECATED(5.5, "LoadAnimSequenceIntoThisSection without taking FLoadAnimSequenceData data is deprecated")
 	virtual bool LoadAnimSequenceIntoThisSection(UAnimSequence* Sequence, UMovieScene* MovieScene, UObject* BoundObject, bool bKeyReduce, float Tolerance, bool bResetControls, FFrameNumber InStartFrame , EMovieSceneKeyInterpolation InInterpolation);
 #endif
 	
diff --git a/Engine/Plugins/Animation/ControlRig/Source/ControlRigEditor/Private/ControlRigEditorModule.cpp b/Engine/Plugins/Animation/ControlRig/Source/ControlRigEditor/Private/ControlRigEditorModule.cpp
--- a/Engine/Plugins/Animation/ControlRig/Source/ControlRigEditor/Private/ControlRigEditorModule.cpp
+++ b/Engine/Plugins/Animation/ControlRig/Source/ControlRigEditor/Private/ControlRigEditorModule.cpp
@@ -1273,8 +1273,9 @@ void FControlRigEditorModule::BakeToControlRig(UClass* ControlRigClass, UAnimSeq
 							FSmartReduceParams SmartReduce;
 							SmartReduce.TolerancePercentage = KeyReduceTolerance;
 							SmartReduce.SampleRate = BakeFrameRate;
+							TOptional<TRange<FFrameNumber>> AnimFrameRamge; //use whole range there
 							FControlRigParameterTrackEditor::LoadAnimationIntoSection(SequencerPtr, AnimSequence, SkelMeshComp, FFrameNumber(0),
-								bKeyReduce, SmartReduce, bResetControls, ParamSection);
+								bKeyReduce, SmartReduce, bResetControls, AnimFrameRamge, ParamSection);
 						}
 						SequencerPtr->EmptySelection();
 						SequencerPtr->SelectSection(ParamSection);
diff --git a/Engine/Plugins/Animation/ControlRig/Source/ControlRigEditor/Private/ControlRigSequencerEditorLibrary.cpp b/Engine/Plugins/Animation/ControlRig/Source/ControlRigEditor/Private/ControlRigSequencerEditorLibrary.cpp
--- a/Engine/Plugins/Animation/ControlRig/Source/ControlRigEditor/Private/ControlRigSequencerEditorLibrary.cpp
+++ b/Engine/Plugins/Animation/ControlRig/Source/ControlRigEditor/Private/ControlRigSequencerEditorLibrary.cpp
@@ -1357,8 +1357,13 @@ bool UControlRigSequencerEditorLibrary::BakeToControlRig(UWorld* World, ULevelSe
 						DefaultInterpolation = WeakSequencer.Pin()->GetKeyInterpolation();
 					}
 					FFrameNumber StartFrame = MovieScene->GetPlaybackRange().GetLowerBoundValue();
+					UMovieSceneControlRigParameterSection::FLoadAnimSequenceData Data;
+					Data.bKeyReduce = bReduceKeys;
+					Data.Tolerance = Tolerance;
+					Data.bResetControls = bResetControls;
+					Data.StartFrame = FFrameNumber(0);
 					ParamSection->LoadAnimSequenceIntoThisSection(TempAnimSequence, StartFrame, MovieScene, SkeletalMeshComp,
-						bReduceKeys, Tolerance, bResetControls, FFrameNumber(0), DefaultInterpolation);
+						Data, DefaultInterpolation);
 
 					//Turn Off Any Skeletal Animation Tracks
 					UMovieSceneSkeletalAnimationTrack* SkelTrack = Cast<UMovieSceneSkeletalAnimationTrack>(MovieScene->FindTrack(UMovieSceneSkeletalAnimationTrack::StaticClass(), Binding.BindingID, NAME_None));
@@ -1405,8 +1410,18 @@ bool UControlRigSequencerEditorLibrary::BakeToControlRig(UWorld* World, ULevelSe
 	return bResult;
 }
 
+//deprecreated 5.6
 bool UControlRigSequencerEditorLibrary::LoadAnimSequenceIntoControlRigSection(UMovieSceneSection* MovieSceneSection, UAnimSequence* AnimSequence, USkeletalMeshComponent* SkelMeshComp,
-	FFrameNumber InStartFrame, EMovieSceneTimeUnit TimeUnit,bool bKeyReduce, float Tolerance, EMovieSceneKeyInterpolation Interpolation, bool bResetControls)
+	FFrameNumber InStartFrame, EMovieSceneTimeUnit TimeUnit, bool bKeyReduce, float Tolerance, EMovieSceneKeyInterpolation Interpolation, bool bResetControls)
+{
+	return  UControlRigSequencerEditorLibrary::LoadAnimSequenceIntoControlRigSectionWithRange(MovieSceneSection, AnimSequence, SkelMeshComp, InStartFrame,
+		false, FFrameNumber(0), FFrameNumber(0),
+		TimeUnit, bKeyReduce, Tolerance, Interpolation, bResetControls);
+}
+
+bool UControlRigSequencerEditorLibrary::LoadAnimSequenceIntoControlRigSectionWithRange(UMovieSceneSection* MovieSceneSection, UAnimSequence* AnimSequence, USkeletalMeshComponent* SkelMeshComp,
+	FFrameNumber InStartFrame, bool bUseCustomAnimRange, FFrameNumber AnimStartRange, FFrameNumber AnimEndRange,
+	EMovieSceneTimeUnit TimeUnit,bool bKeyReduce, float Tolerance, EMovieSceneKeyInterpolation Interpolation, bool bResetControls)
 {
 	if (MovieSceneSection == nullptr || AnimSequence == nullptr || SkelMeshComp == nullptr)
 	{
@@ -1424,7 +1439,19 @@ bool UControlRigSequencerEditorLibrary::LoadAnimSequenceIntoControlRigSection(UM
 			InStartFrame = FFrameRate::TransformTime(FFrameTime(InStartFrame, 0), MovieScene->GetDisplayRate(),MovieScene->GetTickResolution()).RoundToFrame();
 		}
 		FFrameNumber StartFrame  = MovieScene->GetPlaybackRange().GetLowerBoundValue();
-		return Section->LoadAnimSequenceIntoThisSection(AnimSequence, StartFrame,  MovieScene, SkelMeshComp, bKeyReduce, Tolerance, bResetControls, InStartFrame, Interpolation);
+		UMovieSceneControlRigParameterSection::FLoadAnimSequenceData Data;
+		Data.bKeyReduce = bKeyReduce;
+		Data.Tolerance = Tolerance;
+		Data.bResetControls = bResetControls;
+		Data.StartFrame = InStartFrame;
+		if (bUseCustomAnimRange)
+		{
+			TRange<FFrameNumber> Range{ 0 };
+			Range.SetLowerBoundValue(AnimStartRange);
+			Range.SetUpperBoundValue(AnimEndRange);
+			Data.AnimFrameRange = Range;
+		}
+		return Section->LoadAnimSequenceIntoThisSection(AnimSequence, StartFrame,  MovieScene, SkelMeshComp, Data, Interpolation);
 	}
 	return false;
 }
diff --git a/Engine/Plugins/Animation/ControlRig/Source/ControlRigEditor/Private/Sequencer/ControlRigParameterTrackEditor.cpp b/Engine/Plugins/Animation/ControlRig/Source/ControlRigEditor/Private/Sequencer/ControlRigParameterTrackEditor.cpp
--- a/Engine/Plugins/Animation/ControlRig/Source/ControlRigEditor/Private/Sequencer/ControlRigParameterTrackEditor.cpp
+++ b/Engine/Plugins/Animation/ControlRig/Source/ControlRigEditor/Private/Sequencer/ControlRigParameterTrackEditor.cpp
@@ -87,6 +87,7 @@
 #include "Animation/SkeletalMeshActor.h"
 #include "TimerManager.h"
 #include "BakeToControlRigSettings.h"
+#include "LoadAnimToControlRigSettings.h"
 #include "UnrealEdGlobals.h"
 #include "Editor/UnrealEdEngine.h"
 #include "Toolkits/IToolkitHost.h"
@@ -109,6 +110,7 @@
 #include "LevelSequence.h"
 #include "ISequencerModule.h"
 #include "CurveModel.h"
+#include "Sequencer/SLoadAnimToControlRig.h"
 #include "FrontendFilterBase.h"
 
 #define LOCTEXT_NAMESPACE "FControlRigParameterTrackEditor"
@@ -1164,14 +1166,19 @@ void FControlRigParameterTrackEditor::SmartReduce(const TSharedPtr<ISequencer>&
 }
 
 bool FControlRigParameterTrackEditor::LoadAnimationIntoSection(const TSharedPtr<ISequencer>& SequencerPtr, UAnimSequence* AnimSequence, USkeletalMeshComponent* SkelMeshComp,
-	FFrameNumber StartFrame, bool bReduceKeys, const FSmartReduceParams& ReduceParams, bool bResetControls, UMovieSceneControlRigParameterSection* ParamSection)
+	FFrameNumber StartFrame, bool bReduceKeys, const FSmartReduceParams& ReduceParams, bool bResetControls, const TOptional<TRange<FFrameNumber>>& InAnimFrameRange, UMovieSceneControlRigParameterSection* ParamSection)
 {
 	EMovieSceneKeyInterpolation DefaultInterpolation = SequencerPtr->GetKeyInterpolation();
 	UMovieSceneSequence* OwnerSequence = SequencerPtr->GetFocusedMovieSceneSequence();
 	UMovieScene* OwnerMovieScene = OwnerSequence ? OwnerSequence->GetMovieScene() : nullptr;
-	
+	UMovieSceneControlRigParameterSection::FLoadAnimSequenceData Data;
+	Data.bKeyReduce = false; //use smart reduce
+	Data.Tolerance = 0.0;
+	Data.bResetControls = bResetControls;
+	Data.StartFrame = StartFrame;
+	Data.AnimFrameRange = InAnimFrameRange;
 	if (ParamSection->LoadAnimSequenceIntoThisSection(AnimSequence, FFrameNumber(0), OwnerMovieScene, SkelMeshComp,
-		false, 0.0, bResetControls, StartFrame, DefaultInterpolation))
+		Data, DefaultInterpolation))
 	{
 		if (bReduceKeys)
 		{
@@ -1328,9 +1335,10 @@ void FControlRigParameterTrackEditor::BakeToControlRig(UClass* InClass, FGuid Ob
 					Sequencer->NotifyMovieSceneDataChanged(EMovieSceneDataChangeType::MovieSceneStructureItemAdded);
 					TOptional<TRange<FFrameNumber>> OptionalRange = Sequencer->GetSubSequenceRange();
 					FFrameNumber StartFrame = OptionalRange.IsSet() ? OptionalRange.GetValue().GetLowerBoundValue() : OwnerMovieScene->GetPlaybackRange().GetLowerBoundValue();
+					TOptional<TRange<FFrameNumber>> AnimLoadFrameRamge; //do whole range
 
 					LoadAnimationIntoSection(Sequencer, TempAnimSequence, SkelMeshComp, StartFrame,
-						BakeSettings->bReduceKeys, BakeSettings->SmartReduce, BakeSettings->bResetControls, ParamSection);
+						BakeSettings->bReduceKeys, BakeSettings->SmartReduce, BakeSettings->bResetControls, AnimLoadFrameRamge, ParamSection);
 
 					//Turn Off Any Skeletal Animation Tracks
 					TArray<UMovieSceneSkeletalAnimationTrack*> SkelAnimationTracks;
@@ -5961,7 +5969,7 @@ void FControlRigParameterSection::BuildSectionContextMenu(FMenuBuilder& MenuBuil
 			{
 				MenuBuilder.AddSubMenu(
 					LOCTEXT("ImportAnimSequenceIntoThisSection", "Import Anim Sequence Into This Section"), NSLOCTEXT("Sequencer", "ImportAnimSequenceIntoThisSectionTP", "Import Anim Sequence Into This Section"),
-					FNewMenuDelegate::CreateRaw(this, &FControlRigParameterSection::AddAnimationSubMenuForFK, InObjectBinding, Skeleton, ParameterSection)
+					FNewMenuDelegate::CreateRaw(this, &FControlRigParameterSection::LoadAnimationIntoSection, InObjectBinding, Skeleton, ParameterSection)
 				);
 			}
 		}
@@ -6280,19 +6288,19 @@ bool FControlRigParameterSection::RequestDeleteKeyArea(const TArray<FName>& KeyA
 }
 
 
-void FControlRigParameterSection::AddAnimationSubMenuForFK(FMenuBuilder& MenuBuilder, FGuid ObjectBinding, USkeleton* Skeleton, UMovieSceneControlRigParameterSection* Section)
+void FControlRigParameterSection::LoadAnimationIntoSection(FMenuBuilder& MenuBuilder, FGuid ObjectBinding, USkeleton* Skeleton, UMovieSceneControlRigParameterSection* Section)
 {
 	TSharedPtr<ISequencer> SequencerPtr = WeakSequencer.Pin();
 	UMovieSceneSequence* Sequence = SequencerPtr.IsValid() ? SequencerPtr->GetFocusedMovieSceneSequence() : nullptr;
 
 	FAssetPickerConfig AssetPickerConfig;
 	{
-		AssetPickerConfig.OnAssetSelected = FOnAssetSelected::CreateRaw(this, &FControlRigParameterSection::OnAnimationAssetSelectedForFK, ObjectBinding, Section);
-		AssetPickerConfig.OnAssetEnterPressed = FOnAssetEnterPressed::CreateRaw(this, &FControlRigParameterSection::OnAnimationAssetEnterPressedForFK, ObjectBinding, Section);
+		AssetPickerConfig.OnAssetSelected = FOnAssetSelected::CreateRaw(this, &FControlRigParameterSection::OnAnimationAssetSelected, ObjectBinding, Section);
+		AssetPickerConfig.OnAssetEnterPressed = FOnAssetEnterPressed::CreateRaw(this, &FControlRigParameterSection::OnAnimationAssetEnterPressed, ObjectBinding, Section);
 		AssetPickerConfig.bAllowNullSelection = false;
 		AssetPickerConfig.bAddFilterUI = true;
 		AssetPickerConfig.InitialAssetViewType = EAssetViewType::List;
-		AssetPickerConfig.OnShouldFilterAsset = FOnShouldFilterAsset::CreateRaw(this, &FControlRigParameterSection::ShouldFilterAssetForFK);
+		AssetPickerConfig.OnShouldFilterAsset = FOnShouldFilterAsset::CreateRaw(this, &FControlRigParameterSection::ShouldFilterAsset);
 		AssetPickerConfig.Filter.bRecursiveClasses = true;
 		AssetPickerConfig.Filter.ClassPaths.Add(UAnimSequenceBase::StaticClass()->GetClassPathName());
 		AssetPickerConfig.OnShouldFilterAsset = FOnShouldFilterAsset::CreateUObject(Skeleton, &USkeleton::ShouldFilterAsset, TEXT("Skeleton"));
@@ -6316,7 +6324,7 @@ void FControlRigParameterSection::AddAnimationSubMenuForFK(FMenuBuilder& MenuBui
 }
 
 
-void FControlRigParameterSection::OnAnimationAssetSelectedForFK(const FAssetData& AssetData, FGuid ObjectBinding, UMovieSceneControlRigParameterSection* Section)
+void FControlRigParameterSection::OnAnimationAssetSelected(const FAssetData& AssetData, FGuid ObjectBinding, UMovieSceneControlRigParameterSection* Section)
 {
 	FSlateApplication::Get().DismissAllMenus();
 
@@ -6332,20 +6340,34 @@ void FControlRigParameterSection::OnAnimationAssetSelectedForFK(const FAssetData
 
 		if (AnimSequence && SkelMeshComp && AnimSequence->GetDataModel()->GetNumBoneTracks() > 0)
 		{
-			FBakeToControlDelegate BakeCallback = FBakeToControlDelegate::CreateLambda([this, Section,
+			//if we get a new anim sequence we change the start and end range times so it can be the same as the anim sequence
+			static uint32 LastAnimSequenceID = (uint32)(-1);
+			if (AnimSequence->GetUniqueID() != LastAnimSequenceID)
+			{
+				ULoadAnimToControlRigSettings* LoadSettings = GetMutableDefault<ULoadAnimToControlRigSettings>();
+				LoadSettings->StartFrame = 0;
+				LoadSettings->EndFrame = AnimSequence->GetDataModel()->GetNumberOfFrames();
+				LastAnimSequenceID = AnimSequence->GetUniqueID();
+			}
+			FLoadAnimToControlRigDelegate LoadCallback = FLoadAnimToControlRigDelegate::CreateLambda([this, Section,
 				AnimSequence, SkelMeshComp]
-				(bool bKeyReduce, float KeyReduceTolerance, FFrameRate BakeFrameRate, bool bResetControls)
+				(ULoadAnimToControlRigSettings* LoadSettings)
 				{
 					if (TSharedPtr<ISequencer> SequencerPtr = WeakSequencer.Pin())
 					{
-						FScopedTransaction Transaction(LOCTEXT("BakeAnimation_Transaction", "Bake Animation To FK Control Rig"));
+						FScopedTransaction Transaction(LOCTEXT("LoadAnimation", "Load Animation"));
 						Section->Modify();
 						FFrameNumber StartFrame = SequencerPtr->GetLocalTime().Time.GetFrame();
-						FSmartReduceParams SmartReduce;
-						SmartReduce.TolerancePercentage = KeyReduceTolerance;
-						SmartReduce.SampleRate = BakeFrameRate;
+						TOptional<TRange<FFrameNumber>> AnimLoadFrameRamge; //do whole range
+						if (LoadSettings->bUseCustomTimeRange)
+						{
+							TRange<FFrameNumber> Range{ 0 };
+							Range.SetLowerBoundValue(LoadSettings->StartFrame);
+							Range.SetUpperBoundValue(LoadSettings->EndFrame);
+							AnimLoadFrameRamge = Range;
+						}
 						if (!FControlRigParameterTrackEditor::LoadAnimationIntoSection(SequencerPtr, AnimSequence, SkelMeshComp, StartFrame,
-							bKeyReduce, SmartReduce, bResetControls, Section))
+							LoadSettings->bReduceKeys, LoadSettings->SmartReduce, LoadSettings->bResetControls, AnimLoadFrameRamge, Section))
 						{
 							Transaction.Cancel();
 						}
@@ -6354,13 +6376,13 @@ void FControlRigParameterSection::OnAnimationAssetSelectedForFK(const FAssetData
 					
 				});
 
-			FOnWindowClosed BakeClosedCallback = FOnWindowClosed::CreateLambda([](const TSharedRef<SWindow>&) {});
-			BakeToControlRigDialog::GetBakeParams(BakeCallback, BakeClosedCallback);
+			FOnWindowClosed LoadClosedCallback = FOnWindowClosed::CreateLambda([](const TSharedRef<SWindow>&) {});
+			FLoadAnimToControlRigDialog::GetLoadAnimParams(LoadCallback, LoadClosedCallback);
 		}
 	}
 }
 
-bool FControlRigParameterSection::ShouldFilterAssetForFK(const FAssetData& AssetData)
+bool FControlRigParameterSection::ShouldFilterAsset(const FAssetData& AssetData)
 {
 	// we don't want 
 
@@ -6380,11 +6402,11 @@ bool FControlRigParameterSection::ShouldFilterAssetForFK(const FAssetData& Asset
 
 }
 
-void FControlRigParameterSection::OnAnimationAssetEnterPressedForFK(const TArray<FAssetData>& AssetData, FGuid  ObjectBinding, UMovieSceneControlRigParameterSection* Section)
+void FControlRigParameterSection::OnAnimationAssetEnterPressed(const TArray<FAssetData>& AssetData, FGuid  ObjectBinding, UMovieSceneControlRigParameterSection* Section)
 {
 	if (AssetData.Num() > 0)
 	{
-		OnAnimationAssetSelectedForFK(AssetData[0].GetAsset(), ObjectBinding, Section);
+		OnAnimationAssetSelected(AssetData[0].GetAsset(), ObjectBinding, Section);
 	}
 }
 
diff --git a/Engine/Plugins/Animation/ControlRig/Source/ControlRigEditor/Private/Sequencer/ControlRigParameterTrackEditor.h b/Engine/Plugins/Animation/ControlRig/Source/ControlRigEditor/Private/Sequencer/ControlRigParameterTrackEditor.h
--- a/Engine/Plugins/Animation/ControlRig/Source/ControlRigEditor/Private/Sequencer/ControlRigParameterTrackEditor.h
+++ b/Engine/Plugins/Animation/ControlRig/Source/ControlRigEditor/Private/Sequencer/ControlRigParameterTrackEditor.h
@@ -301,7 +301,8 @@ public:
 
 	//** Function to load animation into a section, returns true if successful
 	static bool LoadAnimationIntoSection(const TSharedPtr<ISequencer>& InSequencer,  UAnimSequence* AnimSequence, USkeletalMeshComponent* SkelMeshComp,
-		FFrameNumber StartFrame, bool bReduceKeys, const FSmartReduceParams& ReduceParams, bool bResetControls, UMovieSceneControlRigParameterSection* ParamSection);
+		FFrameNumber StartFrame, bool bReduceKeys, const FSmartReduceParams& ReduceParams, bool bResetControls,const TOptional<TRange<FFrameNumber>> &InFrameRamge,
+		UMovieSceneControlRigParameterSection* ParamSection);
 
 	//** Function to smart reduce all keys on all constraints based on the parameters
 	static void SmartReduce(const TSharedPtr<ISequencer>& InSequencer, const FSmartReduceParams& InParams, UMovieSceneControlRigParameterSection* InSection);
@@ -407,16 +408,16 @@ public:
 	virtual bool RequestDeleteKeyArea(const TArray<FName>& KeyAreaNamePath) override;
 protected:
 	/** Add Sub Menu */
-	void AddAnimationSubMenuForFK(FMenuBuilder& MenuBuilder, FGuid ObjectBinding, USkeleton* Skeleton, UMovieSceneControlRigParameterSection* Section);
+	void LoadAnimationIntoSection(FMenuBuilder& MenuBuilder, FGuid ObjectBinding, USkeleton* Skeleton, UMovieSceneControlRigParameterSection* Section);
 
 	/** Animation sub menu filter function */
-	bool ShouldFilterAssetForFK(const FAssetData& AssetData);
+	bool ShouldFilterAsset(const FAssetData& AssetData);
 
 	/** Animation asset selected */
-	void OnAnimationAssetSelectedForFK(const FAssetData& AssetData,FGuid ObjectBinding, UMovieSceneControlRigParameterSection* Section);
+	void OnAnimationAssetSelected(const FAssetData& AssetData,FGuid ObjectBinding, UMovieSceneControlRigParameterSection* Section);
 
 	/** Animation asset enter pressed */
-	void OnAnimationAssetEnterPressedForFK(const TArray<FAssetData>& AssetData, FGuid ObjectBinding, UMovieSceneControlRigParameterSection* Section);
+	void OnAnimationAssetEnterPressed(const TArray<FAssetData>& AssetData, FGuid ObjectBinding, UMovieSceneControlRigParameterSection* Section);
 
 	/** Select controls channels from selected pressed*/
 	void ShowSelectedControlsChannels();
diff --git a/Engine/Plugins/Animation/ControlRig/Source/ControlRigEditor/Private/Sequencer/LoadAnimToControlRigSettings.cpp b/Engine/Plugins/Animation/ControlRig/Source/ControlRigEditor/Private/Sequencer/LoadAnimToControlRigSettings.cpp
new file mode 100644
--- /dev/null
+++ b/Engine/Plugins/Animation/ControlRig/Source/ControlRigEditor/Private/Sequencer/LoadAnimToControlRigSettings.cpp
@@ -0,0 +1,17 @@
+// Copyright Epic Games, Inc. All Rights Reserved.
+
+#include "LoadAnimToControlRigSettings.h"
+
+#include UE_INLINE_GENERATED_CPP_BY_NAME(LoadAnimToControlRigSettings)
+
+ULoadAnimToControlRigSettings::ULoadAnimToControlRigSettings(const FObjectInitializer& Initializer)
+	: Super(Initializer)
+{}
+
+void ULoadAnimToControlRigSettings::Reset()
+{
+	bReduceKeys = false;
+	SmartReduce.Reset();
+	bUseCustomTimeRange = false;
+}
+
diff --git a/Engine/Plugins/Animation/ControlRig/Source/ControlRigEditor/Private/Sequencer/LoadAnimToControlRigSettings.h b/Engine/Plugins/Animation/ControlRig/Source/ControlRigEditor/Private/Sequencer/LoadAnimToControlRigSettings.h
new file mode 100644
--- /dev/null
+++ b/Engine/Plugins/Animation/ControlRig/Source/ControlRigEditor/Private/Sequencer/LoadAnimToControlRigSettings.h
@@ -0,0 +1,42 @@
+// Copyright Epic Games, Inc. All Rights Reserved.
+#pragma once
+
+#include "CoreMinimal.h"
+#include "UObject/ObjectMacros.h"
+#include "UObject/Object.h"
+#include "Filters/CurveEditorSmartReduceFilter.h"
+#include "LoadAnimToControlRigSettings.generated.h"
+
+//settings used when loading animation sequences into a control rig section
+UCLASS(BlueprintType, config = EditorSettings)
+class  ULoadAnimToControlRigSettings : public UObject
+{
+public:
+	ULoadAnimToControlRigSettings(const FObjectInitializer& Initializer);
+
+	GENERATED_BODY()
+	/** Reduce Keys */
+	UPROPERTY(EditAnywhere, Category = "Reduce Keys")
+	bool bReduceKeys = false;
+
+	UPROPERTY(EditAnywhere, Category = "Reduce Keys", meta = (EditCondition = "bReduceKeys"))
+	FSmartReduceParams SmartReduce;
+
+	UPROPERTY(EditAnywhere, Category = "Animation Time Range")
+	bool bUseCustomTimeRange = false;
+
+	UPROPERTY(EditAnywhere, Category = "Animation Time Range", meta = (EditCondition = "bUseCustomTimeRange"))
+	FFrameNumber StartFrame;
+
+	UPROPERTY(EditAnywhere, Category = "Animation Time Range", meta = (EditCondition = "bUseCustomTimeRange"))
+	FFrameNumber EndFrame;
+
+	/** Reset controls to initial value on every frame */
+	UPROPERTY(EditAnywhere, Category = "Reset Controls")
+	bool bResetControls = true;
+
+	/** Resets the default properties. */
+	void Reset();
+};
+
+
diff --git a/Engine/Plugins/Animation/ControlRig/Source/ControlRigEditor/Private/Sequencer/SLoadAnimToControlRig.cpp b/Engine/Plugins/Animation/ControlRig/Source/ControlRigEditor/Private/Sequencer/SLoadAnimToControlRig.cpp
new file mode 100644
--- /dev/null
+++ b/Engine/Plugins/Animation/ControlRig/Source/ControlRigEditor/Private/Sequencer/SLoadAnimToControlRig.cpp
@@ -0,0 +1,112 @@
+// Copyright Epic Games, Inc. All Rights Reserved.
+
+#include "SLoadAnimToControlRig.h"
+#include "CoreMinimal.h"
+#include "Widgets/SWidget.h"
+#include "Widgets/SCompoundWidget.h"
+#include "LoadAnimToControlRigSettings.h"
+#include "Editor.h"
+#include "Modules/ModuleManager.h"
+#include "Framework/Application/SlateApplication.h"
+#include "PropertyEditorModule.h"
+#include "Widgets/SWindow.h"
+#include "Widgets/Input/SButton.h"
+#include "Misc/FrameRate.h"
+
+class SLoadAnimToControlRigDialog : public SCompoundWidget
+{
+	SLATE_BEGIN_ARGS(SLoadAnimToControlRigDialog) {}
+	SLATE_END_ARGS()
+	~SLoadAnimToControlRigDialog()
+	{
+	}
+	
+	void Construct(const FArguments& InArgs)
+	{
+		FPropertyEditorModule& PropertyEditor = FModuleManager::LoadModuleChecked<FPropertyEditorModule>("PropertyEditor");
+
+		FDetailsViewArgs DetailsViewArgs;
+		DetailsViewArgs.bShowOptions = false;
+		DetailsViewArgs.bAllowSearch = false;
+		DetailsViewArgs.bShowPropertyMatrixButton = false;
+		DetailsViewArgs.bUpdatesFromSelection = false;
+		DetailsViewArgs.bLockable = false;
+		DetailsViewArgs.NameAreaSettings = FDetailsViewArgs::HideNameArea;
+		DetailsViewArgs.ViewIdentifier = "Load Animation To Control Rig";
+
+		DetailView = PropertyEditor.CreateDetailView(DetailsViewArgs);
+
+		ChildSlot
+			[
+				SNew(SVerticalBox)
+
+				+ SVerticalBox::Slot()
+			[
+				DetailView.ToSharedRef()
+			]
+
+		+ SVerticalBox::Slot()
+			.AutoHeight()
+			.HAlign(HAlign_Right)
+			.Padding(5.f)
+			[
+				SNew(SButton)
+				.ContentPadding(FMargin(10, 5))
+			.Text(NSLOCTEXT("ControlRig", "LoadAnimationToControlRig", "Load Animation"))
+			.OnClicked(this, &SLoadAnimToControlRigDialog::OnLoadAnim)
+			]
+
+		];
+
+		ULoadAnimToControlRigSettings* LoadSettings = GetMutableDefault<ULoadAnimToControlRigSettings>();
+		DetailView->SetObject(LoadSettings);
+	}
+
+	void SetDelegate(FLoadAnimToControlRigDelegate& InDelegate)
+	{
+		Delegate = InDelegate;
+	}
+
+private:
+
+	FReply OnLoadAnim()
+	{
+		ULoadAnimToControlRigSettings* LoadSettings = GetMutableDefault<ULoadAnimToControlRigSettings>();
+		TSharedPtr<SWindow> Window = FSlateApplication::Get().FindWidgetWindow(AsShared());
+		if (LoadSettings && Delegate.IsBound())
+		{
+			Delegate.Execute(LoadSettings);
+		}
+		if (Window.IsValid())
+		{
+			Window->RequestDestroyWindow();
+		}
+		return FReply::Handled();
+	}
+	TSharedPtr<IDetailsView> DetailView;
+	FLoadAnimToControlRigDelegate  Delegate;
+};
+
+
+void FLoadAnimToControlRigDialog::GetLoadAnimParams(FLoadAnimToControlRigDelegate& InDelegate, const FOnWindowClosed& OnClosedDelegate)
+{
+	const FText TitleText = NSLOCTEXT("ControlRig", "LoadAnimation", "Load Animation");
+
+	// Create the window to choose our options
+	TSharedRef<SWindow> Window = SNew(SWindow)
+		.Title(TitleText)
+		.HasCloseButton(true)
+		.SizingRule(ESizingRule::UserSized)
+		.ClientSize(FVector2D(400.0f, 200.0f))
+		.AutoCenter(EAutoCenter::PreferredWorkArea)
+		.SupportsMinimize(false);
+
+	TSharedRef<SLoadAnimToControlRigDialog> DialogWidget = SNew(SLoadAnimToControlRigDialog);
+	DialogWidget->SetDelegate(InDelegate);
+	Window->SetContent(DialogWidget);
+	Window->SetOnWindowClosed(OnClosedDelegate);
+
+	FSlateApplication::Get().AddWindow(Window);
+
+}
+
diff --git a/Engine/Plugins/Animation/ControlRig/Source/ControlRigEditor/Private/Sequencer/SLoadAnimToControlRig.h b/Engine/Plugins/Animation/ControlRig/Source/ControlRigEditor/Private/Sequencer/SLoadAnimToControlRig.h
new file mode 100644
--- /dev/null
+++ b/Engine/Plugins/Animation/ControlRig/Source/ControlRigEditor/Private/Sequencer/SLoadAnimToControlRig.h
@@ -0,0 +1,16 @@
+// Copyright Epic Games, Inc. All Rights Reserved.
+#pragma once
+
+#include "CoreMinimal.h"
+#include "Delegates/Delegate.h"
+#include "Widgets/SWindow.h"
+
+class ULoadAnimToControlRigSettings;
+
+DECLARE_DELEGATE_OneParam(FLoadAnimToControlRigDelegate, ULoadAnimToControlRigSettings*);
+
+//dialog to show ULoadAnimToControlRigSettings properties for loading animation into a control rig section
+struct  FLoadAnimToControlRigDialog
+{
+	static void GetLoadAnimParams(FLoadAnimToControlRigDelegate& Delegate, const FOnWindowClosed& OnClosedDelegate);
+};
\ No newline at end of file
diff --git a/Engine/Plugins/Animation/ControlRig/Source/ControlRigEditor/Public/ControlRigSequencerEditorLibrary.h b/Engine/Plugins/Animation/ControlRig/Source/ControlRigEditor/Public/ControlRigSequencerEditorLibrary.h
--- a/Engine/Plugins/Animation/ControlRig/Source/ControlRigEditor/Public/ControlRigSequencerEditorLibrary.h
+++ b/Engine/Plugins/Animation/ControlRig/Source/ControlRigEditor/Public/ControlRigSequencerEditorLibrary.h
@@ -134,6 +134,9 @@ public:
 	* @param MovieScene The MovieScene getting loaded into
 	* @param SkelMeshComponent The Skeletal Mesh component getting loaded into.
 	* @param InStartFrame Frame to insert the animation
+	* @param bUseCustomAnimRange If true we will load in a specific animation range from the anim sequence
+	* @param AnimStartRange  The start of the animation in the anim sequence frame rate that we load the animation from, used if bUseCustomAnimRange is true
+	* @param AnimEndRange The end of the animation in the anim sequence frame rate we load the animation from, used if bUseCustomAnimRange is true
 	* @param TimeUnit Unit for all frame and time values, either in display rate or tick resolution
 	* @param bKeyReduce If true do key reduction based upon Tolerance, if false don't
 	* @param Tolerance If reducing keys, tolerance about which keys will be removed, smaller tolerance, more keys usually.
@@ -142,6 +145,13 @@ public:
 	* @return returns True if successful, False otherwise
 	*/
 	UFUNCTION(BlueprintCallable, Category = "Editor Scripting | Sequencer Tools | Control Rig")
+	static bool LoadAnimSequenceIntoControlRigSectionWithRange(UMovieSceneSection* MovieSceneSection, UAnimSequence* AnimSequence, USkeletalMeshComponent* SkelMeshComp,
+		FFrameNumber InStartFrame, bool bUseCustomAnimRange, FFrameNumber AnimStartRange, FFrameNumber AnimEndRange,
+		EMovieSceneTimeUnit TimeUnit = EMovieSceneTimeUnit::DisplayRate, bool bKeyReduce = false, float Tolerance = 0.001f,
+		EMovieSceneKeyInterpolation Interpolation = EMovieSceneKeyInterpolation::SmartAuto, bool bResetControls = true);
+
+	UE_DEPRECATED(5.6, "Use LoadAnimSequenceIntoControlRigSectionWithRange that takes UseCustomTimeRange")
+	UFUNCTION(BlueprintCallable, Category = "Editor Scripting | Sequencer Tools | Control Rig", meta = (DeprecatedFunction, DeprecationMessage = "Use LoadAnimSequenceIntoControlRigSectionWithRange that takes UseCustomTimeRange"))
 	static bool LoadAnimSequenceIntoControlRigSection(UMovieSceneSection* MovieSceneSection, UAnimSequence* AnimSequence, USkeletalMeshComponent* SkelMeshComp,
 		FFrameNumber InStartFrame, EMovieSceneTimeUnit TimeUnit = EMovieSceneTimeUnit::DisplayRate, bool bKeyReduce = false, float Tolerance = 0.001f,
 		EMovieSceneKeyInterpolation Interpolation = EMovieSceneKeyInterpolation::SmartAuto, bool bResetControls = true);
diff --git a/Engine/Plugins/Importers/USDImporter/Source/USDStage/Private/USDLevelSequenceHelper.cpp b/Engine/Plugins/Importers/USDImporter/Source/USDStage/Private/USDLevelSequenceHelper.cpp
--- a/Engine/Plugins/Importers/USDImporter/Source/USDStage/Private/USDLevelSequenceHelper.cpp
+++ b/Engine/Plugins/Importers/USDImporter/Source/USDStage/Private/USDLevelSequenceHelper.cpp
@@ -447,17 +447,18 @@ namespace UE::USDLevelSequenceHelper::Private
 			Track->SetTrackName(FName(*ObjectName));
 			Track->SetDisplayName(FText::FromString(ObjectName));
 
-			const bool bResetControls = true;
 			const FFrameNumber SequenceStart{0};
+			UMovieSceneControlRigParameterSection::FLoadAnimSequenceData Data;
+			Data.bKeyReduce = bReduceKeys;
+			Data.Tolerance = Tolerance;
+			Data.bResetControls = true;
+			Data.StartFrame = ControlRigSectionStartFrame;
 			ParamSection->LoadAnimSequenceIntoThisSection(
 				AnimSequence,
 				SequenceStart,
 				MovieScene,
 				SkeletalMeshComp,
-				bReduceKeys,
-				Tolerance,
-				bResetControls,
-				ControlRigSectionStartFrame,
+				Data,
 				EMovieSceneKeyInterpolation::SmartAuto
 			);
 
