Fixed Control Rig Hiearchy random initialization issues by converting static names into lazy names

diff --git a/Engine/Plugins/Animation/ControlRig/Source/ControlRig/Public/Rigs/RigHierarchy.h b/Engine/Plugins/Animation/ControlRig/Source/ControlRig/Public/Rigs/RigHierarchy.h
--- a/Engine/Plugins/Animation/ControlRig/Source/ControlRig/Public/Rigs/RigHierarchy.h
+++ b/Engine/Plugins/Animation/ControlRig/Source/ControlRig/Public/Rigs/RigHierarchy.h
@@ -168,12 +168,12 @@ public:
 	typedef TMap<int32, TArray<int32>> TElementDependencyMap;
 	typedef TPair<int32, TArray<int32>> TElementDependencyMapPair;
 	typedef TTuple<int32, int32, int32, ERigTransformType::Type> TInstructionSliceElement;
-	inline static const FName TagMetadataName = TEXT("Tags");
-	inline static const FName ShortModuleNameMetadataName = TEXT("ShortModuleName");
-	inline static const FName DesiredNameMetadataName = TEXT("DesiredName");
-	inline static const FName DesiredKeyMetadataName = TEXT("DesiredKey");
-	inline static const FName ModuleMetadataName = TEXT("Module");
-	inline static const FName NameSpaceMetadataName = TEXT("NameSpace");
+	inline static const FLazyName TagMetadataName = TEXT("Tags");
+	inline static const FLazyName ShortModuleNameMetadataName = TEXT("ShortModuleName");
+	inline static const FLazyName DesiredNameMetadataName = TEXT("DesiredName");
+	inline static const FLazyName DesiredKeyMetadataName = TEXT("DesiredKey");
+	inline static const FLazyName ModuleMetadataName = TEXT("Module");
+	inline static const FLazyName NameSpaceMetadataName = TEXT("NameSpace");
 
 	URigHierarchy();
 
