Control Rig Pose Adapter: ensure that the parent's bone index is safe to use

diff --git a/Engine/Plugins/Animation/ControlRig/Source/ControlRig/Private/Tools/ControlRigPoseAdapter.cpp b/Engine/Plugins/Animation/ControlRig/Source/ControlRig/Private/Tools/ControlRigPoseAdapter.cpp
--- a/Engine/Plugins/Animation/ControlRig/Source/ControlRig/Private/Tools/ControlRigPoseAdapter.cpp
+++ b/Engine/Plugins/Animation/ControlRig/Source/ControlRig/Private/Tools/ControlRigPoseAdapter.cpp
@@ -277,7 +277,10 @@ void FControlRigPoseAdapter::UpdateInputOutputMappingIfRequired(UControlRig* InC
 				if (CompactPoseParentIndex.IsValid())
 				{
 					const FSkeletonPoseBoneIndex SkeletonIndex = InRequiredBones.GetSkeletonPoseIndexFromCompactPoseIndex(CompactPoseParentIndex);
-					PoseParentName = RefSkeleton->GetBoneName(SkeletonIndex.GetInt());
+					if (SkeletonIndex.IsValid() && RefSkeleton->IsValidIndex(SkeletonIndex.GetInt()))
+					{
+						PoseParentName = RefSkeleton->GetBoneName(SkeletonIndex.GetInt());
+					}
 				}
 
 				if (HierarchyParentName.IsEqual(PoseParentName, ENameCase::CaseSensitive))
