Anim In Engine: Motion trails 2.0.
Rewrote exisiting motion trails to peform only N number of evaluations per frame to keep interactivity. Created new evaluation classes to handle doing these multiple frame evaluations for specific set of actors/control rigs to calculate their world transforms.  Added a dependency graph of tracks to keep track of what objects need to get re-evaluted when a movie scene change happens.  Update visual with new styles/color options. Added new interactive abilities like setting an offset transform or putting trail in a new space. Removed old popup and created new menu item in new toolbar.

diff --git a/Engine/Plugins/Animation/ControlRig/Source/ControlRigEditor/Private/EditMode/ControlRigEditMode.cpp b/Engine/Plugins/Animation/ControlRig/Source/ControlRigEditor/Private/EditMode/ControlRigEditMode.cpp
--- a/Engine/Plugins/Animation/ControlRig/Source/ControlRigEditor/Private/EditMode/ControlRigEditMode.cpp
+++ b/Engine/Plugins/Animation/ControlRig/Source/ControlRigEditor/Private/EditMode/ControlRigEditMode.cpp
@@ -496,8 +496,7 @@ void FControlRigEditMode::Enter()
 
 	UE::ControlRig::PopulateControlRigViewportToolbarTransformSubmenu("LevelEditor.ViewportToolbar.Transform");
 	UE::ControlRig::PopulateControlRigViewportToolbarShowSubmenu(
-		"LevelEditor.ViewportToolbar.Show", GetToolkit()->GetToolkitCommands()
-	);
+		"LevelEditor.ViewportToolbar.Show");
 
 	RegisterPendingFocusMode();
 
diff --git a/Engine/Plugins/Animation/ControlRig/Source/ControlRigEditor/Private/EditMode/ControlRigEditModeToolkit.cpp b/Engine/Plugins/Animation/ControlRig/Source/ControlRigEditor/Private/EditMode/ControlRigEditModeToolkit.cpp
--- a/Engine/Plugins/Animation/ControlRig/Source/ControlRigEditor/Private/EditMode/ControlRigEditModeToolkit.cpp
+++ b/Engine/Plugins/Animation/ControlRig/Source/ControlRigEditor/Private/EditMode/ControlRigEditModeToolkit.cpp
@@ -26,6 +26,7 @@
 #include "Sequencer/AnimLayers/SAnimLayers.h"
 #include "Sequencer/AnimLayers/AnimLayers.h"
 #include "Tween/SControlRigTweenWidget.h"
+#include "Tools/MotionTrailOptions.h"
 
 #define LOCTEXT_NAMESPACE "FControlRigEditModeToolkit"
 
@@ -35,7 +36,7 @@ namespace
 	const TArray<FName> AnimationPaletteNames = { AnimationName };
 }
 
-bool FControlRigEditModeToolkit::bMotionTrailsTabOpen = false;
+bool FControlRigEditModeToolkit::bMotionTrailsOn = false;
 bool FControlRigEditModeToolkit::bAnimLayerTabOpen = false;
 bool FControlRigEditModeToolkit::bPoseTabOpen = false;
 bool FControlRigEditModeToolkit::bSnapperTabOpen = false;
@@ -417,10 +418,6 @@ void FControlRigEditModeToolkit::InvokeUI()
 		{
 			CreateAndShowTweenOverlay();
 		}
-		if (bMotionTrailsTabOpen)
-		{
-			TryInvokeToolkitUI(MotionTrailTabName);
-		}
 		if (bAnimLayerTabOpen)
 		{
 			TryInvokeToolkitUI(AnimLayerTabName);
@@ -441,6 +438,16 @@ void FControlRigEditModeToolkit::InvokeUI()
 		{
 			TryInvokeToolkitUI(PoseTabName);
 		}
+
+		if (bMotionTrailsOn)
+		{
+			if (UMotionTrailToolOptions* Settings = GetMutableDefault<UMotionTrailToolOptions>())
+			{
+				Settings->bShowTrails = true;
+				FPropertyChangedEvent ShowTrailEvent(UMotionTrailToolOptions::StaticClass()->FindPropertyByName(GET_MEMBER_NAME_CHECKED(UMotionTrailToolOptions, bShowTrails)));
+				Settings->PostEditChangeProperty(ShowTrailEvent);
+			}
+		}
 	}	
 }
 
@@ -468,16 +475,18 @@ void FControlRigEditModeToolkit::UnregisterAndRemoveFloatingTabs()
 		{
 			TSharedPtr<FAssetEditorModeUILayer> ModeUILayerPtr = ModeUILayer.Pin();
 
-			TSharedPtr<SDockTab> MotionTrailTab = ModeUILayerPtr->GetTabManager()->FindExistingLiveTab(FTabId(MotionTrailTabName));
-			if (MotionTrailTab)
+			bMotionTrailsOn = false;
+			if (UMotionTrailToolOptions* Settings = GetMutableDefault<UMotionTrailToolOptions>())
 			{
-				bMotionTrailsTabOpen = true;
-				MotionTrailTab->RequestCloseTab();
-			}
-			else
-			{
-				bMotionTrailsTabOpen = false;
+				if (Settings->bShowTrails)
+				{
+					bMotionTrailsOn = true;
+					Settings->bShowTrails = false;
+					FPropertyChangedEvent ShowTrailEvent(UMotionTrailToolOptions::StaticClass()->FindPropertyByName(GET_MEMBER_NAME_CHECKED(UMotionTrailToolOptions, bShowTrails)));
+					Settings->PostEditChangeProperty(ShowTrailEvent);
+				}
 			}
+
 			ModeUILayerPtr->GetTabManager()->UnregisterTabSpawner(MotionTrailTabName);
 
 			TSharedPtr<SDockTab> AnimLayerTab = ModeUILayerPtr->GetTabManager()->FindExistingLiveTab(FTabId(AnimLayerTabName));
diff --git a/Engine/Plugins/Animation/ControlRig/Source/ControlRigEditor/Private/EditMode/ControlRigEditModeToolkit.h b/Engine/Plugins/Animation/ControlRig/Source/ControlRigEditor/Private/EditMode/ControlRigEditModeToolkit.h
--- a/Engine/Plugins/Animation/ControlRig/Source/ControlRigEditor/Private/EditMode/ControlRigEditModeToolkit.h
+++ b/Engine/Plugins/Animation/ControlRig/Source/ControlRigEditor/Private/EditMode/ControlRigEditModeToolkit.h
@@ -90,7 +90,7 @@ protected:
 	//this also saves the layout
 	void UnregisterAndRemoveFloatingTabs();
 
-	static bool bMotionTrailsTabOpen;
+	static bool bMotionTrailsOn;
 	static bool bAnimLayerTabOpen;
 	static bool bPoseTabOpen;
 	static bool bSnapperTabOpen;
diff --git a/Engine/Plugins/Animation/ControlRig/Source/ControlRigEditor/Private/EditMode/SControlRigEditModeTools.cpp b/Engine/Plugins/Animation/ControlRig/Source/ControlRigEditor/Private/EditMode/SControlRigEditModeTools.cpp
--- a/Engine/Plugins/Animation/ControlRig/Source/ControlRigEditor/Private/EditMode/SControlRigEditModeTools.cpp
+++ b/Engine/Plugins/Animation/ControlRig/Source/ControlRigEditor/Private/EditMode/SControlRigEditModeTools.cpp
@@ -1306,6 +1306,7 @@ void SControlRigEditModeTools::CustomizeToolBarPalette(FToolBarBuilder& ToolBarB
 		EUserInterfaceActionType::Button
 	);
 
+	/*  leaving this around for a bit to get animator feedback on the new menu system
 	// Motion Trail
 	ToolBarBuilder.AddToolBarButton(
 		FExecuteAction::CreateRaw(OwningToolkit.Pin().Get(), &FControlRigEditModeToolkit::TryInvokeToolkitUI, FControlRigEditModeToolkit::MotionTrailTabName),
@@ -1315,6 +1316,7 @@ void SControlRigEditModeTools::CustomizeToolBarPalette(FToolBarBuilder& ToolBarB
 		FSlateIcon(TEXT("ControlRigEditorStyle"), TEXT("ControlRig.EditableMotionTrails")),
 		EUserInterfaceActionType::Button
 	);
+	*/
 
 	// Anim Layer
 	ToolBarBuilder.AddToolBarButton(
@@ -1327,6 +1329,7 @@ void SControlRigEditModeTools::CustomizeToolBarPalette(FToolBarBuilder& ToolBarB
 	);
 
 	//Pivot
+	/** like motion trail leaving it in around for a bi
 	ToolBarBuilder.AddToolBarButton(
 		FUIAction(
 		FExecuteAction::CreateSP(this, &SControlRigEditModeTools::ToggleEditPivotMode),
@@ -1353,6 +1356,7 @@ void SControlRigEditModeTools::CustomizeToolBarPalette(FToolBarBuilder& ToolBarB
 		FSlateIcon(TEXT("ControlRigEditorStyle"), TEXT("ControlRig.TemporaryPivot")),
 		EUserInterfaceActionType::ToggleButton
 		);
+		*/
 }
 
 void SControlRigEditModeTools::ToggleEditPivotMode()
diff --git a/Engine/Plugins/Animation/ControlRig/Source/ControlRigEditor/Private/EditMode/SControlRigSnapper.cpp b/Engine/Plugins/Animation/ControlRig/Source/ControlRigEditor/Private/EditMode/SControlRigSnapper.cpp
--- a/Engine/Plugins/Animation/ControlRig/Source/ControlRigEditor/Private/EditMode/SControlRigSnapper.cpp
+++ b/Engine/Plugins/Animation/ControlRig/Source/ControlRigEditor/Private/EditMode/SControlRigSnapper.cpp
@@ -1,6 +1,7 @@
 // Copyright Epic Games, Inc. All Rights Reserved.
 
 #include "EditMode/SControlRigSnapper.h"
+#include "EditMode/SComponentPickerPopup.h"
 #include "Framework/Notifications/NotificationManager.h"
 #include "Widgets/Notifications/SNotificationList.h"
 #include "Framework/MultiBox/MultiBoxBuilder.h"
@@ -33,100 +34,6 @@
 
 #define LOCTEXT_NAMESPACE "ControlRigSnapper"
 
-class SComponentPickerPopup : public SCompoundWidget
-{
-public:
-	DECLARE_DELEGATE_OneParam(FOnComponentChosen, FName);
-
-	SLATE_BEGIN_ARGS(SComponentPickerPopup)
-		: _Actor(NULL)
-	{}
-
-	/** An actor with components */
-	SLATE_ARGUMENT(AActor*, Actor)
-
-		/** Called when the text is chosen. */
-		SLATE_EVENT(FOnComponentChosen, OnComponentChosen)
-
-		SLATE_END_ARGS()
-
-		/** Delegate to call when component is selected */
-		FOnComponentChosen OnComponentChosen;
-
-	/** List of tag names selected in the tag containers*/
-	TArray< TSharedPtr<FName> > ComponentNames;
-
-private:
-	TSharedRef<ITableRow> MakeListViewWidget(TSharedPtr<FName> InItem, const TSharedRef<STableViewBase>& OwnerTable)
-	{
-		return SNew(STableRow< TSharedPtr<FName> >, OwnerTable)
-			[
-				SNew(STextBlock).Text(FText::FromName(*InItem.Get()))
-			];
-	}
-
-	void OnComponentSelected(TSharedPtr<FName> InItem, ESelectInfo::Type InSelectInfo)
-	{
-		FSlateApplication::Get().DismissAllMenus();
-
-		if (OnComponentChosen.IsBound())
-		{
-			OnComponentChosen.Execute(*InItem.Get());
-		}
-	}
-
-public:
-	void Construct(const FArguments& InArgs)
-	{
-		OnComponentChosen = InArgs._OnComponentChosen;
-		AActor* Actor = InArgs._Actor;
-
-		TInlineComponentArray<USceneComponent*> Components(Actor);
-
-		ComponentNames.Empty();
-		for (USceneComponent* Component : Components)
-		{
-			if (Component->HasAnySockets())
-			{
-				ComponentNames.Add(MakeShareable(new FName(Component->GetFName())));
-			}
-		}
-
-		// Then make widget
-		this->ChildSlot
-			[
-				SNew(SBorder)
-				.BorderImage(FAppStyle::GetBrush(TEXT("Menu.Background")))
-			.Padding(5)
-			.Content()
-			[
-				SNew(SVerticalBox)
-				+ SVerticalBox::Slot()
-			.AutoHeight()
-			.Padding(0.0f, 1.0f)
-			[
-				SNew(STextBlock)
-				.Font(FAppStyle::GetFontStyle(TEXT("SocketChooser.TitleFont")))
-			.Text(NSLOCTEXT("ComponentChooser", "ChooseComponentLabel", "Choose Component"))
-			]
-		+ SVerticalBox::Slot()
-			.AutoHeight()
-			.MaxHeight(512)
-			[
-				SNew(SBox)
-				.WidthOverride(256)
-			.Content()
-			[
-				SNew(SListView< TSharedPtr<FName> >)
-				.ListItemsSource(&ComponentNames)
-			.OnGenerateRow(this, &SComponentPickerPopup::MakeListViewWidget)
-			.OnSelectionChanged(this, &SComponentPickerPopup::OnComponentSelected)
-			]
-			]
-			]
-			];
-	}
-};
 
 //function to keep framenumber from beign too large do to user error, it can cause a crash if we snap too many frames.
 static void KeepFrameInRange(TSharedPtr<ISequencer>& SequencerPtr, FFrameNumber& KeepFrameInRange, UMovieScene* MovieScene)
diff --git a/Engine/Plugins/Animation/ControlRig/Source/ControlRigEditor/Private/EditMode/SControlRigSnapper.h b/Engine/Plugins/Animation/ControlRig/Source/ControlRigEditor/Private/EditMode/SControlRigSnapper.h
--- a/Engine/Plugins/Animation/ControlRig/Source/ControlRigEditor/Private/EditMode/SControlRigSnapper.h
+++ b/Engine/Plugins/Animation/ControlRig/Source/ControlRigEditor/Private/EditMode/SControlRigSnapper.h
@@ -8,15 +8,16 @@
 #include "Widgets/SWidget.h"
 #include "Widgets/DeclarativeSyntaxSupport.h"
 #include "Widgets/SCompoundWidget.h"
+#include "Widgets/Views/STableRow.h"
+#include "Widgets/Text/STextBlock.h"
+#include "Widgets/Views/SListView.h"
 #include "Tools/ControlRigSnapper.h"
 #include "Misc/FrameNumber.h"
 #include "IDetailsView.h"
 #include "MovieSceneSequenceID.h"
 
-class UControlRig;
-class ISequencer;
-class AActor;
 
+class UControlRig;
 
 class SControlRigSnapper : public SCompoundWidget
 {
diff --git a/Engine/Plugins/Animation/ControlRig/Source/ControlRigEditor/Private/Editor/ControlRigViewportToolbarExtensions.cpp b/Engine/Plugins/Animation/ControlRig/Source/ControlRigEditor/Private/Editor/ControlRigViewportToolbarExtensions.cpp
--- a/Engine/Plugins/Animation/ControlRig/Source/ControlRigEditor/Private/Editor/ControlRigViewportToolbarExtensions.cpp
+++ b/Engine/Plugins/Animation/ControlRig/Source/ControlRigEditor/Private/Editor/ControlRigViewportToolbarExtensions.cpp
@@ -8,9 +8,23 @@
 #include "EditMode/ControlRigEditModeSettings.h"
 #include "ToolMenus.h"
 #include "Tools/MotionTrailOptions.h"
+#include "Tools/SMotionTrailOptions.h"
 #include "Widgets/Input/SSpinBox.h"
-
+#include "Widgets/Colors/SColorBlock.h"
+#include "Widgets/Colors/SColorPicker.h"
+#include "MovieSceneCommonHelpers.h"
+#include "Widgets/Input/SComboBox.h"
+#include "LevelEditor.h"
+#include "SSocketChooser.h"
 #include "LevelEditorActions.h"
+#include "ActorPickerMode.h"
+#include "EditMode/SComponentPickerPopup.h"
+#include "Components/SkeletalMeshComponent.h"
+#include "ILevelEditor.h"
+#include "LevelEditor.h"
+#include "InteractiveToolManager.h"
+#include "EditorModeManager.h"
+#include "EdModeInteractiveToolsContext.h"
 
 #define LOCTEXT_NAMESPACE "ControlRigViewportToolbar"
 
@@ -21,6 +35,815 @@ const FName ControlRigOwnerName = "ControlRigViewportToolbar";
 
 } // namespace UE::ControlRig::Private
 
+
+void OffsetActionExecuteAction(const FToolMenuContext& InContext, UMotionTrailToolOptions* Settings, int32 Index)
+{
+	if (UMotionTrailToolOptions::FPinnedTrail* Trail = Settings->GetPinnedTrail(Index))
+	{
+		Settings->SetHasOffset(Index, !Trail->bHasOffset);
+	}
+}
+
+ECheckBoxState OffsetActionGetActionCheckState(const FToolMenuContext& Context, UMotionTrailToolOptions* Settings, int32 Index)
+{
+	if (UMotionTrailToolOptions::FPinnedTrail* Trail = Settings->GetPinnedTrail(Index))
+	{
+		return (Trail->bHasOffset) ? ECheckBoxState::Checked : ECheckBoxState::Unchecked;
+	}
+	return ECheckBoxState::Undetermined;
+}
+
+void SpaceActionExecuteAction(const FToolMenuContext& InContext, UMotionTrailToolOptions* Settings, int32 Index)
+{
+	if (UMotionTrailToolOptions::FPinnedTrail* Trail = Settings->GetPinnedTrail(Index))
+	{
+		if (Trail->SpaceName.IsSet() == false)
+		{
+			// FIXME temp approach for selecting the parent
+			FSlateApplication::Get().DismissAllMenus();
+
+			static const FActorPickerModeModule& ActorPickerMode = FModuleManager::Get().GetModuleChecked<FActorPickerModeModule>("ActorPickerMode");
+
+			ActorPickerMode.BeginActorPickingMode(
+				FOnGetAllowedClasses(),
+				FOnShouldFilterActor::CreateLambda([](const AActor* InActor)
+					{
+						return true; //todo make sure in sequencer
+					}),
+				FOnActorSelected::CreateLambda([Settings, Index](AActor* InActor)
+					{
+						FLevelEditorModule& LevelEditorModule = FModuleManager::GetModuleChecked<FLevelEditorModule>("LevelEditor");
+						TSharedPtr< ILevelEditor > LevelEditor = LevelEditorModule.GetFirstLevelEditor();
+
+						TSharedPtr<SWidget> MenuWidget = SNew(SComponentPickerPopup)
+							.Actor(InActor)
+							.bCheckForSockets(false)
+							.OnComponentChosen_Lambda([Settings, Index, InActor](FName InComponentName) mutable
+								{
+									Settings->PutPinnnedInSpace(Index, InActor, InComponentName);
+								}
+							);
+						// Create as context menu
+						FSlateApplication::Get().PushMenu(
+							LevelEditor.ToSharedRef(),
+							FWidgetPath(),
+							MenuWidget.ToSharedRef(),
+							FSlateApplication::Get().GetCursorPos(),
+							FPopupTransitionEffect(FPopupTransitionEffect::ContextMenu)
+						);
+
+					})
+
+			);
+		}
+		else
+		{
+			Settings->PutPinnnedInSpace(Index, nullptr, NAME_None);
+		}
+	}
+}
+
+ECheckBoxState SpaceActionGetActionCheckState(const FToolMenuContext& Context, UMotionTrailToolOptions* Settings, int32 Index)
+{
+	if (UMotionTrailToolOptions::FPinnedTrail* Trail = Settings->GetPinnedTrail(Index))
+	{
+		return (Trail->SpaceName.IsSet()) ? ECheckBoxState::Checked : ECheckBoxState::Unchecked;
+	}
+	return ECheckBoxState::Undetermined;
+}
+
+void CreatePinnedMenuDelegate(UToolMenu* SubMenu, UMotionTrailToolOptions* Settings, int32 Index)
+{
+	FToolMenuSection& Section = SubMenu->AddSection(NAME_None);
+
+	FToolUIAction OffsetAction;
+	OffsetAction.ExecuteAction = FToolMenuExecuteAction::CreateStatic(&OffsetActionExecuteAction, Settings, Index);
+	OffsetAction.GetActionCheckState = FToolMenuGetActionCheckState::CreateStatic(&OffsetActionGetActionCheckState, Settings, Index);
+
+	FToolMenuEntry OffsetEntry = FToolMenuEntry::InitMenuEntry(
+		"Offset",
+		LOCTEXT("OffsetLabel", "Offset"),
+		LOCTEXT("OffsetLabelTooltip", "Toggle offset on selects the curve in the viewport, and allows you to move it like shift select does. Toggling it off will remove any offset."),
+		FSlateIcon(),
+		OffsetAction,
+		EUserInterfaceActionType::ToggleButton
+	);
+	Section.AddEntry(OffsetEntry);
+
+	FToolUIAction SpaceAction;
+	SpaceAction.ExecuteAction = FToolMenuExecuteAction::CreateStatic(&SpaceActionExecuteAction, Settings, Index);
+	SpaceAction.GetActionCheckState = FToolMenuGetActionCheckState::CreateStatic(&SpaceActionGetActionCheckState, Settings, Index);
+
+	FToolMenuEntry SpaceEntry = FToolMenuEntry::InitMenuEntry(
+		"Space",
+		LOCTEXT("SpaceLabel", "Space"),
+		LOCTEXT("SpaceLabelTooltip", "Toggling on space will put you into eye drop selection mode to pick the scene compponent/socket that you want to have this trail in. Toggling it off puts it back in world space."),
+		FSlateIcon(),
+		SpaceAction,
+		EUserInterfaceActionType::ToggleButton
+	);
+	Section.AddEntry(SpaceEntry);
+}
+
+void CreatePinnedItems(UMotionTrailToolOptions* Settings, FToolMenuSection& PinnedTrails)
+{
+	const int32  NumPinned = Settings->GetNumPinned();
+	if (NumPinned > 0)
+	{
+		for (int32 Index = 0; Index < NumPinned; ++Index)
+		{
+			if (UMotionTrailToolOptions::FPinnedTrail* Trail = Settings->GetPinnedTrail(Index))
+			{
+				FNewToolMenuDelegate MakeMenuDelegate = FNewToolMenuDelegate::CreateStatic(&CreatePinnedMenuDelegate, Settings, Index);
+				FUIAction TogglePinnedAction(
+					FExecuteAction::CreateLambda([Settings, Index]()
+						{
+							if (UMotionTrailToolOptions::FPinnedTrail* Trail = Settings->GetPinnedTrail(Index))
+							{
+								Settings->DeletePinned(Index);
+							}
+						}
+					),
+					FCanExecuteAction()
+				);
+
+				FText Label = Trail->TrailName;
+				FName Name = FName(*Label.ToString());
+				FToolMenuEntry Entry = FToolMenuEntry::InitSubMenu(
+					Name,
+					Label,
+					// TODO: Update this and other labels/tooltips in this file.
+					LOCTEXT("PinnenTrailtip", "Modify Pinned States"),
+					MakeMenuDelegate,
+					TogglePinnedAction,
+					EUserInterfaceActionType::Button,
+					false,
+					FSlateIcon(FAppStyle::GetAppStyleSetName(), "Icons.Pinned")
+				);
+				PinnedTrails.AddEntry(Entry);
+
+			}
+		}
+	}
+}
+
+TSharedRef<SWidget> CreateFramesBeforeWidget(UMotionTrailToolOptions* Settings)
+{
+	// clang-format off
+	return
+		SNew(SBox)
+		.HAlign(HAlign_Right)
+		[
+			SNew(SBox)
+				.Padding(FMargin(0.0f, 0.0f, 0.0f, 0.0f))
+				.WidthOverride(100.0f)
+				[
+					SNew(SBorder)
+						.BorderImage(FAppStyle::Get().GetBrush("Menu.WidgetBorder"))
+						.Padding(FMargin(1.0f))
+						[
+							SNew(SSpinBox<int32>)
+								.IsEnabled_Lambda([Settings]() { return !Settings->bShowFullTrail; })
+								.MinValue(0)
+								.MaxValue(100)
+								.MinDesiredWidth(50)
+								.ToolTipText_Lambda(
+									[Settings]() -> FText
+									{
+										return FText::AsNumber(
+
+											Settings->FramesBefore
+										);
+									}
+								)
+								.Value_Lambda(
+									[Settings]() -> float
+									{
+										return Settings->FramesBefore;
+									}
+								)
+								.OnValueChanged_Lambda(
+									[Settings](float InValue)
+									{
+										Settings->FramesBefore = InValue;
+									}
+								)
+						]
+				]
+		];
+
+	// clang-format on
+}
+
+
+FToolMenuEntry CreateFramesBefore(UMotionTrailToolOptions* Settings)
+{
+	FToolMenuEntry Entry = FToolMenuEntry::InitWidget("FramesBefore", CreateFramesBeforeWidget(Settings), LOCTEXT("FramesBefore", "Frames Before"));
+	return Entry;
+}
+
+TSharedRef<SWidget> CreateFramesAfterWidget(UMotionTrailToolOptions* Settings)
+{
+	// clang-format off
+	return
+		SNew(SBox)
+		.HAlign(HAlign_Right)
+		[
+			SNew(SBox)
+				.Padding(FMargin(0.0f, 0.0f, 0.0f, 0.0f))
+				.WidthOverride(100.0f)
+				[
+					SNew(SBorder)
+						.BorderImage(FAppStyle::Get().GetBrush("Menu.WidgetBorder"))
+						.Padding(FMargin(1.0f))
+						[
+							SNew(SSpinBox<int32>)
+								.IsEnabled_Lambda([Settings]() { return !Settings->bShowFullTrail; })
+								.MinValue(0)
+								.MaxValue(100)
+								.MinDesiredWidth(50)
+								.ToolTipText_Lambda(
+									[Settings]() -> FText
+									{
+										return FText::AsNumber(
+
+											Settings->FramesAfter
+										);
+									}
+								)
+								.Value_Lambda(
+									[Settings]() -> float
+									{
+										return Settings->FramesAfter;
+									}
+								)
+								.OnValueChanged_Lambda(
+									[Settings](float InValue)
+									{
+										Settings->FramesAfter = InValue;
+									}
+								)
+						]
+				]
+		];
+
+	// clang-format on
+}
+
+FToolMenuEntry CreateFramesAfter(UMotionTrailToolOptions* Settings)
+{
+	FToolMenuEntry Entry = FToolMenuEntry::InitWidget("FramesAfter", CreateFramesAfterWidget(Settings), LOCTEXT("FramesAfter", "Frames After"));
+
+	return Entry;
+}
+
+TSharedRef<SWidget> CreateTrailStyleWidget(UMotionTrailToolOptions* Settings)
+{
+	// clang-format off
+	return
+
+		SNew(SBox)
+		.HAlign(HAlign_Right)
+		[
+			SNew(SBox)
+				.Padding(FMargin(0.0f, 0.0f, 0.0f, 0.0f))
+				.WidthOverride(100.0f)
+				[
+					SNew(SBorder)
+						.BorderImage(FAppStyle::Get().GetBrush("Menu.WidgetBorder"))
+						.Padding(FMargin(1.0f))
+						[
+							SNew(SComboButton)
+								.OnGetMenuContent_Lambda([Settings]()
+									{
+
+										FMenuBuilder MenuBuilder(true, NULL); //maybe todo look at settting these up with commands
+
+										MenuBuilder.BeginSection("TrailStyles");
+
+										TArray <TPair<FText, FText>>& TrailStyles = Settings->GetTrailStyles();
+										for (int32 Index = 0; Index < TrailStyles.Num(); ++Index)
+										{
+											FUIAction ItemAction(FExecuteAction::CreateUObject(Settings, &UMotionTrailToolOptions::SetTrailStyle, Index));
+											MenuBuilder.AddMenuEntry(TrailStyles[Index].Key, TAttribute<FText>(), FSlateIcon(), ItemAction);
+										}
+
+										MenuBuilder.EndSection();
+
+										return MenuBuilder.MakeWidget();
+									})
+								.ButtonContent()
+								[
+									SNew(SHorizontalBox)
+										+ SHorizontalBox::Slot()
+										[
+											SNew(STextBlock)
+												.Text_Lambda([Settings]()
+													{
+														const int32 Index = Settings->GetTrailStyleIndex();
+														const TArray <TPair<FText, FText>>& TrailStyles = Settings->GetTrailStyles();
+														return TrailStyles[Index].Key;
+													})
+												.ToolTipText_Lambda([Settings]()
+													{
+
+														const int32 Index = Settings->GetTrailStyleIndex();
+														const TArray <TPair<FText, FText>>& TrailStyles = Settings->GetTrailStyles();
+														return TrailStyles[Index].Value;
+													})
+										]
+
+								]
+						]
+				]
+		];
+
+	// clang-format on
+}
+
+
+FToolMenuEntry CreateTrailStyle(UMotionTrailToolOptions* Settings)
+{
+	FToolMenuEntry Entry = FToolMenuEntry::InitWidget("TrailStyle", CreateTrailStyleWidget(Settings), LOCTEXT("TrailStyle", "Trail Style"));
+
+	return Entry;
+}
+
+
+void CreatePinnedSubMenu(UToolMenu* InSubMenu, UMotionTrailToolOptions* Settings)
+{
+	FToolMenuSection& PinnedSection = InSubMenu->AddSection("PinnedSection", LOCTEXT("PinnedSection", "Pinned"));
+
+	FUIAction PinSelectedAction(
+		FExecuteAction::CreateLambda([Settings]()
+			{
+				Settings->PinSelection();
+			}
+		),
+		FCanExecuteAction()
+	);
+
+	FToolMenuEntry PinSelected = FToolMenuEntry::InitMenuEntry(
+		"PinSelected",
+		LOCTEXT("PinSelected", "Pin Selected"),
+		LOCTEXT("PinSelectedTrails", "Pin Selected Trails"),
+		FSlateIcon(),
+		PinSelectedAction,
+		EUserInterfaceActionType::Button
+	);
+	PinSelected.InsertPosition.Name = PinSelected.Name;
+	PinSelected.InsertPosition.Position = EToolMenuInsertType::First;
+	PinnedSection.AddEntry(PinSelected);
+
+	FUIAction SelectSocketAction(
+		FExecuteAction::CreateLambda([Settings]()
+			{
+				// FIXME temp approach for selecting the parent
+				FSlateApplication::Get().DismissAllMenus();
+
+				static const FActorPickerModeModule& ActorPickerMode = FModuleManager::Get().GetModuleChecked<FActorPickerModeModule>("ActorPickerMode");
+
+				ActorPickerMode.BeginActorPickingMode(
+					FOnGetAllowedClasses(),
+					FOnShouldFilterActor::CreateLambda([](const AActor* InActor)
+						{
+							UActorComponent* Component = InActor->GetComponentByClass(USkeletalMeshComponent::StaticClass());
+							return Component != nullptr;
+						}),
+					FOnActorSelected::CreateLambda([Settings](AActor* InActor)
+						{
+							FLevelEditorModule& LevelEditorModule = FModuleManager::GetModuleChecked<FLevelEditorModule>("LevelEditor");
+							TSharedPtr< ILevelEditor > LevelEditor = LevelEditorModule.GetFirstLevelEditor();
+
+							if (USkeletalMeshComponent* Component = Cast<USkeletalMeshComponent>(InActor->GetComponentByClass(USkeletalMeshComponent::StaticClass())))
+							{
+								TSharedPtr<SWidget> MenuWidget =
+									SNew(SSocketChooserPopup)
+									.SceneComponent(Component)
+									.OnSocketChosen_Lambda([Settings, Component](FName InSocketName) mutable
+										{
+											Settings->PinComponent(Component, InSocketName);
+										}
+									);
+								// Create as context menu
+								FSlateApplication::Get().PushMenu(
+									LevelEditor.ToSharedRef(),
+									FWidgetPath(),
+									MenuWidget.ToSharedRef(),
+									FSlateApplication::Get().GetCursorPos(),
+									FPopupTransitionEffect(FPopupTransitionEffect::ContextMenu)
+								);
+							}
+						})
+				);
+			}
+		),
+		FCanExecuteAction()
+	);
+
+	FToolMenuEntry SelectSocket = FToolMenuEntry::InitMenuEntry(
+		"SelectSocket",
+		LOCTEXT("SelectSocket", "Pin Socket"),
+		LOCTEXT("SelectSocketTrails", "Pin a Skeletal Mesh Socket by selecting it"),
+		FSlateIcon(),
+		SelectSocketAction,
+		EUserInterfaceActionType::Button
+	);
+	SelectSocket.InsertPosition.Name = SelectSocket.Name;
+
+
+	PinnedSection.AddEntry(SelectSocket);
+
+	FUIAction UnpinAllAction(
+		FExecuteAction::CreateLambda([Settings]()
+			{
+				int32 NumPinned = Settings->GetNumPinned();
+				for (int32 Index = NumPinned - 1; Index >= 0; --Index)
+				{
+					Settings->DeletePinned(Index);
+				}
+			}
+		),
+		FCanExecuteAction()
+	);
+
+	FToolMenuEntry UnpinAll = FToolMenuEntry::InitMenuEntry(
+		"UnpinAll",
+		LOCTEXT("UnpinAll", "Unpin All"),
+		LOCTEXT("UnpinAllTrails", "Unpin All Trails"),
+		FSlateIcon(),
+		UnpinAllAction,
+		EUserInterfaceActionType::Button
+	);
+	UnpinAll.InsertPosition.Name = UnpinAll.Name;
+
+
+
+	PinnedSection.AddEntry(UnpinAll);
+
+	//add pinned items
+	if (Settings->GetNumPinned() > 0)
+	{
+		FToolMenuSection& PinnedTrails = InSubMenu->AddSection("PinnedTrails", LOCTEXT("PinnedTrails", "Pinned Trails"));
+		CreatePinnedItems(Settings, PinnedTrails);
+	}
+
+}
+
+TSharedRef<SWidget> CreateTrailColorWidget(UMotionTrailToolOptions* Settings, FName PropertyName)
+{
+	// clang-format off
+	return
+
+		SNew(SBox)
+		.HAlign(HAlign_Right)
+		[
+			SNew(SBox)
+				.Padding(FMargin(0.0f, 0.0f, 0.0f, 0.0f))
+				.WidthOverride(100.0f)
+				[
+					SNew(SBorder)
+						.BorderImage(FAppStyle::Get().GetBrush("Menu.WidgetBorder"))
+						.Padding(FMargin(1.0f))
+						[
+							SNew(SColorBlock).IsEnabled(true)
+								//.Size(FVector2D(6.0, 38.0))
+								.Color_Lambda([Settings, PropertyName]()
+									{
+										FTrackInstancePropertyBindings Binding(PropertyName, PropertyName.ToString());
+										FLinearColor Color = Binding.GetCurrentValue<FLinearColor>(*Settings);
+										return Color;
+									}
+								)
+
+								.OnMouseButtonDown_Lambda([Settings, PropertyName](const FGeometry&, const FPointerEvent&)
+									{
+
+										FColorPickerArgs PickerArgs;
+										PickerArgs.bUseAlpha = false;
+										FTrackInstancePropertyBindings Binding(PropertyName, PropertyName.ToString());
+										FLinearColor Color = Binding.GetCurrentValue<FLinearColor>(*Settings);
+										PickerArgs.InitialColor = Color;
+										PickerArgs.OnColorCommitted = FOnLinearColorValueChanged::CreateLambda([Settings, PropertyName](FLinearColor Color)
+											{
+												FTrackInstancePropertyBindings Binding(PropertyName, PropertyName.ToString());
+												Binding.CallFunction<FLinearColor>(*Settings, Color);
+												FPropertyChangedEvent Event(Binding.GetProperty(*Settings));
+												Settings->PostEditChangeProperty(Event);
+											});
+										OpenColorPicker(PickerArgs);
+										return FReply::Handled();
+									})
+						]
+				]
+		];
+
+	// clang-format on
+}
+
+FToolMenuEntry CreateTrailColor(UMotionTrailToolOptions* Settings, FName PropertyName)
+{
+	FText Text = FText::FromString(PropertyName.ToString());
+	FToolMenuEntry Entry = FToolMenuEntry::InitWidget(PropertyName, CreateTrailColorWidget(Settings, PropertyName), Text);
+
+	return Entry;
+}
+
+TSharedRef<SWidget> CreateDoublePropertyWidget(UMotionTrailToolOptions* Settings, FName PropertyName)
+{
+	// clang-format off
+	return
+		SNew(SBox)
+		.HAlign(HAlign_Right)
+		[
+			SNew(SBox)
+				.Padding(FMargin(0.0f, 0.0f, 0.0f, 0.0f))
+				.WidthOverride(100.0f)
+				[
+					SNew(SBorder)
+						.BorderImage(FAppStyle::Get().GetBrush("Menu.WidgetBorder"))
+						.Padding(FMargin(1.0f))
+						[
+							SNew(SSpinBox<double>)
+								.MinValue(0)
+								.MaxValue(100)
+								.MinDesiredWidth(50)
+								.ToolTipText_Lambda(
+									[Settings, PropertyName]() -> FText
+									{
+										FTrackInstancePropertyBindings Binding(PropertyName, PropertyName.ToString());
+										double Value = Binding.GetCurrentValue<double>(*Settings);
+										return FText::AsNumber(
+											Value
+										);
+									}
+								)
+								.Value_Lambda(
+									[Settings, PropertyName]() -> double
+									{
+										FTrackInstancePropertyBindings Binding(PropertyName, PropertyName.ToString());
+										double Value = Binding.GetCurrentValue<double>(*Settings);
+										return Value;
+									}
+								)
+								.OnValueChanged_Lambda(
+									[Settings, PropertyName](double InValue)
+									{
+										FTrackInstancePropertyBindings Binding(PropertyName, PropertyName.ToString());
+										Binding.CallFunction<double>(*Settings, InValue);
+										FPropertyChangedEvent Event(Binding.GetProperty(*Settings));
+										Settings->PostEditChangeProperty(Event);
+									}
+								)
+						]
+				]
+		];
+
+	// clang-format on
+}
+
+FToolMenuEntry CreateDoubleProperty(UMotionTrailToolOptions* Settings, FName PropertyName)
+{
+	FText Text = FText::FromString(PropertyName.ToString());
+	FToolMenuEntry Entry = FToolMenuEntry::InitWidget(PropertyName, CreateDoublePropertyWidget(Settings, PropertyName), Text);
+
+	return Entry;
+}
+
+
+void CreateAdvancedSubMenu(UToolMenu* InSubMenu, UMotionTrailToolOptions* Settings)
+{
+	FToolMenuSection& TrailSettings = InSubMenu->AddSection("TrailSettings", LOCTEXT("TrailSettings", "Trail Settings"));
+
+	{
+		FUIAction Action;
+
+		Action.ExecuteAction = FExecuteAction::CreateLambda(
+			[Settings]() -> void
+			{
+				Settings->bShowKeys = !Settings->bShowKeys;
+				FPropertyChangedEvent Event(UMotionTrailToolOptions::StaticClass()->FindPropertyByName(GET_MEMBER_NAME_CHECKED(UMotionTrailToolOptions, bShowKeys)));
+				Settings->PostEditChangeProperty(Event);
+			}
+		);
+		Action.GetActionCheckState = FGetActionCheckState::CreateLambda(
+			[Settings]() -> ECheckBoxState
+			{
+				return Settings->bShowKeys ? ECheckBoxState::Checked
+					: ECheckBoxState::Unchecked;
+			}
+		);
+
+		FToolMenuEntry Entry = FToolMenuEntry::InitMenuEntry(
+			"ShowKeys",
+			LOCTEXT("ShowKeys", "Show Keys"),
+			LOCTEXT("ShowKeysTooltip", "Show keys"),
+			FSlateIcon(),
+			Action,
+			EUserInterfaceActionType::Check
+		);
+		TrailSettings.AddEntry(Entry);
+	}
+
+	{
+		FUIAction Action;
+
+		Action.ExecuteAction = FExecuteAction::CreateLambda(
+			[Settings]() -> void
+			{
+				Settings->bShowMarks = !Settings->bShowMarks;
+				FPropertyChangedEvent Event(UMotionTrailToolOptions::StaticClass()->FindPropertyByName(GET_MEMBER_NAME_CHECKED(UMotionTrailToolOptions, bShowMarks)));
+				Settings->PostEditChangeProperty(Event);
+			}
+		);
+		Action.GetActionCheckState = FGetActionCheckState::CreateLambda(
+			[Settings]() -> ECheckBoxState
+			{
+				return Settings->bShowMarks ? ECheckBoxState::Checked
+					: ECheckBoxState::Unchecked;
+			}
+		);
+
+		FToolMenuEntry Entry = FToolMenuEntry::InitMenuEntry(
+			"ShowMarks",
+			LOCTEXT("ShowMarks", "Show Marks"),
+			LOCTEXT("ShowMarksTooltip", "Show Marks"),
+			FSlateIcon(),
+			Action,
+			EUserInterfaceActionType::Check
+		);
+		TrailSettings.AddEntry(Entry);
+	}
+	{
+		const FName DoubleProperty("KeySize");
+		TrailSettings.AddEntry(CreateDoubleProperty(Settings, DoubleProperty));
+	}
+
+	{
+		const FName DoubleProperty("MarkSize");
+		TrailSettings.AddEntry(CreateDoubleProperty(Settings, DoubleProperty));
+	}
+	{
+		const FName DoubleProperty("TrailThickness");
+		TrailSettings.AddEntry(CreateDoubleProperty(Settings, DoubleProperty));
+	}
+
+	FToolMenuSection& ColorSettings = InSubMenu->AddSection("ColorSettings", LOCTEXT("ColorSettings", "Color Settings"));
+	{
+		const FName ColorProperty("DefaultColor");
+		ColorSettings.AddEntry(CreateTrailColor(Settings, ColorProperty));
+	}
+	{
+		const FName ColorProperty("TimePreColor");
+		ColorSettings.AddEntry(CreateTrailColor(Settings, ColorProperty));
+	}
+	{
+		const FName ColorProperty("TimePostColor");
+		ColorSettings.AddEntry(CreateTrailColor(Settings, ColorProperty));
+	}
+	{
+		const FName ColorProperty("DashPreColor");
+		ColorSettings.AddEntry(CreateTrailColor(Settings, ColorProperty));
+	}
+	{
+		const FName ColorProperty("DashPostColor");
+		ColorSettings.AddEntry(CreateTrailColor(Settings, ColorProperty));
+	}
+	{
+		const FName ColorProperty("KeyColor");
+		ColorSettings.AddEntry(CreateTrailColor(Settings, ColorProperty));
+	}
+
+
+}
+
+void CreateMotionTrailMenu(UToolMenu* InMenu)
+{
+	if (!InMenu)
+	{
+		return;
+	}
+
+	FToolMenuSection& PathModeSection = InMenu->AddSection("PathModeSection", LOCTEXT("PathMode", "Path Mode"));
+	UMotionTrailToolOptions* Settings = GetMutableDefault<UMotionTrailToolOptions>();
+
+
+	//show full trail
+	{
+		FUIAction Action;
+
+		Action.ExecuteAction = FExecuteAction::CreateLambda(
+			[Settings]() -> void
+			{
+				Settings->bShowFullTrail = !Settings->bShowFullTrail;
+				FPropertyChangedEvent Event(UMotionTrailToolOptions::StaticClass()->FindPropertyByName(GET_MEMBER_NAME_CHECKED(UMotionTrailToolOptions, bShowFullTrail)));
+				Settings->PostEditChangeProperty(Event);
+			}
+		);
+		Action.GetActionCheckState = FGetActionCheckState::CreateLambda(
+			[Settings]() -> ECheckBoxState
+			{
+				return Settings->bShowFullTrail ? ECheckBoxState::Checked
+					: ECheckBoxState::Unchecked;
+			}
+		);
+
+		FToolMenuEntry Entry = FToolMenuEntry::InitMenuEntry(
+			"FullTrail",
+			LOCTEXT("FullTrailLabel", "Full trail"),
+			LOCTEXT("FullTrailTooltip", "Show full trail"),
+			FSlateIcon(),
+			Action,
+			EUserInterfaceActionType::RadioButton
+		);
+		PathModeSection.AddEntry(Entry);
+	}
+	//set frames (oppposite of full trail basically)
+	{
+		FUIAction Action;
+
+		Action.ExecuteAction = FExecuteAction::CreateLambda(
+			[Settings]() -> void
+			{
+				Settings->bShowFullTrail = !Settings->bShowFullTrail;
+				FPropertyChangedEvent Event(UMotionTrailToolOptions::StaticClass()->FindPropertyByName(GET_MEMBER_NAME_CHECKED(UMotionTrailToolOptions, bShowFullTrail)));
+				Settings->PostEditChangeProperty(Event);
+			}
+		);
+		Action.GetActionCheckState = FGetActionCheckState::CreateLambda(
+			[Settings]() -> ECheckBoxState
+			{
+				return Settings->bShowFullTrail == false ? ECheckBoxState::Checked
+					: ECheckBoxState::Unchecked;
+			}
+		);
+
+		FToolMenuEntry Entry = FToolMenuEntry::InitMenuEntry(
+			"Set Frames",
+			LOCTEXT("SetFramesLabel", "Set Frames"),
+			LOCTEXT("SetframesTooltip", "Specify frame range"),
+			FSlateIcon(),
+			Action,
+			EUserInterfaceActionType::RadioButton
+		);
+		PathModeSection.AddEntry(Entry);
+	}
+	PathModeSection.AddEntry(CreateFramesBefore(Settings));
+	PathModeSection.AddEntry(CreateFramesAfter(Settings));
+
+	CreatePinnedSubMenu(InMenu, Settings);
+	/* in case design wants to go back to a subment
+	FToolMenuSection& PinnedSection = InMenu->AddSection(TEXT("Pinned"), LOCTEXT("Pinned", "Pinned"));
+	PinnedSection.AddSubMenu(TEXT("PinnedItems"), LOCTEXT("PinnedItems", "Pinned Items"), LOCTEXT("PinnedItems_tooltip", "Managed pinned motion trails"),
+		FNewToolMenuDelegate::CreateLambda([Settings](UToolMenu* InSubMenu)
+			{
+				CreatePinnedSubMenu(InSubMenu, Settings);
+
+			}));
+
+			*/
+	FToolMenuSection& PathOptionsMenu = InMenu->AddSection(TEXT("PathOptions"), LOCTEXT("PathOptions", "Path Options"));
+	//show selected trails
+	{
+		FUIAction Action;
+
+		Action.ExecuteAction = FExecuteAction::CreateLambda(
+			[Settings]() -> void
+			{
+				Settings->bShowSelectedTrails = !Settings->bShowSelectedTrails;
+				FPropertyChangedEvent Event(UMotionTrailToolOptions::StaticClass()->FindPropertyByName(GET_MEMBER_NAME_CHECKED(UMotionTrailToolOptions, bShowSelectedTrails)));
+				Settings->PostEditChangeProperty(Event);
+			}
+		);
+		Action.GetActionCheckState = FGetActionCheckState::CreateLambda(
+			[Settings]() -> ECheckBoxState
+			{
+				return Settings->bShowSelectedTrails ? ECheckBoxState::Checked
+					: ECheckBoxState::Unchecked;
+			}
+		);
+
+		FToolMenuEntry Entry = FToolMenuEntry::InitMenuEntry(
+			"ShowSelectedtrails",
+			LOCTEXT("ShowSelectedtrailsLabel", "Show Trails On Selection"),
+			LOCTEXT("ShowSelectedtrailsLabelTooltip", "Show trails on selected sequencer items."),
+			FSlateIcon(),
+			Action,
+			EUserInterfaceActionType::ToggleButton
+		);
+		PathOptionsMenu.AddEntry(Entry);
+
+	}
+	PathOptionsMenu.AddEntry(CreateTrailStyle(Settings));
+
+	PathOptionsMenu.AddSubMenu(TEXT("Advanced"), LOCTEXT("Advanced", "Advanced"), LOCTEXT("Advanced_tooltip", "Advanced display options"),
+		FNewToolMenuDelegate::CreateLambda([Settings](UToolMenu* InSubMenu)
+			{
+				CreateAdvancedSubMenu(InSubMenu, Settings);
+
+			}));
+
+}
 void UE::ControlRig::PopulateControlRigViewportToolbarTransformSubmenu(const FName InMenuName)
 {
 	FToolMenuOwnerScoped ScopeOwner(UE::ControlRig::Private::ControlRigOwnerName);
@@ -63,7 +886,6 @@ void UE::ControlRig::PopulateControlRigViewportToolbarTransformSubmenu(const FNa
 			);
 			// We want to appear early in the section.
 			Entry.InsertPosition.Position = EToolMenuInsertType::First;
-			Entry.SetShowInToolbarTopLevel(true);
 			GizmoSection.AddEntry(Entry);
 		}
 
@@ -100,11 +922,128 @@ void UE::ControlRig::PopulateControlRigViewportToolbarTransformSubmenu(const FNa
 			);
 			// We want to appear early in the section.
 			Entry.InsertPosition.Position = EToolMenuInsertType::First;
-			Entry.SetShowInToolbarTopLevel(true);
 			GizmoSection.AddEntry(Entry);
 		}
 	}
+	{
+
+		FToolMenuSection& PreviewToolsSection  = Menu->FindOrAddSection("PreviewTools", LOCTEXT("PreviewToolsLabel", "Preview Tools"));
+		{
+			// Add "Temporary Pivot" checkbox.
+			{
+				FUIAction Action;
+
+				Action.ExecuteAction = FExecuteAction::CreateLambda(
+					[]() -> void
+					{
+						if (FLevelEditorModule* LevelEditorModule = FModuleManager::GetModulePtr<FLevelEditorModule>(TEXT("LevelEditor")))
+						{
+							TSharedPtr<ILevelEditor> LevelEditorPtr = LevelEditorModule->GetLevelEditorInstance().Pin();
+							if (LevelEditorPtr.IsValid())
+							{
+								FString ActiveToolName = LevelEditorPtr->GetEditorModeManager().GetInteractiveToolsContext()->ToolManager->GetActiveToolName(EToolSide::Left);							
+								if (ActiveToolName == TEXT("SequencerPivotTool"))
+								{
+									LevelEditorPtr->GetEditorModeManager().GetInteractiveToolsContext()->ToolManager->DeactivateTool(EToolSide::Left, EToolShutdownType::Completed);
+								}
+								else
+								{
+									LevelEditorPtr->GetEditorModeManager().GetInteractiveToolsContext()->ToolManager->SelectActiveToolType(EToolSide::Left, TEXT("SequencerPivotTool"));
+									LevelEditorPtr->GetEditorModeManager().GetInteractiveToolsContext()->ToolManager->ActivateTool(EToolSide::Left);
+								}
+							}
+						}
+					}
+				);
+				Action.GetActionCheckState = FGetActionCheckState::CreateLambda(
+					[]() -> ECheckBoxState
+					{
+						if (FLevelEditorModule* LevelEditorModule = FModuleManager::GetModulePtr<FLevelEditorModule>(TEXT("LevelEditor")))
+						{
+							TSharedPtr<ILevelEditor> LevelEditorPtr = LevelEditorModule->GetLevelEditorInstance().Pin();
+							if (LevelEditorPtr.IsValid())
+							{
+								FString ActiveToolName = LevelEditorPtr->GetEditorModeManager().GetInteractiveToolsContext()->ToolManager->GetActiveToolName(EToolSide::Left);
+								if (ActiveToolName == TEXT("SequencerPivotTool"))
+								{
+									return ECheckBoxState::Checked;
+								}
+							}
+						}
+						return  ECheckBoxState::Unchecked;
+					}
+				);
 
+				FToolMenuEntry Entry = FToolMenuEntry::InitMenuEntry(
+					"TemporaryPivot",
+					LOCTEXT("TemporaryPivotLabel", "Temporary Pivot"),
+					LOCTEXT(
+						"TemporaryPivotTooltip",
+						"Toggle Temporary Pivot Tool"
+					),
+					FSlateIcon(TEXT("ControlRigEditorStyle"), TEXT("ControlRig.TemporaryPivot")),
+					Action,
+					EUserInterfaceActionType::ToggleButton
+				);
+				Entry.SetShowInToolbarTopLevel(true);
+				PreviewToolsSection.AddEntry(Entry);
+			}
+		}
+		
+		FNewToolMenuDelegate MakeMenuDelegate = FNewToolMenuDelegate::CreateLambda(
+			[](UToolMenu* Submenu)
+			{
+				UMotionTrailToolOptions* Settings = GetMutableDefault<UMotionTrailToolOptions>();
+				{
+					CreateMotionTrailMenu(Submenu);
+				}
+
+			}
+		);
+
+		UMotionTrailToolOptions* Settings = GetMutableDefault<UMotionTrailToolOptions>();
+
+		// Create the checkbox actions for the MotionPaths submenu itself.
+		FToolUIAction CheckboxMenuAction;
+		{
+			CheckboxMenuAction.ExecuteAction = FToolMenuExecuteAction::CreateLambda(
+				[Settings](const FToolMenuContext& Context) -> void
+				{
+					Settings->bShowTrails = !Settings->bShowTrails;
+					FPropertyChangedEvent ShowTrailEvent(UMotionTrailToolOptions::StaticClass()->FindPropertyByName(GET_MEMBER_NAME_CHECKED(UMotionTrailToolOptions, bShowTrails)));
+					Settings->PostEditChangeProperty(ShowTrailEvent);
+				}
+			);
+			CheckboxMenuAction.GetActionCheckState = FToolMenuGetActionCheckState::CreateLambda(
+				[Settings](const FToolMenuContext& Context) -> ECheckBoxState
+				{
+					if (Settings->bShowTrails)
+					{
+						return ECheckBoxState::Checked;
+					}
+					else
+					{
+						return ECheckBoxState::Unchecked;
+					}
+				}
+			);
+		}
+
+		FToolMenuEntry MotionPathsSubmenu = FToolMenuEntry::InitSubMenu(
+			"MotionPaths",
+			LOCTEXT("MotionPathsLabel", "Motion Paths"),
+			// TODO: Update this and other labels/tooltips in this file.
+			LOCTEXT("MotionPathsTooltip", "Check to enable motion paths. Submenu contains settings for motion paths."),
+			MakeMenuDelegate,
+			CheckboxMenuAction,
+			EUserInterfaceActionType::ToggleButton,
+			false,
+			FSlateIcon(TEXT("ControlRigEditorStyle"), TEXT("ControlRig.EditableMotionTrails"))
+		);
+		MotionPathsSubmenu.SetShowInToolbarTopLevel(true);
+
+		PreviewToolsSection.AddEntry(MotionPathsSubmenu);
+	}
 	{
 		FToolMenuSection& SelectionSection = Menu->FindOrAddSection("Selection");
 
@@ -143,9 +1082,8 @@ void UE::ControlRig::PopulateControlRigViewportToolbarTransformSubmenu(const FNa
 	}
 }
 
-void UE::ControlRig::PopulateControlRigViewportToolbarShowSubmenu(
-	const FName InMenuName, const TSharedPtr<const FUICommandList>& InCommandList
-)
+
+void UE::ControlRig::PopulateControlRigViewportToolbarShowSubmenu(const FName InMenuName)
 {
 	FToolMenuOwnerScoped ScopeOwner(UE::ControlRig::Private::ControlRigOwnerName);
 
@@ -157,195 +1095,196 @@ void UE::ControlRig::PopulateControlRigViewportToolbarShowSubmenu(
 		LOCTEXT("AnimationLabel", "Animation"),
 		LOCTEXT("AnimationTooltip", "Animation-related show flags"),
 		FNewToolMenuDelegate::CreateLambda(
-			[WeakCommandList = InCommandList.ToWeakPtr()](UToolMenu* AnimationShowFlagsSubmenu)
+			[](UToolMenu* AnimationShowFlagsSubmenu)
 			{
 				FToolMenuSection& UnnamedSection = AnimationShowFlagsSubmenu->FindOrAddSection(NAME_None);
 
-				// TODO: These don't work because the commands have not yet been bound.
-				if (const TSharedPtr<const FUICommandList>& CommandList = WeakCommandList.Pin())
+				if (UControlRigEditModeSettings* const Settings = GetMutableDefault<UControlRigEditModeSettings>())
 				{
-					UnnamedSection.AddMenuEntryWithCommandList(
-						FControlRigEditorCommands::Get().ToggleDrawAxesOnSelection, CommandList
-					);
-					UnnamedSection.AddMenuEntryWithCommandList(
-						FControlRigEditorCommands::Get().ToggleControlsAsOverlay, CommandList
-					);
-					UnnamedSection.AddMenuEntryWithCommandList(
-						FControlRigEditorCommands::Get().ToggleControlVisibility, CommandList
-					);
-					UnnamedSection.AddMenuEntryWithCommandList(FControlRigEditorCommands::Get().ToggleDrawNulls, CommandList);
-					UnnamedSection.AddMenuEntryWithCommandList(FControlRigEditorCommands::Get().ToggleDrawSockets, CommandList);
-				}
-
-				FNewToolMenuDelegate MakeMenuDelegate = FNewToolMenuDelegate::CreateLambda(
-					[](UToolMenu* Submenu)
 					{
-						UMotionTrailToolOptions* Settings = GetMutableDefault<UMotionTrailToolOptions>();
-
-						// TODO: This section needs to be completed to show all radio buttons. Perhaps the settings
-						// object should be updated to better match the design, or the design updated to match what the
-						// settings object can express.
-						{
-							FToolMenuSection& PathModeSection =
-								Submenu->FindOrAddSection("PathMode", LOCTEXT("PathModeLabel", "Path Mode"));
-
-							// Add "Full Trail" radio button.
+						FUIAction Action;
+						Action.ExecuteAction = FExecuteAction::CreateLambda(
+							[Settings]() -> void
 							{
-								FUIAction Action;
-
-								Action.ExecuteAction = FExecuteAction::CreateLambda(
-									[Settings]() -> void
-									{
-										Settings->bShowFullTrail = true;
-										Settings->PostEditChange();
-									}
-								);
-								Action.GetActionCheckState = FGetActionCheckState::CreateLambda(
-									[Settings]() -> ECheckBoxState
-									{
-										return Settings->bShowFullTrail ? ECheckBoxState::Checked
-																		: ECheckBoxState::Unchecked;
-									}
-								);
-
-								FToolMenuEntry Entry = FToolMenuEntry::InitMenuEntry(
-									"FullTrail",
-									LOCTEXT("FullTrailLabel", "Full trail"),
-									LOCTEXT("FullTrailTooltip", "Whether or not to show the full motion trail."),
-									FSlateIcon(),
-									Action,
-									EUserInterfaceActionType::RadioButton
-								);
-								PathModeSection.AddEntry(Entry);
+								Settings->bDisplayAxesOnSelection = !Settings->bDisplayAxesOnSelection;
 							}
-						}
-
-						{
-							FToolMenuSection& PathOptionsSection =
-								Submenu->FindOrAddSection("PathOptions", LOCTEXT("PathOptionsLabel", "Path Options"));
-
-							// Add "Show Keys" checkbox.
+						);
+						Action.GetActionCheckState = FGetActionCheckState::CreateLambda(
+							[Settings]() -> ECheckBoxState
 							{
-								FUIAction Action;
+								return Settings->bDisplayAxesOnSelection ? ECheckBoxState::Checked
+									: ECheckBoxState::Unchecked;
+							}
+						);
 
-								Action.ExecuteAction = FExecuteAction::CreateLambda(
-									[Settings]() -> void
-									{
-										Settings->bShowKeys = !Settings->bShowKeys;
-										Settings->PostEditChange();
-									}
-								);
-								Action.GetActionCheckState = FGetActionCheckState::CreateLambda(
-									[Settings]() -> ECheckBoxState
-									{
-										return Settings->bShowKeys ? ECheckBoxState::Checked : ECheckBoxState::Unchecked;
-									}
-								);
+						FToolMenuEntry Entry = FToolMenuEntry::InitMenuEntry(
+							"AxisOnSelection",
+							LOCTEXT("AxisOnSelectionLabel", "Axis On Selection"),
+							LOCTEXT("AxisOnSelectionTooltip", "Should we show axes for the selected elements"),
+							FSlateIcon(),
+							Action,
+							EUserInterfaceActionType::RadioButton
+						);
+						UnnamedSection.AddEntry(Entry);
 
-								FToolMenuEntry Entry = FToolMenuEntry::InitMenuEntry(
-									"ShowKeys",
-									LOCTEXT("ShowKeysLabel", "Show Keys"),
-									LOCTEXT("ShowKeysTooltip", "Whether or not to show keys on the motion trail."),
-									FSlateIcon(),
-									Action,
-									EUserInterfaceActionType::ToggleButton
-								);
-								PathOptionsSection.AddEntry(Entry);
+					}
+					{
+						FUIAction Action;
+						Action.ExecuteAction = FExecuteAction::CreateLambda(
+							[Settings]() -> void
+							{
+								Settings->bDisplayHierarchy = !Settings->bDisplayHierarchy;
+							}
+						);
+						Action.GetActionCheckState = FGetActionCheckState::CreateLambda(
+							[Settings]() -> ECheckBoxState
+							{
+								return Settings->bDisplayHierarchy ? ECheckBoxState::Checked
+									: ECheckBoxState::Unchecked;
 							}
+						);
 
-							// TODO: This currently only uses FramesBefore. This should be updated to support all cases.
-							TSharedRef<SWidget> FramesShownWidget =
-								// clang-format off
-								SNew(SHorizontalBox)
-								+ SHorizontalBox::Slot()
-								.FillWidth(0.9f)
-								[
-									SNew(SSpinBox<int32>)
-										.MinValue(0)
-										.MaxValue(100)
-										.MinDesiredWidth(100)
-										.ToolTipText_Lambda(
-											[Settings]() -> FText
-											{
-												return FText::AsNumber(
-													
-													Settings->FramesBefore
-												);
-											}
-										)
-										.Value_Lambda(
-											[Settings]() -> float
-											{
-												return Settings->FramesBefore;
-											}
-										)
-										.OnValueChanged_Lambda(
-											[Settings](float InValue)
-											{
-												Settings->FramesBefore = InValue;
-											}
-										)
-								]
-								+ SHorizontalBox::Slot()
-								.FillWidth(0.1f);
-							// clang-format on
-							PathOptionsSection.AddEntry(FToolMenuEntry::InitWidget(
-								"FramesShown", FramesShownWidget, LOCTEXT("FramesShownLabel", "Frames Shown")
-							));
-						}
+						FToolMenuEntry Entry = FToolMenuEntry::InitMenuEntry(
+							"Hierarchy",
+							LOCTEXT("HierarchyLabel", "Hierarchy"),
+							LOCTEXT("HierarchyTooltip", "Whether to show all bones in the hierarchy"),
+							FSlateIcon(),
+							Action,
+							EUserInterfaceActionType::RadioButton
+						);
+						UnnamedSection.AddEntry(Entry);
 					}
-				);
+					{
+						FUIAction Action;
+						Action.ExecuteAction = FExecuteAction::CreateLambda(
+							[Settings]() -> void
+							{
+								Settings->bShowControlsAsOverlay = !Settings->bShowControlsAsOverlay;
+							}
+						);
+						Action.GetActionCheckState = FGetActionCheckState::CreateLambda(
+							[Settings]() -> ECheckBoxState
+							{
+								return Settings->bShowControlsAsOverlay ? ECheckBoxState::Checked
+									: ECheckBoxState::Unchecked;
+							}
+						);
 
-				UMotionTrailToolOptions* Settings = GetMutableDefault<UMotionTrailToolOptions>();
+						FToolMenuEntry Entry = FToolMenuEntry::InitMenuEntry(
+							"ControlsAsOverlay",
+							LOCTEXT("ControlsAsOverlayLabel", "Controls As Overlay"),
+							LOCTEXT("ControlsAsOverlaylTooltip", "Whether to show controls as overlay"),
+							FSlateIcon(),
+							Action,
+							EUserInterfaceActionType::RadioButton
+						);
+						UnnamedSection.AddEntry(Entry);
+					}
+					{
+						FUIAction Action;
+						Action.ExecuteAction = FExecuteAction::CreateLambda(
+							[Settings]() -> void
+							{
+								Settings->bHideControlShapes = !Settings->bHideControlShapes;
+							}
+						);
+						Action.GetActionCheckState = FGetActionCheckState::CreateLambda(
+							[Settings]() -> ECheckBoxState
+							{
+								return Settings->bHideControlShapes == false ? ECheckBoxState::Checked
+									: ECheckBoxState::Unchecked;
+							}
+						);
 
-				// Create the checkbox actions for the MotionPaths submenu itself.
-				FToolUIAction CheckboxMenuAction;
-				{
-					CheckboxMenuAction.ExecuteAction = FToolMenuExecuteAction::CreateLambda(
-						[Settings](const FToolMenuContext& Context) -> void
-						{
-							Settings->bShowTrails = !Settings->bShowTrails;
-						}
-					);
-					CheckboxMenuAction.GetActionCheckState = FToolMenuGetActionCheckState::CreateLambda(
-						[Settings](const FToolMenuContext& Context) -> ECheckBoxState
-						{
-							if (Settings->bShowTrails)
+						FToolMenuEntry Entry = FToolMenuEntry::InitMenuEntry(
+							"ControlShapes",
+							LOCTEXT("ControlShapesLabel", "Control Shapes"),
+							LOCTEXT("ControlShapesTooltip", "Should we always hide control shapes in viewport"),
+							FSlateIcon(),
+							Action,
+							EUserInterfaceActionType::RadioButton
+						);
+						UnnamedSection.AddEntry(Entry);
+					}
+					{
+						FUIAction Action;
+						Action.ExecuteAction = FExecuteAction::CreateLambda(
+							[Settings]() -> void
 							{
-								return ECheckBoxState::Checked;
+								Settings->bDisplayNulls = !Settings->bDisplayNulls;
 							}
-							else
+						);
+						Action.GetActionCheckState = FGetActionCheckState::CreateLambda(
+							[Settings]() -> ECheckBoxState
 							{
-								return ECheckBoxState::Unchecked;
+								return Settings->bDisplayNulls ? ECheckBoxState::Checked
+									: ECheckBoxState::Unchecked;
 							}
-						}
-					);
-				}
+						);
 
-				FToolMenuEntry MotionPathsSubmenu = FToolMenuEntry::InitSubMenu(
-					"MotionPaths",
-					LOCTEXT("MotionPathsLabel", "Motion Paths"),
-					// TODO: Update this and other labels/tooltips in this file.
-					LOCTEXT("MotionPathsTooltip", "Check to enable motion paths. Submenu contains settings for motion paths."),
-					MakeMenuDelegate,
-					CheckboxMenuAction,
-					EUserInterfaceActionType::ToggleButton,
-					false,
-					// TODO: Temporary icon.
-					FSlateIcon(FAppStyle::GetAppStyleSetName(), "AnimGraph.Attribute.InertialBlending.Icon")
-				);
-				MotionPathsSubmenu.SetShowInToolbarTopLevel(true);
-				// Override the label when this item is raised to the top-level toolbar.
-				MotionPathsSubmenu.ToolBarData.LabelOverride = TAttribute<FText>::CreateLambda(
-					[Settings]() -> FText
+						FToolMenuEntry Entry = FToolMenuEntry::InitMenuEntry(
+							"Nulls",
+							LOCTEXT("NullsLabel", "Nulls"),
+							LOCTEXT("NullTooltip", "Whether to show all nulls in the hierarchy"),
+							FSlateIcon(),
+							Action,
+							EUserInterfaceActionType::RadioButton
+						);
+						UnnamedSection.AddEntry(Entry);
+					}
 					{
-						// TODO: This only uses the FramesBefore value at the moment. It should be updated.
-						const int32 FramesShown = Settings->FramesBefore;
-						return FText::AsNumber(FramesShown);
+						FUIAction Action;
+						Action.ExecuteAction = FExecuteAction::CreateLambda(
+							[Settings]() -> void
+							{
+								Settings->bShowAllProxyControls = !Settings->bShowAllProxyControls;
+							}
+						);
+						Action.GetActionCheckState = FGetActionCheckState::CreateLambda(
+							[Settings]() -> ECheckBoxState
+							{
+								return Settings->bShowAllProxyControls ? ECheckBoxState::Checked
+									: ECheckBoxState::Unchecked;
+							}
+						);
+
+						FToolMenuEntry Entry = FToolMenuEntry::InitMenuEntry(
+							"ProxyControls",
+							LOCTEXT("ProxyControlsLabel", "Proxy Controls"),
+							LOCTEXT("ProxyControlsTooltip", "Whether to show Proxy Controls"),
+							FSlateIcon(),
+							Action,
+							EUserInterfaceActionType::RadioButton
+						);
+						UnnamedSection.AddEntry(Entry);
 					}
-				);
+					{
+						FUIAction Action;
+						Action.ExecuteAction = FExecuteAction::CreateLambda(
+							[Settings]() -> void
+							{
+								Settings->bDisplaySockets = !Settings->bDisplaySockets;
+							}
+						);
+						Action.GetActionCheckState = FGetActionCheckState::CreateLambda(
+							[Settings]() -> ECheckBoxState
+							{
+								return Settings->bDisplaySockets ? ECheckBoxState::Checked
+									: ECheckBoxState::Unchecked;
+							}
+						);
 
-				UnnamedSection.AddEntry(MotionPathsSubmenu);
+						FToolMenuEntry Entry = FToolMenuEntry::InitMenuEntry(
+							"Sockets",
+							LOCTEXT("SocketsLabel", "Sockets"),
+							LOCTEXT("SocketsTooltip", "Whether to show Sockets"),
+							FSlateIcon(),
+							Action,
+							EUserInterfaceActionType::RadioButton
+						);
+						UnnamedSection.AddEntry(Entry);
+					}
+				}
 			}
 		),
 		false,
diff --git a/Engine/Plugins/Animation/ControlRig/Source/ControlRigEditor/Private/Editor/ControlRigViewportToolbarExtensions.h b/Engine/Plugins/Animation/ControlRig/Source/ControlRigEditor/Private/Editor/ControlRigViewportToolbarExtensions.h
--- a/Engine/Plugins/Animation/ControlRig/Source/ControlRigEditor/Private/Editor/ControlRigViewportToolbarExtensions.h
+++ b/Engine/Plugins/Animation/ControlRig/Source/ControlRigEditor/Private/Editor/ControlRigViewportToolbarExtensions.h
@@ -1,4 +1,4 @@
-﻿// Copyright Epic Games, Inc. All Rights Reserved.
+// Copyright Epic Games, Inc. All Rights Reserved.
 
 #pragma once
 
@@ -11,7 +11,7 @@ namespace UE::ControlRig
 {
 
 void PopulateControlRigViewportToolbarTransformSubmenu(const FName InMenuName);
-void PopulateControlRigViewportToolbarShowSubmenu(const FName InMenuName, const TSharedPtr<const FUICommandList>& InCommandList);
+void PopulateControlRigViewportToolbarShowSubmenu(const FName InMenuName);
 
 void RemoveControlRigViewportToolbarExtensions();
 
diff --git a/Engine/Plugins/Editor/SequencerAnimTools/Source/SequencerAnimTools/Private/MotionTrailTool.cpp b/Engine/Plugins/Editor/SequencerAnimTools/Source/SequencerAnimTools/Private/MotionTrailTool.cpp
--- a/Engine/Plugins/Editor/SequencerAnimTools/Source/SequencerAnimTools/Private/MotionTrailTool.cpp
+++ b/Engine/Plugins/Editor/SequencerAnimTools/Source/SequencerAnimTools/Private/MotionTrailTool.cpp
@@ -52,6 +52,8 @@ void FMotionTrailCommands::RegisterCommands()
 	UI_COMMAND(TranslateSelectedKeysLeft, "Translate Selected Keys Left", "Translate selected keys one frame to the left", EUserInterfaceActionType::Button, FInputChord(EModifierKey::Control, EKeys::Left));
 	UI_COMMAND(TranslateSelectedKeysRight, "Translate Selected Keys Right", "Translate selected keys one frame to the right", EUserInterfaceActionType::Button, FInputChord(EModifierKey::Control, EKeys::Right));
 	UI_COMMAND(FrameSelection, "Frame Selection", "Frame camera around current key selection", EUserInterfaceActionType::Button, FInputChord(EKeys::F));
+	UI_COMMAND(DeselectAll, "Deslect All", "Deselect All Keys", EUserInterfaceActionType::Button, FInputChord(EKeys::Escape));
+
 }
 
 FString UMotionTrailTool::TrailKeyTransformGizmoInstanceIdentifier = TEXT("TrailKeyTransformGizmoInstanceIdentifier");
@@ -143,6 +145,11 @@ void UMotionTrailTool::SetupIntegration()
 		FExecuteAction::CreateUObject(this, &UMotionTrailTool::FrameSelection),
 		FCanExecuteAction::CreateUObject(this, &UMotionTrailTool::SomeKeysAreSelected));
 
+	CommandBindings->MapAction(
+		Commands.DeselectAll,
+		FExecuteAction::CreateUObject(this, &UMotionTrailTool::SelectNone),
+		FCanExecuteAction::CreateUObject(this, &UMotionTrailTool::SomeKeysAreSelected));
+
 	CommandBindings->MapAction(
 		FGenericCommands::Get().Delete,
 		FExecuteAction::CreateUObject(this, &UMotionTrailTool::DeleteSelectedKeys),
@@ -246,7 +253,6 @@ void UMotionTrailTool::DrawHUD(FCanvas* Canvas, IToolsContextRenderAPI* RenderAP
 	}
 }
 
-
 void UMotionTrailTool::OnTick(float DeltaTime)
 {
 	for (TUniquePtr<UE::SequencerAnimTools::FTrailHierarchy>& TrailHierarchy : TrailHierarchies)
@@ -289,7 +295,8 @@ void UMotionTrailTool::GizmoTransformChanged(UTransformProxy* Proxy, FTransform
 	for (const TUniquePtr<UE::SequencerAnimTools::FTrailHierarchy>& TrailHierarchy : TrailHierarchies)
 	{
 		FVector WidgetLocation = Transform.GetLocation();
-		bool bIsHandled = TrailHierarchy->ApplyDelta(LocationDiff, Diff.GetRotation().Rotator(), WidgetLocation);
+		const bool bShiftIsDown = FSlateApplication::Get().GetModifierKeys().IsShiftDown();
+		bool bIsHandled = TrailHierarchy->ApplyDelta(LocationDiff, Diff.GetRotation().Rotator(), WidgetLocation, bShiftIsDown);
 		if (bIsHandled)
 		{
 			bManipulatorMadeChange = true;
@@ -393,26 +400,23 @@ void UMotionTrailTool::FrameSelection()
 {
 	if (SomeKeysAreSelected())
 	{
-		//if (CurrentViewportClient)
+		FBox Bounds(EForceInit::ForceInit);
+		TArray<FVector> Positions;
+		for (TUniquePtr<UE::SequencerAnimTools::FTrailHierarchy>& TrailHierarchy : TrailHierarchies)
 		{
-			FBox Bounds(EForceInit::ForceInit);
-			TArray<FVector> Positions;
-			for (TUniquePtr<UE::SequencerAnimTools::FTrailHierarchy>& TrailHierarchy : TrailHierarchies)
+			if (TrailHierarchy->IsAnythingSelected(Positions, true))
 			{
-				if (TrailHierarchy->IsAnythingSelected(Positions))
+				for (const FVector Pos : Positions)
 				{
-					for (const FVector Pos : Positions)
-					{
-						Bounds += Pos;
-						Bounds += Pos + FVector::OneVector * 5.f;
-						Bounds += Pos - FVector::OneVector * 5.f;
-					}
+					Bounds += Pos;
+					Bounds += Pos + FVector::OneVector * 5.f;
+					Bounds += Pos - FVector::OneVector * 5.f;
 				}
 			}
-			if (Bounds.IsValid)
-			{
-				//CurrentViewportClient->FocusViewportOnBox(Bounds);
-			}
+		}
+		if (Bounds.IsValid)
+		{
+			GEditor->MoveViewportCamerasToBox(Bounds, true);
 		}
 	}
 }
diff --git a/Engine/Plugins/Editor/SequencerAnimTools/Source/SequencerAnimTools/Private/Sequencer/MotionTrailMovieSceneKey.cpp b/Engine/Plugins/Editor/SequencerAnimTools/Source/SequencerAnimTools/Private/Sequencer/MotionTrailMovieSceneKey.cpp
--- a/Engine/Plugins/Editor/SequencerAnimTools/Source/SequencerAnimTools/Private/Sequencer/MotionTrailMovieSceneKey.cpp
+++ b/Engine/Plugins/Editor/SequencerAnimTools/Source/SequencerAnimTools/Private/Sequencer/MotionTrailMovieSceneKey.cpp
@@ -15,6 +15,7 @@
 #include "PrimitiveDrawingUtils.h"
 #include "SceneView.h"
 
+
 namespace UE
 {
 namespace SequencerAnimTools
@@ -46,20 +47,22 @@ void FMotionTraiMovieScenelKeyTool::Render(const FGuid& Guid, const FSceneView*
 		ClearSelection();
 		return;
 	}
+	const FTransform OffsetTransform = OwningTrail->GetOffsetTransform();
+	const FTransform ParentSpaceTransform = OwningTrail->GetParentSpaceTransform();
 	const bool bHitTesting = PDI && PDI->IsHitTesting();
 	const float KeySize = UMotionTrailToolOptions::GetTrailOptions()->KeySize;
 	for (const TPair<FFrameNumber, TUniquePtr<FTrailKeyInfo>>& FrameKeyPair : Keys.Keys)
 	{
-		if (OwningTrail->GetDrawInfo()->GetCachedViewRange().Contains(OwningTrail->WeakSequencer.Pin()->GetFocusedTickResolution().AsSeconds(FrameKeyPair.Value->FrameNumber)))
+		if (CachedViewFrameRange.Contains(FrameKeyPair.Value->FrameNumber))
 		{
 			FLinearColor KeyColor = CachedSelection.Contains(FrameKeyPair.Value.Get()) ? FLinearColor(1.0,1.0,0.0f) :  UMotionTrailToolOptions::GetTrailOptions()->KeyColor;
-
+			const FTransform Transform = FrameKeyPair.Value->Transform * OffsetTransform * ParentSpaceTransform;
 			if (bHitTesting)
 			{
-				PDI->SetHitProxy(new HMotionTrailMovieSceneKeyProxy(Guid, FrameKeyPair.Value->Transform, FrameKeyPair.Value.Get()));
+				PDI->SetHitProxy(new HMotionTrailMovieSceneKeyProxy(Guid, Transform, FrameKeyPair.Value.Get()));
 			}
 
-			PDI->DrawPoint(FrameKeyPair.Value->Transform.GetLocation(), KeyColor, KeySize, SDPG_MAX);
+			PDI->DrawPoint(Transform.GetLocation(), KeyColor, KeySize, SDPG_MAX);
 
 			if (bHitTesting)
 			{
@@ -71,6 +74,8 @@ void FMotionTraiMovieScenelKeyTool::Render(const FGuid& Guid, const FSceneView*
 
 void FMotionTraiMovieScenelKeyTool::DrawHUD(const FSceneView* View, FCanvas* Canvas)
 {
+	return;
+	/* no longer show frame number 
 	const bool bIsVisible = (UMotionTrailToolOptions::GetTrailOptions()->bShowKeys && UMotionTrailToolOptions::GetTrailOptions()->bShowFrameNumber);
 	if (!bIsVisible)
 	{
@@ -79,17 +84,21 @@ void FMotionTraiMovieScenelKeyTool::DrawHUD(const FSceneView* View, FCanvas* Can
 
 	FFrameRate TickResolution = OwningTrail->GetSequencer()->GetFocusedTickResolution();
 	FFrameRate DisplayRate = OwningTrail->GetSequencer()->GetFocusedDisplayRate();
+	const FTransform OffsetTransform = OwningTrail->GetOffsetTransform();
+	const FTransform ParentSpaceTransform = OwningTrail->GetParentSpaceTransform();
 
 	const float KeySize = UMotionTrailToolOptions::GetTrailOptions()->KeySize;
 	for (const TPair<FFrameNumber, TUniquePtr<FTrailKeyInfo>>& FrameKeyPair : Keys.Keys)
 	{
-		if (OwningTrail->GetDrawInfo()->GetCachedViewRange().Contains(OwningTrail->WeakSequencer.Pin()->GetFocusedTickResolution().AsSeconds(FrameKeyPair.Value->FrameNumber)))
+		if (CachedViewFrameRange.Contains(FrameKeyPair.Value->FrameNumber))
 		{
 			FLinearColor KeyColor = CachedSelection.Contains(FrameKeyPair.Value.Get()) ? FLinearColor(1.0, 1.0, 0.0f) : UMotionTrailToolOptions::GetTrailOptions()->KeyColor;
 
 
 			FVector2D PixelLocation;
-			if (View->WorldToPixel(FrameKeyPair.Value->Transform.GetLocation(), PixelLocation))
+			const FTransform Transform = FrameKeyPair.Value->Transform * OffsetTransform * ParentSpaceTransform;
+
+			if (View->WorldToPixel(Transform.GetLocation(), PixelLocation))
 			{
 				PixelLocation /= Canvas->GetDPIScale();
 
@@ -102,18 +111,34 @@ void FMotionTraiMovieScenelKeyTool::DrawHUD(const FSceneView* View, FCanvas* Can
 			}
 		}
 	}
+	*/
 }
 
 bool FMotionTraiMovieScenelKeyTool::IsSelected(FVector& OutVectorPosition) const
 {
 	if (CachedSelection.Num() > 0)
 	{
+		UpdateSelectedKeysTransform();
 		OutVectorPosition = SelectedKeysTransform.GetLocation();
 		return true;
 	}
 	return false;
 }
 
+void FMotionTraiMovieScenelKeyTool::GetSelectedKeyPositions(TArray<FVector>& OutPositions) const
+{
+	if (CachedSelection.Num() > 0)
+	{
+		const FTransform OffsetTransform = OwningTrail->GetOffsetTransform();
+		const FTransform ParentSpaceTransform = OwningTrail->GetParentSpaceTransform();
+		for (FTrailKeyInfo* KeyInfo : CachedSelection)
+		{
+			const FTransform Transform = KeyInfo->Transform * OffsetTransform * ParentSpaceTransform;
+			OutPositions.Add(Transform.GetLocation());
+		}
+	}
+}
+
 bool FMotionTraiMovieScenelKeyTool::IsSelected() const
 {
 	return (CachedSelection.Num() > 0);
@@ -162,15 +187,17 @@ bool FMotionTraiMovieScenelKeyTool::HandleClick(const FGuid& Guid, FEditorViewpo
 	return false;
 }
 
-
-void FMotionTraiMovieScenelKeyTool::UpdateSelectedKeysTransform()
+void FMotionTraiMovieScenelKeyTool::UpdateSelectedKeysTransform() const
 {
 	if (CachedSelection.Num() > 0)
 	{
+		const FTransform OffsetTransform = OwningTrail->GetOffsetTransform();
+		const FTransform ParentSpaceTransform = OwningTrail->GetParentSpaceTransform();
 		FVector NewGizmoLocation = FVector::ZeroVector;
 		for (FTrailKeyInfo* KeyInfo : CachedSelection)
 		{
-			NewGizmoLocation += KeyInfo->Transform.GetLocation();
+			const FTransform Transform = KeyInfo->Transform * OffsetTransform * ParentSpaceTransform;
+			NewGizmoLocation += Transform.GetLocation();
 		}
 		NewGizmoLocation /= (double)CachedSelection.Num();
 		SelectedKeysTransform.SetLocation(NewGizmoLocation);
@@ -564,28 +591,6 @@ bool FMotionTraiMovieScenelKeyTool::ShouldRebuildKeys()
 		{
 			return true;
 		}
-		for (int32 ChannelIdx = 0; ChannelIdx <= MaxChannel; ChannelIdx++)
-		{
-			const ETransformChannel TransformChannel = ETransformChannel(ChannelIdx);
-			if (DoubleChannels.Num() > MaxChannel)
-			{
-				if ((!TimeKeyPair.Value.Contains(TransformChannel) && Keys[TimeKeyPair.Key]->IdxMap.Contains(TransformChannel)) ||
-					(Keys[TimeKeyPair.Key]->IdxMap.Contains(TransformChannel) && DoubleChannels[int32(TransformChannel)]->GetData().GetIndex(Keys[TimeKeyPair.Key]->IdxMap[TransformChannel]) == INDEX_NONE) ||
-					(TimeKeyPair.Value.Contains(TransformChannel) && !Keys[TimeKeyPair.Key]->IdxMap.Contains(TransformChannel)))
-				{
-					return true;
-				}
-			}
-			else if (FloatChannels.Num() > MaxChannel)
-			{
-				if ((!TimeKeyPair.Value.Contains(TransformChannel) && Keys[TimeKeyPair.Key]->IdxMap.Contains(TransformChannel)) ||
-					(Keys[TimeKeyPair.Key]->IdxMap.Contains(TransformChannel) && FloatChannels[int32(TransformChannel)]->GetData().GetIndex(Keys[TimeKeyPair.Key]->IdxMap[TransformChannel]) == INDEX_NONE) ||
-					(TimeKeyPair.Value.Contains(TransformChannel) && !Keys[TimeKeyPair.Key]->IdxMap.Contains(TransformChannel)))
-				{
-					return true;
-				}
-			}
-		}
 	}
 	return false;
 }
@@ -602,8 +607,12 @@ void FMotionTraiMovieScenelKeyTool::DirtyKeyTransforms()
 		FrameKeyPair.Value->bDirty = true;
 	}
 }
+void FMotionTraiMovieScenelKeyTool::UpdateViewRange(const TRange<FFrameNumber>& InViewRange)
+{
+	CachedViewFrameRange = InViewRange;
+}
 
-void FMotionTraiMovieScenelKeyTool::UpdateKeysInRange(const TRange<double>& ViewRange)
+void FMotionTraiMovieScenelKeyTool::UpdateKeys()
 {
 	for (const TPair<FFrameNumber, TUniquePtr<FTrailKeyInfo>>& FrameKeyPair : Keys.Keys)
 	{
@@ -659,12 +668,12 @@ FTrailKeyInfo::FTrailKeyInfo(const FFrameNumber InFrameNumber, UMovieSceneSectio
 void FTrailKeyInfo::UpdateKeyTransform(EGetKeyFrom UpdateType)
 {	
 	bDirty = false;
-	 if (UpdateType == EGetKeyFrom::FromTrailCache)
+	if (UpdateType == EGetKeyFrom::FromTrailCache)
 	{
-		const double EvalTime = OwningTrail->GetSequencer()->GetFocusedTickResolution().AsSeconds(FrameNumber);
-		OwningTrail->GetTrajectoryTransforms()->GetInterp(EvalTime,Transform,ParentTransform);
+		 OwningTrail->Interp(FrameNumber, Transform, ParentTransform);
 	}
 }
 
 } // namespace MovieScene
 } // namespace UE
+
diff --git a/Engine/Plugins/Editor/SequencerAnimTools/Source/SequencerAnimTools/Private/Sequencer/MotionTrailMovieSceneKey.h b/Engine/Plugins/Editor/SequencerAnimTools/Source/SequencerAnimTools/Private/Sequencer/MotionTrailMovieSceneKey.h
--- a/Engine/Plugins/Editor/SequencerAnimTools/Source/SequencerAnimTools/Private/Sequencer/MotionTrailMovieSceneKey.h
+++ b/Engine/Plugins/Editor/SequencerAnimTools/Source/SequencerAnimTools/Private/Sequencer/MotionTrailMovieSceneKey.h
@@ -221,9 +221,11 @@ public:
 	bool HandleClick(const FGuid& Guid, FEditorViewportClient* InViewportClient, HHitProxy* HitProxy, FInputClick Click);
 	void StartDragging() {};
 	bool IsSelected(FVector& OutVectorPosition) const;
+	void GetSelectedKeyPositions(TArray<FVector>& OutPositions) const;
 	bool IsSelected() const;
 	void OnSectionChanged();
-	void UpdateKeysInRange(const TRange<double>& ViewRange);
+	void UpdateKeys();
+	void UpdateViewRange(const TRange<FFrameNumber>& InViewRange);
 	void ClearSelection();
 	FTrailKeyInfo* FindKey(const FFrameNumber& FrameNumber) const;
 	TArray<FFrameNumber> SelectedKeyTimes() const;
@@ -234,22 +236,21 @@ public:
 	void DeleteSelectedKeys();
 	bool IsAllSelected() const { return Keys.Num() == CachedSelection.Num(); }
 	void GetAllKeys(TArray<FTrailKeyInfo*>& OutKeys) {Keys.GetAllKeys(OutKeys); }
-
+	const FSortedKeys& GetSortedKeys()const { return Keys; }
 	TArray<FFrameNumber> GetTimesFromModifiedTimes(const TArray<FFrameNumber>& Frames, const FFrameNumber& LastFrame, const FFrameNumber& Step);
 protected:
-
-	void UpdateSelectedKeysTransform();
+	TRange<FFrameNumber> CachedViewFrameRange;
+	void UpdateSelectedKeysTransform() const;
 	bool ShouldRebuildKeys();
 	void BuildKeys();
 	void DirtyKeyTransforms();
 	TArray<FKeyHandle> GetSelectedKeyHandles(FMovieSceneChannel* Channel);
 
-	//TSortedMap<FFrameNumber, TUniquePtr<FTrailKeyInfo>> Keys;
 	FSortedKeys Keys;
 	TSet<FTrailKeyInfo*> CachedSelection; //current selection
 
 	FMovieSceneTransformTrail* OwningTrail;
-	FTransform SelectedKeysTransform;
+	mutable FTransform SelectedKeysTransform;
 
 	friend class FMovieSceneTransformTrail;
 	friend class FMovieSceneComponentTransformTrail;
