CIS FIX, namespace issue

diff --git a/Engine/Plugins/Animation/ControlRig/Source/ControlRigEditor/Private/Tools/EvaluateSequencerTools.cpp b/Engine/Plugins/Animation/ControlRig/Source/ControlRigEditor/Private/Tools/EvaluateSequencerTools.cpp
--- a/Engine/Plugins/Animation/ControlRig/Source/ControlRigEditor/Private/Tools/EvaluateSequencerTools.cpp
+++ b/Engine/Plugins/Animation/ControlRig/Source/ControlRigEditor/Private/Tools/EvaluateSequencerTools.cpp
@@ -274,7 +274,7 @@ namespace AIE
 		return true;
 	}
 
-	static FTransform GetControlRigComponentTransform(UControlRig* ControlRig)
+	static FTransform GetThisControlRigComponentTransform(UControlRig* ControlRig)
 	{
 		FTransform Transform = FTransform::Identity;
 		TSharedPtr<IControlRigObjectBinding> ObjectBinding = ControlRig->GetObjectBinding();
@@ -320,7 +320,7 @@ namespace AIE
 		FTransform GlobalTransform = ControlRig->GetControlGlobalTransform(ControlElement->GetKey().Name);
 
 		// switch to constraint space
-		FTransform ToWorldTransform = GetControlRigComponentTransform(ControlRig);
+		FTransform ToWorldTransform = GetThisControlRigComponentTransform(ControlRig);
 		const FTransform WorldTransform = GlobalTransform * ToWorldTransform;
 
 		const TOptional<FTransform> RelativeTransform =
