Control Rig: Fix post construction event not being emitted
#rnx

diff --git a/Engine/Plugins/Animation/ControlRig/Source/ControlRig/Private/ControlRig.cpp b/Engine/Plugins/Animation/ControlRig/Source/ControlRig/Private/ControlRig.cpp
--- a/Engine/Plugins/Animation/ControlRig/Source/ControlRig/Private/ControlRig.cpp
+++ b/Engine/Plugins/Animation/ControlRig/Source/ControlRig/Private/ControlRig.cpp
@@ -984,7 +984,6 @@ bool UControlRig::Execute(const FName& InEventName)
 	}
 
 	bool bSuccess = true;
-	bool bPostConstructionEventWasRun = false;
 
 	// we'll special case the construction event here
 	if (bIsConstructionEvent)
@@ -1073,10 +1072,9 @@ bool UControlRig::Execute(const FName& InEventName)
 					
 				} // destroy FTransientControlScope
 
-				if(!bPostConstructionEventInQueue && !bPostConstructionEventWasRun)
+				if(!bPostConstructionEventInQueue)
 				{
-					RunPostConstructionEvent();
-					bPostConstructionEventWasRun = true;
+					EmitPostConstructionEventFinished();
 				}
 
 				// Reset the input pose after construction
@@ -1093,12 +1091,6 @@ bool UControlRig::Execute(const FName& InEventName)
 			
 		} // destroy DisableSelectionNotifications
 
-		if ((bIsPostConstructionEvent || (bIsConstructionEvent && !bPostConstructionEventInQueue)) && !bPostConstructionEventWasRun)
-		{
-			RunPostConstructionEvent();
-			bPostConstructionEventWasRun = true;
-		}
-
 		if (bConstructionModeEnabled)
 		{
 #if WITH_EDITOR
@@ -1191,6 +1183,11 @@ bool UControlRig::Execute(const FName& InEventName)
 
 		bSuccess = Execute_Internal(InEventName);
 
+		if (bIsPostConstructionEvent)
+		{
+			EmitPostConstructionEventFinished();
+		}
+
 #if WITH_EDITOR
 		if (bEnableAnimAttributeTrace && ExternalAnimAttributeContainer != nullptr)
 		{
@@ -1649,8 +1646,8 @@ void UControlRig::RequestInit()
 
 void UControlRig::RequestConstruction()
 {
-	RequestRunOnceEvent(FRigUnit_PrepareForExecution::EventName, 0);
-	RequestRunOnceEvent(FRigUnit_PostPrepareForExecution::EventName, 1);
+	RequestRunOnceEvent(FRigUnit_PrepareForExecution::EventName);
+	RequestRunOnceEvent(FRigUnit_PostPrepareForExecution::EventName);
 }
 
 bool UControlRig::IsConstructionRequired() const
@@ -2455,7 +2452,7 @@ bool UControlRig::IsControlSelected(const FName& InControlName)const
 	return false;
 }
 
-void UControlRig::RunPostConstructionEvent()
+void UControlRig::EmitPostConstructionEventFinished()
 {
 	if (PostConstructionEvent.IsBound())
 	{
diff --git a/Engine/Plugins/Animation/ControlRig/Source/ControlRig/Private/ModularRig.cpp b/Engine/Plugins/Animation/ControlRig/Source/ControlRig/Private/ModularRig.cpp
--- a/Engine/Plugins/Animation/ControlRig/Source/ControlRig/Private/ModularRig.cpp
+++ b/Engine/Plugins/Animation/ControlRig/Source/ControlRig/Private/ModularRig.cpp
@@ -726,10 +726,10 @@ TArray<FString> UModularRig::GetModulePaths() const
 	return Paths;
 }
 
-void UModularRig::RunPostConstructionEvent()
+void UModularRig::EmitPostConstructionEventFinished()
 {
 	RecomputeShortestDisplayPathCache();
-	Super::RunPostConstructionEvent();
+	Super::EmitPostConstructionEventFinished();
 }
 
 FRigModuleInstance* UModularRig::AddModuleInstance(const FName& InModuleName, TSubclassOf<UControlRig> InModuleClass, const FRigModuleInstance* InParent,
diff --git a/Engine/Plugins/Animation/ControlRig/Source/ControlRig/Public/ControlRig.h b/Engine/Plugins/Animation/ControlRig/Source/ControlRig/Public/ControlRig.h
--- a/Engine/Plugins/Animation/ControlRig/Source/ControlRig/Public/ControlRig.h
+++ b/Engine/Plugins/Animation/ControlRig/Source/ControlRig/Public/ControlRig.h
@@ -509,7 +509,7 @@ private:
 
 protected:
 	
-	virtual void RunPostConstructionEvent();
+	virtual void EmitPostConstructionEventFinished();
 
 private:
 #if WITH_EDITOR
diff --git a/Engine/Plugins/Animation/ControlRig/Source/ControlRig/Public/ModularRig.h b/Engine/Plugins/Animation/ControlRig/Source/ControlRig/Public/ModularRig.h
--- a/Engine/Plugins/Animation/ControlRig/Source/ControlRig/Public/ModularRig.h
+++ b/Engine/Plugins/Animation/ControlRig/Source/ControlRig/Public/ModularRig.h
@@ -266,7 +266,7 @@ public:
 
 protected:
 
-	virtual void RunPostConstructionEvent() override;
+	virtual void EmitPostConstructionEventFinished() override;
 
 private:
 
diff --git a/Engine/Plugins/Runtime/RigVM/Source/RigVM/Private/RigVMHost.cpp b/Engine/Plugins/Runtime/RigVM/Source/RigVM/Private/RigVMHost.cpp
--- a/Engine/Plugins/Runtime/RigVM/Source/RigVM/Private/RigVMHost.cpp
+++ b/Engine/Plugins/Runtime/RigVM/Source/RigVM/Private/RigVMHost.cpp
@@ -347,14 +347,8 @@ void URigVMHost::Evaluate_AnyThread()
 
 	AdaptEventQueueForEvaluate(EventQueueToRun);
 	
-	// insert the events queued to run once
-	for(const TPair<FName,int32>& Pair : EventsToRunOnce)
-	{
-		if(!EventQueueToRun.Contains(Pair.Key))
-		{
-			EventQueueToRun.Insert(Pair.Key, Pair.Value);
-		}
-	}
+	// insert the events queued to run once at the beginning of the queue (in the order specified)
+	EventQueueToRun.Insert(EventsToRunOnce, 0);
 
 #if WITH_EDITOR
 	FName FirstEvent = NAME_None;
@@ -796,7 +790,7 @@ void URigVMHost::RequestInit()
 
 void URigVMHost::RequestRunOnceEvent(const FName& InEventName, int32 InEventIndex)
 {
-	EventsToRunOnce.FindOrAdd(InEventName) = InEventIndex;
+	EventsToRunOnce.AddUnique(InEventName);
 }
 
 bool URigVMHost::RemoveRunOnceEvent(const FName& InEventName)
diff --git a/Engine/Plugins/Runtime/RigVM/Source/RigVM/Public/RigVMHost.h b/Engine/Plugins/Runtime/RigVM/Source/RigVM/Public/RigVMHost.h
--- a/Engine/Plugins/Runtime/RigVM/Source/RigVM/Public/RigVMHost.h
+++ b/Engine/Plugins/Runtime/RigVM/Source/RigVM/Public/RigVMHost.h
@@ -205,7 +205,10 @@ public:
 		return bRequiresInitExecution;
 	} 
 
-	/** Requests to run an event once */
+	/** Requests to run an event once 
+	 * @param InEventName The event to run
+	 * @param InEventIndex Deprecated argument. Not used.
+	 */
 	UFUNCTION(BlueprintCallable, Category = "RigVM")
 	void RequestRunOnceEvent(const FName& InEventName, int32 InEventIndex = -1);
 
@@ -400,7 +403,7 @@ public:
 	UPROPERTY(transient)
 	TArray<FName> EventQueue;
 	TArray<FName> EventQueueToRun;
-	TMap<FName, int32> EventsToRunOnce;
+	TArray<FName> EventsToRunOnce;
 
 	/** Returns true if Evaluate_AnyThread is currently executing */
 	bool IsEvaluating() const { return !EventQueueToRun.IsEmpty(); }
