RigVM: Rename empty execute context to execute pin

Full cook preflight is here https://horde.devtools.epicgames.com/job/67b33c95941060c63132576c

#rnx

diff --git a/Engine/Plugins/Animation/ControlRig/Source/ControlRig/Private/Units/Deprecated/RigUnit_ApplyFK.cpp b/Engine/Plugins/Animation/ControlRig/Source/ControlRig/Private/Units/Deprecated/RigUnit_ApplyFK.cpp
--- a/Engine/Plugins/Animation/ControlRig/Source/ControlRig/Private/Units/Deprecated/RigUnit_ApplyFK.cpp
+++ b/Engine/Plugins/Animation/ControlRig/Source/ControlRig/Private/Units/Deprecated/RigUnit_ApplyFK.cpp
@@ -63,7 +63,6 @@ IMPLEMENT_RIGUNIT_AUTOMATION_TEST(FRigUnit_ApplyFK)
 	const FRigElementKey Root = Controller->AddBone(TEXT("Root"), FRigElementKey(), FTransform(FVector(1.f, 0.f, 0.f)), true, ERigBoneType::User);
 	const FRigElementKey BoneA = Controller->AddBone(TEXT("BoneA"), Root, FTransform(FVector(1.f, 2.f, 3.f)), true, ERigBoneType::User);
 
-	Unit.ExecuteContext = ExecuteContext;
 	Unit.Joint = TEXT("BoneA");
 	Unit.ApplyTransformMode = EApplyTransformMode::Override;
 	Unit.ApplyTransformSpace = ETransformSpaceMode::GlobalSpace;
diff --git a/Engine/Plugins/Animation/ControlRig/Source/ControlRig/Private/Units/Execution/RigUnit_Collection.cpp b/Engine/Plugins/Animation/ControlRig/Source/ControlRig/Private/Units/Execution/RigUnit_Collection.cpp
--- a/Engine/Plugins/Animation/ControlRig/Source/ControlRig/Private/Units/Execution/RigUnit_Collection.cpp
+++ b/Engine/Plugins/Animation/ControlRig/Source/ControlRig/Private/Units/Execution/RigUnit_Collection.cpp
@@ -515,9 +515,9 @@ FRigUnit_CollectionLoop_Execute()
 	{
 	    Count = Collection.Num();
 		Index = 0;
-		BlockToRun = ExecuteContextName;
+		BlockToRun = ExecutePinName;
 	}
-	else if(BlockToRun == ExecuteContextName)
+	else if(BlockToRun == ExecutePinName)
 	{
 		Index++;
 	}
diff --git a/Engine/Plugins/Animation/ControlRig/Source/ControlRig/Private/Units/Execution/RigUnit_Hierarchy.cpp b/Engine/Plugins/Animation/ControlRig/Source/ControlRig/Private/Units/Execution/RigUnit_Hierarchy.cpp
--- a/Engine/Plugins/Animation/ControlRig/Source/ControlRig/Private/Units/Execution/RigUnit_Hierarchy.cpp
+++ b/Engine/Plugins/Animation/ControlRig/Source/ControlRig/Private/Units/Execution/RigUnit_Hierarchy.cpp
@@ -684,9 +684,9 @@ FRigUnit_PoseLoop_Execute()
 	{
 		Count = Pose.Num();
 		Index = 0;
-		BlockToRun = ExecuteContextName;
+		BlockToRun = ExecutePinName;
 	}
-	else if(BlockToRun == ExecuteContextName)
+	else if(BlockToRun == ExecutePinName)
 	{
 		Index++;
 	}
diff --git a/Engine/Plugins/Animation/ControlRig/Source/ControlRig/Public/Units/Debug/RigUnit_DebugHierarchy.h b/Engine/Plugins/Animation/ControlRig/Source/ControlRig/Public/Units/Debug/RigUnit_DebugHierarchy.h
--- a/Engine/Plugins/Animation/ControlRig/Source/ControlRig/Public/Units/Debug/RigUnit_DebugHierarchy.h
+++ b/Engine/Plugins/Animation/ControlRig/Source/ControlRig/Public/Units/Debug/RigUnit_DebugHierarchy.h
@@ -23,7 +23,7 @@ namespace EControlRigDrawHierarchyMode
 /**
  * Draws vectors on each bone in the viewport across the entire hierarchy
  */
-USTRUCT(meta=(DisplayName="Draw Hierarchy"))
+USTRUCT(meta=(DisplayName="Draw Hierarchy", ExecuteContext="FControlRigExecuteContext"))
 struct CONTROLRIG_API FRigUnit_DebugHierarchy : public FRigVMFunction_DebugBase
 {
 	GENERATED_BODY()
@@ -41,7 +41,7 @@ struct CONTROLRIG_API FRigUnit_DebugHierarchy : public FRigVMFunction_DebugBase
 	virtual void Execute() override;
 
 	UPROPERTY(DisplayName = "Execute", Transient, meta = (Input, Output))
-	FControlRigExecuteContext ExecuteContext;
+	FRigVMExecutePin ExecutePin;
 
 	// the items to draw the pose for.
 	// if this is empty we'll draw the whole hierarchy
@@ -72,7 +72,7 @@ struct CONTROLRIG_API FRigUnit_DebugHierarchy : public FRigVMFunction_DebugBase
 /**
 * Draws vectors on each bone in the viewport across the entire pose
 */
-USTRUCT(meta=(DisplayName="Draw Pose Cache"))
+USTRUCT(meta=(DisplayName="Draw Pose Cache", ExecuteContext="FControlRigExecuteContext"))
 struct CONTROLRIG_API FRigUnit_DebugPose : public FRigVMFunction_DebugBase
 {
 	GENERATED_BODY()
@@ -90,7 +90,7 @@ struct CONTROLRIG_API FRigUnit_DebugPose : public FRigVMFunction_DebugBase
 	virtual void Execute() override;
 
 	UPROPERTY(DisplayName = "Execute", Transient, meta = (Input, Output))
-	FControlRigExecuteContext ExecuteContext;
+	FRigVMExecutePin ExecutePin;
 
 	UPROPERTY(meta = (Input))
 	FRigPose Pose;
diff --git a/Engine/Plugins/Animation/ControlRig/Source/ControlRig/Public/Units/Execution/RigUnit_BeginExecution.h b/Engine/Plugins/Animation/ControlRig/Source/ControlRig/Public/Units/Execution/RigUnit_BeginExecution.h
--- a/Engine/Plugins/Animation/ControlRig/Source/ControlRig/Public/Units/Execution/RigUnit_BeginExecution.h
+++ b/Engine/Plugins/Animation/ControlRig/Source/ControlRig/Public/Units/Execution/RigUnit_BeginExecution.h
@@ -22,7 +22,7 @@ struct CONTROLRIG_API FRigUnit_BeginExecution : public FRigUnit
 
 	// The execution result
 	UPROPERTY(EditAnywhere, Transient, DisplayName = "Execute", Category = "BeginExecution", meta = (Output))
-	FControlRigExecuteContext ExecuteContext;
+	FRigVMExecutePin ExecutePin;
 
 	static inline const FLazyName EventName = FLazyName(TEXT("Forwards Solve"));
 };
@@ -43,7 +43,7 @@ struct CONTROLRIG_API FRigUnit_PreBeginExecution : public FRigUnit
 
 	// The execution result
 	UPROPERTY(EditAnywhere, Transient, DisplayName = "Execute", Category = "BeginExecution", meta = (Output))
-	FControlRigExecuteContext ExecuteContext;
+	FRigVMExecutePin ExecutePin;
 
 	static inline const FLazyName EventName = FLazyName(TEXT("Pre Forwards Solve"));
 };
@@ -64,7 +64,7 @@ struct CONTROLRIG_API FRigUnit_PostBeginExecution : public FRigUnit
 
 	// The execution result
 	UPROPERTY(EditAnywhere, Transient, DisplayName = "Execute", Category = "BeginExecution", meta = (Output))
-	FControlRigExecuteContext ExecuteContext;
+	FRigVMExecutePin ExecutePin;
 
 	static inline const FLazyName EventName = FLazyName(TEXT("Post Forwards Solve"));
 };
diff --git a/Engine/Plugins/Animation/ControlRig/Source/ControlRig/Public/Units/Execution/RigUnit_Collection.h b/Engine/Plugins/Animation/ControlRig/Source/ControlRig/Public/Units/Execution/RigUnit_Collection.h
--- a/Engine/Plugins/Animation/ControlRig/Source/ControlRig/Public/Units/Execution/RigUnit_Collection.h
+++ b/Engine/Plugins/Animation/ControlRig/Source/ControlRig/Public/Units/Execution/RigUnit_Collection.h
@@ -621,10 +621,10 @@ struct CONTROLRIG_API FRigUnit_CollectionLoop : public FRigUnit_CollectionBaseMu
 	// FRigVMStruct overrides
 	virtual const TArray<FName>& GetControlFlowBlocks_Impl() const override
 	{
-		static const TArray<FName> Blocks = {ExecuteContextName, ForLoopCompletedPinName};
+		static const TArray<FName> Blocks = {ExecutePinName, ForLoopCompletedPinName};
 		return Blocks;
 	}
-	virtual const bool IsControlFlowBlockSliced(const FName& InBlockName) const { return InBlockName == ExecuteContextName; }
+	virtual const bool IsControlFlowBlockSliced(const FName& InBlockName) const { return InBlockName == ExecutePinName; }
 	virtual int32 GetNumSlices() const override { return Count; }
 
 	RIGVM_METHOD()
diff --git a/Engine/Plugins/Animation/ControlRig/Source/ControlRig/Public/Units/Execution/RigUnit_Hierarchy.h b/Engine/Plugins/Animation/ControlRig/Source/ControlRig/Public/Units/Execution/RigUnit_Hierarchy.h
--- a/Engine/Plugins/Animation/ControlRig/Source/ControlRig/Public/Units/Execution/RigUnit_Hierarchy.h
+++ b/Engine/Plugins/Animation/ControlRig/Source/ControlRig/Public/Units/Execution/RigUnit_Hierarchy.h
@@ -780,10 +780,10 @@ struct CONTROLRIG_API FRigUnit_PoseLoop : public FRigUnit_HierarchyBaseMutable
 	// FRigVMStruct overrides
 	virtual const TArray<FName>& GetControlFlowBlocks_Impl() const override
 	{
-		static const TArray<FName> Blocks = {ExecuteContextName, ForLoopCompletedPinName};
+		static const TArray<FName> Blocks = {ExecutePinName, ForLoopCompletedPinName};
 		return Blocks;
 	}
-	virtual const bool IsControlFlowBlockSliced(const FName& InBlockName) const { return InBlockName == ExecuteContextName; }
+	virtual const bool IsControlFlowBlockSliced(const FName& InBlockName) const { return InBlockName == ExecutePinName; }
 	virtual int32 GetNumSlices() const override { return Count; }
 
 	RIGVM_METHOD()
diff --git a/Engine/Plugins/Animation/ControlRig/Source/ControlRig/Public/Units/Execution/RigUnit_InteractionExecution.h b/Engine/Plugins/Animation/ControlRig/Source/ControlRig/Public/Units/Execution/RigUnit_InteractionExecution.h
--- a/Engine/Plugins/Animation/ControlRig/Source/ControlRig/Public/Units/Execution/RigUnit_InteractionExecution.h
+++ b/Engine/Plugins/Animation/ControlRig/Source/ControlRig/Public/Units/Execution/RigUnit_InteractionExecution.h
@@ -21,7 +21,7 @@ struct CONTROLRIG_API FRigUnit_InteractionExecution : public FRigUnit
 
 	// The execution result
 	UPROPERTY(EditAnywhere, Transient, DisplayName = "Execute", Category = "BeginExecution", meta = (Output))
-	FControlRigExecuteContext ExecuteContext;
+	FRigVMExecutePin ExecutePin;
 
 	static inline const FLazyName EventName = FLazyName(TEXT("Interaction"));
 };
diff --git a/Engine/Plugins/Animation/ControlRig/Source/ControlRig/Public/Units/Execution/RigUnit_InverseExecution.h b/Engine/Plugins/Animation/ControlRig/Source/ControlRig/Public/Units/Execution/RigUnit_InverseExecution.h
--- a/Engine/Plugins/Animation/ControlRig/Source/ControlRig/Public/Units/Execution/RigUnit_InverseExecution.h
+++ b/Engine/Plugins/Animation/ControlRig/Source/ControlRig/Public/Units/Execution/RigUnit_InverseExecution.h
@@ -22,7 +22,7 @@ struct CONTROLRIG_API FRigUnit_InverseExecution : public FRigUnit
 
 	// The execution result
 	UPROPERTY(EditAnywhere, Transient, DisplayName = "Execute", Category = "InverseExecution", meta = (Output))
-	FControlRigExecuteContext ExecuteContext;
+	FRigVMExecutePin ExecutePin;
 
 	static inline const FLazyName EventName = FLazyName(TEXT("Backwards Solve"));
 };
diff --git a/Engine/Plugins/Animation/ControlRig/Source/ControlRig/Public/Units/Execution/RigUnit_PrepareForExecution.h b/Engine/Plugins/Animation/ControlRig/Source/ControlRig/Public/Units/Execution/RigUnit_PrepareForExecution.h
--- a/Engine/Plugins/Animation/ControlRig/Source/ControlRig/Public/Units/Execution/RigUnit_PrepareForExecution.h
+++ b/Engine/Plugins/Animation/ControlRig/Source/ControlRig/Public/Units/Execution/RigUnit_PrepareForExecution.h
@@ -22,7 +22,7 @@ struct CONTROLRIG_API FRigUnit_PrepareForExecution : public FRigUnit
 
 	// The execution result
 	UPROPERTY(EditAnywhere, Transient, DisplayName = "Execute", Category = "PrepareForExecution", meta = (Output))
-	FControlRigExecuteContext ExecuteContext;
+	FRigVMExecutePin ExecutePin;
 
 	static inline const FLazyName EventName = FLazyName(TEXT("Construction"));
 };
@@ -43,7 +43,7 @@ struct CONTROLRIG_API FRigUnit_PostPrepareForExecution : public FRigUnit
 
 	// The execution result
 	UPROPERTY(EditAnywhere, Transient, DisplayName = "Execute", Category = "PostPrepareForExecution", meta = (Output))
-	FControlRigExecuteContext ExecuteContext;
+	FRigVMExecutePin ExecutePin;
 
 	static inline const FLazyName EventName = FLazyName(TEXT("PostConstruction"));
 };
diff --git a/Engine/Plugins/Animation/ControlRig/Source/ControlRig/Public/Units/Hierarchy/RigUnit_ControlChannel.h b/Engine/Plugins/Animation/ControlRig/Source/ControlRig/Public/Units/Hierarchy/RigUnit_ControlChannel.h
--- a/Engine/Plugins/Animation/ControlRig/Source/ControlRig/Public/Units/Hierarchy/RigUnit_ControlChannel.h
+++ b/Engine/Plugins/Animation/ControlRig/Source/ControlRig/Public/Units/Hierarchy/RigUnit_ControlChannel.h
@@ -211,7 +211,7 @@ struct CONTROLRIG_API FRigUnit_SetAnimationChannelBase : public FRigUnit_GetAnim
 	}
 
 	UPROPERTY(DisplayName = "Execute", Transient, meta = (Input, Output))
-	FControlRigExecuteContext ExecuteContext;
+	FRigVMExecutePin ExecutePin;
 };
 
 /**
diff --git a/Engine/Plugins/Animation/ControlRig/Source/ControlRig/Public/Units/Hierarchy/RigUnit_ControlChannelFromItem.h b/Engine/Plugins/Animation/ControlRig/Source/ControlRig/Public/Units/Hierarchy/RigUnit_ControlChannelFromItem.h
--- a/Engine/Plugins/Animation/ControlRig/Source/ControlRig/Public/Units/Hierarchy/RigUnit_ControlChannelFromItem.h
+++ b/Engine/Plugins/Animation/ControlRig/Source/ControlRig/Public/Units/Hierarchy/RigUnit_ControlChannelFromItem.h
@@ -193,7 +193,7 @@ struct CONTROLRIG_API FRigUnit_SetAnimationChannelBaseFromItem : public FRigUnit
 	}
 
 	UPROPERTY(DisplayName = "Execute", Transient, meta = (Input, Output))
-	FControlRigExecuteContext ExecuteContext;
+	FRigVMExecutePin ExecutePin;
 };
 
 /**
diff --git a/Engine/Plugins/Animation/ControlRig/Source/ControlRig/Public/Units/RigUnit.h b/Engine/Plugins/Animation/ControlRig/Source/ControlRig/Public/Units/RigUnit.h
--- a/Engine/Plugins/Animation/ControlRig/Source/ControlRig/Public/Units/RigUnit.h
+++ b/Engine/Plugins/Animation/ControlRig/Source/ControlRig/Public/Units/RigUnit.h
@@ -119,7 +119,7 @@ struct CONTROLRIG_API FRigUnitMutable : public FRigUnit
 	 * This property is used to chain multiple mutable units together
 	 */
 	UPROPERTY(DisplayName = "Execute", Transient, meta = (Input, Output))
-	FControlRigExecuteContext ExecuteContext;
+	FRigVMExecutePin ExecutePin;
 };
 
 #if WITH_EDITOR
diff --git a/Engine/Plugins/Animation/ControlRig/Source/ControlRigDeveloper/Private/ControlRigBlueprint.cpp b/Engine/Plugins/Animation/ControlRig/Source/ControlRigDeveloper/Private/ControlRigBlueprint.cpp
--- a/Engine/Plugins/Animation/ControlRig/Source/ControlRigDeveloper/Private/ControlRigBlueprint.cpp
+++ b/Engine/Plugins/Animation/ControlRig/Source/ControlRigDeveloper/Private/ControlRigBlueprint.cpp
@@ -393,7 +393,7 @@ TArray<URigVMNode*> UControlRigBlueprint::ConvertHierarchyElementsToSpawnerNodes
 		NodePosition += NodePositionIncrement;
 	}
 
-	const URigVMPin* LastPin = EventNode->FindPin(GET_MEMBER_NAME_STRING_CHECKED(FRigUnit_PrepareForExecution, ExecuteContext));
+	const URigVMPin* LastPin = EventNode->FindExecutePin();
 	if(LastPin)
 	{
 		// follow the node's execution links to find the last one
@@ -465,7 +465,7 @@ TArray<URigVMNode*> UControlRigBlueprint::ConvertHierarchyElementsToSpawnerNodes
 
 			if(LastPin)
 			{
-				if(const URigVMPin* NextPin = AddBoneNode->FindPin(GET_MEMBER_NAME_STRING_CHECKED(FRigUnit_HierarchyAddBone, ExecuteContext)))
+				if(const URigVMPin* NextPin = AddBoneNode->FindExecutePin())
 				{
 					GraphController->AddLink(LastPin->GetPinPath(), NextPin->GetPinPath(), true);
 					LastPin = NextPin;
@@ -489,7 +489,7 @@ TArray<URigVMNode*> UControlRigBlueprint::ConvertHierarchyElementsToSpawnerNodes
 
 			if(LastPin)
 			{
-				if(const URigVMPin* NextPin = AddNullNode->FindPin(GET_MEMBER_NAME_STRING_CHECKED(FRigUnit_HierarchyAddBone, ExecuteContext)))
+				if(const URigVMPin* NextPin = AddNullNode->FindExecutePin())
 				{
 					GraphController->AddLink(LastPin->GetPinPath(), NextPin->GetPinPath(), true);
 					LastPin = NextPin;
@@ -664,7 +664,7 @@ TArray<URigVMNode*> UControlRigBlueprint::ConvertHierarchyElementsToSpawnerNodes
 
 				if(LastPin)
 				{
-					if(const URigVMPin* NextPin = AddControlNode->FindPin(GET_MEMBER_NAME_STRING_CHECKED(FRigUnit_HierarchyAddBone, ExecuteContext)))
+					if(const URigVMPin* NextPin = AddControlNode->FindExecutePin())
 					{
 						GraphController->AddLink(LastPin->GetPinPath(), NextPin->GetPinPath(), true);
 						LastPin = NextPin;
@@ -754,7 +754,7 @@ TArray<URigVMNode*> UControlRigBlueprint::ConvertHierarchyElementsToSpawnerNodes
 
 				if(LastPin)
 				{
-					if(const URigVMPin* NextPin = AddControlNode->FindPin(GET_MEMBER_NAME_STRING_CHECKED(FRigUnit_HierarchyAddBone, ExecuteContext)))
+					if(const URigVMPin* NextPin = AddControlNode->FindExecutePin())
 					{
 						GraphController->AddLink(LastPin->GetPinPath(), NextPin->GetPinPath(), true);
 						LastPin = NextPin;
@@ -793,7 +793,7 @@ TArray<URigVMNode*> UControlRigBlueprint::ConvertHierarchyElementsToSpawnerNodes
 
 			if(LastPin)
 			{
-				if(const URigVMPin* NextPin = AddSocketNode->FindPin(GET_MEMBER_NAME_STRING_CHECKED(FRigUnit_HierarchyAddBone, ExecuteContext)))
+				if(const URigVMPin* NextPin = AddSocketNode->FindExecutePin())
 				{
 					GraphController->AddLink(LastPin->GetPinPath(), NextPin->GetPinPath(), true);
 					LastPin = NextPin;
diff --git a/Engine/Plugins/Runtime/RigVM/Source/RigVM/Private/RigVMCore/RigVMGraphFunctionDefinition.cpp b/Engine/Plugins/Runtime/RigVM/Source/RigVM/Private/RigVMCore/RigVMGraphFunctionDefinition.cpp
--- a/Engine/Plugins/Runtime/RigVM/Source/RigVM/Private/RigVMCore/RigVMGraphFunctionDefinition.cpp
+++ b/Engine/Plugins/Runtime/RigVM/Source/RigVM/Private/RigVMCore/RigVMGraphFunctionDefinition.cpp
@@ -55,7 +55,7 @@ bool FRigVMGraphFunctionArgument::IsExecuteContext() const
 	{
 		if(const UScriptStruct* ScriptStruct = Cast<UScriptStruct>(CPPTypeObject.Get()))
 		{
-			if(ScriptStruct->IsChildOf(FRigVMEmptyExecuteContext::StaticStruct()))
+			if(ScriptStruct->IsChildOf(FRigVMExecutePin::StaticStruct()))
 			{
 				return true;
 			}
@@ -217,7 +217,7 @@ bool FRigVMGraphFunctionHeader::IsMutable() const
 		{
 			if(UScriptStruct* Struct = Cast<UScriptStruct>(Arg.CPPTypeObject.Get()))
 			{
-				if(Struct->IsChildOf(FRigVMEmptyExecuteContext::StaticStruct()))
+				if(Struct->IsChildOf(FRigVMExecutePin::StaticStruct()))
 				{
 					return true;
 				}
diff --git a/Engine/Plugins/Runtime/RigVM/Source/RigVM/Private/RigVMCore/RigVMRegistry.cpp b/Engine/Plugins/Runtime/RigVM/Source/RigVM/Private/RigVMCore/RigVMRegistry.cpp
--- a/Engine/Plugins/Runtime/RigVM/Source/RigVM/Private/RigVMCore/RigVMRegistry.cpp
+++ b/Engine/Plugins/Runtime/RigVM/Source/RigVM/Private/RigVMCore/RigVMRegistry.cpp
@@ -385,7 +385,7 @@ void FRigVMRegistry_NoLock::RefreshEngineTypes_NoLock()
 		UScriptStruct* ScriptStruct = *ScriptStructIt;
 		
 		// if this is a C++ type - skip it
-		if(ScriptStruct->IsA<UUserDefinedStruct>() || ScriptStruct->IsChildOf(FRigVMEmptyExecuteContext::StaticStruct()))
+		if(ScriptStruct->IsA<UUserDefinedStruct>() || ScriptStruct->IsChildOf(FRigVMExecutePin::StaticStruct()))
 		{
 			// this check for example makes sure we don't add structs defined in verse
 			if(IsAllowedType_NoLock(ScriptStruct))
@@ -615,7 +615,7 @@ TRigVMTypeIndex FRigVMRegistry_NoLock::FindOrAddType_NoLock(const FRigVMTemplate
 		bool bIsExecute = false;
 		if(const UScriptStruct* ScriptStruct = Cast<UScriptStruct>(CPPTypeObject))
 		{
-			bIsExecute = ScriptStruct->IsChildOf(FRigVMEmptyExecuteContext::StaticStruct());
+			bIsExecute = ScriptStruct->IsChildOf(FRigVMExecutePin::StaticStruct());
 		}
 
 		TArray<TRigVMTypeIndex> Indices;
@@ -751,7 +751,7 @@ TRigVMTypeIndex FRigVMRegistry_NoLock::FindOrAddType_NoLock(const FRigVMTemplate
 			}
 			else if(const UStruct* Struct = Cast<UStruct>(CPPTypeObject))
 			{
-				if(Struct->IsChildOf(FRigVMEmptyExecuteContext::StaticStruct()))
+				if(Struct->IsChildOf(FRigVMExecutePin::StaticStruct()))
 				{
 					if(ArrayDimension == 0)
 					{
diff --git a/Engine/Plugins/Runtime/RigVM/Source/RigVM/Private/RigVMCore/RigVMStruct.cpp b/Engine/Plugins/Runtime/RigVM/Source/RigVM/Private/RigVMCore/RigVMStruct.cpp
--- a/Engine/Plugins/Runtime/RigVM/Source/RigVM/Private/RigVMCore/RigVMStruct.cpp
+++ b/Engine/Plugins/Runtime/RigVM/Source/RigVM/Private/RigVMCore/RigVMStruct.cpp
@@ -454,11 +454,11 @@ bool FRigVMStruct::CheckPinType(UScriptStruct* InStruct, const FName& PinName, c
 		if (Property->GetCPPType() != ExpectedType)
 		{
 			if(ExpectedType == FRigVMExecuteContext::StaticStruct()->GetStructCPPName() ||
-				ExpectedType == FRigVMEmptyExecuteContext::StaticStruct()->GetStructCPPName())
+				ExpectedType == FRigVMExecutePin::StaticStruct()->GetStructCPPName())
 			{
 				if(const FStructProperty* StructProperty = CastField<FStructProperty>(Property))
 				{
-					if(StructProperty->Struct->IsChildOf(FRigVMEmptyExecuteContext::StaticStruct()))
+					if(StructProperty->Struct->IsChildOf(FRigVMExecutePin::StaticStruct()))
 					{
 						return true;
 					}
diff --git a/Engine/Plugins/Runtime/RigVM/Source/RigVM/Private/RigVMFunctions/Execution/RigVMFunction_ForLoop.cpp b/Engine/Plugins/Runtime/RigVM/Source/RigVM/Private/RigVMFunctions/Execution/RigVMFunction_ForLoop.cpp
--- a/Engine/Plugins/Runtime/RigVM/Source/RigVM/Private/RigVMFunctions/Execution/RigVMFunction_ForLoop.cpp
+++ b/Engine/Plugins/Runtime/RigVM/Source/RigVM/Private/RigVMFunctions/Execution/RigVMFunction_ForLoop.cpp
@@ -9,9 +9,9 @@ FRigVMFunction_ForLoopCount_Execute()
 	if(BlockToRun.IsNone())
 	{
 		Index = 0;
-		BlockToRun = ExecuteContextName;
+		BlockToRun = ExecutePinName;
 	}
-	else if(BlockToRun == ExecuteContextName)
+	else if(BlockToRun == ExecutePinName)
 	{
 		Index++;
 	}
diff --git a/Engine/Plugins/Runtime/RigVM/Source/RigVM/Public/RigVMCore/RigVMExecuteContext.h b/Engine/Plugins/Runtime/RigVM/Source/RigVM/Public/RigVMCore/RigVMExecuteContext.h
--- a/Engine/Plugins/Runtime/RigVM/Source/RigVM/Public/RigVMCore/RigVMExecuteContext.h
+++ b/Engine/Plugins/Runtime/RigVM/Source/RigVM/Public/RigVMCore/RigVMExecuteContext.h
@@ -253,9 +253,12 @@ struct FRigVMExecuteContextFixLayout
 // This structure is used for semantically describing a pin
 // on a RigVMStruct - but without any content. So rather than
 // adding an actual FRigVMExecuteContext this can be added as
-// the type for a pin and it won't use any memory. 
+// the type for a pin and it won't use any memory.
+// To customize which FRigVMExecuteContext structure to use
+// on your node you can add the ExecuteContext=FMyExecuteContext
+// meta data tag on the USTRUCT macro.
 USTRUCT(BlueprintType)
-struct RIGVM_API FRigVMEmptyExecuteContext
+struct RIGVM_API FRigVMExecutePin
 #if CPP
 	: public FRigVMExecuteContextFixLayout
 #endif
@@ -268,7 +271,7 @@ struct RIGVM_API FRigVMEmptyExecuteContext
  * indicate execution order.
  */
 USTRUCT(BlueprintType, meta=(DisplayName="Execute Context"))
-struct FRigVMExecuteContext : public FRigVMEmptyExecuteContext
+struct FRigVMExecuteContext : public FRigVMExecutePin
 {
 	GENERATED_BODY()
 
diff --git a/Engine/Plugins/Runtime/RigVM/Source/RigVM/Public/RigVMCore/RigVMStruct.h b/Engine/Plugins/Runtime/RigVM/Source/RigVM/Public/RigVMCore/RigVMStruct.h
--- a/Engine/Plugins/Runtime/RigVM/Source/RigVM/Public/RigVMCore/RigVMStruct.h
+++ b/Engine/Plugins/Runtime/RigVM/Source/RigVM/Public/RigVMCore/RigVMStruct.h
@@ -286,6 +286,7 @@ public:
 	static inline const FLazyName SliceContextMetaName = FLazyName(TEXT("SliceContext"));
 	static inline const FLazyName ExecuteName = FLazyName(TEXT("Execute"));
 	static inline const FLazyName ExecuteContextName = FLazyName(TEXT("ExecuteContext"));
+	static inline const FLazyName ExecutePinName = FLazyName(TEXT("ExecutePin"));
 	static inline const FLazyName ForLoopCountPinName = FLazyName(TEXT("Count"));
 	static inline const FLazyName ForLoopContinuePinName = FLazyName(TEXT("Continue"));
 	static inline const FLazyName ForLoopCompletedPinName = FLazyName(TEXT("Completed"));
@@ -331,5 +332,5 @@ public:
 	 * This property is used to chain multiple mutable nodes together
 	 */
 	UPROPERTY(DisplayName = "Execute", Transient, meta = (Input, Output))
-	FRigVMExecuteContext ExecuteContext;
+	FRigVMExecutePin ExecutePin;
 };
diff --git a/Engine/Plugins/Runtime/RigVM/Source/RigVM/Public/RigVMFunctions/Execution/RigVMFunction_ForLoop.h b/Engine/Plugins/Runtime/RigVM/Source/RigVM/Public/RigVMFunctions/Execution/RigVMFunction_ForLoop.h
--- a/Engine/Plugins/Runtime/RigVM/Source/RigVM/Public/RigVMFunctions/Execution/RigVMFunction_ForLoop.h
+++ b/Engine/Plugins/Runtime/RigVM/Source/RigVM/Public/RigVMFunctions/Execution/RigVMFunction_ForLoop.h
@@ -24,10 +24,10 @@ struct RIGVM_API FRigVMFunction_ForLoopCount : public FRigVMStructMutable
 	// FRigVMStruct overrides
 	virtual const TArray<FName>& GetControlFlowBlocks_Impl() const override
 	{
-		static const TArray<FName> Blocks = {ExecuteContextName, ForLoopCompletedPinName};
+		static const TArray<FName> Blocks = {ExecutePinName, ForLoopCompletedPinName};
 		return Blocks;
 	}
-	virtual const bool IsControlFlowBlockSliced(const FName& InBlockName) const { return InBlockName == ExecuteContextName; }
+	virtual const bool IsControlFlowBlockSliced(const FName& InBlockName) const { return InBlockName == ExecutePinName; }
 	virtual int32 GetNumSlices() const override { return Count; }
 
 	RIGVM_METHOD()
diff --git a/Engine/Plugins/Runtime/RigVM/Source/RigVM/Public/RigVMFunctions/Execution/RigVMFunction_UserDefinedEvent.h b/Engine/Plugins/Runtime/RigVM/Source/RigVM/Public/RigVMFunctions/Execution/RigVMFunction_UserDefinedEvent.h
--- a/Engine/Plugins/Runtime/RigVM/Source/RigVM/Public/RigVMFunctions/Execution/RigVMFunction_UserDefinedEvent.h
+++ b/Engine/Plugins/Runtime/RigVM/Source/RigVM/Public/RigVMFunctions/Execution/RigVMFunction_UserDefinedEvent.h
@@ -27,7 +27,7 @@ struct RIGVM_API FRigVMFunction_UserDefinedEvent : public FRigVMStruct
 
 	// The execution result
 	UPROPERTY(EditAnywhere, Transient, DisplayName = "Execute", Category = "UserDefinedEvent", meta = (Output))
-	FRigVMExecuteContext ExecuteContext;
+	FRigVMExecutePin ExecutePin;
 
 	// True if the current interaction is a rotation
 	UPROPERTY(EditAnywhere, Transient, Category = "UserDefinedEvent", meta = (Input,Constant,DetailsOnly))
diff --git a/Engine/Plugins/Runtime/RigVM/Source/RigVMDeveloper/Private/EdGraph/RigVMEdGraphNode.cpp b/Engine/Plugins/Runtime/RigVM/Source/RigVMDeveloper/Private/EdGraph/RigVMEdGraphNode.cpp
--- a/Engine/Plugins/Runtime/RigVM/Source/RigVMDeveloper/Private/EdGraph/RigVMEdGraphNode.cpp
+++ b/Engine/Plugins/Runtime/RigVM/Source/RigVMDeveloper/Private/EdGraph/RigVMEdGraphNode.cpp
@@ -1602,7 +1602,7 @@ void URigVMEdGraphNode::UpdatePinLists()
 					if (ModelPin->IsStruct())
 					{
 						const UScriptStruct* ScriptStruct = ModelPin->GetScriptStruct();
-						if (ScriptStruct && ScriptStruct->IsChildOf(FRigVMEmptyExecuteContext::StaticStruct()))
+						if (ScriptStruct && ScriptStruct->IsChildOf(FRigVMExecutePin::StaticStruct()))
 						{
 							ExecutePins.Add(ModelPin);
 							continue;
diff --git a/Engine/Plugins/Runtime/RigVM/Source/RigVMDeveloper/Private/EdGraph/RigVMEdGraphSchema.cpp b/Engine/Plugins/Runtime/RigVM/Source/RigVMDeveloper/Private/EdGraph/RigVMEdGraphSchema.cpp
--- a/Engine/Plugins/Runtime/RigVM/Source/RigVMDeveloper/Private/EdGraph/RigVMEdGraphSchema.cpp
+++ b/Engine/Plugins/Runtime/RigVM/Source/RigVMDeveloper/Private/EdGraph/RigVMEdGraphSchema.cpp
@@ -900,7 +900,7 @@ FLinearColor URigVMEdGraphSchema::GetPinTypeColor(const FEdGraphPinType& PinType
 	{
 		if (const UStruct* Struct = Cast<UStruct>(PinType.PinSubCategoryObject))
 		{
-			if (Struct->IsChildOf(FRigVMEmptyExecuteContext::StaticStruct()))
+			if (Struct->IsChildOf(FRigVMExecutePin::StaticStruct()))
 			{
 				return FLinearColor::White;
 			}
@@ -1104,7 +1104,7 @@ bool URigVMEdGraphSchema::SupportsPinType(TWeakPtr<const FEdGraphSchemaAction> S
 		{
 			if(SchemaAction.IsValid() && SchemaAction.Pin()->IsAVariable())
 			{
-				if(ScriptStruct->IsChildOf(FRigVMEmptyExecuteContext::StaticStruct()))
+				if(ScriptStruct->IsChildOf(FRigVMExecutePin::StaticStruct()))
 				{
 					return false;
 				}
@@ -1139,7 +1139,7 @@ bool URigVMEdGraphSchema::SupportsPinTypeContainer(TWeakPtr<const FEdGraphSchema
 	// Do not allow containers for execute context type
 	if(const UScriptStruct* ExecuteContextScriptStruct = Cast<UScriptStruct>(PinType.PinSubCategoryObject))
 	{
-		if (ExecuteContextScriptStruct->IsChildOf(FRigVMEmptyExecuteContext::StaticStruct()))
+		if (ExecuteContextScriptStruct->IsChildOf(FRigVMExecutePin::StaticStruct()))
 		{
 			return ContainerType == EPinContainerType::None;
 		}
@@ -1768,7 +1768,7 @@ UEdGraphPin* URigVMEdGraphSchema::DropPinOnNode(UEdGraphNode* InTargetNode, cons
 
 								if(const UScriptStruct* CPPTypeStruct = Cast<UScriptStruct>(ExternalVar.TypeObject))
 								{
-									if(CPPTypeStruct->IsChildOf(FRigVMEmptyExecuteContext::StaticStruct()))
+									if(CPPTypeStruct->IsChildOf(FRigVMExecutePin::StaticStruct()))
 									{
 										SourcePinName = FRigVMStruct::ExecuteContextName;
 										PinDirection = ERigVMPinDirection::IO;										
diff --git a/Engine/Plugins/Runtime/RigVM/Source/RigVMDeveloper/Private/RigVMCompiler/RigVMCodeGenerator.cpp b/Engine/Plugins/Runtime/RigVM/Source/RigVMDeveloper/Private/RigVMCompiler/RigVMCodeGenerator.cpp
--- a/Engine/Plugins/Runtime/RigVM/Source/RigVMDeveloper/Private/RigVMCompiler/RigVMCodeGenerator.cpp
+++ b/Engine/Plugins/Runtime/RigVM/Source/RigVMDeveloper/Private/RigVMCompiler/RigVMCodeGenerator.cpp
@@ -1178,7 +1178,7 @@ void FRigVMCodeGenerator::ParseInclude(UStruct* InDependency, const FName& InMet
 
 	if (InDependency == URigVM::StaticClass() ||
 		InDependency == FRigVMByteCode::StaticStruct() || 
-		InDependency->IsChildOf(FRigVMEmptyExecuteContext::StaticStruct()) ||
+		InDependency->IsChildOf(FRigVMExecutePin::StaticStruct()) ||
 		InDependency->IsChildOf(FRigVMBaseOp::StaticStruct()) ||
 		InDependency->IsChildOf(URigVMMemoryStorage::StaticClass()))
 	{
diff --git a/Engine/Plugins/Runtime/RigVM/Source/RigVMDeveloper/Private/RigVMCompiler/RigVMCompiler.cpp b/Engine/Plugins/Runtime/RigVM/Source/RigVMDeveloper/Private/RigVMCompiler/RigVMCompiler.cpp
--- a/Engine/Plugins/Runtime/RigVM/Source/RigVMDeveloper/Private/RigVMCompiler/RigVMCompiler.cpp
+++ b/Engine/Plugins/Runtime/RigVM/Source/RigVMDeveloper/Private/RigVMCompiler/RigVMCompiler.cpp
@@ -952,7 +952,7 @@ bool URigVMCompiler::Compile(const FRigVMCompileSettings& InSettings, TArray<URi
 							
 							if(const FStructProperty* StructProperty = CastField<FStructProperty>(Property))
 							{
-								if(StructProperty->Struct->IsChildOf(FRigVMEmptyExecuteContext::StaticStruct()))
+								if(StructProperty->Struct->IsChildOf(FRigVMExecutePin::StaticStruct()))
 								{
 									if(!ValidExecuteContextStructs.Contains(StructProperty->Struct))
 									{
@@ -2032,7 +2032,7 @@ bool URigVMCompiler::TraverseCallExtern(const FRigVMCallExternExprAST* InExpr, F
 
 	auto CheckExecuteStruct = [this, &WorkData](URigVMNode* Subject, const UScriptStruct* ExecuteStruct) -> bool
 	{
-		if(ExecuteStruct->IsChildOf(FRigVMEmptyExecuteContext::StaticStruct()))
+		if(ExecuteStruct->IsChildOf(FRigVMExecutePin::StaticStruct()))
 		{
 			// top level expected execute struct is provided by the graph
 			const UScriptStruct* SpecializedExecuteStruct = WorkData.ExecuteContextStruct;
@@ -2609,7 +2609,7 @@ bool URigVMCompiler::TraverseInlineFunction(const FRigVMInlineFunctionExprAST* I
 					{
 						if (const UScriptStruct* ScriptStruct = Cast<UScriptStruct>(Argument.CPPTypeObject.Get()))
 						{
-							if (ScriptStruct->IsChildOf(FRigVMEmptyExecuteContext::StaticStruct()))
+							if (ScriptStruct->IsChildOf(FRigVMExecutePin::StaticStruct()))
 							{
 								continue;
 							}
diff --git a/Engine/Plugins/Runtime/RigVM/Source/RigVMDeveloper/Private/RigVMModel/RigVMController.cpp b/Engine/Plugins/Runtime/RigVM/Source/RigVMDeveloper/Private/RigVMModel/RigVMController.cpp
--- a/Engine/Plugins/Runtime/RigVM/Source/RigVMDeveloper/Private/RigVMModel/RigVMController.cpp
+++ b/Engine/Plugins/Runtime/RigVM/Source/RigVMDeveloper/Private/RigVMModel/RigVMController.cpp
@@ -12515,7 +12515,7 @@ FName URigVMController::AddExposedPin(const FName& InPinName, ERigVMPinDirection
 	{
 		if(const UScriptStruct* CPPTypeStruct = Cast<UScriptStruct>(CPPTypeObject))
 		{
-			if(CPPTypeStruct->IsChildOf(FRigVMEmptyExecuteContext::StaticStruct()))
+			if(CPPTypeStruct->IsChildOf(FRigVMExecutePin::StaticStruct()))
 			{
 				for(const URigVMPin* ExistingPin : LibraryNode->Pins)
 				{
@@ -13003,7 +13003,7 @@ bool URigVMController::ChangeExposedPinType(const FName& InPinName, const FStrin
 		{
 			if(const UScriptStruct* CPPTypeStruct = Cast<UScriptStruct>(CPPTypeObject))
 			{
-				if(CPPTypeStruct->IsChildOf(FRigVMEmptyExecuteContext::StaticStruct()))
+				if(CPPTypeStruct->IsChildOf(FRigVMExecutePin::StaticStruct()))
 				{
 					for(URigVMPin* ExistingPin : LibraryNode->Pins)
 					{
@@ -13033,7 +13033,7 @@ bool URigVMController::ChangeExposedPinType(const FName& InPinName, const FStrin
 	{
 		if(const UScriptStruct* CPPTypeStruct = Cast<UScriptStruct>(CPPTypeObject))
 		{
-			if(CPPTypeStruct->IsChildOf(FRigVMEmptyExecuteContext::StaticStruct()))
+			if(CPPTypeStruct->IsChildOf(FRigVMExecutePin::StaticStruct()))
 			{
 				bIsExecute = true;
 			}
diff --git a/Engine/Plugins/Runtime/RigVM/Source/RigVMDeveloper/Private/RigVMModel/RigVMNode.cpp b/Engine/Plugins/Runtime/RigVM/Source/RigVMDeveloper/Private/RigVMModel/RigVMNode.cpp
--- a/Engine/Plugins/Runtime/RigVM/Source/RigVMDeveloper/Private/RigVMModel/RigVMNode.cpp
+++ b/Engine/Plugins/Runtime/RigVM/Source/RigVMDeveloper/Private/RigVMModel/RigVMNode.cpp
@@ -526,6 +526,17 @@ URigVMPin* URigVMNode::FindPin(const FString& InPinPath) const
 			}
 		}
 	}
+
+	if(Right.IsEmpty())
+	{
+		static const FString ExecuteContextNameString = FRigVMStruct::ExecuteContextName.ToString();
+		static const FString ExecutePinNameString = FRigVMStruct::ExecutePinName.ToString();
+		if(Left.Equals(ExecuteContextNameString, ESearchCase::IgnoreCase) || 
+			Left.Equals(ExecutePinNameString, ESearchCase::IgnoreCase))
+		{
+			return FindExecutePin();
+		}
+	}
 	
 	return nullptr;
 }
@@ -546,6 +557,24 @@ URigVMPin* URigVMNode::FindRootPinByName(const FName& InPinName) const
 			return OrphanedPin;
 		}
 	}
+
+	if(InPinName == FRigVMStruct::ExecuteContextName || 
+		InPinName == FRigVMStruct::ExecutePinName)
+	{
+		return FindExecutePin();
+	}
+	return nullptr;
+}
+
+URigVMPin* URigVMNode::FindExecutePin() const
+{
+	for(TObjectPtr<URigVMPin> Pin : Pins)
+	{
+		if(Pin->IsExecuteContext())
+		{
+			return Pin;
+		}
+	}
 	return nullptr;
 }
 
diff --git a/Engine/Plugins/Runtime/RigVM/Source/RigVMDeveloper/Private/RigVMModel/RigVMPin.cpp b/Engine/Plugins/Runtime/RigVM/Source/RigVMDeveloper/Private/RigVMModel/RigVMPin.cpp
--- a/Engine/Plugins/Runtime/RigVM/Source/RigVMDeveloper/Private/RigVMModel/RigVMPin.cpp
+++ b/Engine/Plugins/Runtime/RigVM/Source/RigVMDeveloper/Private/RigVMModel/RigVMPin.cpp
@@ -646,7 +646,7 @@ bool URigVMPin::IsExecuteContext() const
 {
 	if (const UScriptStruct* ScriptStruct = GetScriptStruct())
 	{
-		if (ScriptStruct->IsChildOf(FRigVMEmptyExecuteContext::StaticStruct()))
+		if (ScriptStruct->IsChildOf(FRigVMExecutePin::StaticStruct()))
 		{
 			return true;
 		}
diff --git a/Engine/Plugins/Runtime/RigVM/Source/RigVMDeveloper/Private/RigVMModel/RigVMSchema.cpp b/Engine/Plugins/Runtime/RigVM/Source/RigVMDeveloper/Private/RigVMModel/RigVMSchema.cpp
--- a/Engine/Plugins/Runtime/RigVM/Source/RigVMDeveloper/Private/RigVMModel/RigVMSchema.cpp
+++ b/Engine/Plugins/Runtime/RigVM/Source/RigVMDeveloper/Private/RigVMModel/RigVMSchema.cpp
@@ -237,7 +237,7 @@ bool URigVMSchema::ShouldUnfoldStruct(URigVMController* InController, const UStr
 	{
 		return false;
 	}
-	if(InStruct->IsChildOf(FRigVMEmptyExecuteContext::StaticStruct()))
+	if(InStruct->IsChildOf(FRigVMExecutePin::StaticStruct()))
 	{
 		return false;
 	}
diff --git a/Engine/Plugins/Runtime/RigVM/Source/RigVMDeveloper/Public/RigVMModel/RigVMNode.h b/Engine/Plugins/Runtime/RigVM/Source/RigVMDeveloper/Public/RigVMModel/RigVMNode.h
--- a/Engine/Plugins/Runtime/RigVM/Source/RigVMDeveloper/Public/RigVMModel/RigVMNode.h
+++ b/Engine/Plugins/Runtime/RigVM/Source/RigVMDeveloper/Public/RigVMModel/RigVMNode.h
@@ -119,6 +119,10 @@ public:
 	UFUNCTION(BlueprintCallable, Category = RigVMNode)
 	URigVMPin* FindRootPinByName(const FName& InPinName) const;
 
+	// Returns the first execute pin
+	UFUNCTION(BlueprintCallable, Category = RigVMNode)
+	URigVMPin* FindExecutePin() const;
+
 	// Returns all of the top-level orphaned Pins of this Node.
 	UFUNCTION(BlueprintCallable, Category = RigVMNode)
     virtual const TArray<URigVMPin*>& GetOrphanedPins() const;
diff --git a/Engine/Plugins/Runtime/RigVM/Source/RigVMEditor/Private/EdGraph/NodeSpawners/RigVMEdGraphUnitNodeSpawner.cpp b/Engine/Plugins/Runtime/RigVM/Source/RigVMEditor/Private/EdGraph/NodeSpawners/RigVMEdGraphUnitNodeSpawner.cpp
--- a/Engine/Plugins/Runtime/RigVM/Source/RigVMEditor/Private/EdGraph/NodeSpawners/RigVMEdGraphUnitNodeSpawner.cpp
+++ b/Engine/Plugins/Runtime/RigVM/Source/RigVMEditor/Private/EdGraph/NodeSpawners/RigVMEdGraphUnitNodeSpawner.cpp
@@ -385,7 +385,7 @@ void URigVMEdGraphUnitNodeSpawner::HookupMutableNode(URigVMNode* InModelNode, UR
 	{
 		if (Pin->GetScriptStruct())
 		{
-			if (Pin->GetScriptStruct()->IsChildOf(FRigVMEmptyExecuteContext::StaticStruct()))
+			if (Pin->GetScriptStruct()->IsChildOf(FRigVMExecutePin::StaticStruct()))
 			{
 				if (Pin->GetDirection() == ERigVMPinDirection::IO || Pin->GetDirection() == ERigVMPinDirection::Input)
 				{
diff --git a/Engine/Plugins/Runtime/RigVM/Source/RigVMEditor/Private/EdGraph/RigVMEdGraphPanelPinFactory.cpp b/Engine/Plugins/Runtime/RigVM/Source/RigVMEditor/Private/EdGraph/RigVMEdGraphPanelPinFactory.cpp
--- a/Engine/Plugins/Runtime/RigVM/Source/RigVMEditor/Private/EdGraph/RigVMEdGraphPanelPinFactory.cpp
+++ b/Engine/Plugins/Runtime/RigVM/Source/RigVMEditor/Private/EdGraph/RigVMEdGraphPanelPinFactory.cpp
@@ -136,7 +136,7 @@ TSharedPtr<SGraphPin> FRigVMEdGraphPanelPinFactory::CreatePin_Internal(UEdGraphP
 			if (InPin->PinType.PinCategory == UEdGraphSchema_K2::PC_Struct)
 			{
 				const UStruct* Struct = Cast<UStruct>(InPin->PinType.PinSubCategoryObject);
-				if (Struct && Struct->IsChildOf(FRigVMEmptyExecuteContext::StaticStruct()))
+				if (Struct && Struct->IsChildOf(FRigVMExecutePin::StaticStruct()))
 				{
 					return SNew(SGraphPinExec, InPin);
 				}
diff --git a/Engine/Plugins/Runtime/RigVM/Source/RigVMEditor/Private/Editor/RigVMDetailsViewWrapperObject.cpp b/Engine/Plugins/Runtime/RigVM/Source/RigVMEditor/Private/Editor/RigVMDetailsViewWrapperObject.cpp
--- a/Engine/Plugins/Runtime/RigVM/Source/RigVMEditor/Private/Editor/RigVMDetailsViewWrapperObject.cpp
+++ b/Engine/Plugins/Runtime/RigVM/Source/RigVMEditor/Private/Editor/RigVMDetailsViewWrapperObject.cpp
@@ -133,7 +133,7 @@ UClass* URigVMDetailsViewWrapperObject::GetClassForStruct(UScriptStruct* InStruc
 			bool bIsEditable = true;
 			if (FStructProperty* StructProperty = CastField<FStructProperty>(NewProperty))
 			{
-				if (StructProperty->Struct->IsChildOf(FRigVMEmptyExecuteContext::StaticStruct()))
+				if (StructProperty->Struct->IsChildOf(FRigVMExecutePin::StaticStruct()))
 				{
 					bIsEditable = false;
 				}
diff --git a/Engine/Plugins/Runtime/RigVM/Source/RigVMEditor/Private/Editor/RigVMGraphDetailCustomization.cpp b/Engine/Plugins/Runtime/RigVM/Source/RigVMEditor/Private/Editor/RigVMGraphDetailCustomization.cpp
--- a/Engine/Plugins/Runtime/RigVM/Source/RigVMEditor/Private/Editor/RigVMGraphDetailCustomization.cpp
+++ b/Engine/Plugins/Runtime/RigVM/Source/RigVMEditor/Private/Editor/RigVMGraphDetailCustomization.cpp
@@ -171,7 +171,7 @@ public:
 		RigVMTypeUtils::CPPTypeFromPinType(InItem.Get()->GetPinType(false), CPPType, &CPPTypeObject);
 		if (UScriptStruct* ScriptStruct = Cast<UScriptStruct>(CPPTypeObject))
 		{
-			if (ScriptStruct->IsChildOf(FRigVMEmptyExecuteContext::StaticStruct()))
+			if (ScriptStruct->IsChildOf(FRigVMExecutePin::StaticStruct()))
 			{
 				if (GraphPtr.IsValid())
 				{
diff --git a/Engine/Plugins/Runtime/RigVM/Source/RigVMEditor/Private/Widgets/SRigVMGraphChangePinType.cpp b/Engine/Plugins/Runtime/RigVM/Source/RigVMEditor/Private/Widgets/SRigVMGraphChangePinType.cpp
--- a/Engine/Plugins/Runtime/RigVM/Source/RigVMEditor/Private/Widgets/SRigVMGraphChangePinType.cpp
+++ b/Engine/Plugins/Runtime/RigVM/Source/RigVMEditor/Private/Widgets/SRigVMGraphChangePinType.cpp
@@ -240,7 +240,7 @@ void SRigVMGraphChangePinType::FillPinTypeMenu(FMenuBuilder& MenuBuilder)
 
 			if(UScriptStruct* ScriptStruct = Cast<UScriptStruct>(Type.CPPTypeObject))
 			{
-				if(ScriptStruct->IsChildOf(FRigVMEmptyExecuteContext::StaticStruct()))
+				if(ScriptStruct->IsChildOf(FRigVMExecutePin::StaticStruct()))
 				{
 					continue;
 				}
diff --git a/Engine/Source/Programs/Shared/EpicGames.UHT/Types/UhtScriptStruct.cs b/Engine/Source/Programs/Shared/EpicGames.UHT/Types/UhtScriptStruct.cs
--- a/Engine/Source/Programs/Shared/EpicGames.UHT/Types/UhtScriptStruct.cs
+++ b/Engine/Source/Programs/Shared/EpicGames.UHT/Types/UhtScriptStruct.cs
@@ -478,7 +478,7 @@ namespace EpicGames.UHT.Types
 				UhtScriptStruct? scriptStruct = structProperty.ScriptStruct;
 				while (scriptStruct != null)
 				{
-					if (scriptStruct.SourceName == "FRigVMEmptyExecuteContext")
+					if (scriptStruct.SourceName == "FRigVMExecutePin")
 					{
 						return true;
 					}
