Fixed Control Rig not returning correctly the value of the initial rotation of a control

diff --git a/Engine/Plugins/Animation/ControlRig/Source/ControlRig/Public/Rigs/RigHierarchy.h b/Engine/Plugins/Animation/ControlRig/Source/ControlRig/Public/Rigs/RigHierarchy.h
--- a/Engine/Plugins/Animation/ControlRig/Source/ControlRig/Public/Rigs/RigHierarchy.h
+++ b/Engine/Plugins/Animation/ControlRig/Source/ControlRig/Public/Rigs/RigHierarchy.h
@@ -2878,7 +2878,7 @@ public:
 			}
 			else
 			{
-				FRotator Rotator = GetControlPreferredRotator(InControlElement);
+				FRotator Rotator = GetControlPreferredRotator(InControlElement, bIsInitial);
 				EulerAngle = FVector(Rotator.Roll, Rotator.Pitch, Rotator.Yaw);
 			}
 		}
