[ControlRig & Sequencer] Added "Recreate Control Rig With New Settings" option to Control Rig Track right click menu to allow users to change rig settings after the rig is added to sequencer

diff --git a/Engine/Plugins/Animation/ControlRig/Source/ControlRigEditor/Private/Sequencer/ControlRigParameterTrackEditor.cpp b/Engine/Plugins/Animation/ControlRig/Source/ControlRigEditor/Private/Sequencer/ControlRigParameterTrackEditor.cpp
--- a/Engine/Plugins/Animation/ControlRig/Source/ControlRigEditor/Private/Sequencer/ControlRigParameterTrackEditor.cpp
+++ b/Engine/Plugins/Animation/ControlRig/Source/ControlRigEditor/Private/Sequencer/ControlRigParameterTrackEditor.cpp
@@ -112,6 +112,7 @@
 #include "CurveModel.h"
 #include "Sequencer/SLoadAnimToControlRig.h"
 #include "FrontendFilterBase.h"
+#include "InstancedPropertyBagStructureDataProvider.h"
 
 #define LOCTEXT_NAMESPACE "FControlRigParameterTrackEditor"
 
@@ -1543,6 +1544,149 @@ void FControlRigParameterTrackEditor::ConvertIsLayered(UMovieSceneControlRigPara
 	UControlRigSequencerEditorLibrary::SetControlRigLayeredMode(Track, bSetAdditive);
 }
 
+void FControlRigParameterTrackEditor::RecreateControlRigWithNewSettings(UMovieSceneControlRigParameterTrack* InTrack)
+{
+	CreateAndShowRigSettingsWindow(InTrack);
+}
+
+bool FControlRigParameterTrackEditor::CanRecreateControlRigWithNewSettings(UMovieSceneControlRigParameterTrack* InTrack)
+{
+	UControlRig* ControlRig = InTrack->GetControlRig();
+	if (!ensure(ControlRig))
+	{
+		return false;
+	}
+
+	if (ControlRig->GetPublicVariables().Num() == 0)
+	{
+		// no public variable means the rig is not configurable
+		return false;
+	}
+
+	return true;
+}
+
+
+void FControlRigParameterTrackEditor::CreateAndShowRigSettingsWindow(UMovieSceneControlRigParameterTrack* Track)
+{
+	if (RecreateControlRigWithNewSettingsWindow.IsValid())
+	{
+		FSlateApplication::Get().RequestDestroyWindow(RecreateControlRigWithNewSettingsWindow.Pin().ToSharedRef());
+	}
+	
+	SettingsForNewControlRig.Reset();
+	
+	UControlRig* ControlRig = Track->GetControlRig();
+	
+	TArray<FPropertyBagPropertyDesc> BagProperties;
+	TArray<const FProperty*> SourceProperties;
+	TArray<FRigVMExternalVariable> PublicVariables = ControlRig->GetPublicVariables();
+	for (const FRigVMExternalVariable& Variable : PublicVariables)
+	{
+		const FProperty* SourceProperty = ControlRig->GetClass()->FindPropertyByName(Variable.Name);
+		SourceProperties.Add(SourceProperty);
+		
+		BagProperties.Add(FPropertyBagPropertyDesc(Variable.Name , SourceProperty));
+		
+	}
+
+	SettingsForNewControlRig.AddProperties(BagProperties);
+
+	for (const FProperty* SourceProperty : SourceProperties)
+	{
+		SettingsForNewControlRig.SetValue(SourceProperty->GetFName(), SourceProperty, ControlRig); 
+	}
+	
+	FDetailsViewArgs DetailsViewArgs;
+	DetailsViewArgs.NameAreaSettings = FDetailsViewArgs::HideNameArea;
+	FStructureDetailsViewArgs StructureDetailsViewArgs;
+	
+	TSharedRef<IStructureDetailsView> DetailsView=
+		FModuleManager::GetModuleChecked<FPropertyEditorModule>("PropertyEditor")
+		.CreateStructureProviderDetailView(
+			DetailsViewArgs,
+			StructureDetailsViewArgs,
+			MakeShared<FInstancePropertyBagStructureDataProvider>(SettingsForNewControlRig));
+
+	DetailsView->GetOnFinishedChangingPropertiesDelegate().AddLambda(
+		[WeakTrack = TWeakObjectPtr(Track), WeakRig = TWeakObjectPtr(ControlRig), WeakThis = SharedThis(this).ToWeakPtr()]( const FPropertyChangedEvent& )
+		{
+			if (!WeakThis.IsValid())
+			{
+				return;
+			}
+				
+			TSharedPtr<FControlRigParameterTrackEditor> TrackEditor = WeakThis.Pin();
+			if (!WeakTrack.IsValid())
+			{
+				return;
+			}
+			
+			UMovieSceneControlRigParameterTrack* Track = WeakTrack.Get();
+			UControlRig* ControlRig = Track->GetControlRig();
+					
+			if (!WeakRig.IsValid() || WeakRig != ControlRig)
+			{
+				return;
+			}
+			
+			const FScopedTransaction Transaction(LOCTEXT("RecreateControlRigWithNewSettings_Transaction", "Recreated Control Rig with New Settings"));
+			Track->Modify();
+			ControlRig->Modify();
+
+			ControlRig->Initialize();
+			
+			for (const FPropertyBagPropertyDesc& PropertyDesc : TrackEditor->SettingsForNewControlRig.GetPropertyBagStruct()->GetPropertyDescs())
+			{
+				FRigVMExternalVariable Variable = ControlRig->GetPublicVariableByName(PropertyDesc.Name);
+				if (ensure(Variable.IsValid()))
+				{
+					if (Variable.Property->GetClass() == PropertyDesc.CachedProperty->GetClass())
+					{
+						// copy from property bag to control rig
+						void * TargetAddress = Variable.Property->ContainerPtrToValuePtr<void>(ControlRig);
+						void const* SourceAddress = TrackEditor->SettingsForNewControlRig.GetValue().GetMemory() + PropertyDesc.CachedProperty->GetOffset_ForInternal();
+						Variable.Property->CopyCompleteValue(TargetAddress, SourceAddress);
+					}
+				}
+			}
+			
+			ControlRig->Evaluate_AnyThread();
+
+			for (UMovieSceneSection* Section : Track->GetAllSections())
+			{
+				if (Section)
+				{
+					UMovieSceneControlRigParameterSection* CRSection = Cast<UMovieSceneControlRigParameterSection>(Section);
+					if (CRSection)
+					{
+						Section->Modify();
+						CRSection->ClearAllParameters();
+						CRSection->RecreateWithThisControlRig(CRSection->GetControlRig(), true);
+					}
+				}
+			}
+
+			TrackEditor->GetSequencer()->NotifyMovieSceneDataChanged(EMovieSceneDataChangeType::MovieSceneStructureItemsChanged);
+		});
+	
+	
+	TSharedRef<SWindow> Window = SNew(SWindow)
+	.Title( LOCTEXT("RecreateControlRigWithNewSettingsSettingsWindowTitle", "Recreate Control Rig With New Settings") )
+	.bDragAnywhere(true)
+	.Type(EWindowType::Normal)
+	.IsTopmostWindow(true)
+	.SizingRule( ESizingRule::Autosized )
+	.FocusWhenFirstShown(true)
+	.ActivationPolicy(EWindowActivationPolicy::FirstShown)
+	[
+		DetailsView->GetWidget().ToSharedRef()
+	];
+	
+	RecreateControlRigWithNewSettingsWindow = Window;
+	FSlateApplication::Get().AddWindow( Window, true);
+}
+
 void FControlRigParameterTrackEditor::BuildObjectBindingTrackMenu(FMenuBuilder& MenuBuilder, const TArray<FGuid>& ObjectBindings, const UClass* ObjectClass)
 {
 	if(!ObjectClass)
@@ -5121,6 +5265,18 @@ void FControlRigParameterTrackEditor::BuildTrackContextMenu(FMenuBuilder& MenuBu
 					   NAME_None,
 					   EUserInterfaceActionType::ToggleButton);
 			}
+
+
+			MenuBuilder.AddMenuEntry(
+				   LOCTEXT("RecreateControlRigWithNewSettingsSettings", "Recreate Control Rig With New Settings"),
+				   LOCTEXT("RecreateControlRigWithNewSettingsSettingsToolTip", "Recreate Control Rig With New Settings"),
+				   FSlateIcon(),
+				   FUIAction(
+					   FExecuteAction::CreateRaw(this, &FControlRigParameterTrackEditor::RecreateControlRigWithNewSettings, Track),
+					   FCanExecuteAction::CreateRaw(this, &FControlRigParameterTrackEditor::CanRecreateControlRigWithNewSettings, Track)
+				   ),
+				   NAME_None,
+				   EUserInterfaceActionType::Button);
 		}
 		MenuBuilder.EndSection();
 	}
diff --git a/Engine/Plugins/Animation/ControlRig/Source/ControlRigEditor/Private/Sequencer/ControlRigParameterTrackEditor.h b/Engine/Plugins/Animation/ControlRig/Source/ControlRigEditor/Private/Sequencer/ControlRigParameterTrackEditor.h
--- a/Engine/Plugins/Animation/ControlRig/Source/ControlRigEditor/Private/Sequencer/ControlRigParameterTrackEditor.h
+++ b/Engine/Plugins/Animation/ControlRig/Source/ControlRigEditor/Private/Sequencer/ControlRigParameterTrackEditor.h
@@ -260,6 +260,17 @@ private:
 	/** Convert an absolute control rit to a layered control rig */
 	void ConvertIsLayered(UMovieSceneControlRigParameterTrack* Track);
 
+	
+	/** Change rig settings and recreates the rig using the new settings */
+	void RecreateControlRigWithNewSettings(UMovieSceneControlRigParameterTrack* Track);
+	
+	bool CanRecreateControlRigWithNewSettings(UMovieSceneControlRigParameterTrack* Track);
+
+	void CreateAndShowRigSettingsWindow(UMovieSceneControlRigParameterTrack* Track);
+private:
+	FInstancedPropertyBag SettingsForNewControlRig;
+	TWeakPtr<SWindow> RecreateControlRigWithNewSettingsWindow;
+public:
 	/** Set Up EditMode for Specified Control Rig*/
 	void SetUpEditModeIfNeeded(UControlRig* ControlRig);
 
