Don't issue a warning for GetComponent since it already outputs success/failure

diff --git a/Engine/Plugins/Animation/ControlRig/Source/ControlRig/Private/Units/Hierarchy/RigUnit_Component.cpp b/Engine/Plugins/Animation/ControlRig/Source/ControlRig/Private/Units/Hierarchy/RigUnit_Component.cpp
--- a/Engine/Plugins/Animation/ControlRig/Source/ControlRig/Private/Units/Hierarchy/RigUnit_Component.cpp
+++ b/Engine/Plugins/Animation/ControlRig/Source/ControlRig/Private/Units/Hierarchy/RigUnit_Component.cpp
@@ -291,9 +291,7 @@ void FRigDispatch_GetComponentContent::Execute(FRigVMExtendedExecuteContext& InC
 	if(SourceComponent == nullptr)
 	{
 		ComponentStruct->InitializeDefaultValue(TargetComponent);
-#if WITH_EDITOR
-		Context.Report(FRigVMLogSettings(EMessageSeverity::Warning), Context.GetFunctionName(), Context.GetInstructionIndex(), FString::Printf(TEXT("%s not found."), *ComponentKey.ToString()));
-#endif
+		// Since we return success/failure, this is not an error/warning
 		bSuccess = false;
 		return;
 	}
diff --git a/Engine/Plugins/Animation/ControlRig/Source/ControlRig/Public/Units/Hierarchy/RigUnit_Component.h b/Engine/Plugins/Animation/ControlRig/Source/ControlRig/Public/Units/Hierarchy/RigUnit_Component.h
--- a/Engine/Plugins/Animation/ControlRig/Source/ControlRig/Public/Units/Hierarchy/RigUnit_Component.h
+++ b/Engine/Plugins/Animation/ControlRig/Source/ControlRig/Public/Units/Hierarchy/RigUnit_Component.h
@@ -76,7 +76,7 @@ protected:
 };
 
 /**
- * Break / Decompose node for a hierarchy component
+ * Spawns a top level component
  */
 USTRUCT(meta=(DisplayName="Spawn Top Level Component"))
 struct CONTROLRIG_API FRigDispatch_SpawnTopLevelComponent : public FRigDispatch_ComponentBase
@@ -106,7 +106,7 @@ protected:
 };
 
 /**
- * Break / Decompose node for a hierarchy component
+ * Gets the component
  */
 USTRUCT(meta=(DisplayName="Get Component"))
 struct CONTROLRIG_API FRigDispatch_GetComponentContent : public FRigDispatch_ComponentBase
@@ -158,7 +158,7 @@ protected:
 };
 
 /**
- * Break / Decompose node for a hierarchy component
+ * Gets the top level component
  */
 USTRUCT(meta=(DisplayName="Get Top Level Component"))
 struct CONTROLRIG_API FRigDispatch_GetTopLevelComponentContent : public FRigDispatch_ComponentBase
