Control Rig: Unique modular rig module names

#rnx

#virtualized

diff --git a/Engine/Plugins/Animation/ControlRig/Source/ControlRig/Private/ControlRig.cpp b/Engine/Plugins/Animation/ControlRig/Source/ControlRig/Private/ControlRig.cpp
--- a/Engine/Plugins/Animation/ControlRig/Source/ControlRig/Private/ControlRig.cpp
+++ b/Engine/Plugins/Animation/ControlRig/Source/ControlRig/Private/ControlRig.cpp
@@ -679,7 +679,7 @@ bool UControlRig::AllConnectorsAreResolved(FString* OutFailureReason, FRigElemen
 						if(OutFailureReason)
 						{
 							static constexpr TCHAR Format[] = TEXT("Connector '%s%s' has invalid target '%s'.");
-							*OutFailureReason = FString::Printf(Format, *GetRigModuleNameSpace(), *Connector.Name, *CachedRigElement.GetKey().ToString());
+							*OutFailureReason = FString::Printf(Format, *GetRigModulePrefix(), *Connector.Name, *CachedRigElement.GetKey().ToString());
 						}
 						if(OutConnector)
 						{
@@ -1986,27 +1986,23 @@ UControlRig* UControlRig::GetParentRig() const
 	return GetTypedOuter<UControlRig>();
 }
 
-const FString& UControlRig::GetRigModuleNameSpace() const
+const FString& UControlRig::GetRigModulePrefix() const
 {
 	if(IsRigModule())
 	{
-		if(const UControlRig* ParentRig = GetParentRig())
+		if(IsRigModuleInstance() || GetRigModuleSettings().Identifier.Name.IsEmpty())
 		{
-			const FString& ParentNameSpace = ParentRig->GetRigModuleNameSpace();
-			static constexpr TCHAR JoinFormat[] = TEXT("%s%s:");
-			RigModuleNameSpace = FString::Printf(JoinFormat, *ParentNameSpace, *GetFName().ToString());
-			return RigModuleNameSpace;
+			RigModulePrefix = GetFName().ToString() + FRigHierarchyModulePath::ModuleNameSuffix;
 		}
-
-		// this means we are not an instance - so we are authoring the rig module right now.
-		// for visualization purposes we'll use the name of the rig module chosen in the settings.
-		static constexpr TCHAR NameSpaceFormat[] = TEXT("%s:"); 
-		RigModuleNameSpace = FString::Printf(NameSpaceFormat, *GetRigModuleSettings().Identifier.Name);
-		return RigModuleNameSpace;
+		else
+		{
+			RigModulePrefix = GetRigModuleSettings().Identifier.Name + FRigHierarchyModulePath::ModuleNameSuffix;
+		}
+		return RigModulePrefix;
 	}
 
-	static const FString EmptyNameSpace;
-	return EmptyNameSpace;
+	static const FString EmptyPrefix;
+	return EmptyPrefix;
 }
 
 FRigElementKeyRedirector& UControlRig::GetElementKeyRedirector()
@@ -2047,7 +2043,15 @@ FRigControlElement* UControlRig::FindControl(const FName& InControlName) const
 	{
 		return nullptr;
 	}
-	return DynamicHierarchy->Find<FRigControlElement>(FRigElementKey(InControlName, ERigElementType::Control));
+
+	// look up the control based on the stored name. 
+	const FRigElementKey ControlKey(InControlName, ERigElementType::Control);
+	if(FRigControlElement* Control = DynamicHierarchy->Find<FRigControlElement>(ControlKey))
+	{
+		return Control;
+	}
+
+	return nullptr;
 }
 
 bool UControlRig::IsConstructionModeEnabled() const
diff --git a/Engine/Plugins/Animation/ControlRig/Source/ControlRig/Private/ModularRig.cpp b/Engine/Plugins/Animation/ControlRig/Source/ControlRig/Private/ModularRig.cpp
--- a/Engine/Plugins/Animation/ControlRig/Source/ControlRig/Private/ModularRig.cpp
+++ b/Engine/Plugins/Animation/ControlRig/Source/ControlRig/Private/ModularRig.cpp
@@ -17,13 +17,28 @@
 
 FModuleInstanceHandle::FModuleInstanceHandle(const UModularRig* InModularRig, const FString& InPath)
 : ModularRig(const_cast<UModularRig*>(InModularRig))
-, Path(InPath)
+, ModuleName(NAME_None)
+{
+	FRigHierarchyModulePath Path(InPath);
+	if(Path.IsValid())
+	{
+		ModuleName = Path.GetElementFName();
+	}
+	else if(!InPath.IsEmpty())
+	{
+		ModuleName = *InPath;
+	}
+}
+
+FModuleInstanceHandle::FModuleInstanceHandle(const UModularRig* InModularRig, const FName& InModuleName)
+: ModularRig(const_cast<UModularRig*>(InModularRig))
+, ModuleName(InModuleName)
 {
 }
 
 FModuleInstanceHandle::FModuleInstanceHandle(const UModularRig* InModularRig, const FRigModuleInstance* InModule)
 : ModularRig(const_cast<UModularRig*>(InModularRig))
-, Path(InModule->GetPath())
+, ModuleName(InModule->Name)
 {
 }
 
@@ -31,7 +46,7 @@ const FRigModuleInstance* FModuleInstanceHandle::Get() const
 {
 	if(const UModularRig* ResolvedRig = ModularRig.Get())
 	{
-		return ResolvedRig->FindModule(Path);
+		return ResolvedRig->FindModule(ModuleName);
 	}
 	return nullptr;
 }
@@ -57,33 +72,6 @@ void UModularRig::BeginDestroy()
 #endif
 }
 
-FString FRigModuleInstance::GetShortName() const
-{
-	if(const FRigModuleReference* ModuleReference = GetModuleReference())
-	{
-		const FString ShortName = ModuleReference->GetShortName();
-		if(!ShortName.IsEmpty())
-		{
-			return ShortName;
-		}
-	}
-	return Name.ToString();
-}
-
-FString FRigModuleInstance::GetPath() const
-{
-	if (!IsRootModule())
-	{
-		return URigHierarchy::JoinNameSpace(ParentPath, Name.ToString()); 
-	}
-	return Name.ToString();
-}
-
-FString FRigModuleInstance::GetNamespace() const
-{
-	return FString::Printf(TEXT("%s:"), *GetPath());
-}
-
 UControlRig* FRigModuleInstance::GetRig() const
 {
 	if(IsValid(RigPtr))
@@ -127,7 +115,7 @@ const FRigModuleReference* FRigModuleInstance::GetModuleReference() const
 		if(const UModularRig* ModularRig = Cast<UModularRig>(Rig->GetParentRig()))
 		{
 			const FModularRigModel& Model = ModularRig->GetModularRigModel();
-			return Model.FindModule(GetPath());
+			return Model.FindModule(Name);
 		}
 	}
 	return nullptr;
@@ -135,15 +123,19 @@ const FRigModuleReference* FRigModuleInstance::GetModuleReference() const
 
 const FRigModuleInstance* FRigModuleInstance::GetParentModule() const
 {
-	if(ParentPath.IsEmpty())
+	if(ParentModuleName.IsNone())
 	{
-		return this;
+		return nullptr;
+	}
+	if(CachedParentModule)
+	{
+		return CachedParentModule;
 	}
 	if(const UControlRig* Rig = GetRig())
 	{
 		if(const UModularRig* ModularRig = Cast<UModularRig>(Rig->GetParentRig()))
 		{
-			return ModularRig->FindModule(ParentPath);
+			return ModularRig->FindModule(ParentModuleName);
 		}
 	}
 	return nullptr;
@@ -151,18 +143,13 @@ const FRigModuleInstance* FRigModuleInstance::GetParentModule() const
 
 const FRigModuleInstance* FRigModuleInstance::GetRootModule() const
 {
-	if(ParentPath.IsEmpty())
+	if(ParentModuleName.IsNone())
 	{
 		return this;
 	}
-	if(const UControlRig* Rig = GetRig())
+	if(const FRigModuleInstance* ParentModule = GetParentModule())
 	{
-		if(const UModularRig* ModularRig = Cast<UModularRig>(Rig->GetParentRig()))
-		{
-			FString RootPath = ParentPath;
-			(void)URigHierarchy::SplitNameSpace(ParentPath, &RootPath, nullptr, false);
-			return ModularRig->FindModule(RootPath);
-		}
+		return ParentModule->GetRootModule();
 	}
 	return nullptr;
 }
@@ -178,17 +165,16 @@ const FRigConnectorElement* FRigModuleInstance::FindPrimaryConnector() const
 			{
 				return Cast<FRigConnectorElement>(PrimaryConnector.GetElement());
 			}
-	
-			const FString MyModulePath = GetPath();
+
 			const TArray<FRigConnectorElement*> AllConnectors = Hierarchy->GetConnectors();
 			for(const FRigConnectorElement* Connector : AllConnectors)
 			{
 				if(Connector->IsPrimary())
 				{
-					const FString ModulePath = Hierarchy->GetModulePath(Connector->GetKey());
-					if(!ModulePath.IsEmpty())
+					const FName ConnectorModuleName = Hierarchy->GetModuleFName(Connector->GetKey());
+					if(!ConnectorModuleName.IsNone())
 					{
-						if(ModulePath.Equals(MyModulePath, ESearchCase::IgnoreCase))
+						if(ConnectorModuleName == Name)
 						{
 							PrimaryConnector.UpdateCache(Connector->GetKey(), Hierarchy);
 							return Cast<FRigConnectorElement>(PrimaryConnector.GetElement());
@@ -208,14 +194,13 @@ TArray<const FRigConnectorElement*> FRigModuleInstance::FindConnectors() const
 	{
 		if(const URigHierarchy* Hierarchy = Rig->GetHierarchy())
 		{
-			const FString MyModulePath = GetPath();
 			const TArray<FRigConnectorElement*> AllConnectors = Hierarchy->GetConnectors();
 			for(const FRigConnectorElement* Connector : AllConnectors)
 			{
-				const FString ModulePath = Hierarchy->GetModulePath(Connector->GetKey());
-				if(!ModulePath.IsEmpty())
+				const FName ConnectorModuleName = Hierarchy->GetModuleFName(Connector->GetKey());
+				if(!ConnectorModuleName.IsNone())
 				{
-					if(ModulePath.Equals(MyModulePath, ESearchCase::IgnoreCase))
+					if(ConnectorModuleName == Name)
 					{
 						Connectors.Add(Connector);
 					}
@@ -228,7 +213,42 @@ TArray<const FRigConnectorElement*> FRigModuleInstance::FindConnectors() const
 
 bool FRigModuleInstance::IsRootModule() const
 {
-	return ParentPath.IsEmpty();
+	return ParentModuleName.IsNone();
+}
+
+FString FRigModuleInstance::GetModulePrefix() const
+{
+	return Name.ToString() + FRigHierarchyModulePath::ModuleNameSuffix;
+}
+
+FString FRigModuleInstance::GetModulePath_Deprecated() const
+{
+	return ParentPath_DEPRECATED.IsEmpty() ?
+		Name.ToString() :
+		URigHierarchy::JoinNameSpace_Deprecated(ParentPath_DEPRECATED, Name.ToString());
+}
+
+bool FRigModuleInstance::HasChildModule(const FName& InModuleName, bool bRecursive) const
+{
+	if(InModuleName.IsNone())
+	{
+		return false;
+	}
+	for(const FRigModuleInstance* ChildModule : CachedChildren)
+	{
+		if(ChildModule->Name == InModuleName)
+		{
+			return true;
+		}
+		if(bRecursive)
+		{
+			if(ChildModule->HasChildModule(InModuleName, bRecursive))
+			{
+				return true;
+			}
+		}
+	}
+	return false;
 }
 
 void UModularRig::PostInitProperties()
@@ -255,7 +275,6 @@ void UModularRig::PostLoad()
 	Super::PostLoad();
 	ModularRigModel.UpdateCachedChildren();
 	ModularRigModel.Connections.UpdateFromConnectionList();
-	ResetShortestDisplayPathCache();
 }
 
 void UModularRig::InitializeVMs(bool bRequestInit)
@@ -305,7 +324,7 @@ void UModularRig::UpdateModuleHierarchyFromCDO()
 			{
 				if(IsValid(ModuleRig))
 				{
-					PreviousModuleRigs.Add(Module.GetPath(), ModuleRig);
+					PreviousModuleRigs.Add(Module.Name, ModuleRig);
 				}
 			}
         }
@@ -330,8 +349,8 @@ void UModularRig::UpdateModuleHierarchyFromCDO()
 				(void)AddModuleInstance(
 					InModuleReference->Name,
 					InModuleReference->Class.Get(),
-					FindModule(InModuleReference->ParentPath),
-					Model.Connections.GetModuleConnectionMap(InModuleReference->GetPath()),
+					FindModule(InModuleReference->ParentModuleName),
+					Model.Connections.GetModuleConnectionMap(InModuleReference->Name),
 					InModuleReference->ConfigValues);
 			}
 
@@ -340,7 +359,7 @@ void UModularRig::UpdateModuleHierarchyFromCDO()
 		});
 
 		// discard any remaining rigs
-		for(const TPair<FString, UControlRig*>& Pair : PreviousModuleRigs)
+		for(const TPair<FName, UControlRig*>& Pair : PreviousModuleRigs)
 		{
 			DiscardModuleRig(Pair.Value);
 		}
@@ -350,9 +369,9 @@ void UModularRig::UpdateModuleHierarchyFromCDO()
 		// modules have to exist first
 		ForEachModule([this, Model](const FRigModuleInstance* Module) -> bool
 		{
-			if(const FRigModuleReference* ModuleReference = Model.FindModule(Module->GetPath()))
+			if(const FRigModuleReference* ModuleReference = Model.FindModule(Module->Name))
 			{
-				(void)SetModuleVariableBindings(ModuleReference->GetPath(), ModuleReference->Bindings);
+				(void)SetModuleVariableBindings(ModuleReference->Name, ModuleReference->Bindings);
 			}
 			if(UControlRig* ModuleRig = Module->GetRig())
 			{
@@ -388,11 +407,10 @@ bool UModularRig::Execute_Internal(const FName& InEventName)
 				// belong to the module's namespace
 				if (InEventName == FRigUnit_InteractionExecution::EventName)
 				{
-					const FString ModuleNamespace = Module->GetNamespace();
 					const bool bIsInteracting = UnitContext.ElementsBeingInteracted.ContainsByPredicate(
-						[ModuleNamespace, Hierarchy](const FRigElementKey& InteractionElement)
+						[Module, Hierarchy](const FRigElementKey& InteractionElement)
 						{
-							return ModuleNamespace == Hierarchy->GetNameSpace(InteractionElement);
+							return Module->Name == Hierarchy->GetModuleFName(InteractionElement);
 						});
 					if (!bIsInteracting)
 					{
@@ -528,7 +546,7 @@ void UModularRig::ExecuteQueue()
 					}
 
 					// put the variable defaults back
-					if(const FRigModuleReference* ModuleReference = GetModularRigModel().FindModule(ExecutionElement.ModulePath))
+					if(const FRigModuleReference* ModuleReference = GetModularRigModel().FindModule(ExecutionElement.ModuleName))
 					{
 						for (const TPair<FName, FString>& Variable : ModuleReference->ConfigValues)
 						{
@@ -538,11 +556,10 @@ void UModularRig::ExecuteQueue()
 				}
 
 				// Update the interaction elements to show only the ones belonging to this module
-				const FString ModuleNamespace = FString::Printf(TEXT("%s:"), *ExecutionElement.ModulePath);
 				RigUnitContext.ElementsBeingInteracted = RigUnitContext.ElementsBeingInteracted.FilterByPredicate(
-					[ModuleNamespace, Hierarchy](const FRigElementKey& Key)
+					[&ExecutionElement, Hierarchy](const FRigElementKey& Key)
 				{
-					return ModuleNamespace == Hierarchy->GetNameSpace(Key);
+					return ExecutionElement.ModuleName == Hierarchy->GetModuleFName(Key);
 				});
 				RigUnitContext.InteractionType = RigUnitContext.ElementsBeingInteracted.IsEmpty() ?
 					(uint8) EControlRigInteractionType::None
@@ -611,11 +628,11 @@ void UModularRig::OnObjectsReplaced(const TMap<UObject*, UObject*>& OldToNewInst
 		UControlRig* NewRig = CastChecked<UControlRig>(NewObject);
 
 		// relying on GetFName since URigVMHost is overloading GetName()
-		const FString& Path = NewRig->GetFName().ToString();
+		const FName& ModuleName = NewRig->GetFName();
 
 		// if we find a matching module update it.
 		// FRigModuleInstance::SetRig takes care of disregarding the previous module instance.
-		if(FRigModuleInstance* Module = const_cast<FRigModuleInstance*>(FindModule(Path)))
+		if(FRigModuleInstance* Module = const_cast<FRigModuleInstance*>(FindModule(ModuleName)))
 		{
 			Module->SetRig(NewRig);
 			NewRig->bCopyHierarchyBeforeConstruction = false;
@@ -653,7 +670,6 @@ void UModularRig::ResetModules(bool bDestroyModuleRigs)
 	RootModules.Reset();
 	Modules.Reset();
 	SupportedEvents.Reset();
-	ResetShortestDisplayPathCache();
 }
 
 const FModularRigModel& UModularRig::GetModularRigModel() const
@@ -668,11 +684,35 @@ const FModularRigModel& UModularRig::GetModularRigModel() const
 
 void UModularRig::UpdateCachedChildren()
 {
+	bool bContainsDeprecatedParentPaths = false;
+	for (FRigModuleInstance& Module : Modules)
+	{
+		if(!Module.ParentPath_DEPRECATED.IsEmpty())
+		{
+			bContainsDeprecatedParentPaths = true;
+			break;
+		}
+	}
+
+	TMap<FName, FRigModuleInstance*> NameToModule;
+	NameToModule.Reserve(Modules.Num());
+
 	TMap<FString, FRigModuleInstance*> PathToModule;
+	if(bContainsDeprecatedParentPaths)
+	{
+		PathToModule.Reserve(Modules.Num());
+	}
+
 	for (FRigModuleInstance& Module : Modules)
 	{
+		Module.CachedParentModule = nullptr;
 		Module.CachedChildren.Reset();
-		PathToModule.Add(Module.GetPath(), &Module);
+		NameToModule.Add(Module.Name, &Module);
+
+		if(bContainsDeprecatedParentPaths)
+		{
+			PathToModule.Add(Module.GetModulePath_Deprecated(), &Module);
+		}
 	}
 	
 	RootModules.Reset();
@@ -684,10 +724,16 @@ void UModularRig::UpdateCachedChildren()
 		}
 		else
 		{
-			if (FRigModuleInstance** ParentModule = PathToModule.Find(Module.ParentPath))
+			if (FRigModuleInstance** ParentModule = NameToModule.Find(Module.ParentModuleName))
 			{
+				Module.CachedParentModule = *ParentModule;
 				(*ParentModule)->CachedChildren.Add(&Module);
 			}
+			else if (FRigModuleInstance** ParentModuleByPath = PathToModule.Find(Module.GetModulePath_Deprecated()))
+			{
+				Module.CachedParentModule = *ParentModuleByPath;
+				(*ParentModuleByPath)->CachedChildren.Add(&Module);
+			}
 		}
 	}
 }
@@ -720,16 +766,23 @@ TArray<FString> UModularRig::GetModulePaths() const
 	{
 		// don't need to use AddUnique since module paths are
 		// guaranteed to be unique already.
-		Paths.Add(Module->GetPath());
+		Paths.Add(Module->GetModuleReference()->GetModulePath());
 		return true;
 	});
 	return Paths;
 }
 
-void UModularRig::EmitPostConstructionEventFinished()
+TArray<FName> UModularRig::GetModuleNames() const
 {
-	RecomputeShortestDisplayPathCache();
-	Super::EmitPostConstructionEventFinished();
+	TArray<FName> Names;
+	Names.Reserve(Modules.Num());
+	ForEachModule([&Names](const FRigModuleInstance* Module) -> bool
+	{
+		// module names are unique
+		Names.Add(Module->Name);
+		return true;
+	});
+	return Names;
 }
 
 FRigModuleInstance* UModularRig::AddModuleInstance(const FName& InModuleName, TSubclassOf<UControlRig> InModuleClass, const FRigModuleInstance* InParent,
@@ -763,31 +816,22 @@ FRigModuleInstance* UModularRig::AddModuleInstance(const FName& InModuleName, TS
 		return nullptr;
 	}
 
-	// InParent ptr will become invalid once we add the new module in the Modules array
-	// Get the information we need before that.
-	FString ParentPath;
-	if (InParent)
-	{
-		ParentPath = InParent->GetPath();
-	}
-
 	FRigModuleInstance& NewModule = Modules.Add_GetRef(FRigModuleInstance());
 	NewModule.Name = InModuleName;
-	if (!ParentPath.IsEmpty())
+	if (InParent)
 	{
-		NewModule.ParentPath = ParentPath;
+		NewModule.ParentModuleName = InParent->Name;
 	}
-	const FString Name = NewModule.GetPath();
 
 	UControlRig* NewModuleRig = nullptr;
 
 	// reuse existing module rig instances first
-	if(UControlRig** ExistingModuleRigPtr = PreviousModuleRigs.Find(Name))
+	if(UControlRig** ExistingModuleRigPtr = PreviousModuleRigs.Find(InModuleName))
 	{
 		if(UControlRig* ExistingModuleRig = *ExistingModuleRigPtr)
 		{
 			// again relying on GetFName since RigVMHost overloads GetName
-			if(ExistingModuleRig->GetFName().ToString().Equals(Name) && ExistingModuleRig->GetClass() == InModuleClass)
+			if(ExistingModuleRig->GetFName() == InModuleName && ExistingModuleRig->GetClass() == InModuleClass)
 			{
 				NewModuleRig = ExistingModuleRig;
 			}
@@ -795,19 +839,18 @@ FRigModuleInstance* UModularRig::AddModuleInstance(const FName& InModuleName, TS
 			{
 				DiscardModuleRig(ExistingModuleRig);
 			}
-			PreviousModuleRigs.Remove(Name);
+			PreviousModuleRigs.Remove(InModuleName);
 		}
 	}
 
 	if(NewModuleRig == nullptr)
 	{
-		NewModuleRig = NewObject<UControlRig>(this, InModuleClass, *Name);
+		NewModuleRig = NewObject<UControlRig>(this, InModuleClass, InModuleName);
 	}
 	
 	NewModule.SetRig(NewModuleRig);
 
 	UpdateCachedChildren();
-	ResetShortestDisplayPathCache();
 	for (const FName& EventName : NewModule.GetRig()->GetSupportedEvents())
 	{
 		SupportedEvents.AddUnique(EventName);
@@ -823,8 +866,8 @@ FRigModuleInstance* UModularRig::AddModuleInstance(const FName& InModuleName, TS
 		NewModuleRig->SetDynamicHierarchy(Hierarchy);
 		ModulePublicContext.Hierarchy = Hierarchy;
 		ModulePublicContext.ControlRig = this;
-		ModulePublicContext.RigModuleNameSpace = NewModuleRig->GetRigModuleNameSpace();
-		ModulePublicContext.RigModuleNameSpaceHash = GetTypeHash(ModulePublicContext.RigModuleNameSpace);
+		ModulePublicContext.RigModulePrefix = NewModuleRig->GetRigModulePrefix();
+		ModulePublicContext.RigModulePrefixHash = GetTypeHash(ModulePublicContext.RigModulePrefix);
 		NewModuleRig->SetElementKeyRedirector(FRigElementKeyRedirector(InConnectionMap, Hierarchy));
 
 		for (const TPair<FName, FString>& Variable : InVariableDefaultValues )
@@ -836,25 +879,25 @@ FRigModuleInstance* UModularRig::AddModuleInstance(const FName& InModuleName, TS
 	return &NewModule;
 }
 
-bool UModularRig::SetModuleVariableBindings(const FString& InModulePath, const TMap<FName, FString>& InVariableBindings)
+bool UModularRig::SetModuleVariableBindings(const FName& InModuleName, const TMap<FName, FString>& InVariableBindings)
 {
-	if(FRigModuleInstance* Module = const_cast<FRigModuleInstance*>(FindModule(InModulePath)))
+	if(FRigModuleInstance* Module = const_cast<FRigModuleInstance*>(FindModule(InModuleName)))
 	{
 		Module->VariableBindings.Reset();
 		
 		for (const TPair<FName, FString>& Pair : InVariableBindings)
 		{
-			FString SourceModulePath, SourceVariableName = Pair.Value;
-			(void)URigHierarchy::SplitNameSpace(Pair.Value, &SourceModulePath, &SourceVariableName);
+			FString SourceModuleName, SourceVariableName = Pair.Value;
+			(void)FRigHierarchyModulePath(Pair.Value).Split(&SourceModuleName, &SourceVariableName);
 			FRigVMExternalVariable SourceVariable;
-			if (SourceModulePath.IsEmpty())
+			if (SourceModuleName.IsEmpty())
 			{
 				if (const FProperty* Property = GetClass()->FindPropertyByName(*SourceVariableName))
 				{
 					SourceVariable = FRigVMExternalVariable::Make(Property, (UObject*)this);
 				}
 			}
-			else if(const FRigModuleInstance* SourceModule = FindModule(SourceModulePath))
+			else if(const FRigModuleInstance* SourceModule = FindModule(*SourceModuleName))
 			{
 				SourceVariable = SourceModule->GetRig()->GetPublicVariableByName(*SourceVariableName);
 			}
@@ -903,111 +946,26 @@ void UModularRig::DiscardModuleRig(UControlRig* InControlRig)
 	}
 }
 
-void UModularRig::ResetShortestDisplayPathCache() const
+const FRigModuleInstance* UModularRig::FindModule(const FName& InModuleName) const
 {
-	ElementKeyToShortestDisplayPath.Reset();
+	return const_cast<UModularRig*>(this)->FindModule(InModuleName);
 }
 
-void UModularRig::RecomputeShortestDisplayPathCache() const
+FRigModuleInstance* UModularRig::FindModule(const FName& InModuleName)
 {
-	ResetShortestDisplayPathCache();
-	
-	const URigHierarchy* Hierarchy = GetHierarchy();
-	check(Hierarchy);
-
-	TArray<FRigElementKey> AllKeys = Hierarchy->GetAllKeys();
-
-	// add all of the keys from the connections. we may be in a situation where
-	// the keys in the connections are actually referring to invalid / unloaded elements
-	// but we still want to show pretty short names.
-	for(const FModularRigSingleConnection& Connection : ModularRigModel.Connections)
+	FRigModuleInstance* Result = Modules.FindByPredicate([InModuleName](const FRigModuleInstance& Module)
 	{
-		AllKeys.AddUnique(Connection.Connector);
-		for(const FRigElementKey& Target : Connection.Targets)
-		{
-			AllKeys.AddUnique(Target);
-		}
-	}
-
-	auto GetNameForElement = [Hierarchy](const FRigElementKey& InElementKey)
-	{
-		const FName DesiredName = Hierarchy->GetNameMetadata(InElementKey, URigHierarchy::DesiredNameMetadataName, NAME_None);
-		if(!DesiredName.IsNone())
-		{
-			return DesiredName;
-		}
-		return InElementKey.Name;
-	};
-	
-	TMap<FName, bool> IsNameUniqueInHierarchy;
-	for(const FRigElementKey& Key : AllKeys)
-	{
-		auto UpdateNameUniqueness = [](const FName& InName, TMap<FName, bool>& UniqueMap)
-		{
-			if(bool* Unique = UniqueMap.Find(InName))
-			{
-				*Unique = false;
-			}
-			else
-			{
-				UniqueMap.Add(InName, true);
-			}
-		};
-		UpdateNameUniqueness(GetNameForElement(Key), IsNameUniqueInHierarchy);
-	}
-
-	const FModularRigModel& Model = GetModularRigModel();
-	for(const FRigElementKey& Key : AllKeys)
+		return Module.Name == InModuleName;
+	});
+	if(Result == nullptr)
 	{
-		const FName Name = GetNameForElement(Key);
-		const FName ModulePath = Hierarchy->GetModulePathFName(Key);
-		
-		if(!ModulePath.IsNone())
-		{
-			if(const FRigModuleReference* Module = Model.FindModule(ModulePath.ToString()))
-			{
-				const FString ModulePathString = ModulePath.ToString();
-				const FString ModulePathPrefix = ModulePath.ToString() + NamespaceSeparator;
-				FString NameString = Name.ToString();
-				if(NameString.StartsWith(ModulePathPrefix))
-				{
-					NameString = NameString.Mid(ModulePathPrefix.Len());
-				}
-				const FString ModuleShortName = Module->GetShortName();
-				const FString NameSpacedName = URigHierarchy::JoinNameSpace(ModuleShortName, NameString);
-				ElementKeyToShortestDisplayPath.Add(Key, { NameSpacedName, IsNameUniqueInHierarchy.FindChecked(Name) ? NameString : NameSpacedName });
-				continue;
-			}
-		}
-
-		if(IsNameUniqueInHierarchy.FindChecked(Name))
-		{
-			const FString NameString = Name.ToString();
-			ElementKeyToShortestDisplayPath.Add(Key, { NameString, NameString});
-		}
-		else
+		if(InModuleName.ToString().EndsWith(FRigHierarchyModulePath::ModuleNameSuffix))
 		{
-			const FString NameString = Key.ToString();
-			ElementKeyToShortestDisplayPath.Add(Key, { Name.ToString(), Name.ToString()});
+			// we should remove this and make sure clients call with proper names
+			Result = FindModule(*InModuleName.ToString().LeftChop(1));
 		}
 	}
-}
-
-const FRigModuleInstance* UModularRig::FindModule(const FString& InPath) const
-{
-	return const_cast<UModularRig*>(this)->FindModule(InPath);
-}
-
-FRigModuleInstance* UModularRig::FindModule(const FString& InPath)
-{
-	if(InPath.EndsWith(NamespaceSeparator))
-	{
-		return FindModule(InPath.Left(InPath.Len() - 1));
-	}
-	return Modules.FindByPredicate([InPath](const FRigModuleInstance& Module)
-	{
-		return Module.GetPath() == InPath;
-	});
+	return Result;
 }
 
 const FRigModuleInstance* UModularRig::FindModule(const UControlRig* InModuleInstance) const
@@ -1055,59 +1013,60 @@ FRigModuleInstance* UModularRig::FindModule(const FRigElementKey& InElementKey)
 {
 	if(const URigHierarchy* Hierarchy = GetHierarchy())
 	{
-		const FString ModulePath = Hierarchy->GetModulePath(InElementKey);
-		if(!ModulePath.IsEmpty())
+		const FName ModuleName = Hierarchy->GetModuleFName(InElementKey);
+		if(!ModuleName.IsNone())
 		{
-			return FindModule(ModulePath);
+			return FindModule(ModuleName);
 		}
 	}
 	return nullptr;
 }
 
-const UControlRig* UModularRig::GetModuleRig(FString InModulePath) const
+const UControlRig* UModularRig::GetModuleRig(FName InModuleName) const
 {
-	return const_cast<UModularRig*>(this)->GetModuleRig(InModulePath);
+	return const_cast<UModularRig*>(this)->GetModuleRig(InModuleName);
 }
 
-UControlRig* UModularRig::GetModuleRig(FString InModulePath)
+UControlRig* UModularRig::GetModuleRig(FName InModuleName)
 {
-	if(const FRigModuleInstance* Module = FindModule(InModulePath))
+	if(const FRigModuleInstance* Module = FindModule(InModuleName))
 	{
 		UpdateModuleVariables(Module);
 		if(UControlRig* ModuleRig = Module->GetRig())
 		{
 			return ModuleRig;
 		}
-		UE_LOG(LogControlRig, Error, TEXT("Module '%s' doesn't contain a rig instance."), *InModulePath);
+		UE_LOG(LogControlRig, Error, TEXT("Module '%s' doesn't contain a rig instance."), *InModuleName.ToString());
 		return nullptr;
 	}
-	UE_LOG(LogControlRig, Error, TEXT("Module '%s' doesn't exist."), *InModulePath);
+	UE_LOG(LogControlRig, Error, TEXT("Module '%s' doesn't exist."), *InModuleName.ToString());
 	return nullptr;
 }
 
 FString UModularRig::GetParentPathForBP(FString InModulePath) const
 {
-	return GetParentPath(InModulePath);
-}
-
-FString UModularRig::GetParentPath(const FString& InModulePath) const
-{
-	if (const FRigModuleInstance* Element = FindModule(InModulePath))
+	if (const FRigModuleInstance* Module = FindModule(*InModulePath))
 	{
-		return Element->ParentPath;
+		if(const FRigModuleReference* Reference = Module->GetModuleReference())
+		{
+			return Reference->GetModulePath();
+		}
 	}
 	return FString();
 }
 
-const FString& UModularRig::GetShortestDisplayPathForElement(const FRigElementKey& InElementKey, bool bAlwaysShowNameSpace) const
+FName UModularRig::GetParentModuleNameForBP(FName InModuleName) const
 {
-	if(const TTuple<FString, FString>* ShortestDisplayPaths = ElementKeyToShortestDisplayPath.Find(InElementKey))
+	return GetParentModuleName(InModuleName);
+}
+
+FName UModularRig::GetParentModuleName(const FName& InModuleName) const
+{
+	if (const FRigModuleInstance* Module = FindModule(InModuleName))
 	{
-		return bAlwaysShowNameSpace ? ShortestDisplayPaths->Get<0>() : ShortestDisplayPaths->Get<1>();
+		return Module->ParentModuleName;
 	}
-
-	static const FString EmptyPath;
-	return EmptyPath;
+	return NAME_None;
 }
 
 void UModularRig::ForEachModule(TFunctionRef<bool(FRigModuleInstance*)> PerModuleFunction)
@@ -1165,7 +1124,7 @@ void UModularRig::ExecuteConnectorEvent(const FRigElementKey& InConnector, const
 	FControlRigExecuteContext& PublicContext = GetRigVMExtendedExecuteContext().GetPublicDataSafe<FControlRigExecuteContext>();
 
 	FString ShortConnectorName = InConnector.Name.ToString();
-	ShortConnectorName.RemoveFromStart(Module->GetNamespace());
+	ShortConnectorName.RemoveFromStart(Module->GetModulePrefix());
 	TGuardValue<FRigElementKey> ConnectorGuard(PublicContext.UnitContext.ConnectionResolve.Connector, FRigElementKey(*ShortConnectorName, InConnector.Type));
 	TGuardValue<TArray<FRigElementResolveResult>> CandidatesGuard(PublicContext.UnitContext.ConnectionResolve.Matches, Candidates);
 	TGuardValue<TArray<FRigElementResolveResult>> MatchesGuard(PublicContext.UnitContext.ConnectionResolve.Excluded, {});
@@ -1195,16 +1154,16 @@ TArray<FName> UModularRig::GetEventsForAllModules() const
 	return Events;
 }
 
-TArray<FName> UModularRig::GetEventsForModule(FString InModulePath) const
+TArray<FName> UModularRig::GetEventsForModule(FName InModuleName) const
 {
-	if(const UControlRig* ModuleRig = GetModuleRig(InModulePath))
+	if(const UControlRig* ModuleRig = GetModuleRig(InModuleName))
 	{
 		return ModuleRig->GetSupportedEvents();
 	}
 	return TArray<FName>();
 }
 
-TArray<FString> UModularRig::ExecuteEventOnAllModules(FName InEvent)
+TArray<FName> UModularRig::ExecuteEventOnAllModules(FName InEvent)
 {
 	TArray<FRigModuleExecutionElement> QueueForEvent;
 
@@ -1220,7 +1179,7 @@ TArray<FString> UModularRig::ExecuteEventOnAllModules(FName InEvent)
 		return true;
 	});
 
-	TArray<FString> ModulesWhichRanEvent;
+	TArray<FName> ModulesWhichRanEvent;
 	if(QueueForEvent.IsEmpty())
 	{
 		UE_LOG(LogControlRig, Error, TEXT("Event '%s' is not supported by any module on this modular rig."), *InEvent.ToString());
@@ -1235,20 +1194,20 @@ TArray<FString> UModularRig::ExecuteEventOnAllModules(FName InEvent)
 	{
 		if(ExecutionElement.bExecuted)
 		{
-			ModulesWhichRanEvent.Add(ExecutionElement.ModulePath);
+			ModulesWhichRanEvent.Add(ExecutionElement.ModuleName);
 		}
 		else
 		{
-			UE_LOG(LogControlRig, Error, TEXT("Module '%s' did not run event '%s' successfully."), *ExecutionElement.ModulePath, *InEvent.ToString());
+			UE_LOG(LogControlRig, Error, TEXT("Module '%s' did not run event '%s' successfully."), *ExecutionElement.ModuleName.ToString(), *InEvent.ToString());
 		}
 	}
 		
 	return ModulesWhichRanEvent;
 }
 
-bool UModularRig::ExecuteEventOnModuleForBP(FName InEvent, FString InModulePath)
+bool UModularRig::ExecuteEventOnModuleForBP(FName InEvent, FName InModuleName)
 {
-	if(FRigModuleInstance* Module = FindModule(InModulePath))
+	if(FRigModuleInstance* Module = FindModule(InModuleName))
 	{
 		if(const UControlRig* ModuleRig = Module->GetRig())
 		{
@@ -1256,13 +1215,13 @@ bool UModularRig::ExecuteEventOnModuleForBP(FName InEvent, FString InModulePath)
 			{
 				return ExecuteEventOnModule(InEvent, Module); 
 			}
-			UE_LOG(LogControlRig, Error, TEXT("Module '%s' doesn't support the event '%s'."), *InModulePath, *InEvent.ToString());
+			UE_LOG(LogControlRig, Error, TEXT("Module '%s' doesn't support the event '%s'."), *InModuleName.ToString(), *InEvent.ToString());
 			return false;
 		}
-		UE_LOG(LogControlRig, Error, TEXT("Module '%s' doesn't contain a rig instance."), *InModulePath);
+		UE_LOG(LogControlRig, Error, TEXT("Module '%s' doesn't contain a rig instance."), *InModuleName.ToString());
 		return false;
 	}
-	UE_LOG(LogControlRig, Error, TEXT("Module '%s' doesn't exist."), *InModulePath);
+	UE_LOG(LogControlRig, Error, TEXT("Module '%s' doesn't exist."), *InModuleName.ToString());
 	return false;
 }
 
