Control Rig: Unique modular rig module names

diff --git a/Engine/Plugins/Animation/ControlRig/Source/ControlRig/Private/ModularRigController.cpp b/Engine/Plugins/Animation/ControlRig/Source/ControlRig/Private/ModularRigController.cpp
--- a/Engine/Plugins/Animation/ControlRig/Source/ControlRig/Private/ModularRigController.cpp
+++ b/Engine/Plugins/Animation/ControlRig/Source/ControlRig/Private/ModularRigController.cpp
@@ -25,19 +25,19 @@ UModularRigController::UModularRigController(const FObjectInitializer& ObjectIni
 {
 }
 
-FString UModularRigController::AddModule(const FName& InModuleName, TSubclassOf<UControlRig> InClass, const FString& InParentModulePath, bool bSetupUndo)
+FName UModularRigController::AddModule(const FName& InModuleName, TSubclassOf<UControlRig> InClass, const FName& InParentModuleName, bool bSetupUndo)
 {
 	if (!InClass)
 	{
 		UE_LOG(LogControlRig, Error, TEXT("Invalid InClass"));
-		return FString();
+		return NAME_None;
 	}
 
 	UControlRig* ClassDefaultObject = InClass->GetDefaultObject<UControlRig>();
 	if (!ClassDefaultObject->IsRigModule())
 	{
 		UE_LOG(LogControlRig, Error, TEXT("Class %s is not a rig module"), *InClass->GetClassPathName().ToString());
-		return FString();
+		return NAME_None;
 	}
 
 #if WITH_EDITOR
@@ -52,42 +52,11 @@ FString UModularRigController::AddModule(const FName& InModuleName, TSubclassOf<
 	}
 #endif 
 
-	FRigModuleReference* NewModule = nullptr;
-	if (InParentModulePath.IsEmpty())
-	{
-		for (FRigModuleReference* Module : Model->RootModules)
-		{
-			if (Module->Name.ToString() == InModuleName)
-			{
-				return FString();
-			}
-		}
-
-		Model->Modules.Add(FRigModuleReference(InModuleName, InClass, FString()));
-		NewModule = &Model->Modules.Last();
-	}
-	else if (FRigModuleReference* ParentModule = FindModule(InParentModulePath))
-	{
-		for (FRigModuleReference* Module : ParentModule->CachedChildren)
-		{
-			if (Module->Name.ToString() == InModuleName)
-			{
-				return FString();
-			}
-		}
-
-		Model->Modules.Add(FRigModuleReference(InModuleName, InClass, ParentModule->GetPath()));
-		NewModule = &Model->Modules.Last();
-	}
+	const FName SanitizedName = GetSafeNewName(InModuleName);
+	Model->Modules.Add(FRigModuleReference(SanitizedName, InClass, InParentModuleName, Model));
+	FRigModuleReference* NewModule = &Model->Modules.Last();
 
 	Model->UpdateCachedChildren();
-	UpdateShortNames();
-
-	if (!NewModule)
-	{
-		UE_LOG(LogControlRig, Error, TEXT("Error while creating module %s"), *InModuleName.ToString());
-		return FString();
-	}
 
 	Notify(EModularRigNotification::ModuleAdded, NewModule);
 
@@ -95,33 +64,34 @@ FString UModularRigController::AddModule(const FName& InModuleName, TSubclassOf<
 	TransactionPtr.Reset();
 #endif
 
-	return NewModule->GetPath();
+	return NewModule->GetFName();
 }
 
-FRigModuleReference* UModularRigController::FindModule(const FString& InPath)
+FRigModuleReference* UModularRigController::FindModule(const FName& InModuleName)
 {
-	return Model->FindModule(InPath);
+	return Model->FindModule(InModuleName);
 }
 
-const FRigModuleReference* UModularRigController::FindModule(const FString& InPath) const
+const FRigModuleReference* UModularRigController::FindModule(const FName& InModuleName) const
 {
-	return const_cast<UModularRigController*>(this)->FindModule(InPath);
+	return const_cast<UModularRigController*>(this)->FindModule(InModuleName);
 }
 
-FRigModuleReference UModularRigController::GetModuleReference(FString InModulePath) const
+FRigModuleReference UModularRigController::GetModuleReference(FName InModuleName) const
 {
-	if(const FRigModuleReference* Module = FindModule(InModulePath))
+	if(const FRigModuleReference* Module = FindModule(InModuleName))
 	{
 		return *Module;
 	}
 	return FRigModuleReference();
 }
 
-TArray<FRigElementKey> UModularRigController::GetConnectorsForModule(FString InModulePath) const
+TArray<FRigElementKey> UModularRigController::GetConnectorsForModule(FName InModuleName) const
 {
 	TArray<FRigElementKey> ConnectorsForModule;
 
-	if(const FRigModuleReference* Module = FindModule(InModulePath))
+	const FString ModuleNameString = InModuleName.ToString();
+	if(const FRigModuleReference* Module = FindModule(InModuleName))
 	{
 		UControlRig* RigCDO = Module->Class->GetDefaultObject<UControlRig>();
 		if (!RigCDO)
@@ -133,12 +103,12 @@ TArray<FRigElementKey> UModularRigController::GetConnectorsForModule(FString InM
 		const TArray<FRigModuleConnector>& ExposedConnectors = RigCDO->GetRigModuleSettings().ExposedConnectors;
 		for(const FRigModuleConnector& ExposedConnector : ExposedConnectors)
 		{
-			ConnectorsForModule.Emplace(*URigHierarchy::JoinNameSpace(InModulePath, ExposedConnector.Name), ERigElementType::Connector);
+			ConnectorsForModule.Emplace(FRigHierarchyModulePath(ModuleNameString, ExposedConnector.Name).GetPathFName(), ERigElementType::Connector);
 		}
 	}
 	else
 	{
-		UE_LOG(LogControlRig, Error, TEXT("Could not find module %s"), *InModulePath);
+		UE_LOG(LogControlRig, Error, TEXT("Could not find module %s"), *ModuleNameString);
 	}
 
 	return ConnectorsForModule;
@@ -152,17 +122,17 @@ bool UModularRigController::CanConnectConnectorToElement(const FRigElementKey& I
 bool UModularRigController::CanConnectConnectorToElements(const FRigElementKey& InConnectorKey, const TArray<FRigElementKey>& InTargetKeys,
 	FText& OutErrorMessage)
 {
-	FString ConnectorModulePath, ConnectorName;
-	if (!URigHierarchy::SplitNameSpace(InConnectorKey.Name.ToString(), &ConnectorModulePath, &ConnectorName))
+	const FRigHierarchyModulePath ConnectorModulePath(InConnectorKey.Name);
+	if (!ConnectorModulePath.IsValid())
 	{
-		OutErrorMessage = FText::FromString(FString::Printf(TEXT("Connector %s does not contain a namespace"), *InConnectorKey.ToString()));
+		OutErrorMessage = FText::FromString(FString::Printf(TEXT("Connector %s does not contain a module / namespace"), *InConnectorKey.ToString()));
 		return false;
 	}
 	
-	FRigModuleReference* Module = FindModule(ConnectorModulePath);
+	FRigModuleReference* Module = FindModule(ConnectorModulePath.GetModuleFName());
 	if (!Module)
 	{
-		OutErrorMessage = FText::FromString(FString::Printf(TEXT("Could not find module %s"), *ConnectorModulePath));
+		OutErrorMessage = FText::FromString(FString::Printf(TEXT("Could not find module %s"), *ConnectorModulePath.GetModuleNameString()));
 		return false;
 	}
 
@@ -174,13 +144,13 @@ bool UModularRigController::CanConnectConnectorToElements(const FRigElementKey&
 	}
 
 	const FRigModuleConnector* ModuleConnector = RigCDO->GetRigModuleSettings().ExposedConnectors.FindByPredicate(
-		[ConnectorName](FRigModuleConnector& Connector)
+		[&ConnectorModulePath](FRigModuleConnector& Connector)
 		{
-			return Connector.Name == ConnectorName;
+			return Connector.Name == ConnectorModulePath.GetElementFName();
 		});
 	if (!ModuleConnector)
 	{
-		OutErrorMessage = FText::FromString(FString::Printf(TEXT("Could not find connector %s in class %s"), *ConnectorName, *Module->Class->GetPathName()));
+		OutErrorMessage = FText::FromString(FString::Printf(TEXT("Could not find connector %s in class %s"), *ConnectorModulePath.GetElementNameString(), *Module->Class->GetPathName()));
 		return false;
 	}
 
@@ -218,8 +188,8 @@ bool UModularRigController::CanConnectConnectorToElements(const FRigElementKey&
 				return Connector.IsPrimary();
 			});
 
-			const FString PrimaryConnectorPath = FString::Printf(TEXT("%s:%s"), *ConnectorModulePath, *PrimaryModuleConnector->Name);
-			const FRigElementKey PrimaryConnectorKey(*PrimaryConnectorPath, ERigElementType::Connector);
+			const FRigHierarchyModulePath PrimaryConnectorPath = ConnectorModulePath.ReplaceElementName(PrimaryModuleConnector->Name);
+			const FRigElementKey PrimaryConnectorKey(PrimaryConnectorPath.GetPathFName(), ERigElementType::Connector);
 			const FRigElementKey PrimaryTarget = Model->Connections.FindTargetFromConnector(PrimaryConnectorKey);
 			if (!PrimaryTarget.IsValid())
 			{
@@ -271,7 +241,7 @@ bool UModularRigController::CanConnectConnectorToElements(const FRigElementKey&
 			return false;
 		}
 
-		const FRigModuleInstance* ModuleInstance = ModularRig->FindModule(Module->GetPath());
+		const FRigModuleInstance* ModuleInstance = ModularRig->FindModule(Module->GetFName());
 		FModularRigResolveResult RuleResults = RuleManager->FindMatches(Connector, ModuleInstance, ModularRig->GetElementKeyRedirector());
 		for(const FRigElementKey& InTargetKey : InTargetKeys)
 		{
@@ -317,16 +287,16 @@ bool UModularRigController::ConnectConnectorToElements(const FRigElementKey& InC
 		return false;
 	}
 	
-	FString ConnectorParentPath, ConnectorName;
-	(void)URigHierarchy::SplitNameSpace(InConnectorKey.Name.ToString(), &ConnectorParentPath, &ConnectorName);
-	FRigModuleReference* Module = FindModule(ConnectorParentPath);
+
+	const FRigHierarchyModulePath ConnectorModulePath(InConnectorKey.Name);
+	FRigModuleReference* Module = FindModule( ConnectorModulePath.GetModuleFName());
 
 	FRigElementKey CurrentTarget = Model->Connections.FindTargetFromConnector(InConnectorKey);
 
 	UBlueprint* Blueprint = Cast<UBlueprint>(GetOuter());
 
 #if WITH_EDITOR
-	TOptional<FName> TargetModulePathName;
+	TOptional<FName> TargetModuleName;
 	UModularRig* ModularRig = GetDebuggedModularRig();
 	if(ModularRig)
 	{
@@ -334,17 +304,17 @@ bool UModularRigController::ConnectConnectorToElements(const FRigElementKey& InC
 		{
 			for(const FRigElementKey& TargetKey : TargetKeys)
 			{
-				if(TargetModulePathName.IsSet())
+				if(TargetModuleName.IsSet())
 				{
-					if(TargetModulePathName.GetValue() != Hierarchy->GetModulePathFName(TargetKey))
+					if(TargetModuleName.GetValue() != Hierarchy->GetModuleFName(TargetKey))
 					{
-						TargetModulePathName.Reset();
+						TargetModuleName.Reset();
 						break;
 					}
 				}
 				else
 				{
-					TargetModulePathName = Hierarchy->GetModulePathFName(TargetKey);
+					TargetModuleName = Hierarchy->GetModuleFName(TargetKey);
 				}
 			}
 		}
@@ -403,12 +373,12 @@ bool UModularRigController::ConnectConnectorToElements(const FRigElementKey& InC
 					{
 						if(Connector->IsPrimary())
 						{
-							if(!TargetModulePathName.Get(NAME_None).IsNone())
+							if(!TargetModuleName.Get(NAME_None).IsNone())
 							{
-								const FString NewModulePath = ReparentModule(Module->GetPath(), TargetModulePathName.GetValue().ToString(), bSetupUndo);
-								if(!NewModulePath.IsEmpty())
+								const FName ModuleName = Module->Name;
+								if(ReparentModule(ModuleName, TargetModuleName.GetValue(), bSetupUndo))
 								{
-									Module = FindModule(NewModulePath);
+									Module = FindModule(ModuleName);
 								}
 							}
 						}
@@ -417,7 +387,7 @@ bool UModularRigController::ConnectConnectorToElements(const FRigElementKey& InC
 
 				if (Module && bAutoResolveOtherConnectors && bResolvedPrimaryConnector)
 				{
-					(void)AutoConnectModules( {Module->GetPath()}, false, bSetupUndo);
+					(void)AutoConnectModules( {Module->Name}, false, bSetupUndo);
 				}
 			}
 		}
@@ -444,9 +414,8 @@ bool UModularRigController::AddTargetToArrayConnector(const FRigElementKey& InCo
 		return ConnectConnectorToElements(InConnectorKey, ExistingTargets, bSetupUndo, bAutoResolveOtherConnectors, bCheckValidConnection);
 	}
 
-	FString ConnectorParentPath, ConnectorName;
-	(void)URigHierarchy::SplitNameSpace(InConnectorKey.Name.ToString(), &ConnectorParentPath, &ConnectorName);
-	const FRigModuleReference* Module = FindModule(ConnectorParentPath);
+	const FRigHierarchyModulePath ConnectorModulePath(InConnectorKey.Name);
+	const FRigModuleReference* Module = FindModule(ConnectorModulePath.GetModuleFName());
 	if(Module == nullptr)
 	{
 		return false;
@@ -481,17 +450,17 @@ bool UModularRigController::DisconnectConnector(const FRigElementKey& InConnecto
 bool UModularRigController::DisconnectConnector_Internal(const FRigElementKey& InConnectorKey, bool bDisconnectSubModules,
 	TMap<FRigElementKey, TArray<FRigElementKey>>* OutRemovedConnections, bool bSetupUndo)
 {
-	FString ConnectorModulePath, ConnectorName;
-	if (!URigHierarchy::SplitNameSpace(InConnectorKey.Name.ToString(), &ConnectorModulePath, &ConnectorName))
+	const FRigHierarchyModulePath ConnectorModulePath(InConnectorKey.Name);
+	if (!ConnectorModulePath.IsValid())
 	{
 		UE_LOG(LogControlRig, Error, TEXT("Connector %s does not contain a namespace"), *InConnectorKey.ToString());
 		return false;
 	}
 	
-	FRigModuleReference* Module = FindModule(ConnectorModulePath);
+	FRigModuleReference* Module = FindModule(ConnectorModulePath.GetModuleFName());
 	if (!Module)
 	{
-		UE_LOG(LogControlRig, Error, TEXT("Could not find module %s"), *ConnectorModulePath);
+		UE_LOG(LogControlRig, Error, TEXT("Could not find module %s"), *ConnectorModulePath.GetModuleNameString());
 		return false;
 	}
 
@@ -502,6 +471,7 @@ bool UModularRigController::DisconnectConnector_Internal(const FRigElementKey& I
 		return false;
 	}
 
+	const FName ConnectorName = ConnectorModulePath.GetElementFName();
 	const FRigModuleConnector* ModuleConnector = RigCDO->GetRigModuleSettings().ExposedConnectors.FindByPredicate(
 		[ConnectorName](FRigModuleConnector& Connector)
 		{
@@ -509,7 +479,7 @@ bool UModularRigController::DisconnectConnector_Internal(const FRigElementKey& I
 		});
 	if (!ModuleConnector)
 	{
-		UE_LOG(LogControlRig, Error, TEXT("Could not find connector %s in class %s"), *ConnectorName, *Module->Class->GetPathName());
+		UE_LOG(LogControlRig, Error, TEXT("Could not find connector %s in class %s"), *ConnectorName.ToString(), *Module->Class->GetPathName());
 		return false;
 	}
 
@@ -541,7 +511,7 @@ bool UModularRigController::DisconnectConnector_Internal(const FRigElementKey& I
 		TArray<FRigElementKey> ConnectionsToRemove;
 		for (const FModularRigSingleConnection& Connection : Model->Connections)
 		{
-			if (Connection.Connector.Name.ToString().StartsWith(ConnectorModulePath, ESearchCase::IgnoreCase))
+			if (FRigHierarchyModulePath(Connection.Connector.Name).HasModuleName(ConnectorModulePath.GetModuleName()))
 			{
 				ConnectionsToRemove.Add(Connection.Connector);
 			}
@@ -559,11 +529,11 @@ bool UModularRigController::DisconnectConnector_Internal(const FRigElementKey& I
 	{
 		// Remove connections from child modules
 		TArray<FRigElementKey> ConnectionsToRemove;
+		const FName ConnectorModuleFName = ConnectorModulePath.GetModuleFName();
 		for (const FModularRigSingleConnection& Connection : Model->Connections)
 		{
-			FString OtherConnectorModulePath, OtherConnectorName;
-			(void)URigHierarchy::SplitNameSpace(Connection.Connector.Name.ToString(), &OtherConnectorModulePath, &OtherConnectorName);
-			if (OtherConnectorModulePath.StartsWith(ConnectorModulePath, ESearchCase::IgnoreCase) && OtherConnectorModulePath.Len() > ConnectorModulePath.Len())
+			const FRigHierarchyModulePath OtherModulePath(Connection.Connector.Name);
+			if(Model->IsModuleParentedTo(OtherModulePath.GetModuleFName(), ConnectorModuleFName))
 			{
 				ConnectionsToRemove.Add(Connection.Connector);
 			}
@@ -585,7 +555,7 @@ bool UModularRigController::DisconnectConnector_Internal(const FRigElementKey& I
 	{
 		if(ModuleConnector->IsPrimary() && !Module->IsRootModule())
 		{
-			(void)ReparentModule(Module->GetPath(), FString(), bSetupUndo);
+			(void)ReparentModule(Module->Name, NAME_None, bSetupUndo);
 		}
 	}
 
@@ -621,29 +591,30 @@ TArray<FRigElementKey> UModularRigController::DisconnectCyclicConnectors(bool bS
 	TArray<FRigElementKey> ConnectorsToDisconnect;
 	for (const FModularRigSingleConnection& Connection : Model->Connections)
 	{
-		const FString ConnectorModulePath = Hierarchy->GetModulePath(Connection.Connector);
-
-		for(const FRigElementKey& Target : Connection.Targets)
+		const FName ConnectorModuleName = Hierarchy->GetModuleFName(Connection.Connector);
+		if(const FRigModuleReference* ConnectorModule = Model->FindModule(ConnectorModuleName))
 		{
-			const FString TargetModulePath = Hierarchy->GetModulePath(Target);
-
-			// targets in the base hierarchy are always allowed
-			if(TargetModulePath.IsEmpty())
+			for(const FRigElementKey& Target : Connection.Targets)
 			{
-				continue;
-			}
+				const FName TargetModuleName = Hierarchy->GetModuleFName(Target);
 
-			const FRigModuleReference* ConnectorModule = Model->FindModule(ConnectorModulePath);
-			const FRigModuleReference* TargetModule = Model->FindModule(TargetModulePath);
-			if(ConnectorModule == nullptr || TargetModule == nullptr || ConnectorModule == TargetModule)
-			{
-				continue;
-			}
+				// targets in the base hierarchy are always allowed
+				if(TargetModuleName.IsNone())
+				{
+					continue;
+				}
 
-			if(!Model->IsModuleParentedTo(ConnectorModule, TargetModule))
-			{
-				ConnectorsToDisconnect.Add(Connection.Connector);
-				break;
+				const FRigModuleReference* TargetModule = Model->FindModule(TargetModuleName);
+				if(TargetModule == nullptr || ConnectorModule == TargetModule)
+				{
+					continue;
+				}
+
+				if(!Model->IsModuleParentedTo(ConnectorModule, TargetModule))
+				{
+					ConnectorsToDisconnect.Add(Connection.Connector);
+					break;
+				}
 			}
 		}
 	}
@@ -715,18 +686,18 @@ bool UModularRigController::AutoConnectSecondaryConnectors(const TArray<FRigElem
 	bool bResolvedAllConnectors = true;
 	for(const FRigElementKey& ConnectorKey : InConnectorKeys)
 	{
-		const FString ModulePath = Hierarchy->GetModulePath(ConnectorKey);
-		if(ModulePath.IsEmpty())
+		const FName ModuleName = Hierarchy->GetModuleFName(ConnectorKey);
+		if(ModuleName.IsNone())
 		{
-			UE_LOG(LogControlRig, Error, TEXT("Connector %s has no associated module path"), *ConnectorKey.ToString());
+			UE_LOG(LogControlRig, Error, TEXT("Connector %s has no associated module"), *ConnectorKey.ToString());
 			bResolvedAllConnectors = false;
 			continue;
 		}
 
-		const FRigModuleReference* Module = Model->FindModule(ModulePath);
+		const FRigModuleReference* Module = Model->FindModule(ModuleName);
 		if(Module == nullptr)
 		{
-			UE_LOG(LogControlRig, Error, TEXT("Could not find module %s"), *ModulePath);
+			UE_LOG(LogControlRig, Error, TEXT("Could not find module %s"), *ModuleName.ToString());
 			bResolvedAllConnectors = false;
 			continue;
 		}
@@ -734,7 +705,7 @@ bool UModularRigController::AutoConnectSecondaryConnectors(const TArray<FRigElem
 		const FRigConnectorElement* PrimaryConnector = Module->FindPrimaryConnector(Hierarchy);
 		if(PrimaryConnector == nullptr)
 		{
-			UE_LOG(LogControlRig, Error, TEXT("Module %s has no primary connector"), *ModulePath);
+			UE_LOG(LogControlRig, Error, TEXT("Module %s has no primary connector"), *ModuleName.ToString());
 			bResolvedAllConnectors = false;
 			continue;
 		}
@@ -748,7 +719,7 @@ bool UModularRigController::AutoConnectSecondaryConnectors(const TArray<FRigElem
 		
 		if(!Model->Connections.HasConnection(PrimaryConnectorKey))
 		{
-			UE_LOG(LogControlRig, Warning, TEXT("Module %s's primary connector is not resolved"), *ModulePath);
+			UE_LOG(LogControlRig, Warning, TEXT("Module %s's primary connector is not resolved"), *ModuleName.ToString());
 			bResolvedAllConnectors = false;
 			continue;
 		}
@@ -756,13 +727,13 @@ bool UModularRigController::AutoConnectSecondaryConnectors(const TArray<FRigElem
 		const UControlRig* RigCDO = Module->Class->GetDefaultObject<UControlRig>();
 		if(RigCDO == nullptr)
 		{
-			UE_LOG(LogControlRig, Error, TEXT("Module %s has no default rig assigned"), *ModulePath);
+			UE_LOG(LogControlRig, Error, TEXT("Module %s has no default rig assigned"), *ModuleName.ToString());
 			bResolvedAllConnectors = false;
 			continue;
 		}
 
 		const UModularRigRuleManager* RuleManager = Hierarchy->GetRuleManager();
-		const FRigModuleInstance* ModuleInstance = ModularRig->FindModule(Module->GetPath());
+		const FRigModuleInstance* ModuleInstance = ModularRig->FindModule(Module->Name);
 		
 		if (bReplaceExistingConnections || !Model->Connections.HasConnection(ConnectorKey))
 		{
@@ -810,7 +781,7 @@ bool UModularRigController::AutoConnectSecondaryConnectors(const TArray<FRigElem
 #endif
 }
 
-bool UModularRigController::AutoConnectModules(const TArray<FString>& InModulePaths, bool bReplaceExistingConnections, bool bSetupUndo)
+bool UModularRigController::AutoConnectModules(const TArray<FName>& InModuleNames, bool bReplaceExistingConnections, bool bSetupUndo)
 {
 #if WITH_EDITOR
 	TArray<FRigElementKey> ConnectorKeys;
@@ -836,12 +807,12 @@ bool UModularRigController::AutoConnectModules(const TArray<FString>& InModulePa
 		return false;
 	}
 
-	for(const FString& ModulePath : InModulePaths)
+	for(const FName& ModuleName : InModuleNames)
 	{
-		const FRigModuleReference* Module = FindModule(ModulePath);
+		const FRigModuleReference* Module = FindModule(ModuleName);
 		if (!Module)
 		{
-			UE_LOG(LogControlRig, Error, TEXT("Could not find module %s"), *ModulePath);
+			UE_LOG(LogControlRig, Error, TEXT("Could not find module %s"), *ModuleName.ToString());
 			return false;
 		}
 
@@ -864,31 +835,31 @@ bool UModularRigController::AutoConnectModules(const TArray<FString>& InModulePa
 #endif
 }
 
-bool UModularRigController::SetConfigValueInModule(const FString& InModulePath, const FName& InVariableName, const FString& InValue, bool bSetupUndo)
+bool UModularRigController::SetConfigValueInModule(const FName& InModuleName, const FName& InVariableName, const FString& InValue, bool bSetupUndo)
 {
-	FRigModuleReference* Module = FindModule(InModulePath);
+	FRigModuleReference* Module = FindModule(InModuleName);
 	if (!Module)
 	{
-		UE_LOG(LogControlRig, Error, TEXT("Could not find module %s"), *InModulePath);
+		UE_LOG(LogControlRig, Error, TEXT("Could not find module %s"), *InModuleName.ToString());
 		return false;
 	}
 
 	if (!Module->Class.IsValid())
 	{
-		UE_LOG(LogControlRig, Error, TEXT("Class defined in module %s is not valid"), *InModulePath);
+		UE_LOG(LogControlRig, Error, TEXT("Class defined in module %s is not valid"), *InModuleName.ToString());
 		return false;
 	}
 
 	const FProperty* Property = Module->Class->FindPropertyByName(InVariableName);
 	if (!Property)
 	{
-		UE_LOG(LogControlRig, Error, TEXT("Could not find variable %s in module %s"), *InVariableName.ToString(), *InModulePath);
+		UE_LOG(LogControlRig, Error, TEXT("Could not find variable %s in module %s"), *InVariableName.ToString(), *InModuleName.ToString());
 		return false;
 	}
 
 	if (Property->HasAllPropertyFlags(CPF_BlueprintReadOnly))
 	{
-		UE_LOG(LogControlRig, Error, TEXT("The target variable %s in module %s is read only"), *InVariableName.ToString(), *InModulePath);
+		UE_LOG(LogControlRig, Error, TEXT("The target variable %s in module %s is read only"), *InVariableName.ToString(), *InModuleName.ToString());
 		return false;
 	}
 
@@ -900,7 +871,7 @@ bool UModularRigController::SetConfigValueInModule(const FString& InModulePath,
 
 	if (!FBlueprintEditorUtils::PropertyValueFromString_Direct(Property, InValue, TempMemory))
 	{
-		UE_LOG(LogControlRig, Error, TEXT("Value %s for variable %s in module %s is not valid"), *InValue, *InVariableName.ToString(), *InModulePath);
+		UE_LOG(LogControlRig, Error, TEXT("Value %s for variable %s in module %s is not valid"), *InValue, *InVariableName.ToString(), *InModuleName.ToString());
 		return false;
 	}
 
@@ -926,10 +897,10 @@ bool UModularRigController::SetConfigValueInModule(const FString& InModulePath,
 	return true;
 }
 
-TArray<FString> UModularRigController::GetPossibleBindings(const FString& InModulePath, const FName& InVariableName)
+TArray<FString> UModularRigController::GetPossibleBindings(const FName& InModuleName, const FName& InVariableName)
 {
 	TArray<FString> PossibleBindings;
-	const FRigModuleReference* Module = FindModule(InModulePath);
+	const FRigModuleReference* Module = FindModule(InModuleName);
 	if (!Module)
 	{
 		return PossibleBindings;
@@ -959,7 +930,7 @@ TArray<FString> UModularRigController::GetPossibleBindings(const FString& InModu
 		{
 			FText ErrorMessage;
 			const FString VariableName = Variable.Name.ToString();
-			if (CanBindModuleVariable(InModulePath, InVariableName, VariableName, ErrorMessage))
+			if (CanBindModuleVariable(InModuleName, InVariableName, VariableName, ErrorMessage))
 			{
 				PossibleBindings.Add(VariableName);
 			}
@@ -967,15 +938,14 @@ TArray<FString> UModularRigController::GetPossibleBindings(const FString& InModu
 	}
 
 	// Add possible module variables
-	const FString InvalidModulePrefix = InModulePath + UModularRig::NamespaceSeparator;
-	Model->ForEachModule([this, &PossibleBindings, InModulePath, InVariableName, InvalidModulePrefix](const FRigModuleReference* InModule) -> bool
+	Model->ForEachModule([this, &PossibleBindings, InModuleName, InVariableName](const FRigModuleReference* InModule) -> bool
 	{
-		const FString CurModulePath = InModule->GetPath();
-		if (InModulePath != CurModulePath && !CurModulePath.StartsWith(InvalidModulePrefix, ESearchCase::IgnoreCase))
+		const FName CurModuleName = InModule->Name;
+		if (InModuleName != CurModuleName)
 		{
 			if (!InModule->Class.IsValid())
 			{
-				InModule->Class.LoadSynchronous();
+				(void)InModule->Class.LoadSynchronous();
 			}
 			if (InModule->Class.IsValid())
 			{
@@ -983,8 +953,8 @@ TArray<FString> UModularRigController::GetPossibleBindings(const FString& InModu
 			   for (const FRigVMExternalVariable& Variable : Variables)
 			   {
 				   FText ErrorMessage;
-				   const FString SourceVariablePath = URigHierarchy::JoinNameSpace(CurModulePath, Variable.Name.ToString());
-				   if (CanBindModuleVariable(InModulePath, InVariableName, SourceVariablePath, ErrorMessage))
+				   const FRigHierarchyModulePath SourceVariablePath(CurModuleName, Variable.Name);
+				   if (CanBindModuleVariable(InModuleName, InVariableName, SourceVariablePath, ErrorMessage))
 				   {
 					   PossibleBindings.Add(SourceVariablePath);
 				   }
@@ -997,50 +967,50 @@ TArray<FString> UModularRigController::GetPossibleBindings(const FString& InModu
 	return PossibleBindings;
 }
 
-bool UModularRigController::CanBindModuleVariable(const FString& InModulePath, const FName& InVariableName, const FString& InSourcePath, FText& OutErrorMessage)
+bool UModularRigController::CanBindModuleVariable(const FName& InModuleName, const FName& InVariableName, const FString& InSourcePath, FText& OutErrorMessage)
 {
-	FRigModuleReference* Module = FindModule(InModulePath);
+	FRigModuleReference* Module = FindModule(InModuleName);
 	if (!Module)
 	{
-		OutErrorMessage = FText::FromString(FString::Printf(TEXT("Could not find module %s"), *InModulePath));
+		OutErrorMessage = FText::FromString(FString::Printf(TEXT("Could not find module %s"), *InModuleName.ToString()));
 		return false;
 	}
 
 	if (!Module->Class.IsValid())
 	{
-		OutErrorMessage = FText::FromString(FString::Printf(TEXT("Class defined in module %s is not valid"), *InModulePath));
+		OutErrorMessage = FText::FromString(FString::Printf(TEXT("Class defined in module %s is not valid"), *InModuleName.ToString()));
 		return false;
 	}
 
 	const FProperty* TargetProperty = Module->Class->FindPropertyByName(InVariableName);
 	if (!TargetProperty)
 	{
-		OutErrorMessage = FText::FromString(FString::Printf(TEXT("Could not find variable %s in module %s"), *InVariableName.ToString(), *InModulePath));
+		OutErrorMessage = FText::FromString(FString::Printf(TEXT("Could not find variable %s in module %s"), *InVariableName.ToString(), *InModuleName.ToString()));
 		return false;
 	}
 
 	if (TargetProperty->HasAnyPropertyFlags(CPF_BlueprintReadOnly | CPF_DisableEditOnInstance))
 	{
-		OutErrorMessage = FText::FromString(FString::Printf(TEXT("The target variable %s in module %s is read only"), *InVariableName.ToString(), *InModulePath));
+		OutErrorMessage = FText::FromString(FString::Printf(TEXT("The target variable %s in module %s is read only"), *InVariableName.ToString(), *InModuleName.ToString()));
 		return false;
 	}
 
-	FString SourceModulePath, SourceVariableName = InSourcePath;
-	(void)URigHierarchy::SplitNameSpace(InSourcePath, &SourceModulePath, &SourceVariableName);
+	FString SourceModuleName, SourceVariableName = InSourcePath;
+	(void)FRigHierarchyModulePath(InSourcePath).Split(&SourceModuleName, &SourceVariableName);
 
 	FRigModuleReference* SourceModule = nullptr;
-	if (!SourceModulePath.IsEmpty())
+	if (!SourceModuleName.IsEmpty())
 	{
-		SourceModule = FindModule(SourceModulePath);
+		SourceModule = FindModule(*SourceModuleName);
 		if (!SourceModule)
 		{
-			OutErrorMessage = FText::FromString(FString::Printf(TEXT("Could not find source module %s"), *SourceModulePath));
+			OutErrorMessage = FText::FromString(FString::Printf(TEXT("Could not find source module %s"), *SourceModuleName));
 			return false;
 		}
 
-		if (SourceModulePath.StartsWith(InModulePath, ESearchCase::IgnoreCase))
+		if(Model->IsModuleParentedTo(SourceModule->Name, InModuleName))
 		{
-			OutErrorMessage = FText::FromString(FString::Printf(TEXT("Cannot bind variable of module %s to a variable of module %s because the source module is a child of the target module"), *InModulePath, *SourceModulePath));
+			OutErrorMessage = FText::FromString(FString::Printf(TEXT("Cannot bind variable of module %s to a variable of module %s because the source module is a child of the target module"), *InModuleName.ToString(), *SourceModuleName));
 			return false;
 		}
 	}
@@ -1063,10 +1033,10 @@ bool UModularRigController::CanBindModuleVariable(const FString& InModulePath, c
 		return false;
 	}
 
-	FString SourcePath = (SourceModulePath.IsEmpty()) ? SourceVariableName : URigHierarchy::JoinNameSpace(SourceModulePath, SourceVariableName);
 	if (!RigVMTypeUtils::AreCompatible(SourceProperty, TargetProperty))
 	{
-		FString TargetPath = FString::Printf(TEXT("%s.%s"), *InModulePath, *InVariableName.ToString());
+		const FString SourcePath = (SourceModuleName.IsEmpty()) ? SourceVariableName : FRigHierarchyModulePath(SourceModuleName, SourceVariableName).GetPath();
+		const FString TargetPath = FString::Printf(TEXT("%s.%s"), *InModuleName.ToString(), *InVariableName.ToString());
 		OutErrorMessage = FText::FromString(FString::Printf(TEXT("Property %s of type %s and %s of type %s are not compatible"), *SourcePath, *SourceProperty->GetCPPType(), *TargetPath, *TargetProperty->GetCPPType()));
 		return false;
 	}
@@ -1074,25 +1044,25 @@ bool UModularRigController::CanBindModuleVariable(const FString& InModulePath, c
 	return true;
 }
 
-bool UModularRigController::BindModuleVariable(const FString& InModulePath, const FName& InVariableName, const FString& InSourcePath, bool bSetupUndo)
+bool UModularRigController::BindModuleVariable(const FName& InModuleName, const FName& InVariableName, const FString& InSourcePath, bool bSetupUndo)
 {
 	FText ErrorMessage;
-	if (!CanBindModuleVariable(InModulePath, InVariableName, InSourcePath, ErrorMessage))
+	if (!CanBindModuleVariable(InModuleName, InVariableName, InSourcePath, ErrorMessage))
 	{
-		UE_LOG(LogControlRig, Error, TEXT("Could not bind module variable %s : %s"), *URigHierarchy::JoinNameSpace(InModulePath, InVariableName.ToString()), *ErrorMessage.ToString());
+		UE_LOG(LogControlRig, Error, TEXT("Could not bind module variable %s : %s"), *FRigHierarchyModulePath(InModuleName, InVariableName).GetPath(), *ErrorMessage.ToString());
 		return false;
 	}
 	
-	FRigModuleReference* Module = FindModule(InModulePath);
+	FRigModuleReference* Module = FindModule(InModuleName);
 	const FProperty* TargetProperty = Module->Class->FindPropertyByName(InVariableName);
 
-	FString SourceModulePath, SourceVariableName = InSourcePath;
-	(void)URigHierarchy::SplitNameSpace(InSourcePath, &SourceModulePath, &SourceVariableName);
+	FString SourceModuleName, SourceVariableName = InSourcePath;
+	(void)FRigHierarchyModulePath(InSourcePath).Split(&SourceModuleName, &SourceVariableName);
 
 	FRigModuleReference* SourceModule = nullptr;
-	if (!SourceModulePath.IsEmpty())
+	if (!SourceModuleName.IsEmpty())
 	{
-		SourceModule = FindModule(SourceModulePath);
+		SourceModule = FindModule(*SourceModuleName);
 	}
 
 	const FProperty* SourceProperty = nullptr;
@@ -1108,7 +1078,7 @@ bool UModularRigController::BindModuleVariable(const FString& InModulePath, cons
 		}
 	}
 
-	FString SourcePath = (SourceModulePath.IsEmpty()) ? SourceVariableName : URigHierarchy::JoinNameSpace(SourceModulePath, SourceVariableName);
+	FString SourcePath = (SourceModuleName.IsEmpty()) ? SourceVariableName : FRigHierarchyModulePath(SourceModuleName, SourceVariableName).GetPath();
 
 #if WITH_EDITOR
 	TSharedPtr<FScopedTransaction> TransactionPtr;
@@ -1134,18 +1104,18 @@ bool UModularRigController::BindModuleVariable(const FString& InModulePath, cons
 	return true;
 }
 
-bool UModularRigController::UnBindModuleVariable(const FString& InModulePath, const FName& InVariableName, bool bSetupUndo)
+bool UModularRigController::UnBindModuleVariable(const FName& InModuleName, const FName& InVariableName, bool bSetupUndo)
 {
-	FRigModuleReference* Module = FindModule(InModulePath);
+	FRigModuleReference* Module = FindModule(InModuleName);
 	if (!Module)
 	{
-		UE_LOG(LogControlRig, Error, TEXT("Could not find module %s"), *InModulePath);
+		UE_LOG(LogControlRig, Error, TEXT("Could not find module %s"), *InModuleName.ToString());
 		return false;
 	}
 
 	if (!Module->Bindings.Contains(InVariableName))
 	{
-		UE_LOG(LogControlRig, Error, TEXT("Variable %s in module %s is not bound"), *InVariableName.ToString(), *InModulePath);
+		UE_LOG(LogControlRig, Error, TEXT("Variable %s in module %s is not bound"), *InVariableName.ToString(), *InModuleName.ToString());
 		return false;
 	}
 
@@ -1172,12 +1142,12 @@ bool UModularRigController::UnBindModuleVariable(const FString& InModulePath, co
 	return true;
 }
 
-bool UModularRigController::DeleteModule(const FString& InModulePath, bool bSetupUndo)
+bool UModularRigController::DeleteModule(const FName& InModuleName, bool bSetupUndo)
 {
-	FRigModuleReference* Module = FindModule(InModulePath);
+	FRigModuleReference* Module = FindModule(InModuleName);
 	if (!Module)
 	{
-		UE_LOG(LogControlRig, Error, TEXT("Could not find module %s"), *InModulePath);
+		UE_LOG(LogControlRig, Error, TEXT("Could not find module %s"), *InModuleName.ToString());
 		return false;
 	}
 
@@ -1193,39 +1163,37 @@ bool UModularRigController::DeleteModule(const FString& InModulePath, bool bSetu
 	}
 #endif
 
-	(void)DeselectModule(Module->GetPath());
+	(void)DeselectModule(Module->Name);
 
 	// Delete children
-	TArray<FString> ChildrenPaths;
-	Algo::Transform(Module->CachedChildren, ChildrenPaths, [](const FRigModuleReference* Child){ return Child->GetPath(); });
-	for (const FString& ChildPath : ChildrenPaths)
+	TArray<FName> ChildrenNames;
+	Algo::Transform(Module->CachedChildren, ChildrenNames, [](const FRigModuleReference* Child){ return Child->Name; });
+	for (const FName& ChildName : ChildrenNames)
 	{
-		DeleteModule(ChildPath, bSetupUndo);
+		DeleteModule(ChildName, bSetupUndo);
 	}
 
 	Model->DeletedModules.Add(*Module);
 	Model->Modules.RemoveSingle(*Module);
 	Model->UpdateCachedChildren();
-	UpdateShortNames();
+
+	const FString InModuleNameString = InModuleName.ToString();
 
 	// Fix connections
 	{
 		TArray<FRigElementKey> ToRemove;
 		for (FModularRigSingleConnection& Connection : Model->Connections)
 		{
-			FString ConnectionModulePath, ConnectionName;
-			(void)URigHierarchy::SplitNameSpace(Connection.Connector.Name.ToString(), &ConnectionModulePath, &ConnectionName);
-
-			if (ConnectionModulePath == InModulePath)
+			const FRigHierarchyModulePath ConnectionModulePath(Connection.Connector.Name);
+			if (ConnectionModulePath.HasModuleName(InModuleNameString))
 			{
 				ToRemove.Add(Connection.Connector);
 			}
 
 			for(const FRigElementKey& Target : Connection.Targets)
 			{
-				FString TargetModulePath, TargetName;
-				(void)URigHierarchy::SplitNameSpace(Target.Name.ToString(), &TargetModulePath, &TargetName);
-				if (TargetModulePath == InModulePath)
+				const FRigHierarchyModulePath TargetModulePath(Target.Name);
+				if (TargetModulePath.HasModuleName(InModuleNameString))
 				{
 					ToRemove.Add(Connection.Connector);
 					break;
@@ -1242,11 +1210,10 @@ bool UModularRigController::DeleteModule(const FString& InModulePath, bool bSetu
 	// Fix bindings
 	for (FRigModuleReference& Reference : Model->Modules)
 	{
-		Reference.Bindings = Reference.Bindings.FilterByPredicate([InModulePath](const TPair<FName, FString>& Binding)
+		Reference.Bindings = Reference.Bindings.FilterByPredicate([InModuleNameString](const TPair<FName, FString>& Binding)
 		{
-			FString ModulePath, VariableName = Binding.Value;
-			(void)URigHierarchy::SplitNameSpace(Binding.Value, &ModulePath, &VariableName);
-			if (ModulePath == InModulePath)
+			const FRigHierarchyModulePath BindingModulePath(Binding.Value);
+			if (BindingModulePath.HasModuleName(InModuleNameString))
 			{
 				return false;
 			}
@@ -1264,27 +1231,27 @@ bool UModularRigController::DeleteModule(const FString& InModulePath, bool bSetu
 	return false;
 }
 
-FString UModularRigController::RenameModule(const FString& InModulePath, const FName& InNewName, bool bSetupUndo)
+FName UModularRigController::RenameModule(const FName& InModuleName, const FName& InNewName, bool bSetupUndo)
 {
-	FRigModuleReference* Module = FindModule(InModulePath);
+	FRigModuleReference* Module = FindModule(InModuleName);
 	if (!Module)
 	{
-		UE_LOG(LogControlRig, Error, TEXT("Could not find module %s"), *InModulePath);
-		return FString();
+		UE_LOG(LogControlRig, Error, TEXT("Could not find module %s"), *InModuleName.ToString());
+		return NAME_None;
 	}
 
-	const FString OldName = Module->Name.ToString();
-	const FString NewName = InNewName.ToString();
-	if (OldName.Equals(NewName))
+	const FName OldName = Module->Name;
+	const FName NewName = InNewName;
+	if (OldName == NewName)
 	{
-		return Module->GetPath();
+		return Module->Name;
 	}
 
 	FText ErrorMessage;
-	if (!CanRenameModule(InModulePath, InNewName, ErrorMessage))
+	if (!CanRenameModule(InModuleName, InNewName, ErrorMessage))
 	{
-		UE_LOG(LogControlRig, Error, TEXT("Could not rename module %s: %s"), *InModulePath, *ErrorMessage.ToString());
-		return FString();
+		UE_LOG(LogControlRig, Error, TEXT("Could not rename module %s: %s"), *InModuleName.ToString(), *ErrorMessage.ToString());
+		return NAME_None;
 	}
 	
 #if WITH_EDITOR
@@ -1299,10 +1266,7 @@ FString UModularRigController::RenameModule(const FString& InModulePath, const F
 	}
 #endif
 	
-	const FString OldPath = (Module->ParentPath.IsEmpty()) ? OldName : URigHierarchy::JoinNameSpace(Module->ParentPath, OldName);
-	const FString NewPath = (Module->ParentPath.IsEmpty()) ? *NewName :  URigHierarchy::JoinNameSpace(Module->ParentPath, NewName);
-
-	const int32 SelectionIndex = Model->SelectedModulePaths.Find(OldPath);
+	const int32 SelectionIndex = Model->SelectedModuleNames.Find(OldName);
 	if(SelectionIndex != INDEX_NONE)
 	{
 		Notify(EModularRigNotification::ModuleDeselected, Module);
@@ -1315,26 +1279,28 @@ FString UModularRigController::RenameModule(const FString& InModulePath, const F
 	for (int32 i=0; i<Children.Num(); ++i)
 	{
 		FRigModuleReference* Child = Children[i];
-		Child->ParentPath.ReplaceInline(*OldPath, *NewPath);
-
+		if(Child->ParentModuleName == OldName)
+		{
+			Child->ParentModuleName = NewName;
+		}
 		Children.Append(Child->CachedChildren);
 	}
 
 	// Fix connections
 	{
-		const FString OldNamespace = OldPath + TEXT(":");
-		const FString NewNamespace = NewPath + TEXT(":");
 		for (FModularRigSingleConnection& Connection : Model->Connections)
 		{
-			if (Connection.Connector.Name.ToString().StartsWith(OldNamespace, ESearchCase::IgnoreCase))
+			FRigHierarchyModulePath ConnectorModulePath(Connection.Connector.Name);
+			if (ConnectorModulePath.ReplaceModuleNameInline(OldName, NewName))
 			{
-				Connection.Connector.Name = *FString::Printf(TEXT("%s%s"), *NewNamespace, *Connection.Connector.Name.ToString().RightChop(OldNamespace.Len()));
+				Connection.Connector.Name = ConnectorModulePath.GetPathFName();
 			}
 			for(FRigElementKey& Target : Connection.Targets)
 			{
-				if (Target.Name.ToString().StartsWith(OldNamespace, ESearchCase::IgnoreCase))
+				FRigHierarchyModulePath TargetModulePath(Target.Name);
+				if (TargetModulePath.ReplaceModuleNameInline(OldName, NewName))
 				{
-					Target.Name = *FString::Printf(TEXT("%s%s"), *NewNamespace, *Target.Name.ToString().RightChop(OldNamespace.Len()));
+					Target.Name = TargetModulePath.GetPathFName();
 				}
 			}
 		}
@@ -1346,21 +1312,28 @@ FString UModularRigController::RenameModule(const FString& InModulePath, const F
 	{
 		for (TPair<FName, FString>& Binding : Reference.Bindings)
 		{
-			FString ModulePath, VariableName = Binding.Value;
-			(void)URigHierarchy::SplitNameSpace(Binding.Value, &ModulePath, &VariableName);
-			if (ModulePath == OldPath)
+			FRigHierarchyModulePath BindingModulePath(Binding.Value);
+			if (BindingModulePath.ReplaceModuleNameInline(OldName, NewName))
 			{
-				Binding.Value = URigHierarchy::JoinNameSpace(NewPath, VariableName);
+				Binding.Value = BindingModulePath.GetPath();
 			}
 		};
 	}
 
-	UpdateShortNames();
+	// make sure to update our backwards compat code from path to module name
+	for(TPair<FRigHierarchyModulePath, FName>& Pair : Model->PreviousModulePaths)
+	{
+		if(Pair.Value == Module->PreviousName)
+		{
+			Pair.Value = Module->Name;
+		}
+	}
+
 	Notify(EModularRigNotification::ModuleRenamed, Module);
 
 	if(SelectionIndex != INDEX_NONE)
 	{
-		Model->SelectedModulePaths[SelectionIndex] = NewPath;
+		Model->SelectedModuleNames[SelectionIndex] = NewName;
 		Notify(EModularRigNotification::ModuleSelected, Module);
 	}
 
@@ -1368,10 +1341,10 @@ FString UModularRigController::RenameModule(const FString& InModulePath, const F
 	TransactionPtr.Reset();
 #endif
 	
-	return NewPath;
+	return NewName;
 }
 
-bool UModularRigController::CanRenameModule(const FString& InModulePath, const FName& InNewName, FText& OutErrorMessage) const
+bool UModularRigController::CanRenameModule(const FName& InModuleName, const FName& InNewName, FText& OutErrorMessage) const
 {
 	if (InNewName.IsNone() || InNewName.ToString().IsEmpty())
 	{
@@ -1379,21 +1352,27 @@ bool UModularRigController::CanRenameModule(const FString& InModulePath, const F
 		return false;
 	}
 
-	if(InNewName.ToString().Contains(UModularRig::NamespaceSeparator))
+	if(InNewName.ToString().Contains(FRigHierarchyModulePath::NamespaceSeparator_Deprecated))
 	{
 		OutErrorMessage = NSLOCTEXT("ModularRigController", "NameContainsNamespaceSeparator", "Name contains namespace separator ':'.");
 		return false;
 	}
 
-	const FRigModuleReference* Module = const_cast<UModularRigController*>(this)->FindModule(InModulePath);
+	if(InNewName.ToString().Contains(FRigHierarchyModulePath::ModuleNameSuffix))
+	{
+		OutErrorMessage = NSLOCTEXT("ModularRigController", "NameContainsModuleSuffix", "Name contains module suffix '/'.");
+		return false;
+	}
+
+	const FRigModuleReference* Module = const_cast<UModularRigController*>(this)->FindModule(InModuleName);
 	if (!Module)
 	{
-		OutErrorMessage = FText::FromString(FString::Printf(TEXT("Module %s not found."), *InModulePath));
+		OutErrorMessage = FText::FromString(FString::Printf(TEXT("Module %s not found."), *InModuleName.ToString()));
 		return false;
 	}
 
 	FString ErrorMessage;
-	if(!IsNameAvailable(Module->ParentPath, InNewName, &ErrorMessage))
+	if(!IsNameAvailable(InNewName, &ErrorMessage))
 	{
 		OutErrorMessage = FText::FromString(ErrorMessage);
 		return false;
@@ -1401,21 +1380,20 @@ bool UModularRigController::CanRenameModule(const FString& InModulePath, const F
 	return true;
 }
 
-FString UModularRigController::ReparentModule(const FString& InModulePath, const FString& InNewParentModulePath, bool bSetupUndo)
+bool UModularRigController::ReparentModule(const FName& InModuleName, const FName& InNewParentModuleName, bool bSetupUndo)
 {
-	FRigModuleReference* Module = FindModule(InModulePath);
+	FRigModuleReference* Module = FindModule(InModuleName);
 	if (!Module)
 	{
-		UE_LOG(LogControlRig, Error, TEXT("Could not find module %s"), *InModulePath);
-		return FString();
+		UE_LOG(LogControlRig, Error, TEXT("Could not find module %s"), *InModuleName.ToString());
+		return false;
 	}
 
-	const FRigModuleReference* NewParentModule = FindModule(InNewParentModulePath);
-	const FString PreviousParentPath = Module->ParentPath;
-	const FString ParentPath = (NewParentModule) ? NewParentModule->GetPath() : FString();
-	if(PreviousParentPath.Equals(ParentPath, ESearchCase::IgnoreCase))
+	const FRigModuleReference* NewParentModule = FindModule(InNewParentModuleName);
+	const FName PreviousParentModuleName = Module->ParentModuleName;
+	if(PreviousParentModuleName == InNewParentModuleName)
 	{
-		return Module->GetPath();
+		return false;
 	}
 
 #if WITH_EDITOR
@@ -1430,155 +1408,72 @@ FString UModularRigController::ReparentModule(const FString& InModulePath, const
 	}
 #endif
 
-	// Reparent or unparent children
-	const FString OldPath = Module->GetPath();
-
-	const int32 SelectionIndex = Model->SelectedModulePaths.Find(OldPath);
-	if(SelectionIndex != INDEX_NONE)
-	{
-		Notify(EModularRigNotification::ModuleDeselected, Module);
-	}
-	
-	Module->PreviousParentPath = Module->ParentPath;
-	Module->PreviousName = Module->Name;
-	Module->ParentPath = (NewParentModule) ? NewParentModule->GetPath() : FString();
-	Module->Name = GetSafeNewName(Module->ParentPath, FRigName(Module->Name));
-	const FString NewPath = Module->GetPath();
-
-	// Fix all the subtree namespaces
-	TArray<FRigModuleReference*> SubTree = Module->CachedChildren;
-	for (int32 Index=0; Index<SubTree.Num(); ++Index)
-	{
-		SubTree[Index]->ParentPath.ReplaceInline(*OldPath, *NewPath);
-		SubTree.Append(SubTree[Index]->CachedChildren);
-	}
+	Module->PreviousParentName = Module->ParentModuleName;
+	Module->ParentModuleName = (NewParentModule) ? NewParentModule->Name : NAME_None;
 
 	Model->UpdateCachedChildren();
-	UpdateShortNames();
-
-	// Fix connections
-	{
-		for (FModularRigSingleConnection& Connection : Model->Connections)
-		{
-			const TPair<FRigName, FRigName> ConnectorNamespace = URigHierarchy::SplitNameSpace(Connection.Connector.Name);
-			if (ConnectorNamespace.Key.ToString().StartsWith(OldPath, ESearchCase::IgnoreCase))
-			{
-				Connection.Connector.Name = *FString::Printf(TEXT("%s%s"), *NewPath, *Connection.Connector.Name.ToString().RightChop(OldPath.Len()));
-			}
-
-			for(FRigElementKey& Target : Connection.Targets)
-			{
-				const TPair<FRigName, FRigName> TargetNamespace = URigHierarchy::SplitNameSpace(Target.Name);
-				if (TargetNamespace.Key.ToString().StartsWith(OldPath, ESearchCase::IgnoreCase))
-				{
-					Target.Name = *FString::Printf(TEXT("%s%s"), *NewPath, *Target.Name.ToString().RightChop(OldPath.Len()));
-				}
-			}
-		}
-		Model->Connections.UpdateFromConnectionList();
-	}
-
-	// Fix bindings
-	for (FRigModuleReference& Reference : Model->Modules)
-	{
-		const FString ReferencePath = Reference.GetPath();
-		for (TPair<FName, FString>& Binding : Reference.Bindings)
-		{
-			FString ModulePath, VariableName = Binding.Value;
-			(void)URigHierarchy::SplitNameSpace(Binding.Value, &ModulePath, &VariableName);
-			if (ModulePath == OldPath)
-			{
-				Binding.Value = URigHierarchy::JoinNameSpace(NewPath, VariableName);
-				ModulePath = NewPath;
-			}
-
-			// Remove any child dependency
-			if (ModulePath.Contains(ReferencePath))
-			{
-				UE_LOG(LogControlRig, Warning, TEXT("Binding lost due to source %s contained in child module of %s"), *Binding.Value, *ReferencePath);
-				Binding.Value.Reset();
-			}
-		};
-
-		Reference.Bindings = Reference.Bindings.FilterByPredicate([](const TPair<FName, FString>& Binding)
-		{
-			return !Binding.Value.IsEmpty();
-		});
-	}
-
-	// Fix connectors in the hierarchies
-	
 
 	// since we've reparented the module now we should clear out all connectors which are cyclic
 	(void)DisconnectCyclicConnectors(bSetupUndo);
 
 	Notify(EModularRigNotification::ModuleReparented, Module);
 
-	if(SelectionIndex != INDEX_NONE)
-	{
-		Model->SelectedModulePaths[SelectionIndex] = NewPath;
-		Notify(EModularRigNotification::ModuleSelected, Module);
-	}
-
 #if WITH_EDITOR
  	TransactionPtr.Reset();
 #endif
 	
-	return NewPath;
+	return true;
 }
 
-FString UModularRigController::MirrorModule(const FString& InModulePath, const FRigVMMirrorSettings& InSettings, bool bSetupUndo)
+FName UModularRigController::MirrorModule(const FName& InModuleName, const FRigVMMirrorSettings& InSettings, bool bSetupUndo)
 {
-	FRigModuleReference* OriginalModule = FindModule(InModulePath);
+	FRigModuleReference* OriginalModule = FindModule(InModuleName);
 	if (!OriginalModule || !OriginalModule->Class.IsValid())
 	{
-		return FString();
+		return NAME_None;
 	}
 
-	FString NewModuleName = OriginalModule->Name.ToString();
+	FString NewModuleNameString = OriginalModule->Name.ToString();
 	if (!InSettings.SearchString.IsEmpty())
 	{
-		NewModuleName = NewModuleName.Replace(*InSettings.SearchString, *InSettings.ReplaceString, ESearchCase::CaseSensitive);
-		NewModuleName = GetSafeNewName(OriginalModule->ParentPath, FRigName(NewModuleName)).ToString();
+		NewModuleNameString = NewModuleNameString.Replace(*InSettings.SearchString, *InSettings.ReplaceString, ESearchCase::CaseSensitive);
+		NewModuleNameString = GetSafeNewName(FRigName(NewModuleNameString)).ToString();
 	}
 
 	// Before any changes, gather all the information we need from the OriginalModule, as the pointer might become invalid afterwards
-	const FRigElementKeyRedirector::FKeyMap OriginalConnectionMap = Model->Connections.GetModuleConnectionMap(InModulePath);
+	const FRigElementKeyRedirector::FKeyMap OriginalConnectionMap = Model->Connections.GetModuleConnectionMap(InModuleName);
 	const TMap<FName, FString> OriginalBindings = OriginalModule->Bindings;
 	const TSubclassOf<UControlRig> OriginalClass = OriginalModule->Class.Get();
-	const FString OriginalParentPath = OriginalModule->ParentPath;
+	const FName OriginalParentName = OriginalModule->ParentModuleName;
 	const TMap<FName, FString> OriginalConfigValues = OriginalModule->ConfigValues;
 
 	FModularRigControllerCompileBracketScope CompileBracketScope(this);
 
-	FString NewModulePath = AddModule(*NewModuleName, OriginalClass, OriginalParentPath, bSetupUndo);
-	FRigModuleReference* NewModule = FindModule(NewModulePath);
+	const FName NewModuleName = AddModule(*NewModuleNameString, OriginalClass, OriginalParentName, bSetupUndo);
+	FRigModuleReference* NewModule = FindModule(NewModuleName);
 	if (!NewModule)
 	{
-		return FString();
+		return NAME_None;
 	}
 
 	for (const  FRigElementKeyRedirector::FKeyPair& Pair : OriginalConnectionMap)
 	{
 		for(const FRigElementKey& Target : Pair.Value)
 		{
-			FString OriginalTargetPath = Target.Name.ToString();
-			FString NewTargetPath = OriginalTargetPath.Replace(*InSettings.SearchString, *InSettings.ReplaceString, ESearchCase::CaseSensitive);
-			FRigElementKey NewTargetKey(*NewTargetPath, Target.Type);
+			FString OriginalTargetName = Target.Name.ToString();
+			FString NewTargetName = OriginalTargetName.Replace(*InSettings.SearchString, *InSettings.ReplaceString, ESearchCase::CaseSensitive);
+			FRigElementKey NewTargetKey(*NewTargetName, Target.Type);
 
-			FString NewConnectorPath = URigHierarchy::JoinNameSpace(NewModulePath, Pair.Key.Name.ToString());
-			FRigElementKey NewConnectorKey(*NewConnectorPath, ERigElementType::Connector);
+			const FRigHierarchyModulePath NewConnectorPath(NewModuleName.ToString(), Pair.Key.Name.ToString());
+			FRigElementKey NewConnectorKey(NewConnectorPath.GetPathFName(), ERigElementType::Connector);
 			ConnectConnectorToElement(NewConnectorKey, NewTargetKey, bSetupUndo, false, false);
 		}
-
-		// Path might change after connecting
-		NewModulePath = NewModule->GetPath();
 	}
 
 	for (const TPair<FName, FString>& Pair : OriginalBindings)
 	{
 		FString NewSourcePath = Pair.Value.Replace(*InSettings.SearchString, *InSettings.ReplaceString, ESearchCase::CaseSensitive);
-		BindModuleVariable(NewModulePath, Pair.Key, NewSourcePath, bSetupUndo);
+		BindModuleVariable(NewModuleName, Pair.Key, NewSourcePath, bSetupUndo);
 	}
 
 	TSet<FName> ConfigValueSet;
@@ -1654,7 +1549,7 @@ FString UModularRigController::MirrorModule(const FString& InModulePath, const F
 		}
 
 		ConfigValueSet.Add(Property->GetFName());
-		SetConfigValueInModule(NewModulePath, Property->GetFName(), NewValueStr, bSetupUndo);
+		SetConfigValueInModule(NewModuleName, Property->GetFName(), NewValueStr, bSetupUndo);
 	}
 #endif
 
@@ -1663,80 +1558,19 @@ FString UModularRigController::MirrorModule(const FString& InModulePath, const F
 	{
 		if (!ConfigValueSet.Contains(Pair.Key))
 		{
-			SetConfigValueInModule(NewModulePath, Pair.Key, Pair.Value, bSetupUndo);
-		}
-	}
-	
-	return NewModulePath;
-}
-
-bool UModularRigController::SetModuleShortName(const FString& InModulePath, const FString& InNewShortName, bool bSetupUndo)
-{
-	FRigModuleReference* Module = FindModule(InModulePath);
-	if (!Module)
-	{
-		UE_LOG(LogControlRig, Error, TEXT("Could not find module %s"), *InModulePath);
-		return false;
-	}
-
-	FText ErrorMessage;
-	if (!CanSetModuleShortName(InModulePath, InNewShortName, ErrorMessage))
-	{
-		UE_LOG(LogControlRig, Error, TEXT("Could not rename module %s: %s"), *InModulePath, *ErrorMessage.ToString());
-		return false;
-	}
-
-	const FString OldShortName = Module->GetShortName();
-	const FString NewShortName = InNewShortName;
-	if (OldShortName.Equals(NewShortName))
-	{
-		return true;
-	}
-	
-#if WITH_EDITOR
-	TSharedPtr<FScopedTransaction> TransactionPtr;
-	if (bSetupUndo)
-	{
-		TransactionPtr = MakeShared<FScopedTransaction>(NSLOCTEXT("ModularRigController", "SetModuleShortNameTransaction", "Set Module Display Name"), !GIsTransacting);
-		if(UBlueprint* Blueprint = Cast<UBlueprint>(GetOuter()))
-		{
-			Blueprint->Modify();
+			SetConfigValueInModule(NewModuleName, Pair.Key, Pair.Value, bSetupUndo);
 		}
 	}
-#endif
-	
-	Module->ShortName = InNewShortName;
-	Module->bShortNameBasedOnPath = false;
-
-	Notify(EModularRigNotification::ModuleShortNameChanged, Module);
-
-	// update all other display named to avoid collision
-	UpdateShortNames();
-
-#if WITH_EDITOR
-	TransactionPtr.Reset();
-#endif
 	
-	return true;
+	return NewModuleName;
 }
 
-bool UModularRigController::CanSetModuleShortName(const FString& InModulePath, const FString& InNewShortName, FText& OutErrorMessage) const
+bool UModularRigController::SwapModuleClass(const FName& InModuleName, TSubclassOf<UControlRig> InNewClass, bool bSetupUndo)
 {
-	FString ErrorMessage;
-	if(!IsShortNameAvailable(FRigName(InNewShortName), &ErrorMessage))
-	{
-		OutErrorMessage = FText::FromString(ErrorMessage);
-		return false;
-	}
-	return true;
-}
-
-bool UModularRigController::SwapModuleClass(const FString& InModulePath, TSubclassOf<UControlRig> InNewClass, bool bSetupUndo)
-{
-	FRigModuleReference* Module = FindModule(InModulePath);
+	FRigModuleReference* Module = FindModule(InModuleName);
 	if (!Module)
 	{
-		UE_LOG(LogControlRig, Error, TEXT("Could not find module %s"), *InModulePath);
+		UE_LOG(LogControlRig, Error, TEXT("Could not find module %s"), *InModuleName.ToString());
 		return false;
 	}
 
@@ -1782,13 +1616,12 @@ bool UModularRigController::SwapModuleClass(const FString& InModulePath, TSubcla
 		TArray<FRigElementKey> ConnectionsToRemove;
 		for (const FModularRigSingleConnection& Connection : Connections)
 		{
-			FString Namespace, ConnectorName;
-			URigHierarchy::SplitNameSpace(Connection.Connector.Name.ToString(), &Namespace, &ConnectorName);
-			if (Namespace.Equals(InModulePath))
+			const FRigHierarchyModulePath ConnectorModulePath(Connection.Connector.Name);
+			if (ConnectorModulePath.HasModuleName(InModuleName))
 			{
-				if (!ExposedConnectors.ContainsByPredicate([ConnectorName](const FRigModuleConnector& Exposed)
+				if (!ExposedConnectors.ContainsByPredicate([&ConnectorModulePath](const FRigModuleConnector& Exposed)
 				{
-				   return Exposed.Name == ConnectorName;
+				   return ConnectorModulePath.HasElementName(Exposed.Name);
 				}))
 				{
 					ConnectionsToRemove.Add(Connection.Connector);
@@ -1843,7 +1676,7 @@ bool UModularRigController::SwapModulesOfClass(TSubclassOf<UControlRig> InOldCla
 	{
 		if (Module->Class.Get() == InOldClass)
 		{
-			SwapModuleClass(Module->GetPath(), InNewClass, bSetupUndo);
+			SwapModuleClass(Module->Name, InNewClass, bSetupUndo);
 		}
 		return true;
 	});
@@ -1855,15 +1688,15 @@ bool UModularRigController::SwapModulesOfClass(TSubclassOf<UControlRig> InOldCla
 	return true;
 }
 
-bool UModularRigController::SelectModule(const FString& InModulePath, const bool InSelected)
+bool UModularRigController::SelectModule(const FName& InModuleName, const bool InSelected)
 {
-	const bool bCurrentlySelected = Model->SelectedModulePaths.Contains(InModulePath);
+	const bool bCurrentlySelected = Model->SelectedModuleNames.Contains(InModuleName);
 	if(bCurrentlySelected == InSelected)
 	{
 		return false;
 	}
 
-	const FRigModuleReference* Module = FindModule(InModulePath);
+	const FRigModuleReference* Module = FindModule(InModuleName);
 	if(Module == nullptr)
 	{
 		return false;
@@ -1871,30 +1704,30 @@ bool UModularRigController::SelectModule(const FString& InModulePath, const bool
 
 	if(InSelected)
 	{
-		Model->SelectedModulePaths.Add(InModulePath);
+		Model->SelectedModuleNames.Add(InModuleName);
 	}
 	else
 	{
-		Model->SelectedModulePaths.Remove(InModulePath);
+		Model->SelectedModuleNames.Remove(InModuleName);
 	}
 
 	Notify(InSelected ? EModularRigNotification::ModuleSelected : EModularRigNotification::ModuleDeselected, Module);
 	return true;
 }
 
-bool UModularRigController::DeselectModule(const FString& InModulePath)
+bool UModularRigController::DeselectModule(const FName& InModuleName)
 {
-	return SelectModule(InModulePath, false);
+	return SelectModule(InModuleName, false);
 }
 
-bool UModularRigController::SetModuleSelection(const TArray<FString>& InModulePaths)
+bool UModularRigController::SetModuleSelection(const TArray<FName>& InModuleNames)
 {
 	bool bResult = false;
-	const TArray<FString> OldSelection = GetSelectedModules();
+	const TArray<FName> OldSelection = GetSelectedModules();
 
-	for(const FString& PreviouslySelectedModule : OldSelection)
+	for(const FName& PreviouslySelectedModule : OldSelection)
 	{
-		if(!InModulePaths.Contains(PreviouslySelectedModule))
+		if(!InModuleNames.Contains(PreviouslySelectedModule))
 		{
 			if(DeselectModule(PreviouslySelectedModule))
 			{
@@ -1902,7 +1735,7 @@ bool UModularRigController::SetModuleSelection(const TArray<FString>& InModulePa
 			}
 		}
 	}
-	for(const FString& NewModuleToSelect : InModulePaths)
+	for(const FName& NewModuleToSelect : InModuleNames)
 	{
 		if(!OldSelection.Contains(NewModuleToSelect))
 		{
@@ -1916,20 +1749,20 @@ bool UModularRigController::SetModuleSelection(const TArray<FString>& InModulePa
 	return bResult;
 }
 
-TArray<FString> UModularRigController::GetAllModules() const
+TArray<FName> UModularRigController::GetAllModules() const
 {
-	TArray<FString> Paths;
-	Model->ForEachModule([&Paths](const FRigModuleReference* Module) -> bool
+	TArray<FName> Names;
+	Model->ForEachModule([&Names](const FRigModuleReference* Module) -> bool
 	{
-		Paths.Add(Module->GetPath());
+		Names.Add(Module->Name);
 		return true;
 	});
-	return Paths;
+	return Names;
 }
 
-TArray<FString> UModularRigController::GetSelectedModules() const
+TArray<FName> UModularRigController::GetSelectedModules() const
 {
-	return Model->SelectedModulePaths;
+	return Model->SelectedModuleNames;
 }
 
 void UModularRigController::RefreshModuleVariables(bool bSetupUndo)
@@ -1957,8 +1790,8 @@ void UModularRigController::RefreshModuleVariables(const FRigModuleReference* In
 	}
 
 	// Make sure the provided module belongs to our ModularRigModel
-	const FString& ModulePath = InModule->GetPath();
-	FRigModuleReference* Module = FindModule(ModulePath);
+	const FName ModuleName = InModule->Name;
+	FRigModuleReference* Module = FindModule(ModuleName);
 	if (Module != InModule)
 	{
 		return;
@@ -1998,15 +1831,15 @@ void UModularRigController::RefreshModuleVariables(const FRigModuleReference* In
 	Module->Bindings.Reset();
 	for (const TPair<FName, FString>& Pair : ConfigValues)
 	{
-		SetConfigValueInModule(ModulePath, Pair.Key, Pair.Value, false);
+		SetConfigValueInModule(ModuleName, Pair.Key, Pair.Value, false);
 	}
 	for (const TPair<FName, FString>& Pair : Bindings)
 	{
-		BindModuleVariable(ModulePath, Pair.Key, Pair.Value, false);
+		BindModuleVariable(ModuleName, Pair.Key, Pair.Value, false);
 	}
 
 	// If the module is the source of another module's binding, make sure it is still a valid binding
-	Model->ForEachModule([this, InModule, ModulePath, ModuleClass](const FRigModuleReference* OtherModule) -> bool
+	Model->ForEachModule([this, InModule, ModuleName, ModuleClass](const FRigModuleReference* OtherModule) -> bool
 	{
 		if (InModule == OtherModule)
 		{
@@ -2015,11 +1848,10 @@ void UModularRigController::RefreshModuleVariables(const FRigModuleReference* In
 		TArray<FName> BindingsToRemove;
 		for (const TPair<FName, FString>& Binding : OtherModule->Bindings)
 		{
-			FString BindingModulePath, VariableName = Binding.Value;
-			(void)URigHierarchy::SplitNameSpace(Binding.Value, &BindingModulePath, &VariableName);
-			if (BindingModulePath == ModulePath)
+			const FRigHierarchyModulePath BindingModulePath(Binding.Value);
+			if (BindingModulePath.HasModuleName(ModuleName))
 			{
-				if (const FProperty* Property = ModuleClass->FindPropertyByName(*VariableName))
+				if (const FProperty* Property = ModuleClass->FindPropertyByName(BindingModulePath.GetElementFName()))
 				{
 					// remove advanced, private or not editable properties
 					const bool bIsAdvanced = Property->HasAnyPropertyFlags(CPF_AdvancedDisplay);
@@ -2032,7 +1864,7 @@ void UModularRigController::RefreshModuleVariables(const FRigModuleReference* In
 					else 
 					{
 						FText ErrorMessage;
-						if (!CanBindModuleVariable(OtherModule->GetPath(), Binding.Key, Binding.Value, ErrorMessage))
+						if (!CanBindModuleVariable(OtherModule->Name, Binding.Key, Binding.Value, ErrorMessage))
 						{
 							BindingsToRemove.Add(Binding.Key);
 						}
@@ -2043,7 +1875,7 @@ void UModularRigController::RefreshModuleVariables(const FRigModuleReference* In
 
 		for (const FName& ToRemove : BindingsToRemove)
 		{
-			UnBindModuleVariable(OtherModule->GetPath(), ToRemove);
+			UnBindModuleVariable(OtherModule->Name, ToRemove);
 		}
 		return true;
 	});
@@ -2069,7 +1901,7 @@ void UModularRigController::SanitizeName(FRigName& InOutName, bool bAllowNameSpa
 
 		if (!bGoodChar)
 		{
-			if(bAllowNameSpaces && C == ':')
+			if(bAllowNameSpaces && C == FRigHierarchyModulePath::ModuleNameSuffixChar)
 			{
 				continue;
 			}
@@ -2098,7 +1930,7 @@ FRigName UModularRigController::GetSanitizedName(const FRigName& InName, bool bA
 	return Name;
 }
 
-bool UModularRigController::IsNameAvailable(const FString& InParentModulePath, const FRigName& InDesiredName, FString* OutErrorMessage) const
+bool UModularRigController::IsNameAvailable(const FRigName& InDesiredName, FString* OutErrorMessage, const FRigModuleReference* InModuleToSkip) const
 {
 	const FRigName DesiredName = GetSanitizedName(InDesiredName, false);
 	if(DesiredName != InDesiredName)
@@ -2111,50 +1943,19 @@ bool UModularRigController::IsNameAvailable(const FString& InParentModulePath, c
 		return false;
 	}
 
-	TArray<FRigModuleReference*>* Children = &Model->RootModules;
-	if (!InParentModulePath.IsEmpty())
-	{
-		if (FRigModuleReference* Parent = const_cast<UModularRigController*>(this)->FindModule(InParentModulePath))
-		{
-			Children = &Parent->CachedChildren;
-		}
-	}
+	static const FString NameAlreadyInUse = TEXT("This name is already in use.");
 
-	for (const FRigModuleReference* Child : *Children)
-	{
-		if (FRigName(Child->Name).Equals(DesiredName, ESearchCase::IgnoreCase))
-		{
-			if(OutErrorMessage)
-			{
-				static const FString NameAlreadyInUse = TEXT("This name is already in use.");
-				*OutErrorMessage = NameAlreadyInUse;
-			}
-			return false;
-		}
-	}
-	return true;
-}
-
-bool UModularRigController::IsShortNameAvailable(const FRigName& InDesiredShortName, FString* OutErrorMessage) const
-{
-	const FRigName DesiredShortName = GetSanitizedName(InDesiredShortName, false);
-	if(DesiredShortName != InDesiredShortName)
+	// the default is to have unique names per module
+	for(const FRigModuleReference& Module : Model->Modules)
 	{
-		if(OutErrorMessage)
+		if(&Module == InModuleToSkip)
 		{
-			static const FString ContainsInvalidCharactersMessage = TEXT("Display Name contains invalid characters.");
-			*OutErrorMessage = ContainsInvalidCharactersMessage;
+			continue;
 		}
-		return false;
-	}
-
-	for (const FRigModuleReference& Child : Model->Modules)
-	{
-		if (InDesiredShortName == FRigName(Child.GetShortName()))
+		if(Module.Name == DesiredName)
 		{
 			if(OutErrorMessage)
 			{
-				static const FString NameAlreadyInUse = TEXT("This name is already in use.");
 				*OutErrorMessage = NameAlreadyInUse;
 			}
 			return false;
@@ -2163,7 +1964,7 @@ bool UModularRigController::IsShortNameAvailable(const FRigName& InDesiredShortN
 	return true;
 }
 
-FRigName UModularRigController::GetSafeNewName(const FString& InParentModulePath, const FRigName& InDesiredName) const
+FRigName UModularRigController::GetSafeNewName(const FRigName& InDesiredName, const FRigModuleReference* InModuleToSkip) const
 {
 	bool bSafeToUse = false;
 
@@ -2174,7 +1975,7 @@ FRigName UModularRigController::GetSafeNewName(const FString& InParentModulePath
 	while (!bSafeToUse)
 	{
 		bSafeToUse = true;
-		if(!IsNameAvailable(InParentModulePath, NewName))
+		if(!IsNameAvailable(NewName, nullptr, InModuleToSkip))
 		{
 			bSafeToUse = false;
 			NewName = FString::Printf(TEXT("%s_%d"), *DesiredName.ToString(), ++Index);
@@ -2183,26 +1984,6 @@ FRigName UModularRigController::GetSafeNewName(const FString& InParentModulePath
 	return NewName;
 }
 
-FRigName UModularRigController::GetSafeNewShortName(const FRigName& InDesiredShortName) const
-{
-	bool bSafeToUse = false;
-
-	// create a copy of the desired name so that the string conversion can be cached
-	const FRigName DesiredShortName = GetSanitizedName(InDesiredShortName, true);
-	FRigName NewShortName = DesiredShortName;
-	int32 Index = 0;
-	while (!bSafeToUse)
-	{
-		bSafeToUse = true;
-		if(!IsShortNameAvailable(NewShortName))
-		{
-			bSafeToUse = false;
-			NewShortName = FString::Printf(TEXT("%s_%d"), *DesiredShortName.ToString(), ++Index);
-		}
-	}
-	return NewShortName;
-}
-
 void UModularRigController::Notify(const EModularRigNotification& InNotification, const FRigModuleReference* InElement)
 {
 	if(!bSuspendNotifications)
@@ -2211,72 +1992,6 @@ void UModularRigController::Notify(const EModularRigNotification& InNotification
 	}
 }
 
-void UModularRigController::UpdateShortNames()
-{
-	TMap<FString, int32> TokenToCount;
-
-	// collect all usages of all paths and their segments
-	for(const FRigModuleReference& Module : Model->Modules)
-	{
-		if(Module.bShortNameBasedOnPath)
-		{
-			FString RemainingPath = Module.GetPath();
-			TokenToCount.FindOrAdd(RemainingPath, 0)++;
-			while(URigHierarchy::SplitNameSpace(RemainingPath, nullptr, &RemainingPath, false))
-			{
-				TokenToCount.FindOrAdd(RemainingPath, 0)++;
-			}
-		}
-		else
-		{
-			TokenToCount.FindOrAdd(Module.ShortName, 0)++;
-		}
-	}
-
-	for(FRigModuleReference& Module : Model->Modules)
-	{
-		if(Module.bShortNameBasedOnPath)
-		{
-			FString ShortPath = Module.GetPath();
-			if(!Module.ParentPath.IsEmpty())
-			{
-				FString Left, Right, RemainingPath = Module.GetPath();
-				ShortPath.Reset();
-
-				while(URigHierarchy::SplitNameSpace(RemainingPath, &Left, &Right))
-				{
-					ShortPath = ShortPath.IsEmpty() ? Right : URigHierarchy::JoinNameSpace(Right, ShortPath);
-
-					// if the short path only exists once - that's what we use for the display name
-					if(TokenToCount.FindChecked(ShortPath) == 1)
-					{
-						RemainingPath.Reset();
-						break;
-					}
-
-					RemainingPath = Left;
-				}
-
-				if(!RemainingPath.IsEmpty())
-				{
-					ShortPath = URigHierarchy::JoinNameSpace(RemainingPath, ShortPath);
-				}
-			}
-
-			if(!Module.ShortName.Equals(ShortPath, ESearchCase::IgnoreCase))
-			{
-				Module.ShortName = ShortPath;
-				Notify(EModularRigNotification::ModuleShortNameChanged, &Module);
-			}
-		}
-		else
-		{
-			// the display name is user defined so we won't touch it
-		}
-		
-	}
-}
-
 UModularRig* UModularRigController::GetDebuggedModularRig()
 {
 #if WITH_EDITOR
diff --git a/Engine/Plugins/Animation/ControlRig/Source/ControlRig/Private/ModularRigModel.cpp b/Engine/Plugins/Animation/ControlRig/Source/ControlRig/Private/ModularRigModel.cpp
--- a/Engine/Plugins/Animation/ControlRig/Source/ControlRig/Private/ModularRigModel.cpp
+++ b/Engine/Plugins/Animation/ControlRig/Source/ControlRig/Private/ModularRigModel.cpp
@@ -7,46 +7,72 @@
 
 #include UE_INLINE_GENERATED_CPP_BY_NAME(ModularRigModel)
 
-FString FRigModuleReference::GetShortName() const
+FString FRigModuleReference::GetElementPrefix() const
 {
-	if(!ShortName.IsEmpty())
+	return GetName() + FRigHierarchyModulePath::ModuleNameSuffix;
+}
+
+FRigHierarchyModulePath FRigModuleReference::GetModulePath() const
+{
+	if(ParentModuleName.IsNone())
 	{
-		return ShortName;
+		if(!ParentPath_DEPRECATED.IsEmpty())
+		{
+			return URigHierarchy::JoinNameSpace_Deprecated(ParentPath_DEPRECATED, Name.ToString()); 
+		}
 	}
-	return GetPath();
+	else if(Model)
+	{
+		if(const FRigModuleReference* ParentModule = GetParentModule())
+		{
+			return FRigHierarchyModulePath(ParentModule->GetModulePath(), Name.ToString()); 
+		}
+	}
+	return Name.ToString();
 }
 
-FString FRigModuleReference::GetPath() const
+const FRigModuleReference* FRigModuleReference::GetParentModule() const
 {
-	if (ParentPath.IsEmpty())
+	if(ParentModuleName.IsNone())
+	{
+		if(!ParentPath_DEPRECATED.IsEmpty())
+		{
+			return Model->FindModuleByPath(ParentPath_DEPRECATED);
+		}
+	}
+	else if(Model)
 	{
-		return Name.ToString();
+		return Model->GetParentModule(this);
 	}
-	return URigHierarchy::JoinNameSpace(ParentPath, Name.ToString());
+	return nullptr;
 }
 
-FString FRigModuleReference::GetNamespace() const
+const FRigModuleReference* FRigModuleReference::GetRootModule() const
 {
-	return GetPath() + UModularRig::NamespaceSeparator;
+	if(ParentModuleName.IsNone())
+	{
+		return this;
+	}
+	if(const FRigModuleReference* ParentModule = GetParentModule())
+	{
+		return ParentModule->GetRootModule();
+	}
+	return nullptr;
 }
 
 const FRigConnectorElement* FRigModuleReference::FindPrimaryConnector(const URigHierarchy* InHierarchy) const
 {
 	if(InHierarchy)
 	{
-		const FString MyModulePath = GetPath();
 		const TArray<FRigConnectorElement*> AllConnectors = InHierarchy->GetConnectors();
 		for(const FRigConnectorElement* Connector : AllConnectors)
 		{
 			if(Connector->IsPrimary())
 			{
-				const FString ModulePath = InHierarchy->GetModulePath(Connector->GetKey());
-				if(!ModulePath.IsEmpty())
+				const FName ModuleName = InHierarchy->GetModuleFName(Connector->GetKey());
+				if(Name == ModuleName)
 				{
-					if(ModulePath.Equals(MyModulePath, ESearchCase::IgnoreCase))
-					{
-						return Connector;
-					}
+					return Connector;
 				}
 			}
 		}
@@ -59,14 +85,14 @@ TArray<const FRigConnectorElement*> FRigModuleReference::FindConnectors(const UR
 	TArray<const FRigConnectorElement*> Connectors;
 	if(InHierarchy)
 	{
-		const FString MyModulePath = GetPath();
+		const FName MyModuleName = GetFName();
 		const TArray<FRigConnectorElement*> AllConnectors = InHierarchy->GetConnectors();
 		for(const FRigConnectorElement* Connector : AllConnectors)
 		{
-			const FString ModulePath = InHierarchy->GetModulePath(Connector->GetKey());
-			if(!ModulePath.IsEmpty())
+			const FName ModuleName = InHierarchy->GetModuleFName(Connector->GetKey());
+			if(!ModuleName.IsNone())
 			{
-				if(ModulePath.Equals(MyModulePath, ESearchCase::IgnoreCase))
+				if(ModuleName == MyModuleName)
 				{
 					Connectors.Add(Connector);
 				}
@@ -83,24 +109,25 @@ FRigElementKeyRedirector::FKeyArray FModularRigSingleConnection::GetTargetArray(
 	return Result;
 }
 
-FRigElementKeyRedirector::FKeyMap FModularRigConnections::GetModuleConnectionMap(const FString& InModulePath) const
+FRigElementKeyRedirector::FKeyMap FModularRigConnections::GetModuleConnectionMap(const FName& InModuleName) const
 {
+	const FString InModuleNameString = InModuleName.ToString();
+	
 	FRigElementKeyRedirector::FKeyMap Result;
 	for (const FModularRigSingleConnection& Connection : ConnectionList)
 	{
-		FString Path, Name;
-		URigHierarchy::SplitNameSpace(Connection.Connector.Name.ToString(), &Path, &Name);
+		const FRigHierarchyModulePath ModulePath(Connection.Connector.Name);
 
 		// Exactly the same path (do not return connectors from child modules)
-		if (Path == InModulePath)
+		if (ModulePath.HasModuleName(InModuleNameString))
 		{
-			Result.Add(FRigElementKey(*Name, ERigElementType::Connector), Connection.GetTargetArray());
+			Result.Add(FRigElementKey(ModulePath.GetElementFName(), ERigElementType::Connector), Connection.GetTargetArray());
 		}
 	}
 	return Result;
 }
 
-void FModularRigConnections::PatchOnLoad()
+void FModularRigConnections::PatchOnLoad(const TMap<FRigHierarchyModulePath, FName>* InModulePathToModuleName)
 {
 	for(FModularRigSingleConnection& Connection : ConnectionList)
 	{
@@ -113,19 +140,32 @@ void FModularRigConnections::PatchOnLoad()
 				Connection.Target_DEPRECATED.Reset();
 			}
 		}
+
+		Connection.Connector.ConvertToModuleNameFormatInline(InModulePathToModuleName);
+		for(FRigElementKey& Target : Connection.Targets)
+		{
+			Target.ConvertToModuleNameFormatInline(InModulePathToModuleName);
+		}
 	}
 }
 
 void FModularRigModel::PatchModelsOnLoad()
 {
+	bool bHasDeprecatedParentPath = false;
+	for(int32 ModuleIndex = 0; ModuleIndex < Modules.Num(); ModuleIndex++)
+	{
+		Modules[ModuleIndex].Model = this;
+		bHasDeprecatedParentPath = bHasDeprecatedParentPath || !Modules[ModuleIndex].ParentPath_DEPRECATED.IsEmpty();
+	}
+	
 	if (Connections.IsEmpty())
 	{
 		ForEachModule([this](const FRigModuleReference* Module) -> bool
 		{
-			FString ModuleNamespace = Module->GetNamespace();
+			FString ElementPrefix = Module->GetElementPrefix();
 			for (const TTuple<FRigElementKey, FRigElementKey>& Connection : Module->Connections_DEPRECATED)
 			{
-				const FString ConnectorPath = FString::Printf(TEXT("%s%s"), *ModuleNamespace, *Connection.Key.Name.ToString());
+				const FString ConnectorPath = FString::Printf(TEXT("%s%s"), *ElementPrefix, *Connection.Key.Name.ToString());
 				FRigElementKey ConnectorKey(*ConnectorPath, ERigElementType::Connector);
 				Connections.AddConnection(ConnectorKey, {Connection.Value});
 			}
@@ -133,7 +173,61 @@ void FModularRigModel::PatchModelsOnLoad()
 		});
 	}
 
-	Connections.PatchOnLoad();
+	// if we need to introduce a unique name list we'll have to also fill the previous module path name list
+	if(bHasDeprecatedParentPath)
+	{
+		PreviousModulePaths.Reset();
+
+		TArray<FString> Paths;
+		TArray<FName> Names;
+		Paths.Reserve(Modules.Num());
+		Names.Reserve(Modules.Num());
+		
+		// temporarily rename all modules
+		for(int32 ModuleIndex = 0; ModuleIndex < Modules.Num(); ModuleIndex++)
+		{
+			FRigModuleReference& Module = Modules[ModuleIndex];
+			Paths.Add(Module.GetModulePath());
+			Names.Add(Module.bShortNameBasedOnPath_DEPRECATED ? Module.Name : *Module.ShortName_DEPRECATED);
+			Module.Name = *FString::Printf(TEXT("____TEMPORARY_MODULE_NAME_%03d"), ModuleIndex);
+		}
+
+		// start renaming it again
+		for(int32 ModuleIndex = 0; ModuleIndex < Modules.Num(); ModuleIndex++)
+		{
+			FRigModuleReference& Module = Modules[ModuleIndex];
+			Module.Name = GetController()->GetSafeNewName(Names[ModuleIndex]);
+			PreviousModulePaths.Add(Paths[ModuleIndex], Module.Name);
+		}
+
+		// update the parent module name
+		for(int32 ModuleIndex = 0; ModuleIndex < Modules.Num(); ModuleIndex++)
+		{
+			FRigModuleReference& Module = Modules[ModuleIndex];
+			if(!Module.ParentPath_DEPRECATED.IsEmpty())
+			{
+				Module.ParentModuleName = PreviousModulePaths.FindChecked(Module.ParentPath_DEPRECATED);
+				Module.ParentPath_DEPRECATED.Reset();
+			}
+		}
+
+		// update the module bindings
+		for(int32 ModuleIndex = 0; ModuleIndex < Modules.Num(); ModuleIndex++)
+		{
+			FRigModuleReference& Module = Modules[ModuleIndex];
+			for(TPair<FName, FString>& Binding : Module.Bindings)
+			{
+				FRigHierarchyModulePath BindingModulePath(Binding.Value);
+				if(BindingModulePath.ConvertToModuleNameFormatInline(&PreviousModulePaths))
+				{
+					Binding.Value = BindingModulePath;
+				}
+			}
+		}
+	}
+	UpdateCachedChildren();
+
+	Connections.PatchOnLoad(&PreviousModulePaths);
 	Connections.UpdateFromConnectionList();
 }
 
@@ -156,82 +250,121 @@ void FModularRigModel::SetOuterClientHost(UObject* InOuterClientHost)
 
 void FModularRigModel::UpdateCachedChildren()
 {
-	TMap<FString, FRigModuleReference*> PathToModule;
+	TMap<FName, FRigModuleReference*> ModuleByName;
+	ModuleByName.Reserve(Modules.Num());
+
+	// also offer a backwards compatible lookup
+	TMap<FString, FRigModuleReference*> ModuleByPath;
+	
 	for (FRigModuleReference& Module : Modules)
 	{
 		Module.CachedChildren.Reset();
-		PathToModule.Add(Module.GetPath(), &Module);
+
+		if(Module.ParentPath_DEPRECATED.IsEmpty())
+		{
+			ModuleByName.Add(Module.GetFName(), &Module);
+		}
+		else
+		{
+			ModuleByPath.Add(Module.GetModulePath(), &Module);
+		}
 	}
 	
 	RootModules.Reset();
 	for (FRigModuleReference& Module : Modules)
 	{
-		if (Module.ParentPath.IsEmpty())
+		if (!Module.HasParentModule())
 		{
 			RootModules.Add(&Module);
 		}
 		else
 		{
-			if (FRigModuleReference** ParentModule = PathToModule.Find(Module.ParentPath))
+			if (FRigModuleReference** ParentModule = ModuleByName.Find(Module.ParentModuleName))
 			{
 				(*ParentModule)->CachedChildren.Add(&Module);
 			}
+			else if (FRigModuleReference** ParentModuleByPath = ModuleByPath.Find(Module.ParentPath_DEPRECATED))
+			{
+				(*ParentModuleByPath)->CachedChildren.Add(&Module);
+			}
 		}
 	}
 }
 
-FRigModuleReference* FModularRigModel::FindModule(const FString& InPath)
+FRigModuleReference* FModularRigModel::FindModule(const FName& InModuleName)
 {
-	const TArray<FRigModuleReference*>* Children = &RootModules;
+	if(InModuleName.IsNone())
+	{
+		return nullptr;
+	}
+	
+	FRigModuleReference* FoundModule = Modules.FindByPredicate([InModuleName](const FRigModuleReference& Module)
+	{
+		return Module.Name == InModuleName;
+	});
 
-	FString Left = InPath, Right;
-	while (URigHierarchy::SplitNameSpace(Left, &Left, &Right, false))
+	if(FoundModule == nullptr)
 	{
-		FRigModuleReference* const * Child = Children->FindByPredicate([Left](FRigModuleReference* Module)
-		{
-			return Module->Name.ToString() == Left;
-		});
-		
-		if (!Child)
+		FoundModule = const_cast<FRigModuleReference*>(FindModuleByPath(InModuleName.ToString()));
+#if WITH_EDITOR
+		if(FoundModule && !IsLoading())
 		{
-			return nullptr;
+			const int32 NumReceivedOldModulePaths = ReceivedOldModulePaths.Num(); 
+			if(ReceivedOldModulePaths.AddUnique(InModuleName) == NumReceivedOldModulePaths)
+			{
+				if(const UObject* Outer = GetOuter())
+				{
+					if(!IsRunningCookCommandlet() && !IsRunningCookOnTheFly())
+					{
+						UE_LOG(
+							LogControlRig,
+							Warning,
+							TEXT("%s: Module '%s' has been accessed using an old module path ('%s'). Please consider updating your code."),
+							*Outer->GetPathName(),
+							*FoundModule->GetName(),
+							*InModuleName.ToString()
+						);
+					}
+				}
+			}
 		}
-		
-		Children = &(*Child)->CachedChildren;
-		Left = Right;
+#endif
 	}
+	
+	return FoundModule;
+}
 
-	FRigModuleReference* const * Child = Children->FindByPredicate([Left](FRigModuleReference* Module)
-		{
-			return Module->Name.ToString() == Left;
-		});
+const FRigModuleReference* FModularRigModel::FindModule(const FName& InName) const
+{
+	FRigModuleReference* Module = const_cast<FModularRigModel*>(this)->FindModule(InName);
+	return Module;
+}
 
-	if (!Child)
+const FRigModuleReference* FModularRigModel::FindModuleByPath(const FString& InModulePath) const
+{
+	if(InModulePath.IsEmpty())
 	{
 		return nullptr;
 	}
-	return *Child;
-}
-
-const FRigModuleReference* FModularRigModel::FindModule(const FString& InPath) const
-{
-	FRigModuleReference* Module = const_cast<FModularRigModel*>(this)->FindModule(InPath);
-	return const_cast<FRigModuleReference*>(Module);
+	
+	return Modules.FindByPredicate([InModulePath](const FRigModuleReference& Module)
+	{
+		return Module.GetModulePath() == InModulePath;
+	});
 }
 
-FRigModuleReference* FModularRigModel::GetParentModule(const FString& InPath)
+FRigModuleReference* FModularRigModel::GetParentModule(const FName& InName)
 {
-	const FString ParentPath = GetParentPath(InPath);
-	if(!ParentPath.IsEmpty())
+	if(const FRigModuleReference* Module = FindModule(InName))
 	{
-		return FindModule(ParentPath);
+		return GetParentModule(Module);
 	}
 	return nullptr;
 }
 
-const FRigModuleReference* FModularRigModel::GetParentModule(const FString& InPath) const
+const FRigModuleReference* FModularRigModel::GetParentModule(const FName& InName) const
 {
-	FRigModuleReference* Module = const_cast<FModularRigModel*>(this)->GetParentModule(InPath);
+	FRigModuleReference* Module = const_cast<FModularRigModel*>(this)->GetParentModule(InName);
 	return const_cast<FRigModuleReference*>(Module);
 }
 
@@ -239,7 +372,10 @@ FRigModuleReference* FModularRigModel::GetParentModule(const FRigModuleReference
 {
 	if(InChildModule)
 	{
-		return FindModule(InChildModule->ParentPath);
+		if(!InChildModule->ParentModuleName.IsNone())
+		{
+			return FindModule(InChildModule->ParentModuleName);
+		}
 	}
 	return nullptr;
 }
@@ -250,10 +386,10 @@ const FRigModuleReference* FModularRigModel::GetParentModule(const FRigModuleRef
 	return const_cast<FRigModuleReference*>(Module);
 }
 
-bool FModularRigModel::IsModuleParentedTo(const FString& InChildModulePath, const FString& InParentModulePath) const
+bool FModularRigModel::IsModuleParentedTo(const FName& InChildModuleName, const FName& InParentModuleName) const
 {
-	const FRigModuleReference* ChildModule = FindModule(InChildModulePath);
-	const FRigModuleReference* ParentModule = FindModule(InParentModulePath);
+	const FRigModuleReference* ChildModule = FindModule(InChildModuleName);
+	const FRigModuleReference* ParentModule = FindModule(InParentModuleName);
 	if(ChildModule == nullptr || ParentModule == nullptr)
 	{
 		return false;
@@ -310,15 +446,6 @@ TArray<const FRigModuleReference*> FModularRigModel::FindModuleInstancesOfClass(
 	return FindModuleInstancesOfClass(*InClass.ToString());
 }
 
-FString FModularRigModel::GetParentPath(const FString& InPath) const
-{
-	if (const FRigModuleReference* Element = FindModule(InPath))
-	{
-		return Element->ParentPath;
-	}
-	return FString();
-}
-
 void FModularRigModel::ForEachModule(TFunction<bool(const FRigModuleReference*)> PerModule) const
 {
 	TArray<FRigModuleReference*> ModuleInstances = RootModules;
@@ -332,17 +459,16 @@ void FModularRigModel::ForEachModule(TFunction<bool(const FRigModuleReference*)>
 	}
 }
 
-TArray<FString> FModularRigModel::SortPaths(const TArray<FString>& InPaths) const
+TArray<FName> FModularRigModel::SortModuleNames(const TArray<FName>& InModuleNames) const
 {
-	TArray<FString> Result;
-	ForEachModule([InPaths, &Result](const FRigModuleReference* Element) -> bool
+	TArray<FName> SortedModuleNames;
+	ForEachModule([InModuleNames, &SortedModuleNames](const FRigModuleReference* Module) -> bool
 	{
-		const FString& Path = Element->GetPath();
-		if(InPaths.Contains(Path))
+		if(InModuleNames.Contains(Module->GetFName()))
 		{
-			Result.AddUnique(Path);
+			SortedModuleNames.AddUnique(Module->GetFName());
 		}
 		return true;
 	});
-	return Result;
+	return SortedModuleNames;
 }
diff --git a/Engine/Plugins/Animation/ControlRig/Source/ControlRig/Private/ModularRigRuleManager.cpp b/Engine/Plugins/Animation/ControlRig/Source/ControlRig/Private/ModularRigRuleManager.cpp
--- a/Engine/Plugins/Animation/ControlRig/Source/ControlRig/Private/ModularRigRuleManager.cpp
+++ b/Engine/Plugins/Animation/ControlRig/Source/ControlRig/Private/ModularRigRuleManager.cpp
@@ -174,7 +174,7 @@ void UModularRigRuleManager::SetHierarchy(const URigHierarchy* InHierarchy)
 void UModularRigRuleManager::ResolveConnector(FWorkData& InOutWorkData)
 {
 	FilterIncompatibleTypes(InOutWorkData);
-	FilterInvalidNameSpaces(InOutWorkData);
+	FilterInvalidModules(InOutWorkData);
 	FilterByConnectorRules(InOutWorkData);
 	FilterByConnectorEvent(InOutWorkData);
 
@@ -205,32 +205,30 @@ void UModularRigRuleManager::FilterIncompatibleTypes(FWorkData& InOutWorkData)
 	});
 }
 
-void UModularRigRuleManager::FilterInvalidNameSpaces(FWorkData& InOutWorkData)
+void UModularRigRuleManager::FilterInvalidModules(FWorkData& InOutWorkData)
 {
 	if(InOutWorkData.Connector == nullptr)
 	{
 		return;
 	}
 	
-	const FName NameSpace = InOutWorkData.Hierarchy->GetNameSpaceFName(InOutWorkData.Connector->GetKey());
-	if(NameSpace.IsNone())
+	const FName ModuleName = InOutWorkData.Hierarchy->GetModuleFName(InOutWorkData.Connector->GetKey());
+	if(ModuleName.IsNone())
 	{
 		return;
 	}
 	
-	const FString NameSpaceString = NameSpace.ToString(); 
-	InOutWorkData.Filter([NameSpaceString, InOutWorkData](FRigElementResolveResult& Result)
+	InOutWorkData.Filter([ModuleName, InOutWorkData](FRigElementResolveResult& Result)
 	{
-		const FName MatchNameSpace = InOutWorkData.Hierarchy->GetNameSpaceFName(Result.GetKey());
-		if(!MatchNameSpace.IsNone())
+		const FName MatchModuleName = InOutWorkData.Hierarchy->GetModuleFName(Result.GetKey());
+		if(MatchModuleName == ModuleName)
 		{
-			const FString MatchNameSpaceString = MatchNameSpace.ToString();
-			if(MatchNameSpaceString.Equals(NameSpaceString, ESearchCase::IgnoreCase))
-			{
-				static const FText CannotConnectWithinNameSpaceMessage = LOCTEXT("CannotConnectWithinNameSpace", "Cannot connect within the same namespace.");
-				Result.SetInvalidTarget(CannotConnectWithinNameSpaceMessage);
-			}
-			else if(MatchNameSpaceString.StartsWith(NameSpaceString, ESearchCase::IgnoreCase))
+			static const FText CannotConnectWithinNameSpaceMessage = LOCTEXT("CannotConnectWithinNameSpace", "Cannot connect within the same namespace.");
+			Result.SetInvalidTarget(CannotConnectWithinNameSpaceMessage);
+		}
+		else
+		{
+			if(InOutWorkData.Module && InOutWorkData.Module->HasChildModule(MatchModuleName, true))
 			{
 				static const FText CannotConnectBelowNameSpaceMessage = LOCTEXT("CannotConnectBelowNameSpace", "Cannot connect to element below the connector's namespace.");
 				Result.SetInvalidTarget(CannotConnectBelowNameSpaceMessage);
diff --git a/Engine/Plugins/Animation/ControlRig/Source/ControlRig/Private/Rigs/RigConnectionRules.cpp b/Engine/Plugins/Animation/ControlRig/Source/ControlRig/Private/Rigs/RigConnectionRules.cpp
--- a/Engine/Plugins/Animation/ControlRig/Source/ControlRig/Private/Rigs/RigConnectionRules.cpp
+++ b/Engine/Plugins/Animation/ControlRig/Source/ControlRig/Private/Rigs/RigConnectionRules.cpp
@@ -107,16 +107,14 @@ const FRigConnectorElement* FRigConnectionRuleInput::FindPrimaryConnector(FText*
 	check(Hierarchy);
 	check(Module);
 
-	const FName ModuleNameSpace = *Module->GetNamespace();
-
 	const FRigConnectorElement* PrimaryConnector = nullptr;
 	Hierarchy->ForEach<FRigConnectorElement>(
-		[this, ModuleNameSpace, &PrimaryConnector](FRigConnectorElement* Connector) -> bool
+		[this, &PrimaryConnector](FRigConnectorElement* Connector) -> bool
 		{
 			if(Connector->IsPrimary())
 			{
-				const FName ConnectorNameSpace = Hierarchy->GetNameSpaceFName(Connector->GetKey());
-				if(!ConnectorNameSpace.IsNone() && ConnectorNameSpace.IsEqual(ModuleNameSpace, ENameCase::IgnoreCase))
+				const FName ConnectorModuleName = Hierarchy->GetModuleFName(Connector->GetKey());
+				if(Module->Name == ConnectorModuleName)
 				{
 					PrimaryConnector = Connector;
 					return false; // stop the search
@@ -131,7 +129,7 @@ const FRigConnectorElement* FRigConnectionRuleInput::FindPrimaryConnector(FText*
 		if(OutErrorMessage)
 		{
 			static constexpr TCHAR Format[] = TEXT("No primary connector found for module '%s'.");
-			*OutErrorMessage = FText::FromString(FString::Printf(Format, *Module->GetPath()));
+			*OutErrorMessage = FText::FromString(FString::Printf(Format, *Module->Name.ToString()));
 		}
 	}
 
@@ -143,16 +141,14 @@ TArray<const FRigConnectorElement*> FRigConnectionRuleInput::FindSecondaryConnec
 	check(Hierarchy);
 	check(Module);
 
-	const FName ModuleNameSpace = *Module->GetNamespace();
-
 	TArray<const FRigConnectorElement*> SecondaryConnectors;
 	Hierarchy->ForEach<FRigConnectorElement>(
-		[this, bOptional, ModuleNameSpace, &SecondaryConnectors](const FRigConnectorElement* Connector) -> bool
+		[this, bOptional, &SecondaryConnectors](const FRigConnectorElement* Connector) -> bool
 		{
 			if(Connector->IsSecondary() && Connector->IsOptional() == bOptional)
 			{
-				const FName ConnectorNameSpace = Hierarchy->GetNameSpaceFName(Connector->GetKey());
-				if(!ConnectorNameSpace.IsNone() && ConnectorNameSpace.IsEqual(ModuleNameSpace, ENameCase::IgnoreCase))
+				const FName ConnectorModuleName = Hierarchy->GetModuleFName(Connector->GetKey());
+				if(Module->Name == ConnectorModuleName)
 				{
 					SecondaryConnectors.Add(Connector);
 				}
