Control Rig: Unique modular rig module names

diff --git a/Engine/Plugins/Animation/ControlRig/Source/ControlRig/Private/Rigs/RigHierarchy.cpp b/Engine/Plugins/Animation/ControlRig/Source/ControlRig/Private/Rigs/RigHierarchy.cpp
--- a/Engine/Plugins/Animation/ControlRig/Source/ControlRig/Private/Rigs/RigHierarchy.cpp
+++ b/Engine/Plugins/Animation/ControlRig/Source/ControlRig/Private/Rigs/RigHierarchy.cpp
@@ -1662,6 +1662,56 @@ bool URigHierarchy::IsProcedural(const FRigHierarchyKey& InKey) const
 	return false;
 }
 
+int32 URigHierarchy::GetIndex(const FRigElementKey& InKey) const
+{
+	if(ElementKeyRedirector)
+	{
+		if(FRigElementKeyRedirector::FCachedKeyArray* CachedRigElements = ElementKeyRedirector->Find(InKey))
+		{
+			if(CachedRigElements->Num() == 1)
+			{
+				if((*CachedRigElements)[0].UpdateCache(this))
+				{
+					return (*CachedRigElements)[0].GetIndex();
+				}
+			}
+			return INDEX_NONE;
+		}
+	}
+	
+	if(const int32* Index = ElementIndexLookup.Find(InKey))
+	{
+		return *Index;
+	}
+
+	const FRigElementKey PatchedKey = PatchElementKeyInLookup(InKey);
+	if(PatchedKey != InKey)
+	{
+#if WITH_EDITOR
+		if(!IsLoading())
+		{
+			const int32 NumReceivedNameSpaceBasedKeys = ReceivedNameSpaceBasedKeys.Num(); 
+			if(ReceivedNameSpaceBasedKeys.AddUnique(InKey) == NumReceivedNameSpaceBasedKeys)
+			{
+				if(!IsRunningCookCommandlet() && !IsRunningCookOnTheFly())
+				{
+					UE_LOG(
+						LogControlRig,
+						Warning,
+						TEXT("%s: Element '%s' has been accessed using a namespace based key ('%s'). Please consider updating your code."),
+						*GetPathName(),
+						*PatchedKey.ToString(),
+						*InKey.ToString()
+					);
+				}
+			}
+		}
+#endif
+		return GetIndex(PatchedKey);
+	}
+	return INDEX_NONE;
+}
+
 const FRigBaseComponent* URigHierarchy::FindComponent(const FRigComponentKey& InKey) const
 {
 	const int32 ComponentIndex = GetComponentIndex(InKey);
@@ -1728,7 +1778,23 @@ FRigBaseComponent* URigHierarchy::GetTopLevelComponent(int32 InTopLevelComponent
 
 int32 URigHierarchy::GetComponentIndex(const FRigComponentKey& InComponentKey, bool bFollowRedirector) const
 {
-	if(bFollowRedirector && ElementKeyRedirector && !InComponentKey.IsTopLevel())
+	if(bFollowRedirector)
+	{
+		if(ElementKeyRedirector == nullptr)
+		{
+			bFollowRedirector = false;
+		}
+	}
+
+	if(!bFollowRedirector)
+	{
+		if(!ComponentIndexLookup.Contains(InComponentKey))
+		{
+			bFollowRedirector = true;
+		}
+	}
+	
+	if(bFollowRedirector && !InComponentKey.IsTopLevel())
 	{
 		const int32 ElementIndex = GetIndex(InComponentKey.ElementKey);
 		if(ElementIndex != INDEX_NONE)
@@ -2160,6 +2226,16 @@ bool URigHierarchy::RemoveAllMetadata(FRigElementKey InItem)
 }
 
 FName URigHierarchy::GetModulePathFName(FRigElementKey InItem) const
+{
+	return GetModuleFName(InItem);
+}
+
+FString URigHierarchy::GetModulePath(FRigElementKey InItem) const
+{
+	return GetModuleName(InItem);
+}
+
+FName URigHierarchy::GetModuleFName(FRigElementKey InItem) const
 {
 	if(!InItem.IsValid())
 	{
@@ -2173,22 +2249,19 @@ FName URigHierarchy::GetModulePathFName(FRigElementKey InItem) const
 	}
 
 	// fall back on the name of the item
-	const FString NameString = InItem.Name.ToString();
-	FString ModulePathFromName;
-	if(SplitNameSpace(NameString, &ModulePathFromName, nullptr))
+	const FRigHierarchyModulePath ModulePath = InItem.Name.ToString();
+	FString ModuleName;
+	if(ModulePath.Split(&ModuleName, nullptr))
 	{
-		if(!ModulePathFromName.IsEmpty())
-		{
-			return *ModulePathFromName;
-		}
+		return *ModuleName;
 	}
 
 	return NAME_None;
 }
 
-FString URigHierarchy::GetModulePath(FRigElementKey InItem) const
+FString URigHierarchy::GetModuleName(FRigElementKey InItem) const
 {
-	const FName ModulePathName = GetModulePathFName(InItem);
+	const FName ModulePathName = GetModuleFName(InItem);
 	if(!ModulePathName.IsNone())
 	{
 		return ModulePathName.ToString();
@@ -2196,41 +2269,31 @@ FString URigHierarchy::GetModulePath(FRigElementKey InItem) const
 	return FString();
 }
 
+FString URigHierarchy::GetModulePrefix(FRigElementKey InItem) const
+{
+	return GetModuleName(InItem) + FRigHierarchyModulePath::ModuleNameSuffix;
+}
+
 FName URigHierarchy::GetNameSpaceFName(FRigElementKey InItem) const
 {
-	if(!InItem.IsValid())
+	const FString NameSpaceString = GetNameSpace(InItem);
+	if(NameSpaceString.IsEmpty())
 	{
 		return NAME_None;
 	}
-	
-	const FName Result = GetNameMetadata(InItem, NameSpaceMetadataName, NAME_None);
-	if(!Result.IsNone())
-	{
-		return Result;
-	}
-
-	// fall back on the name of the item
-	const FString NameString = InItem.Name.ToString();
-	FString NameSpaceFromName;
-	if(SplitNameSpace(NameString, &NameSpaceFromName, nullptr))
-	{
-		if(!NameSpaceFromName.IsEmpty())
-		{
-			return *(NameSpaceFromName + UModularRig::NamespaceSeparator);
-		}
-	}
-	
-	return NAME_None;
+	return *NameSpaceString;
 }
 
 FString URigHierarchy::GetNameSpace(FRigElementKey InItem) const
 {
-	const FName NameSpaceName = GetNameSpaceFName(InItem);
-	if(!NameSpaceName.IsNone())
+	if(!InItem.IsValid())
 	{
-		return NameSpaceName.ToString();
+		return FString();
 	}
-	return FString();
+	
+	// fall back on the name of the item
+	const FString NameString = GetModuleName(InItem);
+	return NameString + FRigHierarchyModulePath::NamespaceSeparator_Deprecated;
 }
 
 TArray<const FRigBaseElement*> URigHierarchy::GetSelectedElements(ERigElementType InTypeFilter) const
@@ -2374,42 +2437,42 @@ TArray<FRigElementKey> URigHierarchy::GetSelectedKeysByPredicate(TFunctionRef<bo
 	return Result;
 }
 
-FRigName URigHierarchy::JoinNameSpace(const FRigName& InLeft, const FRigName& InRight)
-{
-	return FRigName(JoinNameSpace(InLeft.ToString(), InRight.ToString()));
-}
-
-FString URigHierarchy::JoinNameSpace(const FString& InLeft, const FString& InRight)
+FString URigHierarchy::JoinNameSpace_Deprecated(const FString& InLeft, const FString& InRight)
 {
-	if(InLeft.EndsWith(UModularRig::NamespaceSeparator))
+	if(InLeft.EndsWith(FRigHierarchyModulePath::NamespaceSeparator_Deprecated))
 	{
 		return InLeft + InRight;
 	}
-	return InLeft + UModularRig::NamespaceSeparator + InRight;
+	return InLeft + FRigHierarchyModulePath::NamespaceSeparator_Deprecated + InRight;
+}
+
+FRigName URigHierarchy::JoinNameSpace_Deprecated(const FRigName& InLeft, const FRigName& InRight)
+{
+	return FRigName(JoinNameSpace_Deprecated(InLeft.ToString(), InRight.ToString()));
 }
 
-TPair<FString, FString> URigHierarchy::SplitNameSpace(const FString& InNameSpacedPath, bool bFromEnd)
+TPair<FString, FString> URigHierarchy::SplitNameSpace_Deprecated(const FString& InNameSpacedPath, bool bFromEnd)
 {
 	TPair<FString, FString> Result;
-	(void)SplitNameSpace(InNameSpacedPath, &Result.Key, &Result.Value, bFromEnd);
+	(void)SplitNameSpace_Deprecated(InNameSpacedPath, &Result.Key, &Result.Value, bFromEnd);
 	return Result;
 }
 
-TPair<FRigName, FRigName> URigHierarchy::SplitNameSpace(const FRigName& InNameSpacedPath, bool bFromEnd)
+TPair<FRigName, FRigName> URigHierarchy::SplitNameSpace_Deprecated(const FRigName& InNameSpacedPath, bool bFromEnd)
 {
-	const TPair<FString, FString> Result = SplitNameSpace(InNameSpacedPath.GetName(), bFromEnd);
+	const TPair<FString, FString> Result = SplitNameSpace_Deprecated(InNameSpacedPath.GetName(), bFromEnd);
 	return {FRigName(Result.Key), FRigName(Result.Value)};
 }
 
-bool URigHierarchy::SplitNameSpace(const FString& InNameSpacedPath, FString* OutNameSpace, FString* OutName, bool bFromEnd)
+bool URigHierarchy::SplitNameSpace_Deprecated(const FString& InNameSpacedPath, FString* OutNameSpace, FString* OutName, bool bFromEnd)
 {
-	return InNameSpacedPath.Split(UModularRig::NamespaceSeparator, OutNameSpace, OutName, ESearchCase::CaseSensitive, bFromEnd ? ESearchDir::FromEnd : ESearchDir::FromStart);
+	return InNameSpacedPath.Split(FRigHierarchyModulePath::NamespaceSeparator_Deprecated, OutNameSpace, OutName, ESearchCase::CaseSensitive, bFromEnd ? ESearchDir::FromEnd : ESearchDir::FromStart);
 }
 
-bool URigHierarchy::SplitNameSpace(const FRigName& InNameSpacedPath, FRigName* OutNameSpace, FRigName* OutName, bool bFromEnd)
+bool URigHierarchy::SplitNameSpace_Deprecated(const FRigName& InNameSpacedPath, FRigName* OutNameSpace, FRigName* OutName, bool bFromEnd)
 {
 	FString NameSpace, Name;
-	if(SplitNameSpace(InNameSpacedPath.GetName(), &NameSpace, &Name, bFromEnd))
+	if(SplitNameSpace_Deprecated(InNameSpacedPath.GetName(), &NameSpace, &Name, bFromEnd))
 	{
 		if(OutNameSpace)
 		{
@@ -2440,7 +2503,7 @@ void URigHierarchy::SanitizeName(FRigName& InOutName, bool bAllowNameSpaces)
 
 		if (!bGoodChar)
 		{
-			if(bAllowNameSpaces && C == ':')
+			if(bAllowNameSpaces && (C == FRigHierarchyModulePath::NamespaceSeparatorChar_Deprecated || C == FRigHierarchyModulePath::ModuleNameSuffixChar))
 			{
 				continue;
 			}
@@ -2634,12 +2697,13 @@ FRigName URigHierarchy::GetSafeNewName(const FRigName& InPotentialNewName, ERigE
 		const FControlRigExecuteContext& CRContext = ExecuteContext->GetPublicData<FControlRigExecuteContext>();
 		if(CRContext.IsRigModule())
 		{
-			FRigName LastSegmentName;
-			if(SplitNameSpace(SanitizedName, nullptr, &LastSegmentName, true))
+			FString LastSegmentName;
+			const FRigHierarchyModulePath ModulePath = SanitizedName.ToString();
+			if(ModulePath.Split(nullptr, &LastSegmentName))
 			{
 				SanitizedName = LastSegmentName;
 			}
-			SanitizedName = CRContext.GetRigModuleNameSpace() + SanitizedName.GetName();
+			SanitizedName = CRContext.GetRigModulePrefix() + SanitizedName.GetName();
 			bAllowNameSpaceWhenSanitizingName = true;
 		}
 	}
@@ -2774,30 +2838,23 @@ FRigName URigHierarchy::GetSafeNewComponentName(const FRigElementKey& InElementK
 	return FRigName(UniqueName);
 }
 
-FText URigHierarchy::GetDisplayNameForUI(const FRigBaseElement* InElement, bool bIncludeNameSpace) const
+FText URigHierarchy::GetDisplayNameForUI(const FRigBaseElement* InElement, bool bIncludeModuleName) const
 {
 	check(InElement);
 
-	if (const UModularRig* ModularRig = GetTypedOuter<UModularRig>())
-	{
-		const FString& ShortestPath = ModularRig->GetShortestDisplayPathForElement(InElement->GetKey(), false);
-		if (!ShortestPath.IsEmpty())
-		{
-			return FText::FromString(ShortestPath);
-		}
-	}
-
 	const FName& DisplayName = InElement->GetDisplayName();
 	FString DisplayNameString = DisplayName.ToString();
-	(void)SplitNameSpace(DisplayNameString, nullptr, &DisplayNameString);
 	
-	if(bIncludeNameSpace)
+	const FRigHierarchyModulePath ModulePath(DisplayNameString);
+	(void)ModulePath.Split(nullptr, &DisplayNameString);
+	
+	if(bIncludeModuleName)
 	{
-		const FName ModuleShortName = GetNameMetadata(InElement->Key, ShortModuleNameMetadataName, NAME_None);
-		if(!ModuleShortName.IsNone())
+		const FName ModuleName = GetNameMetadata(InElement->Key, ModuleMetadataName, NAME_None);
+		if(!ModuleName.IsNone())
 		{
-			const FString ModuleShortNameString = ModuleShortName.ToString();
-			const FString ModuleDisplayName = JoinNameSpace(ModuleShortNameString, DisplayNameString);
+			const FString ModuleNameString = ModuleName.ToString();
+			const FString ModuleDisplayName = FRigHierarchyModulePath(ModuleNameString, DisplayNameString);
 			return FText::FromString(ModuleDisplayName);
 		}
 	}
@@ -2805,11 +2862,11 @@ FText URigHierarchy::GetDisplayNameForUI(const FRigBaseElement* InElement, bool
 	return FText::FromString(*DisplayNameString);
 }
 
-FText URigHierarchy::GetDisplayNameForUI(const FRigElementKey& InKey, bool bIncludeNameSpace) const
+FText URigHierarchy::GetDisplayNameForUI(const FRigElementKey& InKey, bool bIncludeModuleName) const
 {
 	if(const FRigBaseElement* Element = Find(InKey))
 	{
-		return GetDisplayNameForUI(Element, bIncludeNameSpace);
+		return GetDisplayNameForUI(Element, bIncludeModuleName);
 	}
 	return FText();
 }
@@ -9436,6 +9493,197 @@ TArray<FString> URigHierarchy::ConnectorSettingsToPythonCommands(const FRigConne
 
 #endif
 
+FRigElementKey URigHierarchy::PatchElementKeyInLookup(const FRigElementKey& InKey, const TMap<FRigHierarchyModulePath, FName>* InModulePathToName) const
+{
+	if(!InKey.IsValid())
+	{
+		return InKey;
+	}
+	if(const int32* Index = ElementIndexLookup.Find(InKey))
+	{
+		return Get(*Index)->GetKey();
+	}
+	
+	const FString NameString = InKey.Name.ToString();
+	if(!NameString.Contains(FRigHierarchyModulePath::NamespaceSeparator_Deprecated))
+	{
+		return InKey;
+	}
+
+	const TMap<FRigHierarchyModulePath, FName>* ModulePathToName = InModulePathToName;
+	if(ModulePathToName == nullptr)
+	{
+		const UModularRig* ModularRig = Cast<UModularRig>(GetOuter());
+		if(ModularRig == nullptr)
+		{
+			return InKey;
+		}
+		ModulePathToName = &ModularRig->GetModularRigModel().PreviousModulePaths;
+	}
+	
+	const FRigElementKey PatchedKey = InKey.ConvertToModuleNameFormat(ModulePathToName);
+	if(PatchedKey == InKey)
+	{
+		return InKey;
+	}
+
+	if(const int32* PatchedIndex = ElementIndexLookup.Find(PatchedKey))
+	{
+		const int32 Index = *PatchedIndex;
+		const_cast<URigHierarchy*>(this)->ElementIndexLookup.Add(InKey, Index);
+	}
+
+	return PatchedKey;
+}
+
+void URigHierarchy::PatchElementMetadata(const TMap<FRigHierarchyModulePath, FName>& InModulePathToName)
+{
+	bool bPerformedChange = false;
+	for(int32 Index = 0; Index < ElementMetadata.Num(); Index++)
+	{
+		FMetadataStorage& MetadataStorage = ElementMetadata[Index];
+		if(MetadataStorage.MetadataMap.IsEmpty())
+		{
+			continue;
+		}
+
+		bool bContainsNamespace = false;
+		for(const TPair<FName, FRigBaseMetadata*>& Pair : MetadataStorage.MetadataMap)
+		{
+			if(FRigHierarchyModulePath(Pair.Key).UsesNameSpaceFormat())
+			{
+				bContainsNamespace = true;
+				break;
+			}
+		}
+		
+		if(bContainsNamespace)
+		{
+			TMap<FName, FRigBaseMetadata*> OldMetadataMap;
+			Swap(OldMetadataMap, MetadataStorage.MetadataMap);
+
+			for(const TPair<FName, FRigBaseMetadata*>& Pair : MetadataStorage.MetadataMap)
+			{
+				FRigHierarchyModulePath ModulePath(Pair.Key);
+				if(ModulePath.ConvertToModuleNameFormatInline(&InModulePathToName))
+				{
+					MetadataStorage.MetadataMap.Add(*ModulePath.GetPath(), Pair.Value);
+					bPerformedChange = true;
+					continue;
+				}
+				MetadataStorage.MetadataMap.Add(Pair.Key, Pair.Value);
+			}
+		}
+
+		for(const TPair<FName, FRigBaseMetadata*>& Pair : MetadataStorage.MetadataMap)
+		{
+			if(Pair.Key == ModuleMetadataName && Pair.Value->GetType() == ERigMetadataType::Name)
+			{
+				FRigNameMetadata* NameMetadata = static_cast<FRigNameMetadata*>(Pair.Value);
+				const FName OldValueName = NameMetadata->GetValue();
+				if(const FName* ModuleName = InModulePathToName.Find((OldValueName.ToString())))
+				{
+					NameMetadata->SetValue(*ModuleName);
+					bPerformedChange = true;
+				}
+			}
+			else if(Pair.Value->GetType() == ERigMetadataType::RigElementKey)
+			{
+				FRigElementKeyMetadata* KeyMetadata = static_cast<FRigElementKeyMetadata*>(Pair.Value);
+				FRigElementKey ValueKey = KeyMetadata->GetValue();
+				if(ValueKey.ConvertToModuleNameFormatInline(&InModulePathToName))
+				{
+					KeyMetadata->SetValue(ValueKey);
+					bPerformedChange = true;
+				}
+			}
+			else if(Pair.Value->GetType() == ERigMetadataType::RigElementKeyArray)
+			{
+				FRigElementKeyArrayMetadata* KeyArrayMetadata = static_cast<FRigElementKeyArrayMetadata*>(Pair.Value);
+				TArray<FRigElementKey> ValueKeys = KeyArrayMetadata->GetValue();
+				for(FRigElementKey& ValueKey : ValueKeys)
+				{
+					if(ValueKey.ConvertToModuleNameFormatInline(&InModulePathToName))
+					{
+						bPerformedChange = true;
+					}
+				}
+				KeyArrayMetadata->SetValue(ValueKeys);
+			}
+		}
+	}
+
+	if(bPerformedChange)
+	{
+		MetadataVersion++;
+	}
+}
+
+void URigHierarchy::PatchModularRigComponentKeys(const TMap<FRigHierarchyModulePath, FName>& InModulePathToName)
+{
+	TMap<FRigHierarchyKey,FRigHierarchyKey> PatchedKeys;
+	for(FRigBaseElement* Element : Elements)
+	{
+		FRigElementKey PatchedKey = Element->GetKey();
+		if(PatchedKey.ConvertToModuleNameFormatInline(&InModulePathToName))
+		{
+			PatchedKeys.Add(Element->GetKey(), PatchedKey);
+			continue;
+		}
+
+		// resolve the element key in reverse
+		const FRigHierarchyModulePath ModulePath = Element->GetName();
+		if(ModulePath.UsesModuleNameFormat())
+		{
+			FString ModuleName,ElementName;
+			if(ModulePath.Split(&ModuleName, &ElementName))
+			{
+				if(const FRigHierarchyModulePath* OldModulePath = InModulePathToName.FindKey(*ModuleName))
+				{
+					const FString PathBasedName = JoinNameSpace_Deprecated(OldModulePath->GetPath(), ElementName);
+					PatchedKeys.Add(FRigElementKey(*PathBasedName, Element->GetType()), Element->GetKey());
+				}
+			}
+		}
+	}
+
+	for(FInstancedStruct& InstancedStruct : ElementComponents)
+	{
+		if(!InstancedStruct.IsValid())
+		{
+			continue;
+		}
+
+		FRigBaseComponent* Component = InstancedStruct.GetMutablePtr<FRigBaseComponent>();
+		const FRigElementKey PatchedKey = PatchElementKeyInLookup(Component->GetElementKey(), &InModulePathToName);
+		if(PatchedKey != Component->GetElementKey())
+		{
+			Modify();
+			
+			const FRigComponentKey OldKey = Component->Key;
+			Component->Key.ElementKey = PatchedKey;
+			const FRigComponentKey& NewKey = Component->Key;
+			PatchedKeys.Add(OldKey, NewKey);
+			ComponentIndexLookup.Remove(OldKey);
+			ComponentIndexLookup.Add(NewKey, Component->IndexInHierarchy);
+		}
+	}
+
+	for(FInstancedStruct& InstancedStruct : ElementComponents)
+	{
+		if(!InstancedStruct.IsValid())
+		{
+			continue;
+		}
+
+		FRigBaseComponent* Component = InstancedStruct.GetMutablePtr<FRigBaseComponent>();
+		for(const TPair<FRigHierarchyKey,FRigHierarchyKey>& Pair : PatchedKeys)
+		{
+			Component->OnRigHierarchyKeyChanged(Pair.Key, Pair.Value);
+		}
+	}
+}
+
 FRigHierarchyRedirectorGuard::FRigHierarchyRedirectorGuard(UControlRig* InControlRig)
 : Guard(InControlRig->GetHierarchy()->ElementKeyRedirector, &InControlRig->GetElementKeyRedirector())
 {
diff --git a/Engine/Plugins/Animation/ControlRig/Source/ControlRig/Private/Rigs/RigHierarchyController.cpp b/Engine/Plugins/Animation/ControlRig/Source/ControlRig/Private/Rigs/RigHierarchyController.cpp
--- a/Engine/Plugins/Animation/ControlRig/Source/ControlRig/Private/Rigs/RigHierarchyController.cpp
+++ b/Engine/Plugins/Animation/ControlRig/Source/ControlRig/Private/Rigs/RigHierarchyController.cpp
@@ -161,22 +161,21 @@ bool URigHierarchyController::SelectHierarchyKey(FRigHierarchyKey InKey, bool bS
 			return false;
 		}
 
-		const bool bSelectionState = Hierarchy->OrderedSelection.Contains(Key);
+		const bool bSelectionState = Hierarchy->OrderedSelection.Contains(Element->GetKey());
 		ensure(bSelectionState == Element->bSelected);
 		if(Element->bSelected == bSelect)
 		{
 			return false;
 		}
-
 		Element->bSelected = bSelect;
 	
 		if(bSelect)
 		{
-			Hierarchy->OrderedSelection.Add(Key);
+			Hierarchy->OrderedSelection.Add(Element->GetKey());
 		}
 		else
 		{
-			Hierarchy->OrderedSelection.Remove(Key);
+			Hierarchy->OrderedSelection.Remove(Element->GetKey());
 		}
 
 		if(Element->bSelected)
@@ -198,7 +197,7 @@ bool URigHierarchyController::SelectHierarchyKey(FRigHierarchyKey InKey, bool bS
 			return false;
 		}
 
-		const bool bSelectionState = Hierarchy->OrderedSelection.Contains(Key);
+		const bool bSelectionState = Hierarchy->OrderedSelection.Contains(Component->GetKey());
 		ensure(bSelectionState == Component->bSelected);
 		if(Component->bSelected == bSelect)
 		{
@@ -209,11 +208,11 @@ bool URigHierarchyController::SelectHierarchyKey(FRigHierarchyKey InKey, bool bS
 	
 		if(bSelect)
 		{
-			Hierarchy->OrderedSelection.Add(Key);
+			Hierarchy->OrderedSelection.Add(Component->GetKey());
 		}
 		else
 		{
-			Hierarchy->OrderedSelection.Remove(Key);
+			Hierarchy->OrderedSelection.Remove(Component->GetKey());
 		}
 
 		if(Component->bSelected)
@@ -511,7 +510,7 @@ FRigElementKey URigHierarchyController::AddControl(
 		else if(Hierarchy->HasExecuteContext())
 		{
 			const FControlRigExecuteContext& CRContext = Hierarchy->ExecuteContext->GetPublicData<FControlRigExecuteContext>();
-			if(!CRContext.GetRigModuleNameSpace().IsEmpty())
+			if(!CRContext.GetRigModulePrefix().IsEmpty())
 			{
 				// avoid self name collision
 				NewElement->Settings.DisplayName = NAME_None;
@@ -701,11 +700,11 @@ FRigElementKey URigHierarchyController::AddConnector(FName InName, FRigConnector
 				if(Hierarchy->HasExecuteContext())
 				{
 					const FControlRigExecuteContext& CRContext = Hierarchy->ExecuteContext->GetPublicData<FControlRigExecuteContext>();
-					const FString NameSpace = CRContext.GetRigModuleNameSpace();
-					if(!NameSpace.IsEmpty())
+					const FString ModulePrefix = CRContext.GetRigModulePrefix();
+					if(!ModulePrefix.IsEmpty())
 					{
-						const FString ConnectorNameSpace = Hierarchy->GetNameSpace(Connector->GetKey());
-						if(!ConnectorNameSpace.IsEmpty() && ConnectorNameSpace.Equals(NameSpace, ESearchCase::IgnoreCase))
+						const FString ConnectorModulePrefix = Hierarchy->GetModulePrefix(Connector->GetKey());
+						if(!ConnectorModulePrefix.IsEmpty() && ConnectorModulePrefix.Equals(ModulePrefix, ESearchCase::IgnoreCase))
 						{
 							static constexpr TCHAR Format[] = TEXT("Cannot add connector '%s' - there already is a primary connector.");
 							ReportAndNotifyErrorf(Format, *InName.ToString());
@@ -2782,10 +2781,10 @@ FName URigHierarchyController::GetSafeNewName(const FName& InDesiredName, ERigEl
 	// remove potential namespaces from it
 	if(!bAllowNameSpace)
 	{
-		int32 Index = INDEX_NONE;
-		if(Name.GetName().FindLastChar(TEXT(':'), Index))
+		const FRigHierarchyModulePath ModulePath(InDesiredName);
+		if(ModulePath.IsValid())
 		{
-			Name.SetName(Name.GetName().Mid(Index + 1));
+			Name.SetFName(ModulePath.GetElementFName());
 		}
 	}
 	
@@ -2809,11 +2808,11 @@ int32 URigHierarchyController::AddElement(FRigBaseElement* InElementToAdd, FRigB
 
 	FRigName DesiredName = InDesiredName;
 	URigHierarchy::SanitizeName(DesiredName);
-	
-	FRigName LastSegmentName;
-	if(URigHierarchy::SplitNameSpace(DesiredName, nullptr, &LastSegmentName, true))
+
+	const FRigHierarchyModulePath ModulePath(InDesiredName.ToString());
+	if(ModulePath.IsValid())
 	{
-		DesiredName = LastSegmentName;
+		DesiredName = ModulePath.GetElementFName();
 	}
 
 	if(!InDesiredName.IsNone() &&
@@ -2827,28 +2826,17 @@ int32 URigHierarchyController::AddElement(FRigBaseElement* InElementToAdd, FRigB
 	{
 		const FControlRigExecuteContext& CRContext = Hierarchy->ExecuteContext->GetPublicData<FControlRigExecuteContext>();
 
-		if(!CRContext.GetRigModuleNameSpace().IsEmpty())
+		if(!CRContext.GetRigModulePrefix().IsEmpty())
 		{
-			if(InElementToAdd->GetName().StartsWith(CRContext.GetRigModuleNameSpace(), ESearchCase::IgnoreCase))
+			if(InElementToAdd->GetName().StartsWith(CRContext.GetRigModulePrefix(), ESearchCase::IgnoreCase))
 			{
-				Hierarchy->SetNameMetadata(InElementToAdd->Key, URigHierarchy::NameSpaceMetadataName, *CRContext.GetRigModuleNameSpace());
-				Hierarchy->SetNameMetadata(InElementToAdd->Key, URigHierarchy::ModuleMetadataName, *CRContext.GetRigModuleNameSpace().LeftChop(1));
+				Hierarchy->SetNameMetadata(InElementToAdd->Key, URigHierarchy::ModuleMetadataName, *CRContext.GetRigModulePrefix().LeftChop(1));
 
 				if(Hierarchy->ElementKeyRedirector)
 				{
 					Hierarchy->ElementKeyRedirector->Add(FRigElementKey(DesiredName.GetFName(), InElementToAdd->Key.Type), {InElementToAdd->Key}, Hierarchy);
 				}
 			}
-
-			// also try to find the live module and retrieve the shortname from that to update the element's short name
-			if(const FRigModuleInstance* Module = CRContext.GetRigModuleInstance())
-			{
-				const FString ModuleShortName = Module->GetShortName();
-				if(!ModuleShortName.IsEmpty())
-				{
-					Hierarchy->SetNameMetadata(InElementToAdd->Key, URigHierarchy::ShortModuleNameMetadataName, *ModuleShortName);
-				}
-			}
 		}
 	}
 	
diff --git a/Engine/Plugins/Animation/ControlRig/Source/ControlRig/Private/Rigs/RigHierarchyDefines.cpp b/Engine/Plugins/Animation/ControlRig/Source/ControlRig/Private/Rigs/RigHierarchyDefines.cpp
--- a/Engine/Plugins/Animation/ControlRig/Source/ControlRig/Private/Rigs/RigHierarchyDefines.cpp
+++ b/Engine/Plugins/Animation/ControlRig/Source/ControlRig/Private/Rigs/RigHierarchyDefines.cpp
@@ -89,6 +89,311 @@ void FRigControlLimitEnabled::SetForValueType(ERigControlValueType InValueType,
 	}
 }
 
+////////////////////////////////////////////////////////////////////////////////
+// FRigHierarchyModulePath
+////////////////////////////////////////////////////////////////////////////////
+
+bool FRigHierarchyModulePath::IsValid() const
+{
+	return !IsEmpty() && (UsesNameSpaceFormat() || UsesModuleNameFormat());
+}
+
+bool FRigHierarchyModulePath::UsesNameSpaceFormat() const
+{
+	const int32 Len = ModulePath.Len();
+	if(Len >= 3)
+	{
+		int32 SeparatorIndex = INDEX_NONE;
+		if(ModulePath.FindLastChar(NamespaceSeparatorChar_Deprecated, SeparatorIndex))
+		{
+			return SeparatorIndex > 0 && SeparatorIndex < Len - 1;
+		}
+	};
+	return false;
+}
+
+bool FRigHierarchyModulePath::UsesModuleNameFormat() const
+{
+	const int32 Len = ModulePath.Len();
+	if(Len >= 3)
+	{
+		int32 SeparatorIndex = INDEX_NONE;
+		if(ModulePath.FindLastChar(ModuleNameSuffixChar, SeparatorIndex))
+		{
+			return SeparatorIndex > 0 && SeparatorIndex < Len - 1;
+		}
+	};
+	return false;
+}
+
+FName FRigHierarchyModulePath::GetPathFName() const
+{
+	if(IsEmpty())
+	{
+		return NAME_None;
+	}
+	return FName(*ModulePath);
+}
+
+FStringView FRigHierarchyModulePath::GetModuleName() const
+{
+	const int32 Len = ModulePath.Len();
+	if(Len >= 3)
+	{
+		int32 SeparatorIndex = INDEX_NONE;
+		if(ModulePath.FindLastChar(ModuleNameSuffixChar, SeparatorIndex))
+		{
+			return FStringView(ModulePath).Left(SeparatorIndex);
+		}
+		if(ModulePath.FindLastChar(NamespaceSeparatorChar_Deprecated, SeparatorIndex))
+		{
+			return FStringView(ModulePath).Left(SeparatorIndex);
+		}
+	};
+	static const FString EmptyString;
+	return EmptyString;
+}
+
+const FString& FRigHierarchyModulePath::GetModuleNameString() const
+{
+	if(!CachedModuleNameString.IsSet())
+	{
+		CachedModuleNameString = FString(GetModuleName());
+	}
+	return CachedModuleNameString.GetValue();
+}
+
+const FName& FRigHierarchyModulePath::GetModuleFName() const
+{
+	if(!CachedModuleFName.IsSet())
+	{
+		const FString& ModuleNameString = GetModuleNameString();
+		if(ModuleNameString.IsEmpty())
+		{
+			CachedModuleFName = NAME_None;
+		}
+		else
+		{
+			CachedModuleFName = *ModuleNameString;
+		}
+	}
+	return CachedModuleFName.GetValue();
+}
+
+FStringView FRigHierarchyModulePath::GetModulePrefix() const
+{
+	const int32 Len = ModulePath.Len();
+	if(Len >= 3)
+	{
+		int32 SeparatorIndex = INDEX_NONE;
+		if(ModulePath.FindLastChar(ModuleNameSuffixChar, SeparatorIndex))
+		{
+			return FStringView(ModulePath).Left(SeparatorIndex + 1);
+		}
+		if(ModulePath.FindLastChar(NamespaceSeparatorChar_Deprecated, SeparatorIndex))
+		{
+			return FStringView(ModulePath).Left(SeparatorIndex + 1);
+		}
+	};
+	static const FString EmptyString;
+	return EmptyString;
+}
+
+FString FRigHierarchyModulePath::GetModulePrefixString() const
+{
+	return FString(GetModulePrefix());
+}
+
+FStringView FRigHierarchyModulePath::GetElementName() const
+{
+	const int32 Len = ModulePath.Len();
+	if(Len >= 3)
+	{
+		int32 SeparatorIndex = INDEX_NONE;
+		if(ModulePath.FindLastChar(ModuleNameSuffixChar, SeparatorIndex))
+		{
+			return FStringView(ModulePath).Mid(SeparatorIndex+1);
+		}
+		if(ModulePath.FindLastChar(NamespaceSeparatorChar_Deprecated, SeparatorIndex))
+		{
+			return FStringView(ModulePath).Mid(SeparatorIndex+1);
+		}
+	};
+	static const FString EmptyString;
+	return EmptyString;
+}
+
+const FString& FRigHierarchyModulePath::GetElementNameString() const
+{
+	if(!CachedElementNameString.IsSet())
+	{
+		CachedElementNameString = FString(GetElementName());
+	}
+	return CachedElementNameString.GetValue();
+}
+
+const FName& FRigHierarchyModulePath::GetElementFName() const
+{
+	if(!CachedElementFName.IsSet())
+	{
+		const FString& ElementNameString = GetElementNameString();
+		if(ElementNameString.IsEmpty())
+		{
+			CachedElementFName = NAME_None;
+		}
+		else
+		{
+			CachedElementFName = *ElementNameString;
+		}
+	}
+	return CachedElementFName.GetValue();
+}
+
+FRigHierarchyModulePath FRigHierarchyModulePath::Join(const FString& InModuleName, const FString& InElementName)
+{
+	if(InModuleName.IsEmpty() || InElementName.IsEmpty())
+	{
+		return FRigHierarchyModulePath(); 
+	}
+	if(InModuleName.EndsWith(ModuleNameSuffix))
+	{
+		return InModuleName + InElementName;
+	}
+	FRigHierarchyModulePath Result(InModuleName + ModuleNameSuffix + InElementName);
+	Result.CachedModuleNameString = InModuleName;
+	Result.CachedElementNameString = InElementName;
+	return Result;
+}
+
+FRigHierarchyModulePath FRigHierarchyModulePath::Join(const FName& InModuleFName, const FName& InElementFName)
+{
+	if(InModuleFName.IsNone() || InElementFName.IsNone())
+	{
+		return FRigHierarchyModulePath(); 
+	}
+	FRigHierarchyModulePath Result = Join(InModuleFName.ToString(), InElementFName.ToString());
+	Result.CachedModuleFName = InModuleFName;
+	Result.CachedElementFName = InElementFName;
+	return Result;
+}
+
+bool FRigHierarchyModulePath::Split(FStringView* OutModuleName, FStringView* OutElementName) const
+{
+	if(CachedModuleNameString.IsSet() && CachedElementNameString.IsSet())
+	{
+		if(OutModuleName)
+		{
+			*OutModuleName = CachedModuleNameString.GetValue();
+		}
+		if(OutElementName)
+		{
+			*OutElementName = CachedElementNameString.GetValue();
+		}
+		return true;
+	}
+	
+	const int32 Len = ModulePath.Len();
+	if(Len >= 3)
+	{
+		int32 SeparatorIndex = INDEX_NONE;
+		if(ModulePath.FindLastChar(ModuleNameSuffixChar, SeparatorIndex))
+		{
+			if(OutModuleName)
+			{
+				*OutModuleName = FStringView(ModulePath).Left(SeparatorIndex); 
+			}
+			if(OutElementName)
+			{
+				*OutElementName = FStringView(ModulePath).Mid(SeparatorIndex+1); 
+			}
+			return true;
+		}
+		if(ModulePath.FindLastChar(NamespaceSeparatorChar_Deprecated, SeparatorIndex))
+		{
+			if(OutModuleName)
+			{
+				*OutModuleName = FStringView(ModulePath).Left(SeparatorIndex); 
+			}
+			if(OutElementName)
+			{
+				*OutElementName = FStringView(ModulePath).Mid(SeparatorIndex+1); 
+			}
+			return true;
+		}
+	};
+	return false;
+}
+
+bool FRigHierarchyModulePath::Split(FString* OutModuleName, FString* OutElementName) const
+{
+	if(CachedModuleNameString.IsSet() && CachedElementNameString.IsSet())
+	{
+		if(OutModuleName)
+		{
+			*OutModuleName = CachedModuleNameString.GetValue();
+		}
+		if(OutElementName)
+		{
+			*OutElementName = CachedElementNameString.GetValue();
+		}
+		return true;
+	}
+
+	FStringView ModuleName, ElementName;
+	if(Split(&ModuleName, &ElementName))
+	{
+		if(OutModuleName)
+		{
+			*OutModuleName = ModuleName;
+		}
+		if(OutElementName)
+		{
+			*OutElementName = ElementName;
+		}
+		return true;
+	}
+	return false;
+}
+
+FRigHierarchyModulePath FRigHierarchyModulePath::ConvertToModuleNameFormat(const TMap<FRigHierarchyModulePath, FName>* InModulePathToModuleName) const
+{
+	FRigHierarchyModulePath Result = *this;
+	(void)Result.ConvertToModuleNameFormatInline(InModulePathToModuleName);
+	return Result;
+}
+
+bool FRigHierarchyModulePath::ConvertToModuleNameFormatInline(const TMap<FRigHierarchyModulePath, FName>* InModulePathToModuleName)
+{
+	if(UsesNameSpaceFormat())
+	{
+		FString Left, Right;
+		if(Split(&Left, &Right))
+		{
+			if(InModulePathToModuleName)
+			{
+				if(const FName* RemappedModuleName = InModulePathToModuleName->Find(FString(Left)))
+				{
+					const FString NameString = RemappedModuleName->ToString();
+					*this = Join(NameString, Right);
+					return true;
+				}
+			}
+
+			const FRigHierarchyModulePath OldModulePath = Left;
+			if(OldModulePath.UsesNameSpaceFormat())
+			{
+				*this = Join(FString(OldModulePath.GetElementName()), Right);
+			}
+			else
+			{
+				*this = Join(Left, Right);
+			}
+			return true;
+		}
+	}
+	return false;
+}
+
 ////////////////////////////////////////////////////////////////////////////////
 // FRigElementKey
 ////////////////////////////////////////////////////////////////////////////////
@@ -142,6 +447,23 @@ FString FRigElementKey::ToPythonString() const
 #endif
 }
 
+FRigElementKey FRigElementKey::ConvertToModuleNameFormat(const TMap<FRigHierarchyModulePath, FName>* InModulePathToModuleName) const
+{
+	FRigElementKey PatchedKey = *this;
+	(void)PatchedKey.ConvertToModuleNameFormatInline(InModulePathToModuleName);
+	return PatchedKey;
+}
+
+bool FRigElementKey::ConvertToModuleNameFormatInline(const TMap<FRigHierarchyModulePath, FName>* InModulePathToModuleName)
+{
+	FRigHierarchyModulePath ModulePath = Name.ToString();
+	if(ModulePath.ConvertToModuleNameFormatInline(InModulePathToModuleName))
+	{
+		Name = *ModulePath.GetPath();
+	}
+	return false;
+}
+
 ////////////////////////////////////////////////////////////////////////////////
 // FRigComponentKey
 ////////////////////////////////////////////////////////////////////////////////
@@ -196,6 +518,18 @@ bool FRigComponentKey::IsTopLevel() const
 	return ElementKey == URigHierarchy::GetTopLevelComponentElementKey();
 }
 
+FRigComponentKey FRigComponentKey::ConvertToModuleNameFormat(const TMap<FRigHierarchyModulePath, FName>* InModulePathToModuleName) const
+{
+	FRigComponentKey Result = *this;
+	(void)Result.ConvertToModuleNameFormatInline(InModulePathToModuleName);
+	return Result;
+}
+
+bool FRigComponentKey::ConvertToModuleNameFormatInline(const TMap<FRigHierarchyModulePath, FName>* InModulePathToModuleName)
+{
+	return ElementKey.ConvertToModuleNameFormatInline(InModulePathToModuleName);
+}
+
 ////////////////////////////////////////////////////////////////////////////////
 // FRigHierarchyKey
 ////////////////////////////////////////////////////////////////////////////////
diff --git a/Engine/Plugins/Animation/ControlRig/Source/ControlRig/Private/Rigs/RigHierarchyElements.cpp b/Engine/Plugins/Animation/ControlRig/Source/ControlRig/Private/Rigs/RigHierarchyElements.cpp
--- a/Engine/Plugins/Animation/ControlRig/Source/ControlRig/Private/Rigs/RigHierarchyElements.cpp
+++ b/Engine/Plugins/Animation/ControlRig/Source/ControlRig/Private/Rigs/RigHierarchyElements.cpp
@@ -7,7 +7,6 @@
 #include "ControlRigObjectVersion.h"
 #include "ControlRigGizmoLibrary.h"
 #include "AnimationCoreLibrary.h"
-#include "Algo/Transform.h"
 
 #include UE_INLINE_GENERATED_CPP_BY_NAME(RigHierarchyElements)
 
@@ -87,7 +86,6 @@ void FRigBaseElement::Load(FArchive& Ar, const FRigHierarchySerializationSetting
 	}
 }
 
-
 FRigBaseMetadata* FRigBaseElement::GetMetadata(const FName& InName, ERigMetadataType InType)
 {
 	if (!Owner)
diff --git a/Engine/Plugins/Animation/ControlRig/Source/ControlRig/Private/Sequencer/MovieSceneControlRigParameterSection.cpp b/Engine/Plugins/Animation/ControlRig/Source/ControlRig/Private/Sequencer/MovieSceneControlRigParameterSection.cpp
--- a/Engine/Plugins/Animation/ControlRig/Source/ControlRig/Private/Sequencer/MovieSceneControlRigParameterSection.cpp
+++ b/Engine/Plugins/Animation/ControlRig/Source/ControlRig/Private/Sequencer/MovieSceneControlRigParameterSection.cpp
@@ -910,89 +910,6 @@ void UMovieSceneControlRigParameterSection::PreSave(FObjectPreSaveContext SaveCo
 	Super::PreSave(SaveContext);
 }
 
-bool UMovieSceneControlRigParameterSection::RenameParameterName(const FName& OldParameterName, const FName& NewParameterName, TOptional<ERigControlType> ControlType)
-{
-	bool bWasReplaced = false;
-
-	auto RenameParameterNameInner = [this, &bWasReplaced, OldParameterName, NewParameterName](auto& ParameterNamesAndCurves)
-	{
-		for (auto& ParameterNameAndCurve : ParameterNamesAndCurves)
-		{
-			if (ParameterNameAndCurve.ParameterName == OldParameterName)
-			{
-				if (!bWasReplaced)
-				{
-					Modify();
-					bWasReplaced = true;
-				}
-				ParameterNameAndCurve.ParameterName = NewParameterName;
-				break;
-			}
-		}
-	};
-	if(ControlType.IsSet())
-	{
-		switch (ControlType.GetValue())
-		{
-			case ERigControlType::Float:
-			case ERigControlType::ScaleFloat:
-			{
-				RenameParameterNameInner(ScalarParameterNamesAndCurves);
-				break;
-			}
-			case ERigControlType::Bool:
-			{
-				RenameParameterNameInner(BoolParameterNamesAndCurves);
-				break;
-			}
-			case ERigControlType::Integer:
-			{
-				RenameParameterNameInner(IntegerParameterNamesAndCurves);
-				RenameParameterNameInner(EnumParameterNamesAndCurves);
-				break;
-			}
-			case ERigControlType::Vector2D:
-			{
-				RenameParameterNameInner(Vector2DParameterNamesAndCurves);
-				break;
-			}
-			case ERigControlType::Position:
-			case ERigControlType::Rotator:
-			case ERigControlType::Scale:
-			{
-				RenameParameterNameInner(VectorParameterNamesAndCurves);
-				break;
-			}
-			break;
-			case ERigControlType::Transform:
-			case ERigControlType::EulerTransform:
-			case ERigControlType::TransformNoScale:
-			{
-				RenameParameterNameInner(TransformParameterNamesAndCurves);
-				break;
-			}
-		
-		};
-	}
-	else
-	{
-		RenameParameterNameInner(ScalarParameterNamesAndCurves);
-		RenameParameterNameInner(BoolParameterNamesAndCurves);
-		RenameParameterNameInner(EnumParameterNamesAndCurves);
-		RenameParameterNameInner(IntegerParameterNamesAndCurves);
-		RenameParameterNameInner(Vector2DParameterNamesAndCurves);
-		RenameParameterNameInner(VectorParameterNamesAndCurves);
-		RenameParameterNameInner(ColorParameterNamesAndCurves);
-		RenameParameterNameInner(TransformParameterNamesAndCurves);
-	}
-
-	if (bWasReplaced)
-	{
-		ReconstructChannelProxy();
-	}
-	return bWasReplaced;
-}
-
 //make sure to zero out Scale values if getting to Additive or use the current values if getting set to Override
 void UMovieSceneControlRigParameterSection::SetBlendType(EMovieSceneBlendType InBlendType)
 {
@@ -1686,6 +1603,72 @@ const FName& UMovieSceneControlRigParameterSection::FindControlNameFromConstrain
 	return DummyName;
 }
 
+void UMovieSceneControlRigParameterSection::ForEachParameter(TFunction<void(FBaseParameterNameAndValue*)> InCallback)
+{
+	ForEachParameter(ScalarParameterNamesAndCurves, InCallback);
+	ForEachParameter(BoolParameterNamesAndCurves, InCallback);
+	ForEachParameter(EnumParameterNamesAndCurves, InCallback);
+	ForEachParameter(IntegerParameterNamesAndCurves, InCallback);
+	ForEachParameter(Vector2DParameterNamesAndCurves, InCallback);
+	ForEachParameter(VectorParameterNamesAndCurves, InCallback);
+	ForEachParameter(ColorParameterNamesAndCurves, InCallback);
+	ForEachParameter(TransformParameterNamesAndCurves, InCallback);
+}
+
+void UMovieSceneControlRigParameterSection::ForEachParameter(TOptional<ERigControlType> InControlType,
+	TFunction<void(FBaseParameterNameAndValue*)> InCallback)
+{
+	if(!InControlType.IsSet())
+	{
+		ForEachParameter(InCallback);
+		return;
+	}
+	
+	switch (InControlType.GetValue())
+	{
+		case ERigControlType::Float:
+		case ERigControlType::ScaleFloat:
+		{
+			ForEachParameter(ScalarParameterNamesAndCurves, InCallback);
+			break;
+		}
+		case ERigControlType::Bool:
+		{
+			ForEachParameter(BoolParameterNamesAndCurves, InCallback);
+			break;
+		}
+		case ERigControlType::Integer:
+		{
+			ForEachParameter(IntegerParameterNamesAndCurves, InCallback);
+			ForEachParameter(EnumParameterNamesAndCurves, InCallback);
+			break;
+		}
+		case ERigControlType::Vector2D:
+		{
+			ForEachParameter(Vector2DParameterNamesAndCurves, InCallback);
+			break;
+		}
+		case ERigControlType::Position:
+		case ERigControlType::Rotator:
+		case ERigControlType::Scale:
+		{
+			ForEachParameter(VectorParameterNamesAndCurves, InCallback);
+			break;
+		}
+		case ERigControlType::Transform:
+		case ERigControlType::EulerTransform:
+		case ERigControlType::TransformNoScale:
+		{
+			ForEachParameter(TransformParameterNamesAndCurves, InCallback);
+			break;
+		}
+		default:
+		{
+			break;
+		}
+	}
+}
+
 TArray<FSpaceControlNameAndChannel>& UMovieSceneControlRigParameterSection::GetSpaceChannels()
 {
 	return SpaceChannels;
@@ -1843,6 +1826,8 @@ EMovieSceneChannelProxyType UMovieSceneControlRigParameterSection::CacheChannelP
 					continue;
 				}
 
+				const FName& ControlName = ControlElement->GetFName();
+
 				FName ParentControlName = NAME_None;
 				FText Group;
 
@@ -1865,7 +1850,7 @@ EMovieSceneChannelProxyType UMovieSceneControlRigParameterSection::CacheChannelP
 				{
 					for (FScalarParameterNameAndCurve& Scalar : GetScalarParameterNamesAndCurves())
 					{
-						if (ControlElement->GetFName() == Scalar.ParameterName)
+						if (ControlName == Scalar.ParameterName)
 						{
 							if (Group.IsEmpty())
 							{
@@ -1900,7 +1885,7 @@ EMovieSceneChannelProxyType UMovieSceneControlRigParameterSection::CacheChannelP
 				{
 					for (FBoolParameterNameAndCurve& Bool : GetBoolParameterNamesAndCurves())
 					{
-						if (ControlElement->GetFName() == Bool.ParameterName)
+						if (ControlName == Bool.ParameterName)
 						{
 							if (Group.IsEmpty())
 							{
@@ -1937,7 +1922,7 @@ EMovieSceneChannelProxyType UMovieSceneControlRigParameterSection::CacheChannelP
 					{
 						for (FEnumParameterNameAndCurve& Enum : GetEnumParameterNamesAndCurves())
 						{
-							if (ControlElement->GetFName() == Enum.ParameterName)
+							if (ControlName == Enum.ParameterName)
 							{
 								if (Group.IsEmpty())
 								{
@@ -1971,7 +1956,7 @@ EMovieSceneChannelProxyType UMovieSceneControlRigParameterSection::CacheChannelP
 					{
 						for (FIntegerParameterNameAndCurve& Integer : GetIntegerParameterNamesAndCurves())
 						{
-							if (ControlElement->GetFName() == Integer.ParameterName)
+							if (ControlName == Integer.ParameterName)
 							{
 								if (Group.IsEmpty())
 								{
@@ -2010,7 +1995,7 @@ EMovieSceneChannelProxyType UMovieSceneControlRigParameterSection::CacheChannelP
 				{
 					for (FVector2DParameterNameAndCurves& Vector2D : GetVector2DParameterNamesAndCurves())
 					{
-						if (ControlElement->GetFName() == Vector2D.ParameterName)
+						if (ControlName == Vector2D.ParameterName)
 						{
 							if (Group.IsEmpty())
 							{
@@ -2046,7 +2031,7 @@ EMovieSceneChannelProxyType UMovieSceneControlRigParameterSection::CacheChannelP
 				{
 					for (FVectorParameterNameAndCurves& Vector : GetVectorParameterNamesAndCurves())
 					{
-						if (ControlElement->GetFName() == Vector.ParameterName)
+						if (ControlName == Vector.ParameterName)
 						{
 							if (Group.IsEmpty())
 							{
@@ -2125,9 +2110,8 @@ EMovieSceneChannelProxyType UMovieSceneControlRigParameterSection::CacheChannelP
 				{
 					for (FTransformParameterNameAndCurves& Transform : GetTransformParameterNamesAndCurves())
 					{
-						if (ControlElement->GetFName() == Transform.ParameterName)
+						if (ControlName == Transform.ParameterName)
 						{
-							const FName ControlName = ControlElement->GetFName();
 							if (Group.IsEmpty())
 							{
 								ControlChannelMap.Add(Transform.ParameterName, FChannelMapInfo(ControlIndex, SortOrder, FloatChannelIndex, INDEX_NONE, NAME_None, MaskIndex, CategoryIndex));
@@ -2226,7 +2210,7 @@ EMovieSceneChannelProxyType UMovieSceneControlRigParameterSection::CacheChannelP
 				{
 					for (FScalarParameterNameAndCurve& Scalar : GetScalarParameterNamesAndCurves())
 					{
-						if (ControlElement->GetFName() == Scalar.ParameterName)
+						if (ControlName == Scalar.ParameterName)
 						{
 							ControlChannelMap.Add(Scalar.ParameterName, FChannelMapInfo(ControlIndex, SortOrder, FloatChannelIndex, INDEX_NONE, NAME_None, MaskIndex));
 							Channels.Add(Scalar.ParameterCurve);
@@ -2242,7 +2226,7 @@ EMovieSceneChannelProxyType UMovieSceneControlRigParameterSection::CacheChannelP
 				{
 					for (FBoolParameterNameAndCurve& Bool : GetBoolParameterNamesAndCurves())
 					{
-						if (ControlElement->GetFName() == Bool.ParameterName)
+						if (ControlName == Bool.ParameterName)
 						{
 							ControlChannelMap.Add(Bool.ParameterName, FChannelMapInfo(ControlIndex, SortOrder, BoolChannelIndex, INDEX_NONE, NAME_None, MaskIndex));
 							Channels.Add(Bool.ParameterCurve);
@@ -2260,7 +2244,7 @@ EMovieSceneChannelProxyType UMovieSceneControlRigParameterSection::CacheChannelP
 					{
 						for (FEnumParameterNameAndCurve& Enum : GetEnumParameterNamesAndCurves())
 						{
-							if (ControlElement->GetFName() == Enum.ParameterName)
+							if (ControlName == Enum.ParameterName)
 							{
 								ControlChannelMap.Add(Enum.ParameterName, FChannelMapInfo(ControlIndex, SortOrder, EnumChannelIndex, INDEX_NONE, NAME_None, MaskIndex));
 								Channels.Add(Enum.ParameterCurve);
@@ -2275,7 +2259,7 @@ EMovieSceneChannelProxyType UMovieSceneControlRigParameterSection::CacheChannelP
 					{
 						for (FIntegerParameterNameAndCurve& Integer : GetIntegerParameterNamesAndCurves())
 						{
-							if (ControlElement->GetFName() == Integer.ParameterName)
+							if (ControlName == Integer.ParameterName)
 							{
 								ControlChannelMap.Add(Integer.ParameterName, FChannelMapInfo(ControlIndex, SortOrder, IntegerChannelIndex, INDEX_NONE, NAME_None, MaskIndex));
 								Channels.Add(Integer.ParameterCurve);
@@ -2292,7 +2276,7 @@ EMovieSceneChannelProxyType UMovieSceneControlRigParameterSection::CacheChannelP
 				{
 					for (FVector2DParameterNameAndCurves& Vector2D : GetVector2DParameterNamesAndCurves())
 					{
-						if (ControlElement->GetFName() == Vector2D.ParameterName)
+						if (ControlName == Vector2D.ParameterName)
 						{
 							ControlChannelMap.Add(Vector2D.ParameterName, FChannelMapInfo(ControlIndex, SortOrder, FloatChannelIndex, INDEX_NONE, NAME_None, MaskIndex));
 							Channels.Add(Vector2D.XCurve);
@@ -2311,7 +2295,7 @@ EMovieSceneChannelProxyType UMovieSceneControlRigParameterSection::CacheChannelP
 				{
 					for (FVectorParameterNameAndCurves& Vector : GetVectorParameterNamesAndCurves())
 					{
-						if (ControlElement->GetFName() == Vector.ParameterName)
+						if (ControlName == Vector.ParameterName)
 						{
 							ControlChannelMap.Add(Vector.ParameterName, FChannelMapInfo(ControlIndex, SortOrder, FloatChannelIndex, INDEX_NONE, NAME_None, MaskIndex));
 							bool bDoSpaceChannel = true;
@@ -2357,7 +2341,7 @@ EMovieSceneChannelProxyType UMovieSceneControlRigParameterSection::CacheChannelP
 				{
 					for (FTransformParameterNameAndCurves& Transform : GetTransformParameterNamesAndCurves())
 					{
-						if (ControlElement->GetFName() == Transform.ParameterName)
+						if (ControlName == Transform.ParameterName)
 						{
 							ControlChannelMap.Add(Transform.ParameterName, FChannelMapInfo(ControlIndex, SortOrder, FloatChannelIndex, INDEX_NONE, NAME_None, MaskIndex));
 
@@ -2718,29 +2702,51 @@ void UMovieSceneControlRigParameterSection::RecreateWithThisControlRig(UControlR
 		}
 	}
 
-	for (FRigControlElement* ControlElement : SortedControls)
+	// rename all existing parameters based on previous name information.
+	ForEachParameter([this, Hierarchy](FBaseParameterNameAndValue* InOutParameter)
 	{
-		if (!Hierarchy->IsAnimatable(ControlElement))
+		const FRigElementKey ControlKey = FRigElementKey(InOutParameter->ParameterName, ERigElementType::Control);
+		if(const FRigHierarchyKey* NewControlKey = Hierarchy->GetPreviousNameMap().Find(ControlKey))
 		{
-			continue;
+			InOutParameter->ParameterName = NewControlKey->GetFName();
 		}
-
-		FName PreviousName = Hierarchy->GetPreviousName(ControlElement->GetKey());
-		if (PreviousName != NAME_None && PreviousName != ControlElement->GetKey().Name)
+		else
 		{
-			TOptional<ERigControlType> ControlType = ControlElement->Settings.ControlType;
-			RenameParameterName(PreviousName, ControlElement->GetKey().Name, ControlType);
+			const FRigElementKey CurveKey = FRigElementKey(InOutParameter->ParameterName, ERigElementType::Curve);
+			if(const FRigHierarchyKey* NewCurveKey = Hierarchy->GetPreviousNameMap().Find(CurveKey))
+			{
+				InOutParameter->ParameterName = NewCurveKey->GetFName();
+			}
 		}
-		if (ControlElement->Settings.ControlType == ERigControlType::Float ||
-			ControlElement->Settings.ControlType == ERigControlType::ScaleFloat)
+	});
+
+	// rename all existing parameters based on short name vs long name.
+	// this also recovers from controls being stored with an original long path
+	// which may now be outdated. UControlRig::FindControl has backwards compat for that.
+	if(InControlRig->IsModularRig())
+	{
+		ForEachParameter([this, InControlRig](FBaseParameterNameAndValue* InOutParameter)
 		{
-			if (const FName* OldCurveControlName = CurveControlNameRemapping.Find(ControlElement->GetFName()))
+			if(const FRigControlElement* Control = InControlRig->FindControl(InOutParameter->ParameterName))
 			{
-				TOptional<ERigControlType> ControlType = ControlElement->Settings.ControlType;
-				RenameParameterName(*OldCurveControlName, ControlElement->GetKey().Name, ControlType);
+				const FName& ControlName = Control->GetFName();
+				if(ControlName != InOutParameter->ParameterName)
+				{
+					InOutParameter->ParameterName = ControlName;
+				}
 			}
+		});
+	}
+
+	for (FRigControlElement* ControlElement : SortedControls)
+	{
+		if (!Hierarchy->IsAnimatable(ControlElement))
+		{
+			continue;
 		}
-		
+
+		const FName& ControlName = ControlElement->GetFName();
+
 		switch (ControlElement->Settings.ControlType)
 		{
 		case ERigControlType::Float:
@@ -2752,7 +2758,7 @@ void UMovieSceneControlRigParameterSection::RecreateWithThisControlRig(UControlR
 				//or use IntialValue?
 				DefaultValue = ControlRig->GetControlValue(ControlElement, ERigControlValueType::Current).Get<float>();
 			}
-			AddScalarParameter(ControlElement->GetFName(), DefaultValue, false);
+			AddScalarParameter(ControlName, DefaultValue, false);
 			break;
 		}
 		case ERigControlType::Bool:
@@ -2762,7 +2768,7 @@ void UMovieSceneControlRigParameterSection::RecreateWithThisControlRig(UControlR
 			if (bSetDefault)
 			{
 				DefaultValue = ControlRig->GetControlValue(ControlElement, ERigControlValueType::Current).Get<bool>();
-				AddBoolParameter(ControlElement->GetFName(), DefaultValue, false);
+				AddBoolParameter(ControlName, DefaultValue, false);
 
 			}
 			break;
@@ -2776,7 +2782,7 @@ void UMovieSceneControlRigParameterSection::RecreateWithThisControlRig(UControlR
 				if (bSetDefault)
 				{
 					DefaultValue = (uint8)ControlRig->GetControlValue(ControlElement, ERigControlValueType::Current).Get<int32>();
-					AddEnumParameter(ControlElement->GetFName(), ControlElement->Settings.ControlEnum, DefaultValue, false);
+					AddEnumParameter(ControlName, ControlElement->Settings.ControlEnum, DefaultValue, false);
 				}
 			}
 			else
@@ -2786,7 +2792,7 @@ void UMovieSceneControlRigParameterSection::RecreateWithThisControlRig(UControlR
 				if (bSetDefault)
 				{
 					DefaultValue = ControlRig->GetControlValue(ControlElement, ERigControlValueType::Current).Get<int32>();
-					AddIntegerParameter(ControlElement->GetFName(), DefaultValue, false);
+					AddIntegerParameter(ControlName, DefaultValue, false);
 				}
 			}
 			break;
@@ -2800,7 +2806,7 @@ void UMovieSceneControlRigParameterSection::RecreateWithThisControlRig(UControlR
 				const FVector3f TempValue = ControlRig->GetControlValue(ControlElement, ERigControlValueType::Current).Get<FVector3f>();
 				DefaultValue = FVector2D(TempValue.X, TempValue.Y);
 			}
-			AddVector2DParameter(ControlElement->GetFName(), DefaultValue, false);
+			AddVector2DParameter(ControlName, DefaultValue, false);
 			break;
 		}
 
@@ -2814,7 +2820,7 @@ void UMovieSceneControlRigParameterSection::RecreateWithThisControlRig(UControlR
 				//or use IntialValue?
 				DefaultValue = (FVector)ControlRig->GetControlValue(ControlElement, ERigControlValueType::Current).Get<FVector3f>();
 			}
-			AddVectorParameter(ControlElement->GetFName(), DefaultValue, false);
+			AddVectorParameter(ControlName, DefaultValue, false);
 			//mz todo specify rotator special so we can do quat interps
 			break;
 		}
@@ -2847,7 +2853,7 @@ void UMovieSceneControlRigParameterSection::RecreateWithThisControlRig(UControlR
 					DefaultValue = FEulerTransform(NoScale.Rotation.Rotator(), NoScale.Location, FVector::OneVector);
 				}
 			}
-			AddTransformParameter(ControlElement->GetFName(), DefaultValue, false);
+			AddTransformParameter(ControlName, DefaultValue, false);
 			break;
 		}
 
@@ -4085,16 +4091,6 @@ UObject* UMovieSceneControlRigParameterSection::GetImplicitObjectOwner()
 	return Super::GetImplicitObjectOwner();
 }
 
-FEnumParameterNameAndCurve::FEnumParameterNameAndCurve(FName InParameterName)
-{
-	ParameterName = InParameterName;
-}
-
-FIntegerParameterNameAndCurve::FIntegerParameterNameAndCurve(FName InParameterName)
-{
-	ParameterName = InParameterName;
-}
-
 void UMovieSceneControlRigParameterSection::ImportEntityImpl(UMovieSceneEntitySystemLinker* EntityLinker, const FEntityImportParams& Params, FImportedEntity* OutImportedEntity)
 {
 	using namespace UE::MovieScene;
diff --git a/Engine/Plugins/Animation/ControlRig/Source/ControlRig/Private/Sequencer/MovieSceneControlRigParameterTrack.cpp b/Engine/Plugins/Animation/ControlRig/Source/ControlRig/Private/Sequencer/MovieSceneControlRigParameterTrack.cpp
--- a/Engine/Plugins/Animation/ControlRig/Source/ControlRig/Private/Sequencer/MovieSceneControlRigParameterTrack.cpp
+++ b/Engine/Plugins/Animation/ControlRig/Source/ControlRig/Private/Sequencer/MovieSceneControlRigParameterTrack.cpp
@@ -729,12 +729,24 @@ void UMovieSceneControlRigParameterTrack::RenameParameterName(const FName& OldPa
 {
 	if (OldParameterName != NewParameterName)
 	{
-		TOptional<ERigControlType> ControlType; //not set so need to check all 
 		for (UMovieSceneSection* Section : Sections)
 		{
 			if (UMovieSceneControlRigParameterSection* CRSection = Cast<UMovieSceneControlRigParameterSection>(Section))
 			{
-				CRSection->RenameParameterName(OldParameterName, NewParameterName, ControlType);
+				bool bRequiredToRebuildProxy = false;
+				CRSection->ForEachParameter([OldParameterName, NewParameterName, CRSection, &bRequiredToRebuildProxy](FBaseParameterNameAndValue* InOutParameter)
+				{
+					if(InOutParameter->ParameterName == OldParameterName)
+					{
+						CRSection->Modify();
+						InOutParameter->ParameterName = NewParameterName;
+						bRequiredToRebuildProxy = true;
+					}
+				});
+				if(bRequiredToRebuildProxy)
+				{
+					CRSection->ReconstructChannelProxy();
+				}
 			}
 		}
 		if (FControlRotationOrder* RotationOrder = ControlsRotationOrder.Find(OldParameterName))
@@ -862,7 +874,7 @@ void UMovieSceneControlRigParameterTrack::ChangeControlRotationOrder(const FName
 			{
 				if (NewOrder.IsSet())
 				{
-					FControlRotationOrder& RotationOrder = ControlsRotationOrder.FindOrAdd(InControlName);
+					FControlRotationOrder& RotationOrder = ControlsRotationOrder.FindOrAdd(ControlElement->GetFName());
 					RotationOrder.RotationOrder = NewOrder.GetValue();
 					TOptional<EEulerRotationOrder> Setting = GetControlRotationOrder(ControlElement, false);
 					if (Setting != NewOrder)
@@ -878,7 +890,7 @@ void UMovieSceneControlRigParameterTrack::ChangeControlRotationOrder(const FName
 				{
 					if (UMovieSceneControlRigParameterSection* CRSection = Cast<UMovieSceneControlRigParameterSection>(Section))
 					{
-						CRSection->ChangeControlRotationOrder(InControlName, Current, NewOrder, Interpolation);
+						CRSection->ChangeControlRotationOrder(ControlElement->GetFName(), Current, NewOrder, Interpolation);
 					}
 				}
 			}
diff --git a/Engine/Plugins/Animation/ControlRig/Source/ControlRig/Private/Units/Execution/RigUnit_DynamicHierarchy.cpp b/Engine/Plugins/Animation/ControlRig/Source/ControlRig/Private/Units/Execution/RigUnit_DynamicHierarchy.cpp
--- a/Engine/Plugins/Animation/ControlRig/Source/ControlRig/Private/Units/Execution/RigUnit_DynamicHierarchy.cpp
+++ b/Engine/Plugins/Animation/ControlRig/Source/ControlRig/Private/Units/Execution/RigUnit_DynamicHierarchy.cpp
@@ -5,6 +5,7 @@
 #include "Rigs/RigHierarchyController.h"
 #include "Units/RigUnitContext.h"
 #include "ControlRig.h"
+#include "ModularRig.h"
 #include "Components/SkeletalMeshComponent.h"
 #include "Units/Execution/RigUnit_PrepareForExecution.h"
 
@@ -981,10 +982,14 @@ FRigUnit_HierarchyAddAnimationChannelBool_Execute()
 		return;
 	}
 
-	const FString NameStr = Name.ToString();
-	if (NameStr.Contains(TEXT(":")))
+	const FRigHierarchyModulePath ModulePath(Name);
+	if (ModulePath.UsesModuleNameFormat())
 	{
-		UE_CONTROLRIG_RIGUNIT_REPORT_ERROR(TEXT("Animation channel name %s contains invalid ':' character."), *NameStr);
+		UE_CONTROLRIG_RIGUNIT_REPORT_ERROR(TEXT("Animation channel name %s contains invalid ':' character."), *ModulePath.GetPath());
+	}
+	if (ModulePath.UsesModuleNameFormat())
+	{
+		UE_CONTROLRIG_RIGUNIT_REPORT_ERROR(TEXT("Animation channel name %s contains invalid '/' character."), *ModulePath.GetPath());
 	}
 
 	Item.Reset();
@@ -1026,10 +1031,14 @@ FRigUnit_HierarchyAddAnimationChannelFloat_Execute()
 		return;
 	}
 
-	const FString NameStr = Name.ToString();
-	if (NameStr.Contains(TEXT(":")))
+	const FRigHierarchyModulePath ModulePath(Name);
+	if (ModulePath.UsesModuleNameFormat())
+	{
+		UE_CONTROLRIG_RIGUNIT_REPORT_ERROR(TEXT("Animation channel name %s contains invalid ':' character."), *ModulePath.GetPath());
+	}
+	if (ModulePath.UsesModuleNameFormat())
 	{
-		UE_CONTROLRIG_RIGUNIT_REPORT_ERROR(TEXT("Animation channel name %s contains invalid ':' character."), *NameStr);
+		UE_CONTROLRIG_RIGUNIT_REPORT_ERROR(TEXT("Animation channel name %s contains invalid '/' character."), *ModulePath.GetPath());
 	}
 
 	Item.Reset();
@@ -1072,12 +1081,15 @@ FRigUnit_HierarchyAddAnimationChannelScaleFloat_Execute()
 		return;
 	}
 
-	const FString NameStr = Name.ToString();
-	if (NameStr.Contains(TEXT(":")))
+	const FRigHierarchyModulePath ModulePath(Name);
+	if (ModulePath.UsesModuleNameFormat())
 	{
-		UE_CONTROLRIG_RIGUNIT_REPORT_ERROR(TEXT("Animation channel name %s contains invalid ':' character."), *NameStr);
+		UE_CONTROLRIG_RIGUNIT_REPORT_ERROR(TEXT("Animation channel name %s contains invalid ':' character."), *ModulePath.GetPath());
+	}
+	if (ModulePath.UsesModuleNameFormat())
+	{
+		UE_CONTROLRIG_RIGUNIT_REPORT_ERROR(TEXT("Animation channel name %s contains invalid '/' character."), *ModulePath.GetPath());
 	}
-
 
 	Item.Reset();
 
@@ -1119,12 +1131,15 @@ FRigUnit_HierarchyAddAnimationChannelInteger_Execute()
 		return;
 	}
 
-	const FString NameStr = Name.ToString();
-	if (NameStr.Contains(TEXT(":")))
+	const FRigHierarchyModulePath ModulePath(Name);
+	if (ModulePath.UsesModuleNameFormat())
 	{
-		UE_CONTROLRIG_RIGUNIT_REPORT_ERROR(TEXT("Animation channel name %s contains invalid ':' character."), *NameStr);
+		UE_CONTROLRIG_RIGUNIT_REPORT_ERROR(TEXT("Animation channel name %s contains invalid ':' character."), *ModulePath.GetPath());
+	}
+	if (ModulePath.UsesModuleNameFormat())
+	{
+		UE_CONTROLRIG_RIGUNIT_REPORT_ERROR(TEXT("Animation channel name %s contains invalid '/' character."), *ModulePath.GetPath());
 	}
-
 
 	Item.Reset();
 
@@ -1175,12 +1190,15 @@ FRigUnit_HierarchyAddAnimationChannelVector2D_Execute()
 		return;
 	}
 
-	const FString NameStr = Name.ToString();
-	if (NameStr.Contains(TEXT(":")))
+	const FRigHierarchyModulePath ModulePath(Name);
+	if (ModulePath.UsesModuleNameFormat())
 	{
-		UE_CONTROLRIG_RIGUNIT_REPORT_ERROR(TEXT("Animation channel name %s contains invalid ':' character."), *NameStr);
+		UE_CONTROLRIG_RIGUNIT_REPORT_ERROR(TEXT("Animation channel name %s contains invalid ':' character."), *ModulePath.GetPath());
+	}
+	if (ModulePath.UsesModuleNameFormat())
+	{
+		UE_CONTROLRIG_RIGUNIT_REPORT_ERROR(TEXT("Animation channel name %s contains invalid '/' character."), *ModulePath.GetPath());
 	}
-
 
 	Item.Reset();
 
@@ -1223,12 +1241,15 @@ FRigUnit_HierarchyAddAnimationChannelVector_Execute()
 		return;
 	}
 
-	const FString NameStr = Name.ToString();
-	if (NameStr.Contains(TEXT(":")))
+	const FRigHierarchyModulePath ModulePath(Name);
+	if (ModulePath.UsesModuleNameFormat())
 	{
-		UE_CONTROLRIG_RIGUNIT_REPORT_ERROR(TEXT("Animation channel name %s contains invalid ':' character."), *NameStr);
+		UE_CONTROLRIG_RIGUNIT_REPORT_ERROR(TEXT("Animation channel name %s contains invalid ':' character."), *ModulePath.GetPath());
+	}
+	if (ModulePath.UsesModuleNameFormat())
+	{
+		UE_CONTROLRIG_RIGUNIT_REPORT_ERROR(TEXT("Animation channel name %s contains invalid '/' character."), *ModulePath.GetPath());
 	}
-
 
 	Item.Reset();
 
@@ -1272,12 +1293,15 @@ FRigUnit_HierarchyAddAnimationChannelScaleVector_Execute()
 		return;
 	}
 
-	const FString NameStr = Name.ToString();
-	if (NameStr.Contains(TEXT(":")))
+	const FRigHierarchyModulePath ModulePath(Name);
+	if (ModulePath.UsesModuleNameFormat())
 	{
-		UE_CONTROLRIG_RIGUNIT_REPORT_ERROR(TEXT("Animation channel name %s contains invalid ':' character."), *NameStr);
+		UE_CONTROLRIG_RIGUNIT_REPORT_ERROR(TEXT("Animation channel name %s contains invalid ':' character."), *ModulePath.GetPath());
+	}
+	if (ModulePath.UsesModuleNameFormat())
+	{
+		UE_CONTROLRIG_RIGUNIT_REPORT_ERROR(TEXT("Animation channel name %s contains invalid '/' character."), *ModulePath.GetPath());
 	}
-
 
 	Item.Reset();
 
@@ -1321,12 +1345,15 @@ FRigUnit_HierarchyAddAnimationChannelRotator_Execute()
 		return;
 	}
 
-	const FString NameStr = Name.ToString();
-	if (NameStr.Contains(TEXT(":")))
+	const FRigHierarchyModulePath ModulePath(Name);
+	if (ModulePath.UsesModuleNameFormat())
 	{
-		UE_CONTROLRIG_RIGUNIT_REPORT_ERROR(TEXT("Animation channel name %s contains invalid ':' character."), *NameStr);
+		UE_CONTROLRIG_RIGUNIT_REPORT_ERROR(TEXT("Animation channel name %s contains invalid ':' character."), *ModulePath.GetPath());
+	}
+	if (ModulePath.UsesModuleNameFormat())
+	{
+		UE_CONTROLRIG_RIGUNIT_REPORT_ERROR(TEXT("Animation channel name %s contains invalid '/' character."), *ModulePath.GetPath());
 	}
-
 
 	Item.Reset();
 
diff --git a/Engine/Plugins/Animation/ControlRig/Source/ControlRig/Private/Units/Execution/RigUnit_RigModules.cpp b/Engine/Plugins/Animation/ControlRig/Source/ControlRig/Private/Units/Execution/RigUnit_RigModules.cpp
--- a/Engine/Plugins/Animation/ControlRig/Source/ControlRig/Private/Units/Execution/RigUnit_RigModules.cpp
+++ b/Engine/Plugins/Animation/ControlRig/Source/ControlRig/Private/Units/Execution/RigUnit_RigModules.cpp
@@ -56,16 +56,18 @@ FRigUnit_ResolveArrayConnector_Execute()
 
 FRigUnit_GetCurrentNameSpace_Execute()
 {
-	if(!ExecuteContext.IsRigModule())
-	{
-#if WITH_EDITOR
-		
-		static const FString Message = TEXT("This node should only be used in a Rig Module."); 
-		ExecuteContext.Report(EMessageSeverity::Warning, ExecuteContext.GetFunctionName(), ExecuteContext.GetInstructionIndex(), Message);
-		
-#endif
-	}
-	NameSpace = *ExecuteContext.GetRigModuleNameSpace();
+	FRigUnit_GetModuleName::StaticExecute(ExecuteContext, NameSpace);
+}
+
+FRigVMStructUpgradeInfo FRigUnit_GetCurrentNameSpace::GetUpgradeInfo() const
+{
+	FRigUnit_GetModuleName NewNode;
+	FRigVMStructUpgradeInfo Info(*this, NewNode);
+	Info.AddRemappedPin(
+		GET_MEMBER_NAME_STRING_CHECKED(FRigUnit_GetCurrentNameSpace, NameSpace),
+		GET_MEMBER_NAME_STRING_CHECKED(FRigUnit_GetModuleName, Module)
+	);
+	return Info;
 }
 
 FRigUnit_GetItemShortName_Execute()
@@ -83,38 +85,111 @@ FRigUnit_GetItemShortName_Execute()
 	}
 }
 
+FRigVMStructUpgradeInfo FRigUnit_GetItemShortName::GetUpgradeInfo() const
+{
+	// we don't have a node to upgrade to for this one
+	return FRigVMStructUpgradeInfo();
+}
+
+
 FRigUnit_GetItemNameSpace_Execute()
 {
-	NameSpace.Reset();
-	HasNameSpace = false;
+	FRigUnit_GetItemModuleName::StaticExecute(ExecuteContext, Item, HasNameSpace, NameSpace);
+}
+
+FRigVMStructUpgradeInfo FRigUnit_GetItemNameSpace::GetUpgradeInfo() const
+{
+	FRigUnit_GetItemModuleName NewNode;
+	NewNode.Item = Item;
+	FRigVMStructUpgradeInfo Info(*this, NewNode);
+	Info.AddRemappedPin(
+		GET_MEMBER_NAME_STRING_CHECKED(FRigUnit_GetItemNameSpace, NameSpace),
+		GET_MEMBER_NAME_STRING_CHECKED(FRigUnit_GetItemModuleName, Module)
+	);
+	Info.AddRemappedPin(
+		GET_MEMBER_NAME_STRING_CHECKED(FRigUnit_GetItemNameSpace, HasNameSpace),
+		GET_MEMBER_NAME_STRING_CHECKED(FRigUnit_GetItemModuleName, IsPartOfModule)
+	);
+	return Info;
+}
+
+FRigUnit_IsItemInCurrentNameSpace_Execute()
+{
+	FRigUnit_IsItemInCurrentModule::StaticExecute(ExecuteContext, Item, Result);
+}
+
+FRigVMStructUpgradeInfo FRigUnit_IsItemInCurrentNameSpace::GetUpgradeInfo() const
+{
+	FRigUnit_IsItemInCurrentModule NewNode;
+	NewNode.Item = Item;
+	FRigVMStructUpgradeInfo Info(*this, NewNode);
+	return Info;
+}
+
+FRigUnit_GetItemsInNameSpace_Execute()
+{
+	FRigUnit_GetItemsInModule::StaticExecute(ExecuteContext, TypeToSearch, Items);
+}
+
+FRigVMStructUpgradeInfo FRigUnit_GetItemsInNameSpace::GetUpgradeInfo() const
+{
+	FRigUnit_GetItemsInModule NewNode;
+	NewNode.TypeToSearch = TypeToSearch;
+	FRigVMStructUpgradeInfo Info(*this, NewNode);
+	return Info;
+}
+
+FRigUnit_GetModuleName_Execute()
+{
+	if(!ExecuteContext.IsRigModule())
+	{
+#if WITH_EDITOR
+		
+		static const FString Message = TEXT("This node should only be used in a Rig Module."); 
+		ExecuteContext.Report(EMessageSeverity::Warning, ExecuteContext.GetFunctionName(), ExecuteContext.GetInstructionIndex(), Message);
+		
+#endif
+	}
+	Module = *ExecuteContext.GetRigModulePrefix();
+	if(!Module.IsEmpty())
+	{
+		// remove the module suffix character from the end
+		Module.LeftChopInline(1);
+	}
+}
+
+FRigUnit_GetItemModuleName_Execute()
+{
+	Module.Reset();
+	IsPartOfModule = false;
 
 	if(const URigHierarchy* Hierarchy = ExecuteContext.Hierarchy)
 	{
-		const FString NameSpaceForItem = Hierarchy->GetNameSpace(Item);
-		if(!NameSpaceForItem.IsEmpty())
+		const FString ModuleForItem = Hierarchy->GetModuleName(Item);
+		if(!ModuleForItem.IsEmpty())
 		{
-			NameSpace = NameSpaceForItem;
-			HasNameSpace = true;
+			Module = ModuleForItem;
+			IsPartOfModule = true;
 		}
 	}
 }
 
-FRigUnit_IsItemInCurrentNameSpace_Execute()
+FRigUnit_IsItemInCurrentModule_Execute()
 {
-	FString CurrentNameSpace;
-	FRigUnit_GetCurrentNameSpace::StaticExecute(ExecuteContext, CurrentNameSpace);
-	bool bHasNameSpace = false;
-	FString ItemNameSpace;
-	FRigUnit_GetItemNameSpace::StaticExecute(ExecuteContext, Item, bHasNameSpace, ItemNameSpace);
+	FString CurrentModule;
+	FRigUnit_GetModuleName::StaticExecute(ExecuteContext, CurrentModule);
+	bool bHasModule = false;
+	FString ItemModule;
+	FRigUnit_GetItemModuleName::StaticExecute(ExecuteContext, Item, bHasModule, ItemModule);
 
 	Result = false;
-	if(!CurrentNameSpace.IsEmpty() && !ItemNameSpace.IsEmpty())
+	if(!CurrentModule.IsEmpty() && !ItemModule.IsEmpty())
 	{
-		Result = ItemNameSpace.Equals(CurrentNameSpace, ESearchCase::IgnoreCase); 
+		Result = ItemModule.Equals(CurrentModule, ESearchCase::IgnoreCase); 
 	}
 }
 
-FRigUnit_GetItemsInNameSpace_Execute()
+FRigUnit_GetItemsInModule_Execute()
 {
 	const URigHierarchy* Hierarchy = ExecuteContext.Hierarchy;
 	if (!Hierarchy)
@@ -123,18 +198,18 @@ FRigUnit_GetItemsInNameSpace_Execute()
 		return;
 	}
 
-	FString NameSpace;
-	FRigUnit_GetCurrentNameSpace::StaticExecute(ExecuteContext, NameSpace);
-	if(NameSpace.IsEmpty())
+	FString Module;
+	FRigUnit_GetModuleName::StaticExecute(ExecuteContext, Module);
+	if(Module.IsEmpty())
 	{
 		Items.Reset();
 		return;
 	}
-	const FName NameSpaceName = *NameSpace;
+	const FName ModuleName = *Module;
 	
 	uint32 Hash = GetTypeHash(StaticStruct());
 	Hash = HashCombine(Hash, GetTypeHash((int32)TypeToSearch));
-	Hash = HashCombine(Hash, GetTypeHash(NameSpaceName));
+	Hash = HashCombine(Hash, GetTypeHash(ModuleName));
 
 	if(const FRigElementKeyCollection* Cache = Hierarchy->FindCachedCollection(Hash))
 	{
@@ -143,7 +218,7 @@ FRigUnit_GetItemsInNameSpace_Execute()
 	else
 	{
 		FRigElementKeyCollection Collection;
-		Hierarchy->Traverse([Hierarchy, NameSpaceName, &Collection, TypeToSearch]
+		Hierarchy->Traverse([Hierarchy, ModuleName, &Collection, TypeToSearch]
 			(const FRigBaseElement* InElement, bool &bContinue)
 			{
 				bContinue = true;
@@ -151,10 +226,10 @@ FRigUnit_GetItemsInNameSpace_Execute()
 				const FRigElementKey Key = InElement->GetKey();
 				if(((uint8)TypeToSearch & (uint8)Key.Type) == (uint8)Key.Type)
 				{
-					const FName ItemNameSpace = Hierarchy->GetNameSpaceFName(Key);
-					if(!ItemNameSpace.IsNone())
+					const FName ItemModule = Hierarchy->GetModuleFName(Key);
+					if(!ItemModule.IsNone())
 					{
-						if(ItemNameSpace.IsEqual(NameSpaceName, ENameCase::IgnoreCase))
+						if(ItemModule.IsEqual(ModuleName, ENameCase::IgnoreCase))
 						{
 							Collection.AddUnique(Key);
 						}
@@ -166,4 +241,4 @@ FRigUnit_GetItemsInNameSpace_Execute()
 		Hierarchy->AddCachedCollection(Hash, Collection);
 		Items = Collection.Keys;
 	}
-}
\ No newline at end of file
+}
diff --git a/Engine/Plugins/Animation/ControlRig/Source/ControlRig/Private/Units/Hierarchy/RigUnit_ControlChannel.cpp b/Engine/Plugins/Animation/ControlRig/Source/ControlRig/Private/Units/Hierarchy/RigUnit_ControlChannel.cpp
--- a/Engine/Plugins/Animation/ControlRig/Source/ControlRig/Private/Units/Hierarchy/RigUnit_ControlChannel.cpp
+++ b/Engine/Plugins/Animation/ControlRig/Source/ControlRig/Private/Units/Hierarchy/RigUnit_ControlChannel.cpp
@@ -23,8 +23,8 @@ bool FRigUnit_GetAnimationChannelBase::UpdateCache(const URigHierarchy* InHierar
 	{
 		if(const FRigControlElement* ControlElement = InHierarchy->Find<FRigControlElement>(FRigElementKey(Control, ERigElementType::Control)))
 		{
-			FString Namespace, ChannelName = Channel.ToString();
-			URigHierarchy::SplitNameSpace(ChannelName, &Namespace, &ChannelName);
+			FString ChannelName = Channel.ToString();
+			(void)FRigHierarchyModulePath(ChannelName).Split(nullptr, &ChannelName);
 			
 			for(const FRigBaseElement* Child : InHierarchy->GetChildren(ControlElement))
 			{
diff --git a/Engine/Plugins/Animation/ControlRig/Source/ControlRig/Private/Units/Hierarchy/RigUnit_Metadata.cpp b/Engine/Plugins/Animation/ControlRig/Source/ControlRig/Private/Units/Hierarchy/RigUnit_Metadata.cpp
--- a/Engine/Plugins/Animation/ControlRig/Source/ControlRig/Private/Units/Hierarchy/RigUnit_Metadata.cpp
+++ b/Engine/Plugins/Animation/ControlRig/Source/ControlRig/Private/Units/Hierarchy/RigUnit_Metadata.cpp
@@ -451,10 +451,10 @@ FRigUnit_RemoveAllMetadata_Execute()
 				const TArray<FName> MetadataNames = Hierarchy->GetMetadataNames(Element->GetKey());
 
 				Removed = false;
-				const FString NameSpacePrefix = ExecuteContext.GetElementNameSpace(NameSpace);
+				const FString ModulePrefix = ExecuteContext.GetElementModulePrefix(NameSpace);
 				for(const FName& MetadataName : MetadataNames)
 				{
-					if(MetadataName.ToString().StartsWith(NameSpacePrefix, ESearchCase::CaseSensitive))
+					if(MetadataName.ToString().StartsWith(ModulePrefix, ESearchCase::CaseSensitive))
 					{
 						if(Element->RemoveMetadata(MetadataName))
 						{
diff --git a/Engine/Plugins/Animation/ControlRig/Source/ControlRig/Private/Units/RigUnitContext.cpp b/Engine/Plugins/Animation/ControlRig/Source/ControlRig/Private/Units/RigUnitContext.cpp
--- a/Engine/Plugins/Animation/ControlRig/Source/ControlRig/Private/Units/RigUnitContext.cpp
+++ b/Engine/Plugins/Animation/ControlRig/Source/ControlRig/Private/Units/RigUnitContext.cpp
@@ -13,50 +13,7 @@ bool FControlRigExecuteContext::IsRunningConstructionEvent() const
 		GetEventName() == FRigUnit_PostPrepareForExecution::EventName;
 }
 
-FName FControlRigExecuteContext::AddRigModuleNameSpace(const FName& InName) const
-{
-	if(IsRigModule())
-	{
-		return InName;
-	}
-	return *AddRigModuleNameSpace(InName.ToString());
-}
-
-FString FControlRigExecuteContext::AddRigModuleNameSpace(const FString& InName) const
-{
-	if(IsRigModule())
-	{
-		return InName;
-	}
-	check(!RigModuleNameSpace.IsEmpty());
-	return RigModuleNameSpace + InName;
-}
-
-FName FControlRigExecuteContext::RemoveRigModuleNameSpace(const FName& InName) const
-{
-	if(IsRigModule())
-	{
-		return InName;
-	}
-	return *RemoveRigModuleNameSpace(InName.ToString());
-}
-
-FString FControlRigExecuteContext::RemoveRigModuleNameSpace(const FString& InName) const
-{
-	if(IsRigModule())
-	{
-		return InName;
-	}
-	check(!RigModuleNameSpace.IsEmpty());
-
-	if(InName.StartsWith(RigModuleNameSpace, ESearchCase::CaseSensitive))
-	{
-		return InName.Mid(RigModuleNameSpace.Len());
-	}
-	return InName;
-}
-
-FString FControlRigExecuteContext::GetElementNameSpace(ERigMetaDataNameSpace InNameSpaceType) const
+const FString& FControlRigExecuteContext::GetElementModulePrefix(ERigMetaDataNameSpace InNameSpaceType) const
 {
 	if(IsRigModule())
 	{
@@ -66,28 +23,15 @@ FString FControlRigExecuteContext::GetElementNameSpace(ERigMetaDataNameSpace InN
 		{
 			case ERigMetaDataNameSpace::Self:
 			{
-				return GetRigModuleNameSpace();
+				return GetRigModulePrefix();
 			}
 			case ERigMetaDataNameSpace::Parent:
 			{
-				check(GetRigModuleNameSpace().EndsWith(UModularRig::NamespaceSeparator));
-				FString ParentNameSpace;
-				URigHierarchy::SplitNameSpace(GetRigModuleNameSpace().LeftChop(1), &ParentNameSpace, nullptr, true);
-				if(ParentNameSpace.IsEmpty())
-				{
-					return GetRigModuleNameSpace();
-				}
-				return ParentNameSpace + UModularRig::NamespaceSeparator;
+				return GetRigParentModulePrefix();
 			}
 			case ERigMetaDataNameSpace::Root:
 			{
-				FString RootNameSpace;
-				URigHierarchy::SplitNameSpace(GetRigModuleNameSpace(), &RootNameSpace, nullptr, false);
-				if(RootNameSpace.IsEmpty())
-				{
-					return GetRigModuleNameSpace();
-				}
-				return RootNameSpace + UModularRig::NamespaceSeparator;
+				return GetRigRootModulePrefix();
 			}
 			default:
 			{
@@ -109,12 +53,12 @@ FString FControlRigExecuteContext::GetElementNameSpace(ERigMetaDataNameSpace InN
 			}
 			case ERigMetaDataNameSpace::Parent:
 			{
-				static const FString ParentNameSpace = TEXT("Parent:");
+				static const FString ParentNameSpace = TEXT("Parent/");
 				return ParentNameSpace;
 			}
 			case ERigMetaDataNameSpace::Root:
 			{
-				static const FString RootNameSpace = TEXT("Root:");
+				static const FString RootNameSpace = TEXT("Root/");
 				return RootNameSpace;
 			}
 			default:
@@ -124,7 +68,8 @@ FString FControlRigExecuteContext::GetElementNameSpace(ERigMetaDataNameSpace InN
 		}
 	}
 
-	return FString();
+	static const FString EmptyPrefix;
+	return EmptyPrefix;
 }
 
 const FRigModuleInstance* FControlRigExecuteContext::GetRigModuleInstance(ERigMetaDataNameSpace InNameSpaceType) const
@@ -168,7 +113,7 @@ FName FControlRigExecuteContext::AdaptMetadataName(ERigMetaDataNameSpace InNameS
 		// conscious user decision.
 		const FString MetadataNameString = InMetadataName.ToString();
 		int32 Index = INDEX_NONE;
-		if(MetadataNameString.FindChar(TEXT(':'), Index))
+		if(MetadataNameString.FindChar(FRigHierarchyModulePath::ModuleNameSuffixChar, Index))
 		{
 			return InMetadataName;
 		}
@@ -183,7 +128,7 @@ FName FControlRigExecuteContext::AdaptMetadataName(ERigMetaDataNameSpace InNameS
 				case ERigMetaDataNameSpace::Parent:
 				case ERigMetaDataNameSpace::Root:
 				{
-					return *URigHierarchy::JoinNameSpace(GetElementNameSpace(InNameSpaceType), MetadataNameString);
+					return *(GetElementModulePrefix(InNameSpaceType) + MetadataNameString);
 				}
 				default:
 				{
@@ -206,7 +151,7 @@ FName FControlRigExecuteContext::AdaptMetadataName(ERigMetaDataNameSpace InNameS
 				case ERigMetaDataNameSpace::Parent:
 				case ERigMetaDataNameSpace::Root:
 				{
-					return *URigHierarchy::JoinNameSpace(GetElementNameSpace(InNameSpaceType), MetadataNameString);
+					return *(GetElementModulePrefix(InNameSpaceType) + MetadataNameString);
 				}
 				default:
 				{
@@ -220,24 +165,48 @@ FName FControlRigExecuteContext::AdaptMetadataName(ERigMetaDataNameSpace InNameS
 
 FControlRigExecuteContextRigModuleGuard::FControlRigExecuteContextRigModuleGuard(FControlRigExecuteContext& InContext, const UControlRig* InControlRig)
 	: Context(InContext)
-	, PreviousRigModuleNameSpace(InContext.RigModuleNameSpace)
-	, PreviousRigModuleNameSpaceHash(InContext.RigModuleNameSpaceHash)
+	, PreviousRigModulePrefix(InContext.RigModulePrefix)
+	, PreviousRigParentModulePrefix(InContext.RigParentModulePrefix)
+	, PreviousRigRootModulePrefix(InContext.RigRootModulePrefix)
+	, PreviousRigModulePrefixHash(InContext.RigModulePrefixHash)
 {
-	Context.RigModuleNameSpace = InControlRig->GetRigModuleNameSpace();
-	Context.RigModuleNameSpaceHash = GetTypeHash(Context.RigModuleNameSpace);
+	Context.RigModulePrefix = InControlRig->GetRigModulePrefix();
+	Context.RigParentModulePrefix = Context.RigModulePrefix;
+	Context.RigRootModulePrefix = Context.RigModulePrefix;
+	if(UModularRig* ModularRig = Cast<UModularRig>(InControlRig->GetParentRig()))
+	{
+		if(const FRigModuleInstance* Module = ModularRig->FindModule(InControlRig->GetFName()))
+		{
+			if(const FRigModuleInstance* ParentModule = Module->GetParentModule())
+			{
+				Context.RigParentModulePrefix = ParentModule->GetModulePrefix();
+			}
+			if(const FRigModuleInstance* RootModule = Module->GetRootModule())
+			{
+				Context.RigRootModulePrefix = RootModule->GetModulePrefix();
+			}
+		}
+	}
+	Context.RigModulePrefixHash = GetTypeHash(Context.RigModulePrefix);
 }
 
-FControlRigExecuteContextRigModuleGuard::FControlRigExecuteContextRigModuleGuard(FControlRigExecuteContext& InContext, const FString& InNewModuleNameSpace)
+FControlRigExecuteContextRigModuleGuard::FControlRigExecuteContextRigModuleGuard(FControlRigExecuteContext& InContext, const FString& InNewModulePrefix, const FString& InNewParentModulePrefix, const FString& InNewRootModulePrefix)
 	: Context(InContext)
-	, PreviousRigModuleNameSpace(InContext.RigModuleNameSpace)
-	, PreviousRigModuleNameSpaceHash(InContext.RigModuleNameSpaceHash)
+	, PreviousRigModulePrefix(InContext.RigModulePrefix)
+	, PreviousRigParentModulePrefix(InContext.RigParentModulePrefix)
+	, PreviousRigRootModulePrefix(InContext.RigRootModulePrefix)
+	, PreviousRigModulePrefixHash(InContext.RigModulePrefixHash)
 {
-	Context.RigModuleNameSpace = InNewModuleNameSpace;
-	Context.RigModuleNameSpaceHash = GetTypeHash(Context.RigModuleNameSpace);
+	Context.RigModulePrefix = InNewModulePrefix;
+	Context.RigParentModulePrefix = InNewParentModulePrefix;
+	Context.RigRootModulePrefix = InNewRootModulePrefix;
+	Context.RigModulePrefixHash = GetTypeHash(Context.RigModulePrefix);
 }
 
 FControlRigExecuteContextRigModuleGuard::~FControlRigExecuteContextRigModuleGuard()
 {
-	Context.RigModuleNameSpace = PreviousRigModuleNameSpace;
-	Context.RigModuleNameSpaceHash = PreviousRigModuleNameSpaceHash; 
+	Context.RigModulePrefix = PreviousRigModulePrefix;
+	Context.RigParentModulePrefix = PreviousRigParentModulePrefix;
+	Context.RigRootModulePrefix = PreviousRigRootModulePrefix;
+	Context.RigModulePrefixHash = PreviousRigModulePrefixHash; 
 }
diff --git a/Engine/Plugins/Animation/ControlRig/Source/ControlRig/Public/ControlRig.h b/Engine/Plugins/Animation/ControlRig/Source/ControlRig/Public/ControlRig.h
--- a/Engine/Plugins/Animation/ControlRig/Source/ControlRig/Public/ControlRig.h
+++ b/Engine/Plugins/Animation/ControlRig/Source/ControlRig/Public/ControlRig.h
@@ -112,7 +112,14 @@ public:
 	UControlRig* GetParentRig() const;
 
 	// Returns the namespace of this module (for example ArmModule::)
-	const FString& GetRigModuleNameSpace() const;
+	UE_DEPRECATED(5.6, "Please use UControlRig::GetRigModulePrefix()")
+	const FString& GetRigModuleNameSpace() const
+	{
+		return GetRigModulePrefix();
+	}
+	
+	// Returns the module prefix of this module (for example Arm/ )
+	const FString& GetRigModulePrefix() const;
 
 	// Returns the redirector from key to key for this rig
 	virtual FRigElementKeyRedirector& GetElementKeyRedirector();
@@ -572,7 +579,7 @@ protected:
 	const FRigInfluenceMap* FindInfluenceMap(const FName& InEventName);
 
 
-	FRigElementKeyRedirector ElementKeyRedirector;
+	mutable FRigElementKeyRedirector ElementKeyRedirector;
 
 public:
 
@@ -791,7 +798,7 @@ private:
 	FRigModuleSettings RigModuleSettings;
 
 	UPROPERTY(transient)
-	mutable FString RigModuleNameSpace;
+	mutable FString RigModulePrefix;
 
 
 public:
diff --git a/Engine/Plugins/Animation/ControlRig/Source/ControlRig/Public/ModularRig.h b/Engine/Plugins/Animation/ControlRig/Source/ControlRig/Public/ModularRig.h
--- a/Engine/Plugins/Animation/ControlRig/Source/ControlRig/Public/ModularRig.h
+++ b/Engine/Plugins/Animation/ControlRig/Source/ControlRig/Public/ModularRig.h
@@ -13,10 +13,13 @@ public:
 
 	FModuleInstanceHandle()
 		: ModularRig(nullptr)
-		, Path()
+		, ModuleName(NAME_None)
 	{}
 
+	UE_DEPRECATED(5.6, "Please use FModuleInstanceHandle(const UModularRig* InModularRig, const FName& InModuleName)")
 	FModuleInstanceHandle(const UModularRig* InModularRig, const FString& InPath);
+	
+	FModuleInstanceHandle(const UModularRig* InModularRig, const FName& InModuleName);
 	FModuleInstanceHandle(const UModularRig* InModularRig, const FRigModuleInstance* InElement);
 
 	bool IsValid() const { return Get() != nullptr; }
@@ -30,14 +33,14 @@ public:
 		}
 		return nullptr;
 	}
-	const FString& GetPath() const { return Path; }
+	const FName& GetModuleName() const { return ModuleName; }
 
 	const FRigModuleInstance* Get() const;
 
 private:
 
 	mutable TSoftObjectPtr<UModularRig> ModularRig;
-	FString Path;
+	FName ModuleName;
 };
 
 USTRUCT(BlueprintType)
@@ -50,7 +53,8 @@ public:
 	FRigModuleInstance()
 	: Name(NAME_None)
 	, RigPtr(nullptr)
-	, ParentPath(FString())
+	, ParentModuleName(NAME_None)
+	, CachedParentModule(nullptr)
 	{
 	}
 
@@ -65,21 +69,18 @@ private:
 public:
 	
 	UPROPERTY()
-	FString ParentPath;
+	FString ParentPath_DEPRECATED;
+
+	UPROPERTY()
+	FName ParentModuleName;
 
 	UPROPERTY()
 	TMap<FName, FRigVMExternalVariable> VariableBindings;
 
+	FRigModuleInstance* CachedParentModule;
 	TArray<FRigModuleInstance*> CachedChildren;
 	mutable FCachedRigElement PrimaryConnector;
 
-	FString GetShortName() const;
-	FString GetLongName() const
-	{
-		return GetPath();
-	}
-	FString GetPath() const;
-	FString GetNamespace() const;
 	UControlRig* GetRig() const;
 	void SetRig(UControlRig* InRig);
 	bool ContainsRig(const UControlRig* InRig) const;
@@ -89,6 +90,9 @@ public:
 	const FRigConnectorElement* FindPrimaryConnector() const;
 	TArray<const FRigConnectorElement*> FindConnectors() const;
 	bool IsRootModule() const;
+	FString GetModulePrefix() const;
+	FString GetModulePath_Deprecated() const;
+	bool HasChildModule(const FName& InModuleName, bool bRecursive = true) const;
 };
 
 USTRUCT(BlueprintType)
@@ -97,21 +101,21 @@ struct CONTROLRIG_API FRigModuleExecutionElement
 	GENERATED_USTRUCT_BODY()
 
 	FRigModuleExecutionElement()
-		: ModulePath(FString())
+		: ModuleName(NAME_None)
 		, ModuleInstance(nullptr)
 		, EventName(NAME_None)
 		, bExecuted(false)
 	{}
 
 	FRigModuleExecutionElement(FRigModuleInstance* InModule, FName InEvent)
-		: ModulePath(InModule->GetPath())
+		: ModuleName(InModule->Name)
 		, ModuleInstance(InModule)
 		, EventName(InEvent)
 		, bExecuted(false)
 	{}
 
 	UPROPERTY()
-	FString ModulePath;
+	FName ModuleName;
 	FRigModuleInstance* ModuleInstance;
 
 	UPROPERTY()
@@ -184,35 +188,47 @@ public:
 	/**
 	 * @return Returns all of the module paths within this rig
 	 */
-	UFUNCTION(BlueprintPure, Category=ModularRig)
+	UFUNCTION(BlueprintPure, Category=ModularRig, meta=(DeprecatedFunction, DeprecationMessage="Function has been deprecated, please rely on GetModuleNames instead."))
 	TArray<FString> GetModulePaths() const;
-	
-	const FRigModuleInstance* FindModule(const FString& InPath) const;
+
+	/**
+	 * @return Returns all of the module paths within this rig
+	 */
+	UFUNCTION(BlueprintPure, Category=ModularRig)
+	TArray<FName> GetModuleNames() const;
+
+	const FRigModuleInstance* FindModule(const FName& InModuleName) const;
 	const FRigModuleInstance* FindModule(const UControlRig* InModuleInstance) const;
 	const FRigModuleInstance* FindModule(const FRigBaseElement* InElement) const;
 	const FRigModuleInstance* FindModule(const FRigElementKey& InElementKey) const;
 
-	FRigModuleInstance* FindModule(const FString& InPath);
+	FRigModuleInstance* FindModule(const FName& InModuleName);
 	FRigModuleInstance* FindModule(const UControlRig* InModuleInstance);
 	FRigModuleInstance* FindModule(const FRigBaseElement* InElement);
 	FRigModuleInstance* FindModule(const FRigElementKey& InElementKey);
 
 	/**
-	 * @param InModulePath The path of the module to run the event on
-	 * @return Returns the rig instance for a given module path
+	 * @param InModuleName The name of the module to run the event on
+	 * @return Returns the rig instance for a given module name
 	 */
 	UFUNCTION(BlueprintCallable, Category=ModularRig)
-	UControlRig* GetModuleRig(FString InModulePath);
-	const UControlRig* GetModuleRig(FString InModulePath) const;
+	UControlRig* GetModuleRig(FName InModuleName);
+	const UControlRig* GetModuleRig(FName InModuleName) const;
 
 	/**
 	 * @param InModulePath The path of the module to receive the parent path for
 	 * @return Returns the parent path for a given module path (or an empty string)
 	 */
-	UFUNCTION(BlueprintPure, Category=ModularRig, meta=(ScriptName="GetParentPath", DisplayName="Get Parent Path"))
+	UFUNCTION(BlueprintPure, Category=ModularRig, meta=(ScriptName="GetParentPath", DisplayName="Get Parent Path", DeprecatedFunction, DeprecationMessage="Function has been deprecated, please rely on GetParentModuleName instead."))
 	FString GetParentPathForBP(FString InModulePath) const;
-	FString GetParentPath(const FString& InModulePath) const;
-	const FString& GetShortestDisplayPathForElement(const FRigElementKey& InElementKey, bool bAlwaysShowNameSpace) const;
+
+	/**
+	 * @param InModuleName The name of the module to receive the parent module name for
+	 * @return Returns the parent name for a given module name (or an empty string)
+	 */
+	UFUNCTION(BlueprintPure, Category=ModularRig, meta=(ScriptName="GetParentModuleName", DisplayName="Get Parent Module Name"))
+	FName GetParentModuleNameForBP(FName InModuleName) const;
+	FName GetParentModuleName(const FName& InModuleName) const;
 
 	void ForEachModule(TFunctionRef<bool(FRigModuleInstance*)> PerModuleFunction);
 	void ForEachModule(TFunctionRef<bool(const FRigModuleInstance*)> PerModuleFunction) const;
@@ -226,55 +242,49 @@ public:
 	TArray<FName> GetEventsForAllModules() const;
 
 	/**
-	 * @param InModulePath The path of the module to receive the events for
+	 * @param InModuleName The name of the module to receive the events for
 	 * @return Returns the names of all supported events for a given module path
 	 */
 	UFUNCTION(BlueprintPure, Category=ModularRig)
-	TArray<FName> GetEventsForModule(FString InModulePath) const;
+	TArray<FName> GetEventsForModule(FName InModuleName) const;
 
 	/**
 	 * @param InEvent The name of the event to run
 	 * @return Returns the paths of all modules which ran the event successfully
 	 */
 	UFUNCTION(BlueprintCallable, Category=ModularRig)
-	TArray<FString> ExecuteEventOnAllModules(FName InEvent);
+	TArray<FName> ExecuteEventOnAllModules(FName InEvent);
 
 	/**
 	 * @param InEvent The name of the event to run
-	 * @param InModulePath The path of the module to run the event on
+	 * @param InModuleName The name of the module to run the event on
 	 * @return Returns true if the event was run successfully
 	 */
 	UFUNCTION(BlueprintCallable, Category=ModularRig, meta=(ScriptName="ExecuteEventOnModule", DisplayName="Execute Event On Module"))
-	bool ExecuteEventOnModuleForBP(FName InEvent, FString InModulePath);
+	bool ExecuteEventOnModuleForBP(FName InEvent, FName InModuleName);
 	bool ExecuteEventOnModule(const FName& InEvent, FRigModuleInstance* InModule);
 
 	/**
 	 * Returns a handle to an existing module
-	 * @param InPath The path of the module to retrieve a handle for.
+	 * @param InModuleName The name of the module to retrieve a handle for.
 	 * @return The retrieved handle (may be invalid)
 	 */
-	FModuleInstanceHandle GetHandle(const FString& InPath) const
+	FModuleInstanceHandle GetHandle(const FName& InModuleName) const
 	{
-		if(FindModule(InPath))
+		if(FindModule(InModuleName))
 		{
-			return FModuleInstanceHandle((UModularRig*)this, InPath);
+			return FModuleInstanceHandle((UModularRig*)this, InModuleName);
 		}
 		return FModuleInstanceHandle();
 	}
 
-	static inline const TCHAR* NamespaceSeparator = TEXT(":");
-
-protected:
-
-	virtual void EmitPostConstructionEventFinished() override;
-
 private:
 
 	/** Adds a module to the rig*/
 	FRigModuleInstance* AddModuleInstance(const FName& InModuleName, TSubclassOf<UControlRig> InModuleClass, const FRigModuleInstance* InParent, const FRigElementKeyRedirector::FKeyMap& InConnectionMap, const TMap<FName, FString>& InVariableDefaultValues);
 
 	/** Updates the module's variable bindings */
-	bool SetModuleVariableBindings(const FString& InModulePath, const TMap<FName, FString>& InVariableBindings);
+	bool SetModuleVariableBindings(const FName& InModuleName, const TMap<FName, FString>& InVariableBindings);
 
 	/** Updates the module's variable values based on its binding */
 	void UpdateModuleVariables(const FRigModuleInstance* InModule);
@@ -282,11 +292,7 @@ private:
 	/** Destroys / discards a module instance rig */
 	static void DiscardModuleRig(UControlRig* InControlRig);
 
-	TMap<FString, UControlRig*> PreviousModuleRigs;
-
-	void ResetShortestDisplayPathCache() const;
-	void RecomputeShortestDisplayPathCache() const;
-	mutable TMap<FRigElementKey, TTuple<FString, FString>> ElementKeyToShortestDisplayPath; 
+	TMap<FName, UControlRig*> PreviousModuleRigs;
 
 	friend struct FRigModuleInstance;
 };
diff --git a/Engine/Plugins/Animation/ControlRig/Source/ControlRig/Public/ModularRigController.h b/Engine/Plugins/Animation/ControlRig/Source/ControlRig/Public/ModularRigController.h
--- a/Engine/Plugins/Animation/ControlRig/Source/ControlRig/Public/ModularRigController.h
+++ b/Engine/Plugins/Animation/ControlRig/Source/ControlRig/Public/ModularRigController.h
@@ -3,6 +3,7 @@
 #pragma once
 
 #include "ModularRigModel.h"
+#include "RigVMFunctions/Math/RigVMMathLibrary.h"
 #include "ModularRigController.generated.h"
 
 struct FRigModuleReference;
@@ -26,7 +27,7 @@ class CONTROLRIG_API UModularRigController : public UObject
 	bool bSuspendNotifications;
 
 	UFUNCTION(BlueprintCallable, Category = "Control Rig | Modules")
-	FString AddModule(const FName& InModuleName, TSubclassOf<UControlRig> InClass, const FString& InParentModulePath, bool bSetupUndo = true);
+	FName AddModule(const FName& InModuleName, TSubclassOf<UControlRig> InClass, const FName& InParentModuleName, bool bSetupUndo = true);
 
 	UFUNCTION(BlueprintCallable, Category = "Control Rig | Modules")
 	bool CanConnectConnectorToElement(const FRigElementKey& InConnectorKey, const FRigElementKey& InTargetKey, FText& OutErrorMessage);
@@ -54,56 +55,52 @@ class CONTROLRIG_API UModularRigController : public UObject
 	bool AutoConnectSecondaryConnectors(const TArray<FRigElementKey>& InConnectorKeys, bool bReplaceExistingConnections, bool bSetupUndo = true);
 
 	UFUNCTION(BlueprintCallable, Category = "Control Rig | Modules")
-	bool AutoConnectModules(const TArray<FString>& InModulePaths, bool bReplaceExistingConnections, bool bSetupUndo = true);
+	bool AutoConnectModules(const TArray<FName>& InModuleNames, bool bReplaceExistingConnections, bool bSetupUndo = true);
 
 	UFUNCTION(BlueprintCallable, Category = "Control Rig | Modules")
-	bool SetConfigValueInModule(const FString& InModulePath, const FName& InVariableName, const FString& InValue, bool bSetupUndo = true);
+	bool SetConfigValueInModule(const FName& InModuleName, const FName& InVariableName, const FString& InValue, bool bSetupUndo = true);
 
 	UFUNCTION(BlueprintCallable, Category = "ControlRig | Modules")
-	bool BindModuleVariable(const FString& InModulePath, const FName& InVariableName, const FString& InSourcePath, bool bSetupUndo = true);
-	bool CanBindModuleVariable(const FString& InModulePath, const FName& InVariableName, const FString& InSourcePath, FText& OutErrorMessage);
-	TArray<FString> GetPossibleBindings(const FString& InModulePath, const FName& InVariableName);
+	bool BindModuleVariable(const FName& InModuleName, const FName& InVariableName, const FString& InSourcePath, bool bSetupUndo = true);
+	bool CanBindModuleVariable(const FName& InModuleName, const FName& InVariableName, const FString& InSourcePath, FText& OutErrorMessage);
+	TArray<FString> GetPossibleBindings(const FName& InModuleName, const FName& InVariableName);
 
 	UFUNCTION(BlueprintCallable, Category = "ControlRig | Modules")
-	bool UnBindModuleVariable(const FString& InModulePath, const FName& InVariableName, bool bSetupUndo = true);
+	bool UnBindModuleVariable(const FName& InModuleName, const FName& InVariableName, bool bSetupUndo = true);
 
 	UFUNCTION(BlueprintCallable, Category = "ControlRig | Modules")
-	bool DeleteModule(const FString& InModulePath, bool bSetupUndo = true);
+	bool DeleteModule(const FName& InModuleName, bool bSetupUndo = true);
 
 	UFUNCTION(BlueprintCallable, Category = "ControlRig | Modules")
-	FString RenameModule(const FString& InModulePath, const FName& InNewName, bool bSetupUndo = true);
-	bool CanRenameModule(const FString& InModulePath, const FName& InNewName, FText& OutErrorMessage) const;
+	FName RenameModule(const FName& InModuleName, const FName& InNewName, bool bSetupUndo = true);
+	bool CanRenameModule(const FName& InModuleName, const FName& InNewName, FText& OutErrorMessage) const;
 
 	UFUNCTION(BlueprintCallable, Category = "ControlRig | Modules")
-	FString ReparentModule(const FString& InModulePath, const FString& InNewParentModulePath, bool bSetupUndo = true);
+	bool ReparentModule(const FName& InModuleName, const FName& InNewParentModuleName, bool bSetupUndo = true);
 
 	UFUNCTION(BlueprintCallable, Category = "ControlRig | Modules")
-	FString MirrorModule(const FString& InModulePath, const FRigVMMirrorSettings& InSettings, bool bSetupUndo = true);
+	FName MirrorModule(const FName& InModuleName, const FRigVMMirrorSettings& InSettings, bool bSetupUndo = true);
 
 	UFUNCTION(BlueprintCallable, Category = "ControlRig | Modules")
-	bool SetModuleShortName(const FString& InModulePath, const FString& InNewShortName, bool bSetupUndo = true);
-	bool CanSetModuleShortName(const FString& InModulePath, const FString& InNewShortName, FText& OutErrorMessage) const;
-
-	UFUNCTION(BlueprintCallable, Category = "ControlRig | Modules")
-	bool SwapModuleClass(const FString& InModulePath, TSubclassOf<UControlRig> InNewClass, bool bSetupUndo = true);
+	bool SwapModuleClass(const FName& InModuleName, TSubclassOf<UControlRig> InNewClass, bool bSetupUndo = true);
 
 	UFUNCTION(BlueprintCallable, Category = "ControlRig | Modules")
 	bool SwapModulesOfClass(TSubclassOf<UControlRig> InOldClass, TSubclassOf<UControlRig> InNewClass, bool bSetupUndo = true);
 
 	UFUNCTION(BlueprintCallable, Category = "ControlRig | Modules")
-	bool SelectModule(const FString& InModulePath, const bool InSelected = true);
+	bool SelectModule(const FName& InModuleName, const bool InSelected = true);
 
 	UFUNCTION(BlueprintCallable, Category = "ControlRig | Modules")
-	bool DeselectModule(const FString& InModulePath);
+	bool DeselectModule(const FName& InModuleName);
 
 	UFUNCTION(BlueprintCallable, Category = "ControlRig | Modules")
-	bool SetModuleSelection(const TArray<FString>& InModulePaths);
+	bool SetModuleSelection(const TArray<FName>& InModuleNames);
 
 	UFUNCTION(BlueprintPure, Category = "ControlRig | Modules")
-	TArray<FString> GetAllModules() const;
+	TArray<FName> GetAllModules() const;
 
 	UFUNCTION(BlueprintPure, Category = "ControlRig | Modules")
-	TArray<FString> GetSelectedModules() const;
+	TArray<FName> GetSelectedModules() const;
 
 	void RefreshModuleVariables(bool bSetupUndo = true);
 	void RefreshModuleVariables(const FRigModuleReference* InModule, bool bSetupUndo = true);
@@ -111,33 +108,30 @@ class CONTROLRIG_API UModularRigController : public UObject
 	static int32 GetMaxNameLength() { return 100; }
 	static void SanitizeName(FRigName& InOutName, bool bAllowNameSpaces);
 	static FRigName GetSanitizedName(const FRigName& InName, bool bAllowNameSpaces);
-	bool IsNameAvailable(const FString& InParentModulePath, const FRigName& InDesiredName, FString* OutErrorMessage = nullptr) const;
-	bool IsShortNameAvailable(const FRigName& InDesiredShortName, FString* OutErrorMessage = nullptr) const;
-	FRigName GetSafeNewName(const FString& InParentModulePath, const FRigName& InDesiredName) const;
-	FRigName GetSafeNewShortName(const FRigName& InDesiredShortName) const;
-	FRigModuleReference* FindModule(const FString& InPath);
-	const FRigModuleReference* FindModule(const FString& InPath) const;
+	bool IsNameAvailable(const FRigName& InDesiredName, FString* OutErrorMessage = nullptr, const FRigModuleReference* InModuleToSkip = nullptr) const;
+	FRigName GetSafeNewName(const FRigName& InDesiredName, const FRigModuleReference* InModuleToSkip = nullptr) const;
+	FRigModuleReference* FindModule(const FName& InModuleName);
+	const FRigModuleReference* FindModule(const FName& InModuleName) const;
 	FModularRigModifiedEvent& OnModified() { return ModifiedEvent; }
 
 	/**
-	 * @param InModulePath The path of the module reference to return
-	 * @return Returns all of the module paths within this rig
+	 * @param InModuleName The name of the module reference to return
+	 * @return Returns the module for the given name
 	 */
 	UFUNCTION(BlueprintPure, Category = "ControlRig | Modules")
-	FRigModuleReference GetModuleReference(FString InModulePath) const;
+	FRigModuleReference GetModuleReference(FName InModuleName) const;
 
 	/**
-	 * @param InModulePath The path of the module reference to return
-	 * @return Returns all of associated connectors for a given module path
+	 * @param InModuleName The name of the module reference to return
+	 * @return Returns all of associated connectors for a given module name
 	 */
 	UFUNCTION(BlueprintPure, Category = "ControlRig | Modules")
-	TArray<FRigElementKey> GetConnectorsForModule(FString InModulePath) const;
+	TArray<FRigElementKey> GetConnectorsForModule(FName InModuleName) const;
 
 private:
 
 	void SetModel(FModularRigModel* InModel) { Model = InModel; }
 	void Notify(const EModularRigNotification& InNotification, const FRigModuleReference* InElement);
-	void UpdateShortNames();
 	UModularRig* GetDebuggedModularRig();
 	const UModularRig* GetDebuggedModularRig() const;
 
