Control Rig: Unique modular rig module names

diff --git a/Engine/Plugins/Animation/ControlRig/Source/ControlRig/Public/ModularRigModel.h b/Engine/Plugins/Animation/ControlRig/Source/ControlRig/Public/ModularRigModel.h
--- a/Engine/Plugins/Animation/ControlRig/Source/ControlRig/Public/ModularRigModel.h
+++ b/Engine/Plugins/Animation/ControlRig/Source/ControlRig/Public/ModularRigModel.h
@@ -6,6 +6,7 @@
 #include "ModularRigModel.generated.h"
 
 class UModularRigController;
+struct FModularRigModel;
 
 UENUM()
 enum class EModularRigNotification : uint8
@@ -45,33 +46,40 @@ struct CONTROLRIG_API FRigModuleReference
 {
 	GENERATED_BODY()
 
+public:
+
 	FRigModuleReference()
 		: Name(NAME_None)
-		, ShortName()
-		, bShortNameBasedOnPath(true)
-		, ParentPath(FString())
+		, bShortNameBasedOnPath_DEPRECATED(true)
 		, Class(nullptr)
+		, PreviousName(NAME_None)
+		, PreviousParentName(NAME_None)
 	{}
 	
-	FRigModuleReference(const FName& InName, TSubclassOf<UControlRig> InClass, const FString& InParentPath)
+	FRigModuleReference(const FName& InName, TSubclassOf<UControlRig> InClass, const FName& InParentModuleName, const FModularRigModel* InModel)
 		: Name(InName)
-		, ShortName()
-		, bShortNameBasedOnPath(true)
-		, ParentPath(InParentPath)
+		, bShortNameBasedOnPath_DEPRECATED(true)
+		, ParentModuleName(InParentModuleName)
 		, Class(InClass)
+		, PreviousName(NAME_None)
+		, PreviousParentName(NAME_None)
+		, Model(InModel)
 	{}
 
 	UPROPERTY()
 	FName Name;
 
 	UPROPERTY()
-	FString ShortName;
+	FString ShortName_DEPRECATED;
 
 	UPROPERTY()
-	bool bShortNameBasedOnPath;
+	bool bShortNameBasedOnPath_DEPRECATED;
 
 	UPROPERTY()
-	FString ParentPath;
+	FString ParentPath_DEPRECATED;
+
+	UPROPERTY()
+	FName ParentModuleName;
 
 	UPROPERTY()
 	TSoftClassPtr<UControlRig> Class;
@@ -89,35 +97,50 @@ struct CONTROLRIG_API FRigModuleReference
 	FName PreviousName;
 
 	UPROPERTY(transient)
-	FString PreviousParentPath;
+	FName PreviousParentName;
 
 	TArray<FRigModuleReference*> CachedChildren;
 
-	FString GetShortName() const;
+	FName GetFName() const
+	{
+		return Name;
+	}
+
+	FString GetName() const
+	{
+		return Name.ToString();
+	}
 
-	FString GetLongName() const
+	FName GetElementPrefixFName() const
 	{
-		return GetPath();
+		return *GetElementPrefix();
 	}
 
-	FString GetPath() const;
+	FString GetElementPrefix() const;
 
-	FString GetNamespace() const;
+	FRigHierarchyModulePath GetModulePath() const;
 
-	bool HasParentModule() const { return !ParentPath.IsEmpty(); }
+	bool HasParentModule() const { return !ParentModuleName.IsNone() || !ParentPath_DEPRECATED.IsEmpty(); }
 
 	bool IsRootModule() const { return !HasParentModule(); }
+
+	const FRigModuleReference* GetParentModule() const;
+	const FRigModuleReference* GetRootModule() const;
 	
 	friend bool operator==(const FRigModuleReference& A, const FRigModuleReference& B)
 	{
-		return A.ParentPath == B.ParentPath &&
+		return A.ParentModuleName == B.ParentModuleName &&
 			A.Name == B.Name;
 	}
 
 	const FRigConnectorElement* FindPrimaryConnector(const URigHierarchy* InHierarchy) const;
 	TArray<const FRigConnectorElement*> FindConnectors(const URigHierarchy* InHierarchy) const;
 
+private:
+	const FModularRigModel* Model = nullptr;
+
 	friend class UModularRigController;
+	friend struct FModularRigModel;
 };
 
 USTRUCT(BlueprintType)
@@ -282,9 +305,9 @@ public:
 	}
 
 	/** Gets the connection map for a single module, where the connectors are identified without its namespace*/
-	FRigElementKeyRedirector::FKeyMap GetModuleConnectionMap(const FString& InModulePath) const;
+	FRigElementKeyRedirector::FKeyMap GetModuleConnectionMap(const FName& InModuleName) const;
 
-	void PatchOnLoad();
+	void PatchOnLoad(const TMap<FRigHierarchyModulePath, FName>* InModulePathToModuleName = nullptr);
 };
 
 // A management struct containing all modules in the rig
@@ -295,10 +318,13 @@ public:
 
 	GENERATED_BODY()
 
-	FModularRigModel(){}
+	FModularRigModel()
+	{}
+	
 	FModularRigModel(const FModularRigModel& Other)
 	{
 		Modules = Other.Modules;
+		PreviousModulePaths = Other.PreviousModulePaths;
 		Connections = Other.Connections;
 		UpdateCachedChildren();
 		Connections.UpdateFromConnectionList();
@@ -307,6 +333,7 @@ public:
 	FModularRigModel& operator=(const FModularRigModel& Other)
 	{
 		Modules = Other.Modules;
+		PreviousModulePaths = Other.PreviousModulePaths;
 		Connections = Other.Connections;
 		UpdateCachedChildren();
 		Connections.UpdateFromConnectionList();
@@ -326,6 +353,10 @@ public:
 	UPROPERTY(transient)
 	TObjectPtr<UObject> Controller;
 
+	// remember what modules were called so we can recover.
+	UPROPERTY()
+	TMap<FRigHierarchyModulePath, FName> PreviousModulePaths;
+
 	void PatchModelsOnLoad();
 
 	UModularRigController* GetController(bool bCreateIfNeeded = true);
@@ -336,22 +367,21 @@ public:
 
 	void UpdateCachedChildren();
 
-	FRigModuleReference* FindModule(const FString& InPath);
-	const FRigModuleReference* FindModule(const FString& InPath) const;
+	FRigModuleReference* FindModule(const FName& InModuleName);
+	const FRigModuleReference* FindModule(const FName& InModuleName) const;
+	const FRigModuleReference* FindModuleByPath(const FString& InModulePath) const;
 
-	FRigModuleReference* GetParentModule(const FString& InPath);
-	const FRigModuleReference* GetParentModule(const FString& InPath) const;
+	FRigModuleReference* GetParentModule(const FName& InModuleName);
+	const FRigModuleReference* GetParentModule(const FName& InModuleName) const;
 
 	FRigModuleReference* GetParentModule(const FRigModuleReference* InChildModule);
 	const FRigModuleReference* GetParentModule(const FRigModuleReference* InChildModule) const;
 
-	FString GetParentPath(const FString& InPath) const;
-
 	void ForEachModule(TFunction<bool(const FRigModuleReference*)> PerModule) const;
 
-	TArray<FString> SortPaths(const TArray<FString>& InPaths) const;
+	TArray<FName> SortModuleNames(const TArray<FName>& InModuleNames) const;
 
-	bool IsModuleParentedTo(const FString& InChildModulePath, const FString& InParentModulePath) const;
+	bool IsModuleParentedTo(const FName& InChildModuleName, const FName& InParentModuleName) const;
 	
 	bool IsModuleParentedTo(const FRigModuleReference* InChildModule, const FRigModuleReference* InParentModule) const;
 
@@ -361,7 +391,11 @@ public:
 
 private:
 	TWeakObjectPtr<UObject> OuterClientHost;
-	TArray<FString> SelectedModulePaths;
+	TArray<FName> SelectedModuleNames;
+
+#if WITH_EDITOR
+	mutable TArray<FName> ReceivedOldModulePaths;
+#endif
 
 	friend class UModularRigController;
 	friend struct FRigModuleReference;
diff --git a/Engine/Plugins/Animation/ControlRig/Source/ControlRig/Public/ModularRigRuleManager.h b/Engine/Plugins/Animation/ControlRig/Source/ControlRig/Public/ModularRigRuleManager.h
--- a/Engine/Plugins/Animation/ControlRig/Source/ControlRig/Public/ModularRigRuleManager.h
+++ b/Engine/Plugins/Animation/ControlRig/Source/ControlRig/Public/ModularRigRuleManager.h
@@ -100,7 +100,7 @@ private:
 	void SetHierarchy(const URigHierarchy* InHierarchy);
 	static void ResolveConnector(FWorkData& InOutWorkData);
 	static void FilterIncompatibleTypes(FWorkData& InOutWorkData);
-	static void FilterInvalidNameSpaces(FWorkData& InOutWorkData);
+	static void FilterInvalidModules(FWorkData& InOutWorkData);
 	static void FilterByConnectorRules(FWorkData& InOutWorkData);
 	static void FilterByConnectorEvent(FWorkData& InOutWorkData);
 
diff --git a/Engine/Plugins/Animation/ControlRig/Source/ControlRig/Public/Rigs/RigHierarchy.h b/Engine/Plugins/Animation/ControlRig/Source/ControlRig/Public/Rigs/RigHierarchy.h
--- a/Engine/Plugins/Animation/ControlRig/Source/ControlRig/Public/Rigs/RigHierarchy.h
+++ b/Engine/Plugins/Animation/ControlRig/Source/ControlRig/Public/Rigs/RigHierarchy.h
@@ -169,11 +169,12 @@ public:
 	typedef TPair<int32, TArray<int32>> TElementDependencyMapPair;
 	typedef TTuple<int32, int32, int32, ERigTransformType::Type> TInstructionSliceElement;
 	inline static const FLazyName TagMetadataName = TEXT("Tags");
-	inline static const FLazyName ShortModuleNameMetadataName = TEXT("ShortModuleName");
+	inline static const FLazyName ShortModuleNameMetadataName_Deprecated = TEXT("ShortModuleName");
 	inline static const FLazyName DesiredNameMetadataName = TEXT("DesiredName");
 	inline static const FLazyName DesiredKeyMetadataName = TEXT("DesiredKey");
 	inline static const FLazyName ModuleMetadataName = TEXT("Module");
-	inline static const FLazyName NameSpaceMetadataName = TEXT("NameSpace");
+	inline static const FLazyName NameSpaceMetadataName_Deprecated = TEXT("NameSpace");
+	inline static const FLazyName ShortNameMetadataName_Deprecated = TEXT("ShortName");
 
 	URigHierarchy();
 
@@ -450,29 +451,7 @@ public:
 	 * @param InKey The key of the element to retrieve the index for
 	 * @return The index of the element or INDEX_NONE
 	 */
-	int32 GetIndex(const FRigElementKey& InKey) const
-	{
-		if(ElementKeyRedirector)
-		{
-			if(FRigElementKeyRedirector::FCachedKeyArray* CachedRigElements = ElementKeyRedirector->Find(InKey))
-			{
-				if(CachedRigElements->Num() == 1)
-				{
-					if((*CachedRigElements)[0].UpdateCache(this))
-					{
-						return (*CachedRigElements)[0].GetIndex();
-					}
-			}
-				return INDEX_NONE;
-			}
-		}
-		
-		if(const int32* Index = ElementIndexLookup.Find(InKey))
-		{
-			return *Index;
-		}
-		return INDEX_NONE;
-	}
+	int32 GetIndex(const FRigElementKey& InKey) const;
 
 	/**
 	 * Returns the key and index pair of an element given its key
@@ -1889,28 +1868,49 @@ public:
 	 * Returns the path of the module an element belong to (or NAME_None in case the element doesn't belong to a module)
 	 * @return The path the element belongs to (or NAME_None)
 	 */
-	UFUNCTION(BlueprintPure, Category = URigHierarchy)
+	UFUNCTION(BlueprintPure, Category = URigHierarchy, meta = (DeprecatedFunction, DeprecationMessage = "This function has been deprecated. Please use GetModuleFName() instead."))
 	FName GetModulePathFName(FRigElementKey InItem) const;
 	
 	/**
 	 * Returns the path of the module an element belong to (or an empty string in case the element doesn't belong to a module)
 	 * @return The path the element belongs to (or empty string)
 	 */
-	UFUNCTION(BlueprintPure, Category = URigHierarchy)
+	UFUNCTION(BlueprintPure, Category = URigHierarchy, meta = (DeprecatedFunction, DeprecationMessage = "This function has been deprecated. Please use GetModuleName() instead."))
 	FString GetModulePath(FRigElementKey InItem) const;
 
+	/**
+	 * Returns the name of the module an element belong to (or NAME_None in case the element doesn't belong to a module)
+	 * @return The name the element belongs to (or NAME_None)
+	 */
+	UFUNCTION(BlueprintPure, Category = URigHierarchy)
+	FName GetModuleFName(FRigElementKey InItem) const;
+	
+	/**
+	 * Returns the name of the module an element belong to (or an empty string in case the element doesn't belong to a module)
+	 * @return The name the element belongs to (or empty string)
+	 */
+	UFUNCTION(BlueprintPure, Category = URigHierarchy)
+	FString GetModuleName(FRigElementKey InItem) const;
+
+	/**
+	 * Returns the prefix of the module an element belong to (or an empty string in case the element doesn't belong to a module)
+	 * @return The prefix the element belongs to (or empty string)
+	 */
+	UFUNCTION(BlueprintPure, Category = URigHierarchy)
+	FString GetModulePrefix(FRigElementKey InItem) const;
+
 	/**
 	 * Returns the namespace of an element belong to (or NAME_None in case the element doesn't belong to a module / namespace)
 	 * @return The namespace the element belongs to (or NAME_None)
 	 */
-	UFUNCTION(BlueprintPure, Category = URigHierarchy)
+	UFUNCTION(BlueprintPure, Category = URigHierarchy, meta = (DeprecatedFunction, DeprecationMessage = "This function has been deprecated. Please use GetModuleFName() instead."))
 	FName GetNameSpaceFName(FRigElementKey InItem) const;
 	
 	/**
 	 * Returns the namespace of an element belong to (or an empty string in case the element doesn't belong to a module / namespace)
 	 * @return The namespace the element belongs to (or empty string)
 	 */
-	UFUNCTION(BlueprintPure, Category = URigHierarchy)
+	UFUNCTION(BlueprintPure, Category = URigHierarchy, meta = (DeprecatedFunction, DeprecationMessage = "This function has been deprecated. Please use GetModuleName() instead."))
 	FString GetNameSpace(FRigElementKey InItem) const;
 
 	/*
@@ -2068,16 +2068,16 @@ public:
 	/**
 	 * Returns the two name sections with the right namespace separator
 	 */
-	static FString JoinNameSpace(const FString& InLeft, const FString& InRight);
-	static FRigName JoinNameSpace(const FRigName& InLeft, const FRigName& InRight);
+	static FString JoinNameSpace_Deprecated(const FString& InLeft, const FString& InRight);
+	static FRigName JoinNameSpace_Deprecated(const FRigName& InLeft, const FRigName& InRight);
 
 	/**
 	 * Returns the two name sections with the right namespace separator
 	 */
-	static TPair<FString, FString> SplitNameSpace(const FString& InNameSpacedPath, bool bFromEnd = true);
-	static TPair<FRigName, FRigName> SplitNameSpace(const FRigName& InNameSpacedPath, bool bFromEnd = true);
-	static bool SplitNameSpace(const FString& InNameSpacedPath, FString* OutNameSpace, FString* OutName, bool bFromEnd = true);
-	static bool SplitNameSpace(const FRigName& InNameSpacedPath, FRigName* OutNameSpace, FRigName* OutName, bool bFromEnd = true);
+	static TPair<FString, FString> SplitNameSpace_Deprecated(const FString& InNameSpacedPath, bool bFromEnd = true);
+	static TPair<FRigName, FRigName> SplitNameSpace_Deprecated(const FRigName& InNameSpacedPath, bool bFromEnd = true);
+	static bool SplitNameSpace_Deprecated(const FString& InNameSpacedPath, FString* OutNameSpace, FString* OutName, bool bFromEnd = true);
+	static bool SplitNameSpace_Deprecated(const FRigName& InNameSpacedPath, FRigName* OutNameSpace, FRigName* OutName, bool bFromEnd = true);
 
 	/**
 	 * Returns the max allowed length for a name within the hierarchy.
@@ -2153,8 +2153,8 @@ public:
 	/**
 	 * Returns the display label for an element to be used for the UI
 	 */
-	FText GetDisplayNameForUI(const FRigBaseElement* InElement, bool bIncludeNameSpace = true) const;
-	FText GetDisplayNameForUI(const FRigElementKey& InKey, bool bIncludeNameSpace = true) const;
+	FText GetDisplayNameForUI(const FRigBaseElement* InElement, bool bIncludeModuleName = true) const;
+	FText GetDisplayNameForUI(const FRigElementKey& InKey, bool bIncludeModuleName = true) const;
 
 	/**
 	 * Returns the modified event, which can be used to 
@@ -4473,6 +4473,12 @@ public:
 	UFUNCTION(BlueprintCallable, Category = URigHierarchy)
 	FName GetPreviousHierarchyName(const FRigHierarchyKey& InKey) const;
 
+	/**
+	 * Returns the previous name map used by this hierarchy
+	 * @return The name map from old to new key
+	 */
+	const TMap<FRigHierarchyKey, FRigHierarchyKey>& GetPreviousNameMap() const { return PreviousHierarchyNameMap; }
+
 	/**
 	 * Returns the previous parent of an element prior to a reparent operation
 	 * @param InKey The key of the element to request the old parent  for
@@ -4757,8 +4763,6 @@ private:
 	// Managed lookup from Component Key to Index
 	TMap<FRigComponentKey, int32> ComponentIndexLookup;
 
-	TMap<FRigElementKey, FString> UserDefinedElementName;
-
 	// Per element pose storage. Storage is defined here rather than on the elements
 	// to reduce memory consumption. Only elements created by MakeElement point to
 	// the element storage. Copied elements via the copy constructor or copy operator
@@ -5395,7 +5399,11 @@ private:
 #endif
 
 public:
-	
+
+	FRigElementKey PatchElementKeyInLookup(const FRigElementKey& InKey, const TMap<FRigHierarchyModulePath, FName>* InModulePathToName = nullptr) const;
+	void PatchElementMetadata(const TMap<FRigHierarchyModulePath, FName>& InModulePathToName);
+	void PatchModularRigComponentKeys(const TMap<FRigHierarchyModulePath, FName>& InModulePathToName);
+
 	template<typename RangeType>
 	static void ConvertElementsToKeys(RangeType InElements, TArray<FRigElementKey>& OutKeys)
 	{
@@ -5457,7 +5465,11 @@ public:
 private:
 	
 	// the currently destroyed element - used to avoid notification storms
-	const FRigBaseElement* ElementBeingDestroyed; 
+	const FRigBaseElement* ElementBeingDestroyed;
+
+#if WITH_EDITOR
+	mutable TArray<FRigElementKey> ReceivedNameSpaceBasedKeys;
+#endif
 	
 	friend class URigHierarchyController;
 	friend class UControlRig;
diff --git a/Engine/Plugins/Animation/ControlRig/Source/ControlRig/Public/Rigs/RigHierarchyController.h b/Engine/Plugins/Animation/ControlRig/Source/ControlRig/Public/Rigs/RigHierarchyController.h
--- a/Engine/Plugins/Animation/ControlRig/Source/ControlRig/Public/Rigs/RigHierarchyController.h
+++ b/Engine/Plugins/Animation/ControlRig/Source/ControlRig/Public/Rigs/RigHierarchyController.h
@@ -7,6 +7,7 @@
 #include "RigHierarchyContainer.h"
 #include "Animation/Skeleton.h"
 #include "RigHierarchyContainer.h"
+#include "RigVMFunctions/Math/RigVMMathLibrary.h"
 #include "RigHierarchyController.generated.h"
 
 UCLASS(BlueprintType)
@@ -315,7 +316,7 @@ public:
 	/**
 	 * Sets a control's settings given a control key
 	 * @param InKey The key of the control to set the settings for
-	 * @param The settings to set
+	 * @param InSettings The settings to set
 	 * @return Returns true if the settings have been set correctly
 	 */
 	UFUNCTION(BlueprintCallable, Category = URigHierarchyController)
@@ -612,7 +613,7 @@ public:
 	 * @param bReplaceExistingElements If set to true existing items will be replaced / updated with the content in the buffer
 	 * @param bSelectNewElements If set to true the new elements will be selected
 	 * @param bSetupUndo If set to true the stack will record the change for undo / redo
-	 * @param bPrintPythonCommand If set to true a python command equivalent to this call will be printed out
+	 * @param bPrintPythonCommands If set to true a python command equivalent to this call will be printed out
 	 */
 	UFUNCTION(BlueprintCallable, Category = URigHierarchyController)
 	TArray<FRigElementKey> ImportFromText(
@@ -827,7 +828,7 @@ public:
 	 * @param InKeys The keys of the elements to duplicate
 	 * @param bSelectNewElements If set to true the new elements will be selected
 	 * @param bSetupUndo If set to true the stack will record the change for undo / redo
-	 * @param bPrintPythonCommand If set to true a python command equivalent to this call will be printed out
+	 * @param bPrintPythonCommands If set to true a python command equivalent to this call will be printed out
 	 * @return The keys of the 4d items
 	 */
 	UFUNCTION(BlueprintCallable, Category = URigHierarchyController)
@@ -839,7 +840,7 @@ public:
 	 * @param InSettings The settings to use for the mirror operation
 	 * @param bSelectNewElements If set to true the new elements will be selected
 	 * @param bSetupUndo If set to true the stack will record the change for undo / redo
-	 * @param bPrintPythonCommand If set to true a python command equivalent to this call will be printed out
+	 * @param bPrintPythonCommands If set to true a python command equivalent to this call will be printed out
 	 * @return The keys of the mirrored items
 	 */
 	UFUNCTION(BlueprintCallable, Category = URigHierarchyController)
diff --git a/Engine/Plugins/Animation/ControlRig/Source/ControlRig/Public/Rigs/RigHierarchyDefines.h b/Engine/Plugins/Animation/ControlRig/Source/ControlRig/Public/Rigs/RigHierarchyDefines.h
--- a/Engine/Plugins/Animation/ControlRig/Source/ControlRig/Public/Rigs/RigHierarchyDefines.h
+++ b/Engine/Plugins/Animation/ControlRig/Source/ControlRig/Public/Rigs/RigHierarchyDefines.h
@@ -6,6 +6,7 @@
 #include "TransformNoScale.h"
 #include "EulerTransform.h"
 #include "ControlRigObjectVersion.h"
+#include "RigName.h"
 #include "RigHierarchyDefines.generated.h"
 
 class URigHierarchy;
@@ -108,6 +109,7 @@ enum class ERigHierarchyNotification : uint8
 	ComponentDeselected,
 	ComponentRenamed,
 	ComponentReparented,
+	ShortNameChanged,
 
 	/** MAX - invalid */
 	Max UMETA(Hidden),
@@ -1475,6 +1477,253 @@ struct CONTROLRIG_API FRigControlModifiedContext
 	bool bConstraintUpdate = false;
 };
 
+USTRUCT(BlueprintType)
+struct CONTROLRIG_API FRigHierarchyModulePath
+{
+	GENERATED_BODY();
+
+public:
+	
+	FRigHierarchyModulePath()
+		: ModulePath()
+	{}
+
+	FRigHierarchyModulePath(const FString& InModulePath)
+		: ModulePath(InModulePath)
+	{}
+
+	explicit FRigHierarchyModulePath(const FName& InModulePathFName)
+	: ModulePath(InModulePathFName.IsNone() ? FString() : InModulePathFName.ToString())
+	{}
+
+	FRigHierarchyModulePath(const TCHAR* InModulePath)
+		: ModulePath(InModulePath)
+	{}
+
+	template<typename StringTypeA = FString, typename StringTypeB = FString>
+	FRigHierarchyModulePath(const StringTypeA& InModuleName, const StringTypeB& InElementName)
+		: ModulePath()
+	{
+		*this = Join(InModuleName, InElementName);
+	}
+
+	bool IsEmpty() const
+	{
+		return ModulePath.IsEmpty();
+	}
+	
+	bool IsValid() const;
+	bool UsesNameSpaceFormat() const;
+	bool UsesModuleNameFormat() const;
+
+	operator const FString&() const
+	{
+		return ModulePath;
+	}
+
+	const FString& GetPath() const
+	{
+		return ModulePath;
+	}
+
+	FName GetPathFName() const;
+
+	FStringView GetModuleName() const;
+	const FString& GetModuleNameString() const;
+	const FName& GetModuleFName() const;
+	FStringView GetModulePrefix() const;
+	FString GetModulePrefixString() const;
+	FStringView GetElementName() const;
+	const FString& GetElementNameString() const;
+	const FName& GetElementFName() const;
+
+	template<typename StringType = FString>
+	bool HasModuleName(const StringType& InModuleNameString) const
+	{
+		return GetModuleName().Equals(InModuleNameString, ESearchCase::IgnoreCase);
+	}
+
+	template<typename StringType = FString>
+	bool HasElementName(const StringType& InElementNameString) const
+	{
+		return GetElementName().Equals(InElementNameString, ESearchCase::IgnoreCase);
+	}
+
+	template<typename StringType = FString>
+	bool SetModuleName(const StringType& InModuleNameString)
+	{
+		if(!HasModuleName(InModuleNameString))
+		{
+			*this = Join(InModuleNameString, GetElementNameString());
+			return true;
+		}
+		return false;
+	}
+
+	template<typename StringType = FString>
+	bool SetElementName(const StringType& InElementNameString)
+	{
+		if(!HasElementName(InElementNameString))
+		{
+			*this = Join(GetModuleNameString(), InElementNameString);
+			return true;
+		}
+		return false;
+	}
+
+	template<typename StringType = FString>
+	bool ReplaceModuleNameInline(const StringType& InOldModuleName, const StringType& InNewModuleName)
+	{
+		if(HasModuleName(InOldModuleName))
+		{
+			*this = Join(InNewModuleName, GetElementNameString());
+			return true;
+		}
+		return false;
+	}
+
+	template<typename StringType = FString>
+	FRigHierarchyModulePath ReplaceModuleName(const StringType& InOldModuleName, const StringType& InNewModuleName) const
+	{
+		FRigHierarchyModulePath Result = *this;
+		Result.ReplaceModuleNameInline<StringType>(InOldModuleName, InNewModuleName);
+		return Result;
+	}
+
+	template<typename StringType = FString>
+	FRigHierarchyModulePath ReplaceModuleName(const StringType& InNewModuleName) const
+	{
+		FRigHierarchyModulePath Result = *this;
+		Result.SetModuleName(InNewModuleName);
+		return Result;
+	}
+
+	template<typename StringType = FString>
+	bool ReplaceElementNameInline(const StringType& InOldElementName, const StringType& InNewElementName)
+	{
+		if(HasElementName(InOldElementName))
+		{
+			*this = Join(GetModuleNameString(), InNewElementName);
+			return true;
+		}
+		return false;
+	}
+
+	template<typename StringType = FString>
+	FRigHierarchyModulePath ReplaceElementName(const StringType& InOldElementName, const StringType& InNewElementName) const
+	{
+		FRigHierarchyModulePath Result = *this;
+		Result.ReplaceElementNameInline(InOldElementName, InNewElementName);
+		return Result;
+	}
+
+	template<typename StringType = FString>
+	FRigHierarchyModulePath ReplaceElementName(const StringType& InNewElementName) const
+	{
+		FRigHierarchyModulePath Result = *this;
+		Result.SetElementName(InNewElementName);
+		return Result;
+	}
+
+	friend uint32 GetTypeHash(const FRigHierarchyModulePath& InModulePath)
+	{
+		return GetTypeHash(InModulePath.ModulePath.ToLower());
+	}
+
+	bool operator==(const FRigHierarchyModulePath& InOther) const
+	{
+		return ModulePath.ToLower() == InOther.ModulePath.ToLower();
+	}
+
+	bool Split(FStringView* OutModuleName, FStringView* OutElementName) const;
+	bool Split(FString* OutModuleName, FString* OutElementName) const;
+
+	FRigHierarchyModulePath ConvertToModuleNameFormat(const TMap<FRigHierarchyModulePath, FName>* InModulePathToModuleName = nullptr) const;
+	bool ConvertToModuleNameFormatInline(const TMap<FRigHierarchyModulePath, FName>* InModulePathToModuleName = nullptr);
+
+	static inline const TCHAR* NamespaceSeparator_Deprecated = TEXT(":");
+	static inline constexpr TCHAR NamespaceSeparatorChar_Deprecated = TEXT(':');
+	static inline const TCHAR* ModuleNameSuffix = TEXT("/");
+	static inline constexpr TCHAR ModuleNameSuffixChar = TEXT('/');
+
+private:
+
+	static FRigHierarchyModulePath Join(const FString& InModuleName, const FString& InElementName);
+	static FRigHierarchyModulePath Join(const FName& InModuleFName, const FName& InElementFName);
+
+	UPROPERTY()
+	FString ModulePath;
+
+	mutable TOptional<FString> CachedModuleNameString;
+	mutable TOptional<FString> CachedElementNameString;
+	mutable TOptional<FName> CachedModuleFName;
+	mutable TOptional<FName> CachedElementFName;
+};
+
+template<>
+inline bool FRigHierarchyModulePath::HasModuleName<FName>(const FName& InModuleFName) const
+{
+	if(InModuleFName.IsNone())
+	{
+		return false;
+	}
+	return GetModuleFName() == InModuleFName;
+}
+
+template<>
+inline bool FRigHierarchyModulePath::HasElementName<FName>(const FName& InElementFName) const
+{
+	if(InElementFName.IsNone())
+	{
+		return false;
+	}
+	return GetElementFName() == InElementFName;
+}
+
+template<>
+inline bool FRigHierarchyModulePath::SetModuleName<FName>(const FName& InModuleFName)
+{
+	if(!HasModuleName(InModuleFName))
+	{
+		*this = Join(InModuleFName, GetElementFName());
+		return true;
+	}
+	return false;
+}
+
+template<>
+inline bool FRigHierarchyModulePath::SetElementName<FName>(const FName& InElementFName)
+{
+	if(!HasElementName(InElementFName))
+	{
+		*this = Join(GetModuleFName(), InElementFName);
+		return true;
+	}
+	return false;
+}
+
+template<>
+inline bool FRigHierarchyModulePath::ReplaceModuleNameInline<FName>(const FName& InOldModuleFName, const FName& InNewModuleFName)
+{
+	if(HasModuleName(InOldModuleFName))
+	{
+		*this = Join(InOldModuleFName, GetElementFName());
+		return true;
+	}
+	return false;
+}
+
+template<>
+inline bool FRigHierarchyModulePath::ReplaceElementNameInline<FName>(const FName& InOldElementFName, const FName& InNewElementFName)
+{
+	if(HasElementName(InOldElementFName))
+	{
+		*this = Join(GetModuleFName(), InNewElementFName);
+		return true;
+	}
+	return false;
+}
+
 /*
  * Because it's bitfield, we support some basic functionality
  */
@@ -1524,7 +1773,7 @@ public:
 		, Name(InName)
 	{}
 
-	FRigElementKey(const FString& InKeyString)
+	explicit FRigElementKey(const FString& InKeyString)
 	{
 		FString TypeStr, NameStr;
 		verify(InKeyString.Split(TEXT("("), &TypeStr, &NameStr));
@@ -1652,6 +1901,10 @@ public:
 	}
 
 	FString ToPythonString() const;
+
+	FRigElementKey ConvertToModuleNameFormat(const TMap<FRigHierarchyModulePath, FName>* InModulePathToModuleName = nullptr) const;
+	bool ConvertToModuleNameFormatInline(const TMap<FRigHierarchyModulePath, FName>* InModulePathToModuleName = nullptr);
+
 };
 
 USTRUCT(BlueprintType)
@@ -1767,6 +2020,9 @@ public:
 	FString ToPythonString() const;
 
 	bool IsTopLevel() const;
+
+	FRigComponentKey ConvertToModuleNameFormat(const TMap<FRigHierarchyModulePath, FName>* InModulePathToModuleName = nullptr) const;
+	bool ConvertToModuleNameFormatInline(const TMap<FRigHierarchyModulePath, FName>* InModulePathToModuleName = nullptr);
 };
 
 USTRUCT(BlueprintType)
diff --git a/Engine/Plugins/Animation/ControlRig/Source/ControlRig/Public/Sequencer/MovieSceneControlRigParameterSection.h b/Engine/Plugins/Animation/ControlRig/Source/ControlRig/Public/Sequencer/MovieSceneControlRigParameterSection.h
--- a/Engine/Plugins/Animation/ControlRig/Source/ControlRig/Public/Sequencer/MovieSceneControlRigParameterSection.h
+++ b/Engine/Plugins/Animation/ControlRig/Source/ControlRig/Public/Sequencer/MovieSceneControlRigParameterSection.h
@@ -102,35 +102,20 @@ struct FIntegerParameterNameAndValue
 };
 
 USTRUCT()
-struct CONTROLRIG_API FEnumParameterNameAndCurve
+struct CONTROLRIG_API FEnumParameterNameAndCurve : public FBaseParameterNameAndValue
 {
 	GENERATED_USTRUCT_BODY()
 
-	FEnumParameterNameAndCurve()
-	{}
-
-	FEnumParameterNameAndCurve(FName InParameterName);
-
-	UPROPERTY()
-	FName ParameterName;
-
 	UPROPERTY()
 	FMovieSceneByteChannel ParameterCurve;
 };
 
 
 USTRUCT()
-struct CONTROLRIG_API FIntegerParameterNameAndCurve
+struct CONTROLRIG_API FIntegerParameterNameAndCurve : public FBaseParameterNameAndValue
 {
 	GENERATED_USTRUCT_BODY()
 
-	FIntegerParameterNameAndCurve()
-	{}
-	FIntegerParameterNameAndCurve(FName InParameterName);
-
-	UPROPERTY()
-	FName ParameterName;
-
 	UPROPERTY()
 	FMovieSceneIntegerChannel ParameterCurve;
 };
@@ -253,8 +238,22 @@ public:
 	FSpaceChannelAddedEvent& SpaceChannelAdded() { return OnSpaceChannelAdded; }
 
 	const FName& FindControlNameFromConstraintChannel(const FMovieSceneConstraintChannel* InConstraintChannel) const;
-	
-	bool RenameParameterName(const FName& OldParameterName, const FName& NewParameterName, TOptional<ERigControlType> ControlType);
+
+	void ForEachParameter(TFunction<void(FBaseParameterNameAndValue*)> InCallback);
+	void ForEachParameter(TOptional<ERigControlType> InControlType, TFunction<void(FBaseParameterNameAndValue*)> InCallback);
+
+	template<typename T>
+	void ForEachParameter(TArray<T>& InParameterArray, TFunction<void(FBaseParameterNameAndValue*)> InCallback)
+	{
+		if(!InCallback)
+		{
+			return;
+		}
+		for(T& Parameter : InParameterArray)
+		{
+			InCallback(&Parameter);
+		}
+	}
 
 	void ChangeControlRotationOrder(const FName& InControlName, const TOptional<EEulerRotationOrder>& CurrentOrder, 
 		const  TOptional<EEulerRotationOrder>& NewOrder, EMovieSceneKeyInterpolation Interpolation);
diff --git a/Engine/Plugins/Animation/ControlRig/Source/ControlRig/Public/Units/Execution/RigUnit_RigModules.h b/Engine/Plugins/Animation/ControlRig/Source/ControlRig/Public/Units/Execution/RigUnit_RigModules.h
--- a/Engine/Plugins/Animation/ControlRig/Source/ControlRig/Public/Units/Execution/RigUnit_RigModules.h
+++ b/Engine/Plugins/Animation/ControlRig/Source/ControlRig/Public/Units/Execution/RigUnit_RigModules.h
@@ -109,7 +109,7 @@ struct CONTROLRIG_API FRigUnit_ResolveArrayConnector : public FRigUnit_RigModule
 /**
  * Returns the currently used namespace of the module.
  */
-USTRUCT(meta=(DisplayName="Get NameSpace", Keywords="NameSpace", Varying))
+USTRUCT(meta=(DisplayName="Get NameSpace", Keywords="NameSpace", Varying, Deprecated=5.6))
 struct CONTROLRIG_API FRigUnit_GetCurrentNameSpace : public FRigUnit_RigModulesBase
 {
 	GENERATED_BODY()
@@ -122,6 +122,9 @@ struct CONTROLRIG_API FRigUnit_GetCurrentNameSpace : public FRigUnit_RigModulesB
 	RIGVM_METHOD()
 	virtual void Execute() override;
 
+	RIGVM_METHOD()
+	virtual FRigVMStructUpgradeInfo GetUpgradeInfo() const override;
+
 	/*
 	 * The current namespace of the rig module
 	 */
@@ -132,7 +135,7 @@ struct CONTROLRIG_API FRigUnit_GetCurrentNameSpace : public FRigUnit_RigModulesB
 /**
  * Returns the short name of the given item (without the namespace)
  */
-USTRUCT(meta=(DisplayName="Get Short Name", Keywords="NameSpace", Varying))
+USTRUCT(meta=(DisplayName="Get Short Name", Keywords="NameSpace", Varying, Deprecated=5.6))
 struct CONTROLRIG_API FRigUnit_GetItemShortName : public FRigUnit_RigModulesBase
 {
 	GENERATED_BODY()
@@ -146,6 +149,9 @@ struct CONTROLRIG_API FRigUnit_GetItemShortName : public FRigUnit_RigModulesBase
 	RIGVM_METHOD()
 	virtual void Execute() override;
 
+	RIGVM_METHOD()
+	virtual FRigVMStructUpgradeInfo GetUpgradeInfo() const override;
+
 	/*
 	 * The key of the item to return the short name for
 	 */
@@ -162,7 +168,7 @@ struct CONTROLRIG_API FRigUnit_GetItemShortName : public FRigUnit_RigModulesBase
 /**
  * Returns the namespace of a given item. This may be an empty string if the item doesn't have a namespace.
  */
-USTRUCT(meta=(DisplayName="Get Item NameSpace", Keywords="NameSpace", Varying))
+USTRUCT(meta=(DisplayName="Get Item NameSpace", Keywords="NameSpace", Varying, Deprecated=5.6))
 struct CONTROLRIG_API FRigUnit_GetItemNameSpace : public FRigUnit_RigModulesBase
 {
 	GENERATED_BODY()
@@ -177,6 +183,9 @@ struct CONTROLRIG_API FRigUnit_GetItemNameSpace : public FRigUnit_RigModulesBase
 	RIGVM_METHOD()
 	virtual void Execute() override;
 
+	RIGVM_METHOD()
+	virtual FRigVMStructUpgradeInfo GetUpgradeInfo() const override;
+
 	/*
 	 * The key of the item to return the namespace for
 	 */
@@ -200,7 +209,7 @@ struct CONTROLRIG_API FRigUnit_GetItemNameSpace : public FRigUnit_RigModulesBase
  * Returns true if the given item has been created by this module, 
  * which means that the item's namespace is the module namespace.
  */
-USTRUCT(meta=(DisplayName="Is In Current NameSpace", Keywords="NameSpace", Varying))
+USTRUCT(meta=(DisplayName="Is In Current NameSpace", Keywords="NameSpace", Varying, Deprecated=5.6))
 struct CONTROLRIG_API FRigUnit_IsItemInCurrentNameSpace : public FRigUnit_RigModulesBase
 {
 	GENERATED_BODY()
@@ -214,6 +223,9 @@ struct CONTROLRIG_API FRigUnit_IsItemInCurrentNameSpace : public FRigUnit_RigMod
 	RIGVM_METHOD()
 	virtual void Execute() override;
 
+	RIGVM_METHOD()
+	virtual FRigVMStructUpgradeInfo GetUpgradeInfo() const override;
+
 	/*
 	 * The key of the item to return the namespace for
 	 */
@@ -230,7 +242,7 @@ struct CONTROLRIG_API FRigUnit_IsItemInCurrentNameSpace : public FRigUnit_RigMod
 /**
  * Returns all items (based on a filter) in the current namespace. 
  */
-USTRUCT(meta=(DisplayName="Get Items In NameSpace", Keywords="NameSpace", Varying))
+USTRUCT(meta=(DisplayName="Get Items In NameSpace", Keywords="NameSpace", Varying, Deprecated=5.6))
 struct CONTROLRIG_API FRigUnit_GetItemsInNameSpace : public FRigUnit_RigModulesBase
 {
 	GENERATED_BODY()
@@ -244,6 +256,9 @@ struct CONTROLRIG_API FRigUnit_GetItemsInNameSpace : public FRigUnit_RigModulesB
 	RIGVM_METHOD()
 	virtual void Execute() override;
 
+	RIGVM_METHOD()
+	virtual FRigVMStructUpgradeInfo GetUpgradeInfo() const override;
+
 	UPROPERTY(meta = (Input))
 	ERigElementType TypeToSearch;
 
@@ -252,4 +267,121 @@ struct CONTROLRIG_API FRigUnit_GetItemsInNameSpace : public FRigUnit_RigModulesB
 	 */
 	UPROPERTY(meta = (Output))
 	TArray<FRigElementKey> Items;
+};
+
+/**
+ * Returns the name of the current module instance.
+ */
+USTRUCT(meta=(DisplayName="Get Module Name", Keywords="NameSpace", Varying))
+struct CONTROLRIG_API FRigUnit_GetModuleName : public FRigUnit_RigModulesBase
+{
+	GENERATED_BODY()
+
+	FRigUnit_GetModuleName()
+	{
+		Module = FString();
+	}
+
+	RIGVM_METHOD()
+	virtual void Execute() override;
+
+	/*
+	 * The name of the module
+	 */
+	UPROPERTY(meta = (Output))
+	FString Module;
+};
+
+/**
+ * Returns the namespace of a given item. This may be an empty string if the item doesn't have a namespace.
+ */
+USTRUCT(meta=(DisplayName="Get Item Module", Keywords="NameSpace", Varying))
+struct CONTROLRIG_API FRigUnit_GetItemModuleName : public FRigUnit_RigModulesBase
+{
+	GENERATED_BODY()
+
+	FRigUnit_GetItemModuleName()
+	{
+		Item = FRigElementKey();
+		IsPartOfModule = false;
+		Module = FString();
+	}
+
+	RIGVM_METHOD()
+	virtual void Execute() override;
+
+	/*
+	 * The key of the item to return the module for
+	 */
+	UPROPERTY(meta = (Input, ExpandByDefault))
+	FRigElementKey Item;
+
+	/*
+	 * True if the item is part of a module
+	 */
+	UPROPERTY(meta = (Output))
+	bool IsPartOfModule;
+
+	/*
+	 * The name of the module of the given item
+	 */
+	UPROPERTY(meta = (Output))
+	FString Module;
+};
+
+/**
+ * Returns true if the given item has been created by this module. 
+ */
+USTRUCT(meta=(DisplayName="Is In Current Module", Keywords="NameSpace", Varying))
+struct CONTROLRIG_API FRigUnit_IsItemInCurrentModule : public FRigUnit_RigModulesBase
+{
+	GENERATED_BODY()
+
+	FRigUnit_IsItemInCurrentModule()
+	{
+		Item = FRigElementKey();
+		Result = false;
+	}
+
+	RIGVM_METHOD()
+	virtual void Execute() override;
+
+	/*
+	 * The key of the item to check the module for
+	 */
+	UPROPERTY(meta = (Input, ExpandByDefault))
+	FRigElementKey Item;
+
+	/*
+	 * True if the item is owned by this module.
+	 */
+	UPROPERTY(meta = (Output))
+	bool Result;
+};
+
+/**
+ * Returns all items (based on a filter) in the current module. 
+ */
+USTRUCT(meta=(DisplayName="Get Items In Module", Keywords="NameSpace", Varying))
+struct CONTROLRIG_API FRigUnit_GetItemsInModule : public FRigUnit_RigModulesBase
+{
+	GENERATED_BODY()
+
+	FRigUnit_GetItemsInModule()
+	{
+		TypeToSearch = ERigElementType::All;
+		Items.Reset();
+	}
+
+	RIGVM_METHOD()
+	virtual void Execute() override;
+
+	UPROPERTY(meta = (Input))
+	ERigElementType TypeToSearch;
+
+	/*
+	 * The items within the module
+	 */
+	UPROPERTY(meta = (Output))
+	TArray<FRigElementKey> Items;
 };
\ No newline at end of file
diff --git a/Engine/Plugins/Animation/ControlRig/Source/ControlRig/Public/Units/RigUnitContext.h b/Engine/Plugins/Animation/ControlRig/Source/ControlRig/Public/Units/RigUnitContext.h
--- a/Engine/Plugins/Animation/ControlRig/Source/ControlRig/Public/Units/RigUnitContext.h
+++ b/Engine/Plugins/Animation/ControlRig/Source/ControlRig/Public/Units/RigUnitContext.h
@@ -123,8 +123,8 @@ public:
 		: FRigVMExecuteContext()
 		, Hierarchy(nullptr)
 		, ControlRig(nullptr)
-		, RigModuleNameSpace()
-		, RigModuleNameSpaceHash(0)
+		, RigModulePrefix()
+		, RigModulePrefixHash(0)
 	{
 	}
 
@@ -167,37 +167,41 @@ public:
 	bool IsRunningConstructionEvent() const;
 
 	/**
-	 * Add the namespace from a given name
+	 * Returns true if this context is used on a module currently
 	 */
-	FName AddRigModuleNameSpace(const FName& InName) const;
-	FString AddRigModuleNameSpace(const FString& InName) const;
+	bool IsRigModule() const
+	{
+		return !GetRigModulePrefix().IsEmpty();
+	}
 
 	/**
-	 * Remove the namespace from a given name
+	 * Returns the rig module prefix
 	 */
-	FName RemoveRigModuleNameSpace(const FName& InName) const;
-	FString RemoveRigModuleNameSpace(const FString& InName) const;
+	const FString& GetRigModulePrefix() const
+	{
+		return RigModulePrefix;
+	}
 
 	/**
-	 * Returns true if this context is used on a module currently
+	 * Returns the prefix of the parent rig module
 	 */
-	bool IsRigModule() const
+	const FString& GetRigRootModulePrefix() const
 	{
-		return !GetRigModuleNameSpace().IsEmpty();
+		return RigRootModulePrefix;
 	}
 
 	/**
-	 * Returns the namespace of the currently running rig module
+	 * Returns the prefix of the root rig module
 	 */
-	FString GetRigModuleNameSpace() const
+	const FString& GetRigParentModulePrefix() const
 	{
-		return RigModuleNameSpace;
+		return RigParentModulePrefix;
 	}
 
 	/**
 	 * Returns the namespace given a namespace type
 	 */
-	FString GetElementNameSpace(ERigMetaDataNameSpace InNameSpaceType) const;
+	const FString& GetElementModulePrefix(ERigMetaDataNameSpace InNameSpaceType) const;
 	
 	/**
 	 * Returns the module this unit is running inside of (or nullptr)
@@ -233,7 +237,7 @@ public:
 #if WITH_EDITOR
 	virtual void Report(const FRigVMLogSettings& InLogSettings, const FName& InFunctionName, int32 InInstructionIndex, const FString& InMessage) const override
 	{
-		FString Prefix = GetRigModuleNameSpace();
+		const FString& Prefix = GetRigModulePrefix();
 		if (!Prefix.IsEmpty())
 		{
 			const FString Name = FString::Printf(TEXT("%s %s"), *Prefix, *InFunctionName.ToString());
@@ -247,8 +251,10 @@ public:
 #endif
 
 private:
-	FString RigModuleNameSpace;
-	uint32 RigModuleNameSpaceHash;
+	FString RigModulePrefix;
+	FString RigParentModulePrefix;
+	FString RigRootModulePrefix;
+	uint32 RigModulePrefixHash;
 	const FRigModuleInstance* RigModuleInstance;
 
 	friend class FControlRigExecuteContextRigModuleGuard;
@@ -259,14 +265,16 @@ class CONTROLRIG_API FControlRigExecuteContextRigModuleGuard
 {
 public:
 	FControlRigExecuteContextRigModuleGuard(FControlRigExecuteContext& InContext, const UControlRig* InControlRig);
-	FControlRigExecuteContextRigModuleGuard(FControlRigExecuteContext& InContext, const FString& InNewModuleNameSpace);
+	FControlRigExecuteContextRigModuleGuard(FControlRigExecuteContext& InContext, const FString& InNewModulePrefix, const FString& InNewParentModulePrefix, const FString& InNewRootModulePrefix);
 	~FControlRigExecuteContextRigModuleGuard();
 
 private:
 
 	FControlRigExecuteContext& Context;
-	FString PreviousRigModuleNameSpace;
-	uint32 PreviousRigModuleNameSpaceHash;
+	FString PreviousRigModulePrefix;
+	FString PreviousRigParentModulePrefix;
+	FString PreviousRigRootModulePrefix;
+	uint32 PreviousRigModulePrefixHash;
 };
 
 #if WITH_EDITOR
diff --git a/Engine/Plugins/Animation/ControlRig/Source/ControlRigDeveloper/Private/ControlRigBlueprint.cpp b/Engine/Plugins/Animation/ControlRig/Source/ControlRigDeveloper/Private/ControlRigBlueprint.cpp
--- a/Engine/Plugins/Animation/ControlRig/Source/ControlRigDeveloper/Private/ControlRigBlueprint.cpp
+++ b/Engine/Plugins/Animation/ControlRig/Source/ControlRigDeveloper/Private/ControlRigBlueprint.cpp
@@ -1071,7 +1071,7 @@ TArray<FSoftObjectPath> UControlRigBlueprint::GetReferencesToRigModule(const FAs
 						for (const FRigModuleReference* Module : Modules)
 						{
 							FSoftObjectPath ModulePath = DependencyData.GetSoftObjectPath();
-							ModulePath.SetSubPathString(Module->GetPath());
+							ModulePath.SetSubPathString(Module->GetModulePath().GetPath());
 							Result.Add(ModulePath);
 						}
 					}
@@ -1132,10 +1132,12 @@ void UControlRigBlueprint::UpdateExposedModuleConnectors() const
 {
 	UControlRigBlueprint* MutableThis = ((UControlRigBlueprint*)this);
 	MutableThis->RigModuleSettings.ExposedConnectors.Reset();
-	Hierarchy->ForEach<FRigConnectorElement>([MutableThis](const FRigConnectorElement* ConnectorElement) -> bool
+	Hierarchy->ForEach<FRigConnectorElement>([MutableThis, this](const FRigConnectorElement* ConnectorElement) -> bool
 	{
+		const FRigElementKey ConnectorKey = ConnectorElement->GetKey().ConvertToModuleNameFormat(&ModularRigModel.PreviousModulePaths);
+		
 		FRigModuleConnector ExposedConnector;
-		ExposedConnector.Name = ConnectorElement->GetName();
+		ExposedConnector.Name = ConnectorKey.Name.ToString();
 		ExposedConnector.Settings = ConnectorElement->Settings;
 		MutableThis->RigModuleSettings.ExposedConnectors.Add(ExposedConnector);
 		return true;
@@ -1363,6 +1365,16 @@ void UControlRigBlueprint::PostLoad()
 		}
 	}
 
+	if(ArrayConnectionMap.IsEmpty() && !ConnectionMap_DEPRECATED.IsEmpty())
+	{
+		for(const TPair<FRigElementKey, FRigElementKey>& Pair : ConnectionMap_DEPRECATED)
+		{
+			ArrayConnectionMap.Add(Pair.Key, FRigElementKeyCollection({Pair.Value}));
+		}
+	}
+
+	ModularRigModel.PatchModelsOnLoad();
+	
 #if WITH_EDITOR
 	if(IsControlRigModule() && Hierarchy)
 	{
@@ -1391,17 +1403,29 @@ void UControlRigBlueprint::PostLoad()
 	}
 #endif
 
-	if(ArrayConnectionMap.IsEmpty() && !ConnectionMap_DEPRECATED.IsEmpty())
+	// patch from previously used module paths to unique module names
+	TMap<FRigElementKey, FRigElementKeyCollection> PreviousArrayConnectionMap;
+	Swap(PreviousArrayConnectionMap, ArrayConnectionMap);
+
+	ArrayConnectionMap.Reset();
+	for(TPair<FRigElementKey, FRigElementKeyCollection>& Connection : PreviousArrayConnectionMap)
 	{
-		for(const TPair<FRigElementKey, FRigElementKey>& Pair : ConnectionMap_DEPRECATED)
+		FRigElementKey Key = Connection.Key.ConvertToModuleNameFormat(&ModularRigModel.PreviousModulePaths);
+		FRigElementKeyCollection& Targets = Connection.Value;
+		for(FRigElementKey& TargetKey : Targets.Keys)
 		{
-			ArrayConnectionMap.Add(Pair.Key, FRigElementKeyCollection({Pair.Value}));
+			TargetKey.ConvertToModuleNameFormatInline(&ModularRigModel.PreviousModulePaths);
 		}
+		ArrayConnectionMap.Add(Key, Targets);
 	}
 
-	ModularRigModel.PatchModelsOnLoad();
-
 	UpdateModularDependencyDelegates();
+
+	if(Hierarchy)
+	{
+		Hierarchy->PatchElementMetadata(ModularRigModel.PreviousModulePaths);
+		Hierarchy->PatchModularRigComponentKeys(ModularRigModel.PreviousModulePaths);
+	}
 }
 
 #if WITH_EDITOR
@@ -1476,20 +1500,20 @@ void UControlRigBlueprint::HandleConfigureRigVMController(const FRigVMClient* In
 
 #endif
 
-void UControlRigBlueprint::UpdateConnectionMapAfterRename(const FString& InOldNameSpace)
+void UControlRigBlueprint::UpdateConnectionMapAfterRename(const FString& InOldModuleName)
 {
-	const FString OldNameSpace = InOldNameSpace + UModularRig::NamespaceSeparator;
-	const FString NewNameSpace = RigModuleSettings.Identifier.Name + UModularRig::NamespaceSeparator;
+	const FString OldModuleName = InOldModuleName + FRigHierarchyModulePath::ModuleNameSuffix;
+	const FString NewModuleName = RigModuleSettings.Identifier.Name + FRigHierarchyModulePath::ModuleNameSuffix;
 	
 	TMap<FRigElementKey, FRigElementKeyCollection> FixedConnectionMap;
 	for(const TPair<FRigElementKey, FRigElementKeyCollection>& Pair : ArrayConnectionMap)
 	{
-		auto FixUpConnectionMap = [OldNameSpace, NewNameSpace](const FRigElementKey& InKey) -> FRigElementKey
+		auto FixUpConnectionMap = [OldModuleName, NewModuleName](const FRigElementKey& InKey) -> FRigElementKey
 		{
 			const FString NameString = InKey.Name.ToString();
-			if(NameString.StartsWith(OldNameSpace, ESearchCase::CaseSensitive))
+			if(NameString.StartsWith(OldModuleName, ESearchCase::CaseSensitive))
 			{
-				return FRigElementKey(*(NewNameSpace + NameString.Mid(OldNameSpace.Len())), InKey.Type);
+				return FRigElementKey(*(NewModuleName + NameString.Mid(OldModuleName.Len())), InKey.Type);
 			}
 			return InKey;
 		};
@@ -2642,21 +2666,18 @@ void UControlRigBlueprint::RefreshModuleConnectors(const FRigModuleReference* In
 		return;
 	}
 	
-	const bool bRemoveAllConnectors = !ModularRigModel.FindModule(InModule->GetPath());
+	const bool bRemoveAllConnectors = !ModularRigModel.FindModule(InModule->Name);
 	
 	if (URigHierarchyController* Controller = GetHierarchyController())
 	{
 		if (UControlRig* CDO = GetControlRigClass()->GetDefaultObject<UControlRig>())
 		{
-			const FString Namespace = InModule->GetNamespace();
 			const TArray<FRigElementKey> AllConnectors = Hierarchy->GetKeysOfType<FRigConnectorElement>();
-			const TArray<FRigElementKey> ExistingConnectors = AllConnectors.FilterByPredicate([Namespace](const FRigElementKey& ConnectorKey) -> bool
+			TArray<FRigElementKey> ExistingConnectors = AllConnectors.FilterByPredicate([InModule, this](const FRigElementKey& ConnectorKey) -> bool
 			{
-				const FString ConnectorName = ConnectorKey.Name.ToString();
-				FString ConnectorNamespace;
-				(void)URigHierarchy::SplitNameSpace(ConnectorName, &ConnectorNamespace, nullptr);
-				ConnectorNamespace.Append(UModularRig::NamespaceSeparator);
-				return ConnectorNamespace.Equals(Namespace, ESearchCase::CaseSensitive);
+				const FRigElementKey PatchedKey = ConnectorKey.ConvertToModuleNameFormat(&ModularRigModel.PreviousModulePaths);
+				const FRigHierarchyModulePath ConnectorModulePath(PatchedKey.Name);
+				return ConnectorModulePath.HasModuleName(InModule->Name);
 			});
 
 			// setup the module information. this is needed so that newly added
@@ -2666,24 +2687,32 @@ void UControlRigBlueprint::RefreshModuleConnectors(const FRigModuleReference* In
 			const UControlRig* ModuleCDO = InModule->Class->GetDefaultObject<UControlRig>();
 			const TArray<FRigModuleConnector>& ExpectedConnectors = ModuleCDO->GetRigModuleSettings().ExposedConnectors;
 
+			// rename the connectors since their keys have been patched
+			for(FRigElementKey& ConnectorKey : ExistingConnectors)
+			{
+				const FRigElementKey PatchedKey = ConnectorKey.ConvertToModuleNameFormat(&ModularRigModel.PreviousModulePaths);
+				if(ConnectorKey != PatchedKey)
+				{
+					ConnectorKey = Controller->RenameElement(ConnectorKey, PatchedKey.Name);
+				}
+			}
+
 			// remove the obsolete connectors
-			for(const FRigElementKey& Connector : ExistingConnectors)
+			for(const FRigElementKey& ConnectorKey : ExistingConnectors)
 			{
-				const FString ConnectorNameString = Connector.Name.ToString();
-				FString ConnectorNamespace, ShortName = ConnectorNameString;
-				(void)URigHierarchy::SplitNameSpace(ConnectorNameString, &ConnectorNamespace, &ShortName);
+				const FRigHierarchyModulePath ConnectorModulePath(ConnectorKey.Name);
 				const bool bConnectorExpected = ExpectedConnectors.ContainsByPredicate(
-					[ShortName](const FRigModuleConnector& ExpectedConnector) -> bool
+					[&ConnectorModulePath](const FRigModuleConnector& ExpectedConnector) -> bool
 					{
-						return ExpectedConnector.Name.Equals(ShortName, ESearchCase::CaseSensitive);
+						return ConnectorModulePath.HasElementName(ExpectedConnector.Name);
 					}
 				);
-				
+
 				if(bRemoveAllConnectors || !bConnectorExpected)
 				{
 					Hierarchy->Modify();
-					(void)Controller->RemoveElement(Connector);
-					ArrayConnectionMap.Remove(Connector);
+					(void)Controller->RemoveElement(ConnectorKey);
+					ArrayConnectionMap.Remove(ConnectorKey);
 				}
 			}
 
@@ -2693,19 +2722,22 @@ void UControlRigBlueprint::RefreshModuleConnectors(const FRigModuleReference* In
 				for (const FRigModuleConnector& Connector : ExpectedConnectors)
 				{
 					const FName ConnectorName = *Connector.Name;
-					const FName ConnectorNameWithNameSpace = *URigHierarchy::JoinNameSpace(InModule->GetNamespace(), Connector.Name);
-					const FRigElementKey ConnectorKeyWithNameSpace(ConnectorNameWithNameSpace, ERigElementType::Connector);
-					if(!Hierarchy->Contains(ConnectorKeyWithNameSpace))
+					const FRigHierarchyModulePath ConnectorModulePath(InModule->Name.ToString(), Connector.Name);
+					const FRigElementKey CombinedConnectorKey(ConnectorModulePath.GetPathFName(), ERigElementType::Connector);
+					if(!Hierarchy->Contains(CombinedConnectorKey))
 					{
+						const FString ModulePrefix = InModule->GetElementPrefix();
+						const FString ParentModulePrefix = InModule->GetParentModule() ? InModule->GetParentModule()->GetElementPrefix() : ModulePrefix;
+						const FString RootModulePrefix = InModule->GetRootModule() ? InModule->GetRootModule()->GetElementPrefix() : ModulePrefix;
 						FRigHierarchyExecuteContextBracket HierarchyContextGuard(Hierarchy, &Context);
-						FControlRigExecuteContextRigModuleGuard RigModuleGuard(PublicContext, InModule->GetNamespace());
+						FControlRigExecuteContextRigModuleGuard RigModuleGuard(PublicContext, ModulePrefix, ParentModulePrefix, RootModulePrefix);
 						Hierarchy->Modify();
 						(void)Controller->AddConnector(ConnectorName, Connector.Settings);
 					}
 					else
 					{
 						// copy the connector settings
-						FRigConnectorElement* ExistingConnector = Hierarchy->FindChecked<FRigConnectorElement>(ConnectorKeyWithNameSpace);
+						FRigConnectorElement* ExistingConnector = Hierarchy->FindChecked<FRigConnectorElement>(CombinedConnectorKey);
 						ExistingConnector->Settings = Connector.Settings;
 					}
 				}
@@ -3022,31 +3054,18 @@ void UControlRigBlueprint::HandleRigModulesModified(EModularRigNotification InNo
 							FRigElementKeyCollection TargetConnections;
 							FRigConnectorSettings Settings;
 						};
-						FString OldNamespace;
-						if (InNotification == EModularRigNotification::ModuleRenamed)
-						{
-							OldNamespace = (InModule->ParentPath.IsEmpty()) ? *InModule->PreviousName.ToString() : URigHierarchy::JoinNameSpace(InModule->ParentPath, InModule->PreviousName.ToString());
-						}
-						else if(InNotification == EModularRigNotification::ModuleReparented)
-						{
-							OldNamespace = (InModule->PreviousParentPath.IsEmpty()) ? *InModule->PreviousName.ToString() : URigHierarchy::JoinNameSpace(InModule->PreviousParentPath, InModule->PreviousName.ToString());
-						}
-						FString NewNamespace = (InModule->ParentPath.IsEmpty()) ? *InModule->Name.ToString() : URigHierarchy::JoinNameSpace(InModule->ParentPath, InModule->Name.ToString());
-						OldNamespace.Append(UModularRig::NamespaceSeparator);
-						NewNamespace.Append(UModularRig::NamespaceSeparator);
+						const FName OldModuleName = InModule->PreviousName;
+						const FName NewModuleName = InModule->Name;
 						
 						TArray<FRigElementKey> Connectors = Controller->GetHierarchy()->GetKeysOfType<FRigConnectorElement>();
 						TMap<FRigElementKey, ConnectionInfo> RenamedConnectors; // old key -> new key
 						for (const FRigElementKey& Connector : Connectors)
 						{
-							FString OldConnectorName = Connector.Name.ToString();
-							FString OldConnectorNamespace, ShortName;
-							(void)URigHierarchy::SplitNameSpace(OldConnectorName, &OldConnectorNamespace, &ShortName);
-							OldConnectorNamespace.Append(UModularRig::NamespaceSeparator);
-							if (OldConnectorNamespace.StartsWith(OldNamespace))
+							const FRigHierarchyModulePath ConnectorModulePath(Connector.Name);
+							if (ConnectorModulePath.HasModuleName(OldModuleName))
 							{
 								ConnectionInfo& Info = RenamedConnectors.FindOrAdd(Connector);
-								Info.NewPath = OldConnectorName.Replace(*OldNamespace, *NewNamespace);
+								Info.NewPath = ConnectorModulePath.ReplaceModuleName(NewModuleName);
 								Info.Settings = CastChecked<FRigConnectorElement>(Controller->GetHierarchy()->FindChecked(Connector))->Settings;
 								if (FRigElementKeyCollection* TargetKeys = ArrayConnectionMap.Find(Connector))
 								{
@@ -3068,12 +3087,24 @@ void UControlRigBlueprint::HandleRigModulesModified(EModularRigNotification InNo
 							FControlRigExecuteContext& PublicContext = Context.GetPublicDataSafe<FControlRigExecuteContext>();
 							for (TPair<FRigElementKey, ConnectionInfo>& Pair : RenamedConnectors)
 							{
-								FString Namespace, ConnectorName;
-								(void)URigHierarchy::SplitNameSpace(Pair.Value.NewPath, &Namespace, &ConnectorName);
-								Namespace.Append(UModularRig::NamespaceSeparator);
-								FControlRigExecuteContextRigModuleGuard RigModuleGuard(PublicContext, Namespace);
+								const FName& ConnectorName = FRigHierarchyModulePath(Pair.Value.NewPath).GetElementFName();
+								const FString ModulePrefix = InModule->GetElementPrefix();
+								const FString ParentModulePrefix = InModule->GetParentModule() ? InModule->GetParentModule()->GetElementPrefix() : ModulePrefix;
+								const FString RootModulePrefix = InModule->GetRootModule() ? InModule->GetRootModule()->GetElementPrefix() : ModulePrefix;
+								
+								FControlRigExecuteContextRigModuleGuard RigModuleGuard(PublicContext, ModulePrefix, ParentModulePrefix, RootModulePrefix);
 								const TGuardValue<bool> DisableErrors(Controller->bReportWarningsAndErrors, false);
-								Controller->AddConnector(*ConnectorName, Pair.Value.Settings);
+								Controller->AddConnector(ConnectorName, Pair.Value.Settings);
+							}
+						}
+
+						// update the module table
+						for(TPair<FRigHierarchyModulePath,FName>& ModulePathToName : ModularRigModel.PreviousModulePaths)
+						{
+							if(ModulePathToName.Value == OldModuleName)
+							{
+								Modify();
+								ModulePathToName.Value = NewModuleName;
 							}
 						}
 
diff --git a/Engine/Plugins/Animation/ControlRig/Source/ControlRigDeveloper/Public/ControlRigBlueprint.h b/Engine/Plugins/Animation/ControlRig/Source/ControlRigDeveloper/Public/ControlRigBlueprint.h
--- a/Engine/Plugins/Animation/ControlRig/Source/ControlRigDeveloper/Public/ControlRigBlueprint.h
+++ b/Engine/Plugins/Animation/ControlRig/Source/ControlRigDeveloper/Public/ControlRigBlueprint.h
@@ -53,7 +53,7 @@ public:
 	{
 		if (InModule)
 		{
-			ModulePath = InModule->GetPath();
+			ModulePath = InModule->GetModulePath();
 			if (InModule->Class.IsValid())
 			{
 				ReferencedModule = InModule->Class.Get();
@@ -404,7 +404,7 @@ private:
 	virtual void HandleConfigureRigVMController(const FRigVMClient* InClient, URigVMController* InControllerToConfigure) override;
 #endif
 
-	void UpdateConnectionMapAfterRename(const FString& InOldNameSpace);
+	void UpdateConnectionMapAfterRename(const FString& InOldModuleName);
 
 	// Class used to temporarily cache all 
 	// current control values and reapply them
diff --git a/Engine/Plugins/Animation/ControlRig/Source/ControlRigEditor/Private/ControlRigBlueprintFactory.cpp b/Engine/Plugins/Animation/ControlRig/Source/ControlRigEditor/Private/ControlRigBlueprintFactory.cpp
--- a/Engine/Plugins/Animation/ControlRig/Source/ControlRigEditor/Private/ControlRigBlueprintFactory.cpp
+++ b/Engine/Plugins/Animation/ControlRig/Source/ControlRigEditor/Private/ControlRigBlueprintFactory.cpp
@@ -310,7 +310,7 @@ UObject* UControlRigBlueprintFactory::FactoryCreateNew(UClass* Class, UObject* I
 					if(DefaultRootModuleClass->IsChildOf(UControlRig::StaticClass()))
 					{
 						static const FName RootName = TEXT("Root");
-						ControlRigBlueprint->ModularRigModel.GetController()->AddModule(RootName, DefaultRootModuleClass, FString(), false);
+						ControlRigBlueprint->ModularRigModel.GetController()->AddModule(RootName, DefaultRootModuleClass, NAME_None, false);
 					}
 				}
 			}
diff --git a/Engine/Plugins/Animation/ControlRig/Source/ControlRigEditor/Private/ControlRigDragOps.cpp b/Engine/Plugins/Animation/ControlRig/Source/ControlRigEditor/Private/ControlRigDragOps.cpp
--- a/Engine/Plugins/Animation/ControlRig/Source/ControlRigEditor/Private/ControlRigDragOps.cpp
+++ b/Engine/Plugins/Animation/ControlRig/Source/ControlRigEditor/Private/ControlRigDragOps.cpp
@@ -89,10 +89,10 @@ TSharedPtr<SWidget> FRigHierarchyTagDragDropOp::GetDefaultDecorator() const
 /// FModularRigModuleDragDropOp
 ///////////////////////////////////////////////////////////
 
-TSharedRef<FModularRigModuleDragDropOp> FModularRigModuleDragDropOp::New(const TArray<FString>& InElements)
+TSharedRef<FModularRigModuleDragDropOp> FModularRigModuleDragDropOp::New(const TArray<FName>& InModuleNames)
 {
 	TSharedRef<FModularRigModuleDragDropOp> Operation = MakeShared<FModularRigModuleDragDropOp>();
-	Operation->Elements = InElements;
+	Operation->ModuleNames = InModuleNames;
 	Operation->Construct();
 	return Operation;
 }
@@ -104,17 +104,17 @@ TSharedPtr<SWidget> FModularRigModuleDragDropOp::GetDefaultDecorator() const
 		.BorderImage(FAppStyle::GetBrush("Menu.Background"))
 		[
 			SNew(STextBlock)
-			.Text(FText::FromString(GetJoinedElementNames()))
+			.Text(FText::FromString(GetJoinedModuleNames()))
 			//.Font(FAppStyle::Get().GetFontStyle("FontAwesome.10"))
 		];
 }
 
-FString FModularRigModuleDragDropOp::GetJoinedElementNames() const
+FString FModularRigModuleDragDropOp::GetJoinedModuleNames() const
 {
 	TArray<FString> ElementNameStrings;
-	for (const FString& Element: Elements)
+	for (const FName& ModuleName: ModuleNames)
 	{
-		ElementNameStrings.Add(Element);
+		ElementNameStrings.Add(ModuleName.ToString());
 	}
 	return FString::Join(ElementNameStrings, TEXT(","));
 }
diff --git a/Engine/Plugins/Animation/ControlRig/Source/ControlRigEditor/Private/ControlRigDragOps.h b/Engine/Plugins/Animation/ControlRig/Source/ControlRigEditor/Private/ControlRigDragOps.h
--- a/Engine/Plugins/Animation/ControlRig/Source/ControlRigEditor/Private/ControlRigDragOps.h
+++ b/Engine/Plugins/Animation/ControlRig/Source/ControlRigEditor/Private/ControlRigDragOps.h
@@ -65,27 +65,27 @@ class FModularRigModuleDragDropOp : public FDragDropOperation
 public:
 	DRAG_DROP_OPERATOR_TYPE(FModularRigModuleDragDropOp, FDragDropOperation)
 
-	static TSharedRef<FModularRigModuleDragDropOp> New(const TArray<FString>& InElements);
+	static TSharedRef<FModularRigModuleDragDropOp> New(const TArray<FName>& InModuleNames);
 
 	virtual TSharedPtr<SWidget> GetDefaultDecorator() const override;
 
-	/** @return true if this drag operation contains property paths */
-	bool HasElements() const
+	/** @return true if this drag operation contains dragged modules */
+	bool HasModules() const
 	{
-		return Elements.Num() > 0;
+		return ModuleNames.Num() > 0;
 	}
 
-	/** @return The property paths from this drag operation */
-	const TArray<FString>& GetElements() const
+	/** @return The module names from this drag operation */
+	const TArray<FName>& GetModules() const
 	{
-		return Elements;
+		return ModuleNames;
 	}
 
-	FString GetJoinedElementNames() const;
+	FString GetJoinedModuleNames() const;
 
 private:
 
-	/** Data for the property paths this item represents */
-	TArray<FString> Elements;
+	/** Data for the module names this item represents */
+	TArray<FName> ModuleNames;
 };
 
diff --git a/Engine/Plugins/Animation/ControlRig/Source/ControlRigEditor/Private/ControlRigElementDetails.cpp b/Engine/Plugins/Animation/ControlRig/Source/ControlRigEditor/Private/ControlRigElementDetails.cpp
--- a/Engine/Plugins/Animation/ControlRig/Source/ControlRigEditor/Private/ControlRigElementDetails.cpp
+++ b/Engine/Plugins/Animation/ControlRig/Source/ControlRigEditor/Private/ControlRigElementDetails.cpp
@@ -75,6 +75,10 @@ void RigElementKeyDetails_GetCustomizedInfo(TSharedRef<IPropertyHandle> InStruct
 	InStructPropertyHandle->GetOuterObjects(Objects);
 	for (UObject* Object : Objects)
 	{
+		if(!IsValid(Object))
+		{
+			continue;
+		}
 		if (Object->IsA<UControlRigBlueprint>())
 		{
 			OutBlueprint = CastChecked<UControlRigBlueprint>(Object);
@@ -188,6 +192,8 @@ void SRigElementKeyWidget::Construct(const FArguments& InArgs)
 	OnElementTypeChanged = InArgs._OnElementTypeChanged;
 	
 	UpdateElementNameList();
+
+	TWeakPtr<SRigElementKeyWidget> WeakThisPtr = StaticCastWeakPtr<SRigElementKeyWidget>(AsWeak());
 	
 	ChildSlot
 	[
@@ -199,21 +205,27 @@ void SRigElementKeyWidget::Construct(const FArguments& InArgs)
 				TypeHandle->CreatePropertyValueWidget()
 					:
 				SNew(SEnumComboBox, StaticEnum<ERigElementType>())
-				.CurrentValue_Lambda([this]()
+				.CurrentValue_Lambda([WeakThisPtr]()
 				{
-					if (OnGetElementType.IsBound())
+					if(const TSharedPtr<SRigElementKeyWidget> StrongThisPtr = WeakThisPtr.Pin())
 					{
-						return (int32)OnGetElementType.Execute();
+						if (StrongThisPtr->OnGetElementType.IsBound())
+						{
+							return (int32)StrongThisPtr->OnGetElementType.Execute();
+						}
 					}
 					return (int32)ERigElementType::None;
 				})
-				.OnEnumSelectionChanged_Lambda([this](int32 InEnumValue, ESelectInfo::Type SelectInfo)
+				.OnEnumSelectionChanged_Lambda([WeakThisPtr](int32 InEnumValue, ESelectInfo::Type SelectInfo)
 				{
-					ERigElementType EnumValue = static_cast<ERigElementType>(InEnumValue);
-					OnElementTypeChanged.ExecuteIfBound(EnumValue);
-					UpdateElementNameList();
-					SearchableComboBox->ClearSelection();
-					OnElementNameChanged.ExecuteIfBound(nullptr, ESelectInfo::Direct);
+					if(const TSharedPtr<SRigElementKeyWidget> StrongThisPtr = WeakThisPtr.Pin())
+					{
+						ERigElementType EnumValue = static_cast<ERigElementType>(InEnumValue);
+						StrongThisPtr->OnElementTypeChanged.ExecuteIfBound(EnumValue);
+						StrongThisPtr->UpdateElementNameList();
+						StrongThisPtr->SearchableComboBox->ClearSelection();
+						StrongThisPtr->OnElementNameChanged.ExecuteIfBound(nullptr, ESelectInfo::Direct);
+					}
 				})
 		]
 
@@ -340,6 +352,8 @@ void FRigElementKeyDetails::CustomizeHeader(TSharedRef<IPropertyHandle> InStruct
 	BlueprintBeingCustomized = nullptr;
 	RigElementKeyDetails_GetCustomizedInfo(InStructPropertyHandle, BlueprintBeingCustomized);
 
+	TWeakPtr<FRigElementKeyDetails> WeakThisPtr = StaticCastWeakPtr<FRigElementKeyDetails>(AsWeak());
+
 	if (BlueprintBeingCustomized == nullptr)
 	{
 		HeaderRow
@@ -367,7 +381,14 @@ void FRigElementKeyDetails::CustomizeHeader(TSharedRef<IPropertyHandle> InStruct
 		[
 			SAssignNew(RigElementKeyWidget, SRigElementKeyWidget, NameHandle, TypeHandle)
 			.Blueprint(BlueprintBeingCustomized)
-			.IsEnabled_Lambda([this](){ return !NameHandle->IsEditConst(); })
+			.IsEnabled_Lambda([WeakThisPtr]()
+			{
+				if(const TSharedPtr<FRigElementKeyDetails> StrongThisPtr = WeakThisPtr.Pin())
+				{
+					return !StrongThisPtr->NameHandle->IsEditConst();
+				}
+				return false;
+			})
 			.ActiveBackgroundColor(FSlateColor(FLinearColor(1.f, 1.f, 1.f, FRigElementKeyDetailsDefs::ActivePinBackgroundAlpha)))
 			.ActiveForegroundColor(FSlateColor(FLinearColor(1.f, 1.f, 1.f, FRigElementKeyDetailsDefs::ActivePinForegroundAlpha)))
 			.InactiveBackgroundColor(FSlateColor(FLinearColor(1.f, 1.f, 1.f, FRigElementKeyDetailsDefs::InactivePinBackgroundAlpha)))
@@ -599,9 +620,13 @@ void FRigComponentKeyDetails::CustomizeChildren(TSharedRef<IPropertyHandle> InSt
 		}
 		else
 		{
-			ElementKeyHandle->SetOnChildPropertyValueChanged(FSimpleDelegate::CreateLambda([this]()
+			TWeakPtr<FRigComponentKeyDetails> WeakThisPtr = StaticCastWeakPtr<FRigComponentKeyDetails>(AsWeak());
+			ElementKeyHandle->SetOnChildPropertyValueChanged(FSimpleDelegate::CreateLambda([WeakThisPtr]()
 			{
-				UpdateComponentNameList();
+				if(const TSharedPtr<FRigComponentKeyDetails> StrongThisPtr = WeakThisPtr.Pin())
+				{
+					StrongThisPtr->UpdateComponentNameList();
+				}
 			}));
 			
 			StructBuilder
@@ -1077,8 +1102,6 @@ void FRigBaseElementDetails::CustomizeDetails(IDetailLayoutBuilder& DetailBuilde
 		];
 	}
 
-	DetailBuilder.HideCategory(TEXT("RigElement"));
-
 	// if we are not a bone, control or null
 	if(!IsAnyElementOfType(ERigElementType::Bone) &&
 		!IsAnyElementOfType(ERigElementType::Control) &&
@@ -1333,7 +1356,7 @@ bool FRigBaseElementDetails::IsAnyConnectorImported() const
 {
 	return ContainsElementByPredicate([](const FPerElementInfo& Info)
 	{
-		return Info.Element.GetKey().Name.ToString().Contains(UModularRig::NamespaceSeparator);
+		return Info.Element.GetKey().Name.ToString().Contains(FRigHierarchyModulePath::ModuleNameSuffix);
 	});
 }
 
@@ -2643,11 +2666,11 @@ FDetailWidgetRow& FRigTransformElementDetails::CreateEulerTransformValueWidgetRo
 		bool bIsRelative,
 		bool bSetupUndoRedo)
 	{
-		if (!WeakThisPtr.IsValid())
+		const TSharedPtr<FRigTransformElementDetails> StrongThisPtr = WeakThisPtr.Pin();
+		if(!StrongThisPtr)
 		{
 			return;
 		}
-		FRigTransformElementDetails* StrongThisPtr = WeakThisPtr.Pin().Get();
 		
 		const bool bCurrent = CurrentTransformType == ERigTransformElementDetailsTransform::Current; 
 		const bool bInitial = CurrentTransformType == ERigTransformElementDetailsTransform::Initial; 
@@ -2837,11 +2860,11 @@ FDetailWidgetRow& FRigTransformElementDetails::CreateEulerTransformValueWidgetRo
 		ESlateTransformSubComponent::Type SubComponent,
 		FVector::FReal InNumericValue)
 	{
-		if (!WeakThisPtr.IsValid())
+		const TSharedPtr<FRigTransformElementDetails> StrongThisPtr = WeakThisPtr.Pin();
+		if (!StrongThisPtr)
 		{
 			return;
 		}
-		FRigTransformElementDetails* StrongThisPtr = WeakThisPtr.Pin().Get();
 			
 		const bool bIsRelative = IsComponentRelative((int32)Component);
 
@@ -2879,11 +2902,11 @@ FDetailWidgetRow& FRigTransformElementDetails::CreateEulerTransformValueWidgetRo
 		FVector::FReal InNumericValue,
 		ETextCommit::Type InCommitType)
 	{
-		if (!WeakThisPtr.IsValid())
+		const TSharedPtr<FRigTransformElementDetails> StrongThisPtr = WeakThisPtr.Pin();
+		if (!StrongThisPtr)
 		{
 			return;
 		}
-		FRigTransformElementDetails* StrongThisPtr = WeakThisPtr.Pin().Get();
 		
 		const bool bIsRelative = IsComponentRelative((int32)Component);
 
@@ -2913,11 +2936,11 @@ FDetailWidgetRow& FRigTransformElementDetails::CreateEulerTransformValueWidgetRo
 			ESlateRotationRepresentation::Type Representation,
 			ESlateTransformSubComponent::Type SubComponent)
 		{
-			if (!WeakThisPtr.IsValid())
+			TSharedPtr<FRigTransformElementDetails> StrongThisPtr = WeakThisPtr.Pin();
+			if (!StrongThisPtr)
 			{
 				return;
 			}
-			FRigTransformElementDetails* StrongThisPtr = WeakThisPtr.Pin().Get();
 			
 			if (UControlRig* DebuggedRig = Cast<UControlRig>(StrongThisPtr->PerElementInfos[0].GetBlueprint()->GetObjectBeingDebugged()))
 			{
@@ -2949,11 +2972,11 @@ FDetailWidgetRow& FRigTransformElementDetails::CreateEulerTransformValueWidgetRo
 			ESlateTransformSubComponent::Type SubComponent,
 			FVector::FReal InNumericValue)
 		{
-			if (!WeakThisPtr.IsValid())
+			const TSharedPtr<FRigTransformElementDetails> StrongThisPtr = WeakThisPtr.Pin();
+			if (!StrongThisPtr)
 			{
 				return;
 			}
-			FRigTransformElementDetails* StrongThisPtr = WeakThisPtr.Pin().Get();
 			
 			if (UControlRig* DebuggedRig = Cast<UControlRig>(StrongThisPtr->PerElementInfos[0].GetBlueprint()->GetObjectBeingDebugged()))
 			{
@@ -4949,31 +4972,39 @@ void FRigControlElementDetails::CustomizeShape(IDetailLayoutBuilder& DetailBuild
 	TArray<FRigElementKey> Keys = GetElementKeys();
 	Keys = PerElementInfos[0].GetHierarchy()->SortKeys(Keys);
 
-	auto GetShapeTransform = [this](
+	TWeakPtr<FRigControlElementDetails> WeakThisPtr = StaticCastWeakPtr<FRigControlElementDetails>(AsWeak());
+
+	auto GetShapeTransform = [WeakThisPtr](
 		const FRigElementKey& Key
 		) -> FEulerTransform
 	{
-		if(const FPerElementInfo& Info = FindElement(Key))
+		if(const TSharedPtr<FRigTransformElementDetails> StrongThisPtr = WeakThisPtr.Pin())
 		{
-			if(FRigControlElement* ControlElement = Info.GetElement<FRigControlElement>())
+			if(const FPerElementInfo& Info = StrongThisPtr->FindElement(Key))
 			{
-				return FEulerTransform(Info.GetHierarchy()->GetControlShapeTransform(ControlElement, ERigTransformType::InitialLocal));
+				if(FRigControlElement* ControlElement = Info.GetElement<FRigControlElement>())
+				{
+					return FEulerTransform(Info.GetHierarchy()->GetControlShapeTransform(ControlElement, ERigTransformType::InitialLocal));
+				}
 			}
 		}
 		return FEulerTransform::Identity;
 	};
 
-	auto SetShapeTransform = [this](
+	auto SetShapeTransform = [WeakThisPtr](
 		const FRigElementKey& Key,
 		const FEulerTransform& InTransform,
 		bool bSetupUndo
 		)
 	{
-		if(const FPerElementInfo& Info = FindElement(Key))
+		if(const TSharedPtr<FRigTransformElementDetails> StrongThisPtr = WeakThisPtr.Pin())
 		{
-			if(const FRigControlElement* ControlElement = Info.GetDefaultElement<FRigControlElement>())
+			if(const FPerElementInfo& Info = StrongThisPtr->FindElement(Key))
 			{
-				Info.GetDefaultHierarchy()->SetControlShapeTransform((FRigControlElement*)ControlElement, InTransform.ToFTransform(), ERigTransformType::InitialLocal, bSetupUndo, true, bSetupUndo);
+				if(const FRigControlElement* ControlElement = Info.GetDefaultElement<FRigControlElement>())
+				{
+					Info.GetDefaultHierarchy()->SetControlShapeTransform((FRigControlElement*)ControlElement, InTransform.ToFTransform(), ERigTransformType::InitialLocal, bSetupUndo, true, bSetupUndo);
+				}
 			}
 		}
 	};
diff --git a/Engine/Plugins/Animation/ControlRig/Source/ControlRigEditor/Private/ControlRigHierarchyCommands.cpp b/Engine/Plugins/Animation/ControlRig/Source/ControlRigEditor/Private/ControlRigHierarchyCommands.cpp
--- a/Engine/Plugins/Animation/ControlRig/Source/ControlRigEditor/Private/ControlRigHierarchyCommands.cpp
+++ b/Engine/Plugins/Animation/ControlRig/Source/ControlRigEditor/Private/ControlRigHierarchyCommands.cpp
@@ -33,7 +33,7 @@ void FControlRigHierarchyCommands::RegisterCommands()
 	UI_COMMAND(Unparent, "Unparent", "Unparents the selected elements from the hierarchy", EUserInterfaceActionType::Button, FInputChord(EKeys::P, EModifierKey::Shift));
 	UI_COMMAND(FilteringFlattensHierarchy, "Filtering Flattens Hierarchy", "Whether to keep the hierarchy or flatten it when searching for tree items", EUserInterfaceActionType::ToggleButton, FInputChord());
 	UI_COMMAND(HideParentsWhenFiltering, "Hide Parents When Filtering", "Whether to show parent items grayed out, or hide them entirely when filtering", EUserInterfaceActionType::ToggleButton, FInputChord());
-	UI_COMMAND(ShowShortNames, "Show Short Names", "Whether to show the full path or the short name of each element", EUserInterfaceActionType::ToggleButton, FInputChord());
+	UI_COMMAND(ShowFullPaths, "Show Full Paths", "Whether to show the full path or the short name of each element", EUserInterfaceActionType::ToggleButton, FInputChord());
 	UI_COMMAND(ShowImportedBones, "Show Imported Bones", "Whether to show or hide imported bones", EUserInterfaceActionType::ToggleButton, FInputChord());
 	UI_COMMAND(ShowBones, "Show Bones", "Whether to show or hide bones", EUserInterfaceActionType::ToggleButton, FInputChord());
 	UI_COMMAND(ShowControls, "Show Controls", "Whether to show or hide controls", EUserInterfaceActionType::ToggleButton, FInputChord());
diff --git a/Engine/Plugins/Animation/ControlRig/Source/ControlRigEditor/Private/ControlRigHierarchyCommands.h b/Engine/Plugins/Animation/ControlRig/Source/ControlRigEditor/Private/ControlRigHierarchyCommands.h
--- a/Engine/Plugins/Animation/ControlRig/Source/ControlRigEditor/Private/ControlRigHierarchyCommands.h
+++ b/Engine/Plugins/Animation/ControlRig/Source/ControlRigEditor/Private/ControlRigHierarchyCommands.h
@@ -99,8 +99,8 @@ public:
 	/* Hide parents on filter */
 	TSharedPtr< FUICommandInfo > HideParentsWhenFiltering;
 
-	/* Show short names */
-	TSharedPtr< FUICommandInfo > ShowShortNames;
+	/* Show full paths */
+	TSharedPtr< FUICommandInfo > ShowFullPaths;
 
 	/* Show imported bones */
 	TSharedPtr< FUICommandInfo > ShowImportedBones;
diff --git a/Engine/Plugins/Animation/ControlRig/Source/ControlRigEditor/Private/ControlRigModuleDetails.cpp b/Engine/Plugins/Animation/ControlRig/Source/ControlRigEditor/Private/ControlRigModuleDetails.cpp
--- a/Engine/Plugins/Animation/ControlRig/Source/ControlRigEditor/Private/ControlRigModuleDetails.cpp
+++ b/Engine/Plugins/Animation/ControlRig/Source/ControlRigEditor/Private/ControlRigModuleDetails.cpp
@@ -129,11 +129,11 @@ void FRigModuleInstanceDetails::CustomizeDetails(IDetailLayoutBuilder& DetailBui
 			{
 				if(const FRigModuleInstance* Module = ModularRig->FindModule(ModuleInstance))
 				{
-					const FString Path = Module->GetPath();
+					const FName ModuleName = Module->Name;
 
 					FPerModuleInfo Info;
-					Info.Path = Path;
-					Info.Module = ModularRig->GetHandle(Path);
+					Info.ModuleName = ModuleName;
+					Info.Module = ModularRig->GetHandle(ModuleName);
 					if(!Info.Module.IsValid())
 					{
 						return;
@@ -143,7 +143,7 @@ void FRigModuleInstanceDetails::CustomizeDetails(IDetailLayoutBuilder& DetailBui
 					{
 						if(const UModularRig* DefaultModularRig = Cast<UModularRig>(Blueprint->GeneratedClass->GetDefaultObject()))
 						{
-							Info.DefaultModule = DefaultModularRig->GetHandle(Path);
+							Info.DefaultModule = DefaultModularRig->GetHandle(ModuleName);
 						}
 					}
 
@@ -183,47 +183,6 @@ void FRigModuleInstanceDetails::CustomizeDetails(IDetailLayoutBuilder& DetailBui
 			.OnVerifyTextChanged(this, &FRigModuleInstanceDetails::OnVerifyNameChanged)
 		];
 
-		static const FText ShortNameTooltip = LOCTEXT("ShortNameTooltip", "The short name is used for the user interface, for example the sequencer channels.\nThis value can be edited and adjusted as needed.");
-		GeneralCategory.AddCustomRow(FText::FromString(TEXT("Short Name")))
-		.NameContent()
-		[
-			SNew(STextBlock)
-			.Text(FText::FromString(TEXT("Short Name")))
-			.Font(IDetailLayoutBuilder::GetDetailFont())
-			.ToolTipText(ShortNameTooltip)
-			.IsEnabled(PerModuleInfos.Num() == 1)
-		]
-		.ValueContent()
-		[
-			SNew(SInlineEditableTextBlock)
-			.Font(IDetailLayoutBuilder::GetDetailFont())
-			.Text(this, &FRigModuleInstanceDetails::GetShortName)
-			.OnTextCommitted(this, &FRigModuleInstanceDetails::SetShortName, DetailBuilder.GetPropertyUtilities())
-			.ToolTipText(ShortNameTooltip)
-			.IsEnabled(PerModuleInfos.Num() == 1)
-			.OnVerifyTextChanged(this, &FRigModuleInstanceDetails::OnVerifyShortNameChanged)
-		];
-
-		static const FText LongNameTooltip = LOCTEXT("LongNameTooltip", "The long name represents a unique address within the rig but isn't used for the user interface.");
-		GeneralCategory.AddCustomRow(FText::FromString(TEXT("Long Name")))
-		.NameContent()
-		[
-			SNew(STextBlock)
-			.Text(FText::FromString(TEXT("Long Name")))
-			.Font(IDetailLayoutBuilder::GetDetailFont())
-			.ToolTipText(LOCTEXT("LongNameTooltip", "The long name represents a unique address within the rig but isn't used for the user interface."))
-			.ToolTipText(LongNameTooltip)
-			.IsEnabled(false)
-		]
-		.ValueContent()
-		[
-			SNew(SInlineEditableTextBlock)
-			.Font(IDetailLayoutBuilder::GetDetailFont())
-			.Text(this, &FRigModuleInstanceDetails::GetLongName)
-			.ToolTipText(LongNameTooltip)
-			.IsEnabled(false)
-		];
-
 		GeneralCategory.AddCustomRow(FText::FromString(TEXT("RigClass")))
 		.NameContent()
 		[
@@ -350,8 +309,8 @@ void FRigModuleInstanceDetails::CustomizeDetails(IDetailLayoutBuilder& DetailBui
 						{
 							if (URigHierarchy* Hierarchy = ModularRig->GetHierarchy())
 							{
-								FString ConnectorPath = FString::Printf(TEXT("%s:%s"), *Info.Path, *Connector.Name);
-								FRigElementKey ConnectorKey(*ConnectorPath, ERigElementType::Connector);
+								const FRigHierarchyModulePath ConnectorModulePath(Info.ModuleName.ToString(), Connector.Name);
+								FRigElementKey ConnectorKey(ConnectorModulePath.GetPathFName(), ERigElementType::Connector);
 								if (FRigConnectorElement* ConnectorElement = Cast<FRigConnectorElement>(Hierarchy->Find(ConnectorKey)))
 								{
 									if(ModuleIndex == 0)
@@ -398,8 +357,8 @@ void FRigModuleInstanceDetails::CustomizeDetails(IDetailLayoutBuilder& DetailBui
 				{
 					for (int32 ModuleIndex=0; ModuleIndex<PerModuleInfos.Num(); ++ModuleIndex)
 					{
-						FString ConnectorPath = FString::Printf(TEXT("%s:%s"), *PerModuleInfos[ModuleIndex].Path, *Connector.Name);
-						FRigElementKey ConnectorKey(*ConnectorPath, ERigElementType::Connector);
+						const FRigHierarchyModulePath ConnectorModulePath(PerModuleInfos[ModuleIndex].ModuleName.ToString(), Connector.Name);
+						FRigElementKey ConnectorKey(ConnectorModulePath.GetPathFName(), ERigElementType::Connector);
 						if(const UControlRigBlueprint* Blueprint = PerModuleInfos[ModuleIndex].GetBlueprint())
 						{
 							const TArray<FRigElementKey> CurrentTargets = Blueprint->ModularRigModel.Connections.FindTargetsFromConnector(ConnectorKey);
@@ -434,8 +393,8 @@ void FRigModuleInstanceDetails::CustomizeDetails(IDetailLayoutBuilder& DetailBui
 					}
 				}
 
-				FString ConnectorPath = FString::Printf(TEXT("%s:%s"), *PerModuleInfos[0].Path, *Connector.Name);
-				FRigElementKey ConnectorKey(*ConnectorPath, ERigElementType::Connector);
+				const FRigHierarchyModulePath ConnectorModulePath(PerModuleInfos[0].ModuleName.ToString(), Connector.Name);
+				FRigElementKey ConnectorKey(ConnectorModulePath.GetPathFName(), ERigElementType::Connector);
 
 				FDetailWidgetRow& ConnectorRow = ConnectionsCategory.AddCustomRow(Label)
 				.NameContent()
@@ -515,8 +474,8 @@ void FRigModuleInstanceDetails::CustomizeDetails(IDetailLayoutBuilder& DetailBui
 							{
 								for (FPerModuleInfo& Info : PerModuleInfos)
 								{
-									FString ConnectorPath = FString::Printf(TEXT("%s:%s"), *Info.Path, *Connector.Name);
-									FRigElementKey ConnectorKey(*ConnectorPath, ERigElementType::Connector);
+									const FRigHierarchyModulePath ConnectorModulePath(Info.ModuleName.ToString(), Connector.Name);
+									FRigElementKey ConnectorKey(ConnectorModulePath.GetPathFName(), ERigElementType::Connector);
 									Info.GetBlueprint()->GetModularRigController()->DisconnectConnector(ConnectorKey);
 								}
 								return FReply::Handled();
@@ -561,8 +520,8 @@ void FRigModuleInstanceDetails::CustomizeDetails(IDetailLayoutBuilder& DetailBui
 									{
 										for (FPerModuleInfo& Info : PerModuleInfos)
 										{
-											FString ConnectorPath = FString::Printf(TEXT("%s:%s"), *Info.Path, *Connector.Name);
-											FRigElementKey ConnectorKey(*ConnectorPath, ERigElementType::Connector);
+											const FRigHierarchyModulePath ConnectorModulePath(Info.ModuleName.ToString(), Connector.Name);
+											FRigElementKey ConnectorKey(ConnectorModulePath.GetPathFName(), ERigElementType::Connector);
 											Info.GetBlueprint()->GetModularRigController()->ConnectConnectorToElements(ConnectorKey, Selected, true, ModularRig->GetModularRigSettings().bAutoResolve);
 										}
 									}
@@ -604,8 +563,8 @@ void FRigModuleInstanceDetails::CustomizeDetails(IDetailLayoutBuilder& DetailBui
 							{
 								if (UModularRig* ModularRig = PerModuleInfos[0].GetModularRig())
 								{
-									FString ConnectorPath = FString::Printf(TEXT("%s:%s"), *PerModuleInfos[0].Path, *Connector.Name);
-									FRigElementKey ConnectorKey(*ConnectorPath, ERigElementType::Connector);
+									const FRigHierarchyModulePath ConnectorModulePath(PerModuleInfos[0].ModuleName.ToString(), Connector.Name);
+									FRigElementKey ConnectorKey(ConnectorModulePath.GetPathFName(), ERigElementType::Connector);
 									if (const FRigElementKeyRedirector::FKeyArray* TargetKeys = ModularRig->GetElementKeyRedirector().FindExternalKey(ConnectorKey))
 									{
 										TArray<FRigElementKey> KeysToSelect;
@@ -794,8 +753,8 @@ void FRigModuleInstanceDetails::SetName(const FText& InValue, ETextCommit::Type
 			if (UControlRigBlueprint* Blueprint = Info.GetBlueprint())
 			{
 				UModularRigController* Controller = Blueprint->GetModularRigController();
-				const FString OldPath = ModuleInstance->GetPath();
-				(void)Controller->RenameModule(OldPath, *InValue.ToString(), true);
+				const FName OldModuleName = ModuleInstance->Name;
+				(void)Controller->RenameModule(OldModuleName, *InValue.ToString(), true);
 			}
 		}
 	}
@@ -817,7 +776,7 @@ bool FRigModuleInstanceDetails::OnVerifyNameChanged(const FText& InText, FText&
 			if (UControlRigBlueprint* Blueprint = Info.GetBlueprint())
 			{
 				UModularRigController* Controller = Blueprint->GetModularRigController();
-				if (!Controller->CanRenameModule(ModuleInstance->GetPath(), *InText.ToString(), OutErrorMessage))
+				if (!Controller->CanRenameModule(ModuleInstance->Name, *InText.ToString(), OutErrorMessage))
 				{
 					return false;
 				}
@@ -828,108 +787,6 @@ bool FRigModuleInstanceDetails::OnVerifyNameChanged(const FText& InText, FText&
 	return true;
 }
 
-FText FRigModuleInstanceDetails::GetShortName() const
-{
-	const FRigModuleInstance* FirstModule = PerModuleInfos[0].GetModule();
-	if(FirstModule == nullptr)
-	{
-		return FText();
-	}
-
-	const FString FirstValue = FirstModule->GetShortName();
-	if(PerModuleInfos.Num() > 1)
-	{
-		bool bSame = true;
-		for (int32 i=1; i<PerModuleInfos.Num(); ++i)
-		{
-			if(const FRigModuleInstance* Module = PerModuleInfos[i].GetModule())
-			{
-				if (!Module->GetShortName().Equals(FirstValue, ESearchCase::IgnoreCase))
-				{
-					bSame = false;
-					break;
-				}
-			}
-		}
-		if (!bSame)
-		{
-			return ControlRigModuleDetailsMultipleValues;
-		}
-	}
-	return FText::FromString(FirstValue);
-}
-
-void FRigModuleInstanceDetails::SetShortName(const FText& InValue, ETextCommit::Type InCommitType, const TSharedRef<IPropertyUtilities> PropertyUtilities)
-{
-	if(InValue.IsEmpty())
-	{
-		return;
-	}
-	
-	check(PerModuleInfos.Num() == 1);
-	const FPerModuleInfo& Info = PerModuleInfos[0];
-	if (const FRigModuleInstance* ModuleInstance = Info.GetModule())
-	{
-		if (UControlRigBlueprint* Blueprint = Info.GetBlueprint())
-		{
-			UModularRigController* Controller = Blueprint->GetModularRigController();
-			Controller->SetModuleShortName(ModuleInstance->GetPath(), *InValue.ToString(), true);
-		}
-	}
-}
-
-bool FRigModuleInstanceDetails::OnVerifyShortNameChanged(const FText& InText, FText& OutErrorMessage)
-{
-	check(PerModuleInfos.Num() == 1);
-
-	if(InText.IsEmpty())
-	{
-		return true;
-	}
-
-	const FPerModuleInfo& Info = PerModuleInfos[0];
-	if (const FRigModuleInstance* ModuleInstance = Info.GetModule())
-	{
-		if (UControlRigBlueprint* Blueprint = Info.GetBlueprint())
-		{
-			UModularRigController* Controller = Blueprint->GetModularRigController();
-			return Controller->CanSetModuleShortName(ModuleInstance->GetPath(), *InText.ToString(), OutErrorMessage);
-		}
-	}
-
-	return false;
-}
-
-FText FRigModuleInstanceDetails::GetLongName() const
-{
-	const FRigModuleInstance* FirstModule = PerModuleInfos[0].GetModule();
-	if(FirstModule == nullptr)
-	{
-		return FText();
-	}
-
-	const FString FirstValue = FirstModule->GetLongName();
-	if(PerModuleInfos.Num() > 1)
-	{
-		bool bSame = true;
-		for (int32 i=1; i<PerModuleInfos.Num(); ++i)
-		{
-			if(const FRigModuleInstance* Module = PerModuleInfos[i].GetModule())
-			{
-				if (!Module->GetLongName().Equals(FirstValue, ESearchCase::IgnoreCase))
-				{
-					bSame = false;
-					break;
-				}
-			}
-		}
-		if (!bSame)
-		{
-			return ControlRigModuleDetailsMultipleValues;
-		}
-	}
-	return FText::FromString(FirstValue);
-}
 FText FRigModuleInstanceDetails::GetRigClassPath() const
 {
 	if(PerModuleInfos.Num() > 1)
@@ -1026,7 +883,7 @@ FRigElementKeyRedirector FRigModuleInstanceDetails::GetConnections() const
 
 void FRigModuleInstanceDetails::OnConfigValueChanged(const FName InVariableName)
 {
-	TMap<FString, FString> ModuleValues;
+	TMap<FName, FString> ModuleValues;
 	for(const FPerModuleInfo& Info : PerModuleInfos)
 	{
 		if (const FRigModuleInstance* ModuleInstance = Info.GetModule())
@@ -1034,12 +891,12 @@ void FRigModuleInstanceDetails::OnConfigValueChanged(const FName InVariableName)
 			if (const UControlRig* ModuleRig = ModuleInstance->GetRig())
 			{
 				FString ValueStr = ModuleRig->GetVariableAsString(InVariableName);
-				ModuleValues.Add(ModuleInstance->GetPath(), ValueStr);
+				ModuleValues.Add(ModuleInstance->Name, ValueStr);
 			}
 		}
 	}
 	
-	for (const TPair<FString, FString>& Value : ModuleValues)
+	for (const TPair<FName, FString>& Value : ModuleValues)
 	{
 		if (UControlRigBlueprint* Blueprint = PerModuleInfos[0].GetBlueprint())
 		{
@@ -1057,8 +914,8 @@ bool FRigModuleInstanceDetails::OnConnectorTargetChanged(TArray<FRigElementKey>
 	{
 		if (UModularRigController* Controller = Info.GetBlueprint()->GetModularRigController())
 		{
-			FString ConnectorPath = FString::Printf(TEXT("%s:%s"), *Info.Path, *InConnector.Name);
-			FRigElementKey ConnectorKey(*ConnectorPath, ERigElementType::Connector);
+			const FRigHierarchyModulePath ConnectorModulePath(Info.ModuleName.ToString(), InConnector.Name);
+			FRigElementKey ConnectorKey(ConnectorModulePath.GetPathFName(), ERigElementType::Connector);
 			
 			if (!InTargets.IsEmpty())
 			{
@@ -1080,13 +937,13 @@ bool FRigModuleInstanceDetails::OnConnectorTargetChanged(TArray<FRigElementKey>
 	return bResult;
 }
 
-const FRigModuleInstanceDetails::FPerModuleInfo& FRigModuleInstanceDetails::FindModule(const FString& InPath) const
+const FRigModuleInstanceDetails::FPerModuleInfo& FRigModuleInstanceDetails::FindModule(const FName& InModuleName) const
 {
-	const FPerModuleInfo* Info = FindModuleByPredicate([InPath](const FPerModuleInfo& Info)
+	const FPerModuleInfo* Info = FindModuleByPredicate([InModuleName](const FPerModuleInfo& Info)
 	{
 		if(const FRigModuleInstance* Module = Info.GetModule())
 		{
-			return Module->GetPath() == InPath;
+			return Module->Name == InModuleName;
 		}
 		return false;
 	});
@@ -1182,7 +1039,7 @@ void FRigModuleInstanceDetails::FillBindingMenu(FMenuBuilder& MenuBuilder, const
 	for(int32 Index = 0; Index < PerModuleInfos.Num(); Index++)
 	{
 		const FPerModuleInfo& Info  = PerModuleInfos[Index];
-		const TArray<FString> Bindings = Controller->GetPossibleBindings(Info.GetPath(), InProperty->GetFName());
+		const TArray<FString> Bindings = Controller->GetPossibleBindings(Info.ModuleName, InProperty->GetFName());
 		if(Index == 0)
 		{
 			CombinedBindings = Bindings;
@@ -1241,7 +1098,7 @@ void FRigModuleInstanceDetails::FillBindingMenu(FMenuBuilder& MenuBuilder, const
 			for(const FString& VariablePath : Data.Variables)
 			{
 				FString VariableName = VariablePath;
-				(void)URigHierarchy::SplitNameSpace(VariablePath, nullptr, &VariableName);
+				(void)FRigHierarchyModulePath(VariablePath).Split(nullptr, &VariableName);
 				
 				InMenuBuilder.AddMenuEntry(
 					FUIAction(FExecuteAction::CreateLambda([ThisDetails, InProperty, VariablePath]()
@@ -1293,11 +1150,11 @@ void FRigModuleInstanceDetails::FillBindingMenu(FMenuBuilder& MenuBuilder, const
 	for(const FString& BindingPath : CombinedBindings)
 	{
 		FString MenuPath;
-		(void)URigHierarchy::SplitNameSpace(BindingPath, &MenuPath, nullptr);
+		(void)FRigHierarchyModulePath(BindingPath).Split(&MenuPath, nullptr);
 
 		FString PreviousMenuPath = MenuPath;
 		FString ParentMenuPath = MenuPath, RemainingPath;
-		while(URigHierarchy::SplitNameSpace(ParentMenuPath, &ParentMenuPath, &RemainingPath))
+		while(FRigHierarchyModulePath(ParentMenuPath).Split(&ParentMenuPath, &RemainingPath))
 		{
 			// scope since the map may change at the end of this block
 			{
@@ -1311,7 +1168,7 @@ void FRigModuleInstanceDetails::FillBindingMenu(FMenuBuilder& MenuBuilder, const
 			PerMenuData->FindOrAdd(ParentMenuPath).SubMenuPaths.AddUnique(PreviousMenuPath);
 			PerMenuData->FindOrAdd(PreviousMenuPath).ParentMenuPath = ParentMenuPath;
 			PerMenuData->FindOrAdd(PreviousMenuPath).Name = RemainingPath;
-			if(!ParentMenuPath.Contains(UModularRig::NamespaceSeparator))
+			if(!ParentMenuPath.Contains(FRigHierarchyModulePath::ModuleNameSuffix))
 			{
 				PerMenuData->FindOrAdd(FString()).SubMenuPaths.AddUnique(ParentMenuPath);
 				PerMenuData->FindOrAdd(ParentMenuPath).Name = ParentMenuPath;
@@ -1363,7 +1220,7 @@ void FRigModuleInstanceDetails::HandleRemoveBinding(FName InPropertyName) const
 			if (const FRigModuleInstance* ModuleInstance = Info.GetModule())
 			{
 				UModularRigController* Controller = Blueprint->GetModularRigController();
-				Controller->UnBindModuleVariable(ModuleInstance->GetPath(), InPropertyName);
+				Controller->UnBindModuleVariable(ModuleInstance->Name, InPropertyName);
 			}
 		}
 	}
@@ -1379,7 +1236,7 @@ void FRigModuleInstanceDetails::HandleChangeBinding(const FProperty* InProperty,
 			if (const FRigModuleInstance* ModuleInstance = Info.GetModule())
 			{
 				UModularRigController* Controller = Blueprint->GetModularRigController();
-				Controller->BindModuleVariable(ModuleInstance->GetPath(), InProperty->GetFName(), InNewVariablePath);
+				Controller->BindModuleVariable(ModuleInstance->Name, InProperty->GetFName(), InNewVariablePath);
 			}
 		}
 	}
@@ -1395,8 +1252,8 @@ FReply FRigModuleInstanceDetails::OnAddTargetToArrayConnector(const FString InCo
 			if (const FRigModuleInstance* ModuleInstance = Info.GetModule())
 			{
 				UModularRigController* Controller = Blueprint->GetModularRigController();
-				FString ConnectorPath = FString::Printf(TEXT("%s:%s"), *Info.Path, *InConnectorName);
-				FRigElementKey ConnectorKey(*ConnectorPath, ERigElementType::Connector);
+				const FRigHierarchyModulePath ConnectorModulePath(Info.ModuleName.ToString(), InConnectorName);
+				FRigElementKey ConnectorKey(ConnectorModulePath.GetPathFName(), ERigElementType::Connector);
 				Controller->AddTargetToArrayConnector(ConnectorKey, FRigElementKey(), true, false, false);
 				Blueprint->RecompileModularRig();
 			}
@@ -1416,8 +1273,8 @@ FReply FRigModuleInstanceDetails::OnClearTargetsForArrayConnector(const FString
 			if (const FRigModuleInstance* ModuleInstance = Info.GetModule())
 			{
 				UModularRigController* Controller = Blueprint->GetModularRigController();
-				FString ConnectorPath = FString::Printf(TEXT("%s:%s"), *Info.Path, *InConnectorName);
-				FRigElementKey ConnectorKey(*ConnectorPath, ERigElementType::Connector);
+				const FRigHierarchyModulePath ConnectorModulePath(Info.ModuleName.ToString(), InConnectorName);
+				FRigElementKey ConnectorKey(ConnectorModulePath.GetPathFName(), ERigElementType::Connector);
 				Controller->DisconnectConnector(ConnectorKey);
 			}
 		}
diff --git a/Engine/Plugins/Animation/ControlRig/Source/ControlRigEditor/Private/ControlRigModuleDetails.h b/Engine/Plugins/Animation/ControlRig/Source/ControlRigEditor/Private/ControlRigModuleDetails.h
--- a/Engine/Plugins/Animation/ControlRig/Source/ControlRigEditor/Private/ControlRigModuleDetails.h
+++ b/Engine/Plugins/Animation/ControlRig/Source/ControlRigEditor/Private/ControlRigModuleDetails.h
@@ -43,10 +43,6 @@ public:
 	FText GetName() const;
 	void SetName(const FText& InValue, ETextCommit::Type InCommitType, const TSharedRef<IPropertyUtilities> PropertyUtilities);
 	bool OnVerifyNameChanged(const FText& InText, FText& OutErrorMessage);
-	FText GetShortName() const;
-	void SetShortName(const FText& InValue, ETextCommit::Type InCommitType, const TSharedRef<IPropertyUtilities> PropertyUtilities);
-	bool OnVerifyShortNameChanged(const FText& InText, FText& OutErrorMessage);
-	FText GetLongName() const;
 	FText GetRigClassPath() const;
 	TArray<FRigModuleConnector> GetConnectors() const;
 	FRigElementKeyRedirector GetConnections() const;
@@ -57,7 +53,7 @@ public:
 	struct FPerModuleInfo
 	{
 		FPerModuleInfo()
-			: Path()
+			: ModuleName(NAME_None)
 			, Module()
 			, DefaultModule()
 		{}
@@ -65,9 +61,9 @@ public:
 		bool IsValid() const { return Module.IsValid(); }
 		operator bool() const { return IsValid(); }
 
-		const FString& GetPath() const
+		const FName& GetModuleName() const
 		{
-			return Path;
+			return ModuleName;
 		}
 
 		UModularRig* GetModularRig() const
@@ -111,17 +107,17 @@ public:
 		{
 			if(const UControlRigBlueprint* Blueprint = GetBlueprint())
 			{
-				return Blueprint->ModularRigModel.FindModule(Path);
+				return Blueprint->ModularRigModel.FindModule(ModuleName);
 			}
 			return nullptr;
 		}
 
-		FString Path;
+		FName ModuleName;
 		FModuleInstanceHandle Module;
 		FModuleInstanceHandle DefaultModule;
 	};
 
-	const FPerModuleInfo& FindModule(const FString& InKey) const;
+	const FPerModuleInfo& FindModule(const FName& InModuleName) const;
 	const FPerModuleInfo* FindModuleByPredicate(const TFunction<bool(const FPerModuleInfo&)>& InPredicate) const;
 	bool ContainsModuleByPredicate(const TFunction<bool(const FPerModuleInfo&)>& InPredicate) const;
 
