Control Rig: Unique modular rig module names

diff --git a/Engine/Plugins/Animation/ControlRig/Source/ControlRigEditor/Private/Editor/ControlRigEditor.cpp b/Engine/Plugins/Animation/ControlRig/Source/ControlRigEditor/Private/Editor/ControlRigEditor.cpp
--- a/Engine/Plugins/Animation/ControlRig/Source/ControlRigEditor/Private/Editor/ControlRigEditor.cpp
+++ b/Engine/Plugins/Animation/ControlRig/Source/ControlRigEditor/Private/Editor/ControlRigEditor.cpp
@@ -506,9 +506,9 @@ FReply FControlRigEditor::OnViewportDrop(const FGeometry& MyGeometry, const FDra
 						UModularRigController* Controller = GetControlRigBlueprint()->GetModularRigController();
 						FString ClassName = ControlRigClass->GetName();
 						ClassName.RemoveFromEnd(TEXT("_C"));
-						const FRigName Name = Controller->GetSafeNewName(FString(), FRigName(ClassName));
-						const FString ModulePath = Controller->AddModule(Name, ControlRigClass, FString());
-						if (!ModulePath.IsEmpty())
+						const FRigName DesiredModuleName = Controller->GetSafeNewName(FRigName(ClassName));
+						const FName ModuleName = Controller->AddModule(DesiredModuleName, ControlRigClass, NAME_None);
+						if (!ModuleName.IsNone())
 						{
 							return FReply::Handled();
 						}
@@ -1668,7 +1668,7 @@ void FControlRigEditor::SetDetailViewForRigModules()
 	SetDetailViewForRigModules(ModulesSelected);
 }
 
-void FControlRigEditor::SetDetailViewForRigModules(const TArray<FString> InKeys)
+void FControlRigEditor::SetDetailViewForRigModules(const TArray<FName> InModuleNames)
 {
 	if(IsDetailsPanelRefreshSuspended())
 	{
@@ -1685,12 +1685,12 @@ void FControlRigEditor::SetDetailViewForRigModules(const TArray<FString> InKeys)
 		return;
 	}
 
-	ModulesSelected = InKeys;
+	ModulesSelected = InModuleNames;
 	TArray<UObject*> Objects;
 
-	for(const FString& Key : InKeys)
+	for(const FName& ModuleName : InModuleNames)
 	{
-		const FRigModuleInstance* Element = RigBeingDebugged->FindModule(Key);
+		const FRigModuleInstance* Element = RigBeingDebugged->FindModule(ModuleName);
 		if (Element == nullptr)
 		{
 			continue;
@@ -1707,7 +1707,7 @@ void FControlRigEditor::SetDetailViewForRigModules(const TArray<FString> InKeys)
 	// In case the modules selected are still not available, lets set them again
 	if (Objects.IsEmpty())
 	{
-		ModulesSelected = InKeys;
+		ModulesSelected = InModuleNames;
 	}
 }
 
@@ -1716,7 +1716,7 @@ bool FControlRigEditor::DetailViewShowsAnyRigModule() const
 	return DetailViewShowsStruct(FRigModuleInstance::StaticStruct());
 }
 
-bool FControlRigEditor::DetailViewShowsRigModule(FString InKey) const
+bool FControlRigEditor::DetailViewShowsRigModule(FName InModuleName) const
 {
 	TArray< TWeakObjectPtr<UObject> > SelectedObjects = GetSelectedObjectsFromDetailView();
 	for (TWeakObjectPtr<UObject> SelectedObject : SelectedObjects)
@@ -1729,7 +1729,7 @@ bool FControlRigEditor::DetailViewShowsRigModule(FString InKey) const
 				{
 					if (WrappedStruct->IsChildOf(FRigModuleInstance::StaticStruct()))
 					{
-						if(WrapperObject->GetContent<FRigModuleInstance>().GetPath() == InKey)
+						if(WrapperObject->GetContent<FRigModuleInstance>().Name == InModuleName)
 						{
 							return true;
 						}
@@ -3313,7 +3313,7 @@ void FControlRigEditor::HandlePreviewMeshChanged(USkeletalMesh* InOldSkeletalMes
 						for(const FRigElementKey& Key : ConnectorsAndSockets)
 						{
 							// Remember manually created elements to apply them again
-							if (BPHierarchy->GetNameSpace(Key).IsEmpty())
+							if (BPHierarchy->GetModuleFName(Key).IsNone())
 							{
 								const FRigElementKey& Parent = BPHierarchy->GetDefaultParent(Key);
 								ConnectorsAndSocketsToParents.Emplace(Key, Parent, BPHierarchy->GetLocalTransform(Key));
@@ -3560,7 +3560,7 @@ void FControlRigEditor::FilterDraggedKeys(TArray<FRigElementKey>& Keys, bool bRe
 			{
 				const FString Name = Key.Name.ToString();
 				int32 LastCharIndex = INDEX_NONE;
-				if(Name.FindLastChar(TEXT(':'), LastCharIndex))
+				if(Name.FindLastChar(FRigHierarchyModulePath::ModuleNameSuffixChar, LastCharIndex))
 				{
 					Key.Name = *Name.Mid(LastCharIndex+1);
 				}
@@ -3571,11 +3571,10 @@ void FControlRigEditor::FilterDraggedKeys(TArray<FRigElementKey>& Keys, bool bRe
 				{
 					if(!DebuggedControlRig->GetHierarchy()->Contains(Key))
 					{
-						const FString NameSpace = DebuggedControlRig->GetRigModuleNameSpace();
-						if(!NameSpace.IsEmpty())
+						const FString ModulePrefix = DebuggedControlRig->GetRigModulePrefix();
+						if(!ModulePrefix.IsEmpty())
 						{
-							static constexpr TCHAR Format[] = TEXT("%s%s");
-							Key.Name = *FString::Printf(Format, *NameSpace, *Key.Name.ToString());
+							Key.Name = *(ModulePrefix + Key.Name.ToString());
 						}
 					}
 				}
@@ -4637,7 +4636,7 @@ void FControlRigEditor::HandleModularRigModified(EModularRigNotification InNotif
 	{
 		case EModularRigNotification::ModuleAdded:
 		{
-			ModularRigController->SelectModule(InModule->GetPath());
+			ModularRigController->SelectModule(InModule->Name);
 			break;
 		}
 		case EModularRigNotification::ModuleRemoved:
@@ -4653,15 +4652,6 @@ void FControlRigEditor::HandleModularRigModified(EModularRigNotification InNotif
 		case EModularRigNotification::ModuleReparented:
 		case EModularRigNotification::ModuleRenamed:
 		{
-			FString OldPath;
-			if (InNotification == EModularRigNotification::ModuleRenamed)
-			{
-				OldPath = URigHierarchy::JoinNameSpace(InModule->ParentPath, InModule->PreviousName.ToString());
-			}
-			else
-			{
-				OldPath = URigHierarchy::JoinNameSpace(InModule->PreviousParentPath, InModule->Name.ToString());
-			}
 			break;
 		}
 		case EModularRigNotification::ConnectionChanged:
@@ -4729,7 +4719,7 @@ void FControlRigEditor::SwapModuleWithinAsset()
 					for (const FRigModuleReference* Module : Modules)
 					{
 						FSoftObjectPath ModulePath = Asset.GetSoftObjectPath();
-						ModulePath.SetSubPathString(Module->GetPath());
+						ModulePath.SetSubPathString(Module->GetModulePath().GetPath());
 						Result.Add(ModulePath);
 					}
 				}
@@ -4748,7 +4738,10 @@ void FControlRigEditor::SwapModuleWithinAsset()
 			{
 				if (UControlRigBlueprint* RigBlueprint = Cast<UControlRigBlueprint>(ModulePath.GetWithoutSubPath().ResolveObject()))
 				{
-					return RigBlueprint->GetModularRigController()->SwapModuleClass(ModulePath.GetSubPathString(), NewModuleClass);
+					if(const FRigModuleReference* Module = RigBlueprint->ModularRigModel.FindModuleByPath(ModulePath.GetSubPathString()))
+					{
+						return RigBlueprint->GetModularRigController()->SwapModuleClass(Module->Name, NewModuleClass);
+					}
 				}
 			}
 			return false;
@@ -4795,7 +4788,7 @@ void FControlRigEditor::SwapModuleAcrossProject()
 			{
 				if (UControlRigBlueprint* RigBlueprint = Cast<UControlRigBlueprint>(ModulePath.GetWithoutSubPath().ResolveObject()))
 				{
-					return RigBlueprint->GetModularRigController()->SwapModuleClass(ModulePath.GetSubPathString(), NewModuleClass);
+					return RigBlueprint->GetModularRigController()->SwapModuleClass(*ModulePath.GetSubPathString(), NewModuleClass);
 				}
 			}
 			return false;
diff --git a/Engine/Plugins/Animation/ControlRig/Source/ControlRigEditor/Private/Editor/ControlRigEditor.h b/Engine/Plugins/Animation/ControlRig/Source/ControlRigEditor/Private/Editor/ControlRigEditor.h
--- a/Engine/Plugins/Animation/ControlRig/Source/ControlRigEditor/Private/Editor/ControlRigEditor.h
+++ b/Engine/Plugins/Animation/ControlRig/Source/ControlRigEditor/Private/Editor/ControlRigEditor.h
@@ -142,10 +142,10 @@ public:
 	TArray< TWeakObjectPtr<UObject> > GetSelectedObjectsFromDetailView() const;
 
 	void SetDetailViewForRigModules();
-	void SetDetailViewForRigModules(const TArray<FString> InKeys);
+	void SetDetailViewForRigModules(const TArray<FName> InModuleNames);
 	bool DetailViewShowsAnyRigModule() const;
-	bool DetailViewShowsRigModule(FString InKey) const;
-	TArray<FString> ModulesSelected;
+	bool DetailViewShowsRigModule(FName InModuleName) const;
+	TArray<FName> ModulesSelected;
 
 	virtual void SetDetailObjects(const TArray<UObject*>& InObjects) override;
 	virtual void RefreshDetailView() override;
diff --git a/Engine/Plugins/Animation/ControlRig/Source/ControlRigEditor/Private/Editor/ControlRigSchematicModel.cpp b/Engine/Plugins/Animation/ControlRig/Source/ControlRigEditor/Private/Editor/ControlRigSchematicModel.cpp
--- a/Engine/Plugins/Animation/ControlRig/Source/ControlRigEditor/Private/Editor/ControlRigSchematicModel.cpp
+++ b/Engine/Plugins/Animation/ControlRig/Source/ControlRigEditor/Private/Editor/ControlRigSchematicModel.cpp
@@ -617,10 +617,10 @@ void FControlRigSchematicModel::HandleModularRigModified(EModularRigNotification
 						if(const FSchematicGraphNode* ConnectorNode = FindElementKeyNode(ConnectorKey))
 						{
 							const FString NameString = ConnectorKey.Name.ToString();
-							if(NameString.StartsWith(InModule->GetNamespace(), ESearchCase::IgnoreCase))
+							if(NameString.StartsWith(InModule->GetElementPrefix(), ESearchCase::IgnoreCase))
 							{
-								const FString LocalName = NameString.Mid(InModule->GetNamespace().Len());
-								if(!LocalName.Contains(UModularRig::NamespaceSeparator))
+								const FString LocalName = NameString.Mid(InModule->GetElementPrefix().Len());
+								if(!LocalName.Contains(FRigHierarchyModulePath::ModuleNameSuffix))
 								{
 									UpdateConnector(ConnectorKey);
 								}
@@ -1229,7 +1229,7 @@ TArray<FRigElementKey> FControlRigSchematicModel::GetElementKeysFromDragDropEven
 	if(InDragDropOperation.IsOfType<FModularRigModuleDragDropOp>())
 	{
 		const FModularRigModuleDragDropOp* ModuleDropOp = StaticCast<const FModularRigModuleDragDropOp*>(&InDragDropOperation);
-		const TArray<FString> Sources = ModuleDropOp->GetElements();
+		const TArray<FName> ModuleNames = ModuleDropOp->GetModules();
 
 		URigHierarchy* Hierarchy = InControlRig->GetHierarchy();
 		if (!Hierarchy)
@@ -1238,9 +1238,9 @@ TArray<FRigElementKey> FControlRigSchematicModel::GetElementKeysFromDragDropEven
 		}
 
 		const TArray<FRigConnectorElement*> Connectors = Hierarchy->GetConnectors();
-		for (const FString& ModulePathOrConnectorName : Sources)
+		for (const FName& ModulePathOrConnectorName : ModuleNames)
 		{
-			if(const FRigConnectorElement* Connector = Hierarchy->Find<FRigConnectorElement>(FRigElementKey(*ModulePathOrConnectorName, ERigElementType::Connector)))
+			if(const FRigConnectorElement* Connector = Hierarchy->Find<FRigConnectorElement>(FRigElementKey(ModulePathOrConnectorName, ERigElementType::Connector)))
 			{
 				DraggedKeys.Add(Connector->GetKey());
 			}
@@ -1249,11 +1249,10 @@ TArray<FRigElementKey> FControlRigSchematicModel::GetElementKeysFromDragDropEven
 			{
 				if(const FRigModuleInstance* ModuleInstance = ModularRig->FindModule(ModulePathOrConnectorName))
 				{
-					const FString ModuleNameSpace = ModuleInstance->GetNamespace();
 					for(const FRigConnectorElement* Connector : Connectors)
 					{
-						const FString ConnectorNameSpace = Hierarchy->GetNameSpace(Connector->GetKey());
-						if(ConnectorNameSpace.Equals(ModuleNameSpace))
+						const FName ConnectorModuleName = Hierarchy->GetModuleFName(Connector->GetKey());
+						if(ModuleInstance->Name == ConnectorModuleName)
 						{
 							if(Connector->IsPrimary())
 							{
@@ -1556,19 +1555,19 @@ void FControlRigSchematicModel::HandleSchematicDrop(SSchematicGraphPanel* InPane
 
 						for(const FRigElementKey& TargetKey : TargetKeys)
 						{
-							const FName ModuleName = Controller->GetSafeNewName(FString(), FRigName(AssetBlueprint->RigModuleSettings.Identifier.Name));
-							const FString ModulePath = Controller->AddModule(ModuleName, AssetBlueprint->GetControlRigClass(), FString());
-							if(!ModulePath.IsEmpty())
+							const FName DesiredModuleName = Controller->GetSafeNewName(FRigName(AssetBlueprint->RigModuleSettings.Identifier.Name));
+							const FName ModuleName = Controller->AddModule(DesiredModuleName, AssetBlueprint->GetControlRigClass(), NAME_None);
+							if(!ModuleName.IsNone())
 							{
+								const FString ModuleNameString = ModuleName.ToString();
 								FRigElementKey PrimaryConnectorKey;
 								TArray<FRigConnectorElement*> Connectors = Hierarchy->GetElementsOfType<FRigConnectorElement>();
 								for (FRigConnectorElement* Connector : Connectors)
 								{
 									if (Connector->IsPrimary())
 									{
-										FString Path, Name;
-										(void)URigHierarchy::SplitNameSpace(Connector->GetName(), &Path, &Name);
-										if (Path == ModulePath)
+										const FRigHierarchyModulePath ConnectorModulePath(Connector->GetName());
+										if (ConnectorModulePath.HasModuleName(ModuleNameString))
 										{
 											PrimaryConnectorKey = Connector->GetKey();
 											break;
@@ -1628,8 +1627,8 @@ void FControlRigSchematicModel::HandleSchematicDrop(SSchematicGraphPanel* InPane
 	}
 	else if(ModuleDragDropOperation.IsValid())
 	{
-		const TArray<FString> Sources = ModuleDragDropOperation->GetElements();
-		FFunctionGraphTask::CreateAndDispatchWhenReady([this, Sources, TargetKeys]()
+		const TArray<FName> ModuleNames = ModuleDragDropOperation->GetModules();
+		FFunctionGraphTask::CreateAndDispatchWhenReady([this, ModuleNames, TargetKeys]()
 		{
 			UModularRig* ControlRig = Cast<UModularRig>(ControlRigBlueprint->GetDebuggedControlRig());
 			if (!ControlRig)
@@ -1646,9 +1645,9 @@ void FControlRigSchematicModel::HandleSchematicDrop(SSchematicGraphPanel* InPane
 			if (UModularRigController* Controller = ControlRigBlueprint->GetModularRigController())
 			{
 				const TArray<FRigConnectorElement*> Connectors = Hierarchy->GetConnectors();
-				for (const FString& ModulePathOrConnectorName : Sources)
+				for (const FName& ModulePathOrConnectorName : ModuleNames)
 				{
-					if(const FRigConnectorElement* Connector = Hierarchy->Find<FRigConnectorElement>(FRigElementKey(*ModulePathOrConnectorName, ERigElementType::Connector)))
+					if(const FRigConnectorElement* Connector = Hierarchy->Find<FRigConnectorElement>(FRigElementKey(ModulePathOrConnectorName, ERigElementType::Connector)))
 					{
 						for(const FRigElementKey& TargetKey : TargetKeys)
 						{
@@ -1662,11 +1661,10 @@ void FControlRigSchematicModel::HandleSchematicDrop(SSchematicGraphPanel* InPane
 					}
 					if(const FRigModuleReference* Module = Controller->FindModule(ModulePathOrConnectorName))
 					{
-						const FString ModuleNameSpace = Module->GetNamespace();
 						for(const FRigConnectorElement* Connector : Connectors)
 						{
-							const FString ConnectorNameSpace = Hierarchy->GetNameSpace(Connector->GetKey());
-							if(ConnectorNameSpace.Equals(ModuleNameSpace))
+							const FName ConnectorModuleName = Hierarchy->GetModuleFName(Connector->GetKey());
+							if(Module->Name == ConnectorModuleName)
 							{
 								if(Connector->IsPrimary())
 								{
@@ -1755,8 +1753,8 @@ void FControlRigSchematicModel::OnShowCandidatesForConnector(const FRigElementKe
 		return;
 	}
 
-	const FString ModulePath = Hierarchy->GetModulePath(InConnectorKey);
-	const FRigModuleInstance* ModuleInstance = ModularRig->FindModule(ModulePath);
+	const FName ModuleName = Hierarchy->GetModuleFName(InConnectorKey);
+	const FRigModuleInstance* ModuleInstance = ModularRig->FindModule(ModuleName);
 	if (!ModuleInstance)
 	{
 		return;
diff --git a/Engine/Plugins/Animation/ControlRig/Source/ControlRigEditor/Private/Editor/SModularRigModel.cpp b/Engine/Plugins/Animation/ControlRig/Source/ControlRigEditor/Private/Editor/SModularRigModel.cpp
--- a/Engine/Plugins/Animation/ControlRig/Source/ControlRigEditor/Private/Editor/SModularRigModel.cpp
+++ b/Engine/Plugins/Animation/ControlRig/Source/ControlRigEditor/Private/Editor/SModularRigModel.cpp
@@ -333,20 +333,20 @@ FReply SModularRigModel::OnMouseButtonDown(const FGeometry& MyGeometry, const FP
 					UModularRigController* Controller = ControlRigBlueprint->GetModularRigController();
 					check(Controller);
 
-					if(const FRigModuleReference* Module = Controller->FindModule(Item->ModulePath))
+					if(const FRigModuleReference* Module = Controller->FindModule(Item->ModuleName))
 					{
 						TArray<const FRigModuleReference*> ModulesToSelect = {Module};
-						TArray<FString> ModulePaths;
+						TArray<FName> ModuleNames;
 						for(int32 Index = 0; Index < ModulesToSelect.Num(); Index++)
 						{
-							ModulePaths.Add(ModulesToSelect[Index]->GetPath());
+							ModuleNames.Add(ModulesToSelect[Index]->Name);
 							for(const FRigModuleReference* ChildModule : ModulesToSelect[Index]->CachedChildren)
 							{
 								ModulesToSelect.AddUnique(ChildModule);
 							}
 						}
 						
-						Controller->SetModuleSelection(ModulePaths);
+						Controller->SetModuleSelection(ModuleNames);
 					}
 				}
 			}
@@ -406,9 +406,9 @@ void SModularRigModel::HandlePostCompileModularRigs(URigVMBlueprint* InBlueprint
 	if (ControlRigEditor.IsValid())
 	{
 		TArray<TSharedPtr<FModularRigTreeElement>> SelectedElements;
-		Algo::Transform(ControlRigEditor.Pin()->ModulesSelected, SelectedElements, [this](const FString& Path)
+		Algo::Transform(ControlRigEditor.Pin()->ModulesSelected, SelectedElements, [this](const FName& ModuleName)
 		{
-			return TreeView->FindElement(Path);
+			return TreeView->FindElement(ModuleName.ToString());
 		});
 		TreeView->SetSelection(SelectedElements);
 		ControlRigEditor.Pin()->RefreshDetailView();
@@ -703,13 +703,13 @@ void SModularRigModel::HandleNewItem()
 		return;
 	}
 
-	FString ParentPath;
+	FName ParentModuleName = NAME_None;
 	if (IsSingleSelected())
 	{
 		const TSharedPtr<FModularRigTreeElement> ParentElement = TreeView->FindElement(GetSelectedKeys()[0]);
 		if (ParentElement.IsValid())
 		{
-			ParentPath = ParentElement->ModulePath;
+			ParentModuleName = ParentElement->ModuleName;
 		}
 	}
 	
@@ -726,11 +726,11 @@ void SModularRigModel::HandleNewItem()
 	const bool bPressedOk = SClassPickerDialog::PickClass(TitleText, Options, ChosenClass, UControlRig::StaticClass());
 	if (bPressedOk)
 	{
-		HandleNewItem(ChosenClass, ParentPath);
+		HandleNewItem(ChosenClass, ParentModuleName);
 	}
 }
 
-void SModularRigModel::HandleNewItem(UClass* InClass, const FString &InParentPath)
+void SModularRigModel::HandleNewItem(UClass* InClass, const FName &InParentModuleName)
 {
 	UControlRig* ControlRig = InClass->GetDefaultObject<UControlRig>();
 	if (!ControlRig)
@@ -746,9 +746,9 @@ void SModularRigModel::HandleNewItem(UClass* InClass, const FString &InParentPat
 		
 		FString ClassName = InClass->GetName();
 		ClassName.RemoveFromEnd(TEXT("_C"));
-		const FRigName Name = Controller->GetSafeNewName(InParentPath, FRigName(ClassName));
-		const FString NewPath = Controller->AddModule(Name, InClass, InParentPath);
-		TSharedPtr<FModularRigTreeElement> Element = TreeView->FindElement(NewPath);
+		const FRigName Name = Controller->GetSafeNewName(FRigName(ClassName));
+		const FName NewModuleName = Controller->AddModule(Name, InClass, InParentModuleName);
+		TSharedPtr<FModularRigTreeElement> Element = TreeView->FindElement(NewModuleName.ToString());
 		if (Element.IsValid())
 		{
 			TreeView->SetSelection({Element});
@@ -789,7 +789,7 @@ void SModularRigModel::HandleRenameModule()
 	return;
 }
 
-FName SModularRigModel::HandleRenameModule(const FString& InOldPath, const FName& InNewName)
+FName SModularRigModel::HandleRenameModule(const FName& InOldModuleName, const FName& InNewName)
 {
 	if (ControlRigBlueprint.IsValid())
 	{
@@ -798,7 +798,8 @@ FName SModularRigModel::HandleRenameModule(const FString& InOldPath, const FName
 		UModularRigController* Controller = ControlRigBlueprint->GetModularRigController();
 		check(Controller);
 
-		if (!Controller->RenameModule(InOldPath, InNewName).IsEmpty())
+		const FName NewModuleName = Controller->RenameModule(InOldModuleName, InNewName);
+		if (!NewModuleName.IsNone())
 		{
 			return InNewName;
 		}
@@ -807,28 +808,19 @@ FName SModularRigModel::HandleRenameModule(const FString& InOldPath, const FName
 	return NAME_None;
 }
 
-bool SModularRigModel::HandleVerifyNameChanged(const FString& InOldPath, const FName& InNewName, FText& OutErrorMessage)
+bool SModularRigModel::HandleVerifyNameChanged(const FName& InOldModuleName, const FName& InNewName, FText& OutErrorMessage)
 {
 	if (InNewName.IsNone())
 	{
 		return false;
 	}
 
-	FString ParentPath;
-	FString OldName = InOldPath;
-	(void)URigHierarchy::SplitNameSpace(InOldPath, &ParentPath, &OldName);
-
-	if (InNewName == OldName)
-	{
-		return true;
-	}
-	
 	if (ControlRigBlueprint.IsValid())
 	{
 		UModularRigController* Controller = ControlRigBlueprint->GetModularRigController();
 		check(Controller);
 
-		return Controller->CanRenameModule(InOldPath, InNewName, OutErrorMessage);
+		return Controller->CanRenameModule(InOldModuleName, InNewName, OutErrorMessage);
 	}
 
 	return false;
@@ -847,22 +839,22 @@ void SModularRigModel::HandleDeleteModules()
 		FScopedTransaction Transaction(LOCTEXT("ModularRigModelDeleteSelected", "Delete selected modules"));
 
 		const TArray<TSharedPtr<FModularRigTreeElement>> SelectedItems = GetSelectedItems();
-		TArray<FString> SelectedPaths;
-		Algo::Transform(SelectedItems, SelectedPaths, [](const TSharedPtr<FModularRigTreeElement>& Element)
+		TArray<FName> SelectedModuleNames;
+		Algo::Transform(SelectedItems, SelectedModuleNames, [](const TSharedPtr<FModularRigTreeElement>& Element)
 		{
 			if (Element.IsValid())
 			{
-				return Element->ModulePath;
+				return Element->ModuleName;
 			}
-			return FString();
+			return FName(NAME_None);
 		});
-		HandleDeleteModules(SelectedPaths);
+		HandleDeleteModules(SelectedModuleNames);
 	}
 
 	return;
 }
 
-void SModularRigModel::HandleDeleteModules(const TArray<FString>& InPaths)
+void SModularRigModel::HandleDeleteModules(const TArray<FName>& InModuleNames)
 {
 	if (ControlRigBlueprint.IsValid())
 	{
@@ -872,16 +864,16 @@ void SModularRigModel::HandleDeleteModules(const TArray<FString>& InPaths)
 		check(Controller);
 
 		// Make sure we delete the modules from children to root
-		TArray<FString> SortedPaths = Controller->Model->SortPaths(InPaths);
-		Algo::Reverse(SortedPaths);
-		for (const FString& Path : SortedPaths)
+		TArray<FName> SortedModuleNames = Controller->Model->SortModuleNames(InModuleNames);
+		Algo::Reverse(SortedModuleNames);
+		for (const FName& ModuleName : SortedModuleNames)
 		{
-			Controller->DeleteModule(Path);
+			Controller->DeleteModule(ModuleName);
 		}
 	}
 }
 
-void SModularRigModel::HandleReparentModules(const TArray<FString>& InPaths, const FString& InParentPath)
+void SModularRigModel::HandleReparentModules(const TArray<FName>& InModuleNames, const FName& InParentModuleName)
 {
 	if (ControlRigBlueprint.IsValid())
 	{
@@ -890,9 +882,9 @@ void SModularRigModel::HandleReparentModules(const TArray<FString>& InPaths, con
 		UModularRigController* Controller = ControlRigBlueprint->GetModularRigController();
 		check(Controller);
 
-		for (const FString& Path : InPaths)
+		for (const FName& ModuleName : InModuleNames)
 		{
-			Controller->ReparentModule(Path, InParentPath);
+			Controller->ReparentModule(ModuleName, InParentModuleName);
 		}
 	}
 }
@@ -908,20 +900,20 @@ void SModularRigModel::HandleMirrorModules()
 	if (Rig)
 	{
 		const TArray<TSharedPtr<FModularRigTreeElement>> SelectedItems = GetSelectedItems();
-		TArray<FString> SelectedPaths;
-		Algo::Transform(SelectedItems, SelectedPaths, [](const TSharedPtr<FModularRigTreeElement>& Element)
+		TArray<FName> SelectedModuleNames;
+		Algo::Transform(SelectedItems, SelectedModuleNames, [](const TSharedPtr<FModularRigTreeElement>& Element)
 		{
 			if (Element.IsValid())
 			{
-				return Element->ModulePath;
+				return Element->ModuleName;
 			}
-			return FString();
+			return FName(NAME_None);
 		});
-		HandleMirrorModules(SelectedPaths);
+		HandleMirrorModules(SelectedModuleNames);
 	}
 }
 
-void SModularRigModel::HandleMirrorModules(const TArray<FString>& InPaths)
+void SModularRigModel::HandleMirrorModules(const TArray<FName>& InModuleNames)
 {
 	if (ControlRigBlueprint.IsValid())
 	{
@@ -954,10 +946,10 @@ void SModularRigModel::HandleMirrorModules(const TArray<FString>& InPaths)
 			check(Controller);
 
 			// Make sure we mirror the modules from root to children
-			TArray<FString> SortedPaths = Controller->Model->SortPaths(InPaths);
-			for (const FString& Path : SortedPaths)
+			TArray<FName> SortedModuleNames = Controller->Model->SortModuleNames(InModuleNames);
+			for (const FName& ModuleName : SortedModuleNames)
 			{
-				Controller->MirrorModule(Path, Settings);
+				Controller->MirrorModule(ModuleName, Settings);
 			}
 		}
 	}
@@ -974,24 +966,24 @@ void SModularRigModel::HandleReresolveModules()
 	if (Rig)
 	{
 		const TArray<TSharedPtr<FModularRigTreeElement>> SelectedItems = GetSelectedItems();
-		TArray<FString> SelectedPaths;
-		Algo::Transform(SelectedItems, SelectedPaths, [](const TSharedPtr<FModularRigTreeElement>& Element)
+		TArray<FName> SelectedModuleNames;
+		Algo::Transform(SelectedItems, SelectedModuleNames, [](const TSharedPtr<FModularRigTreeElement>& Element)
 		{
 			if (Element.IsValid())
 			{
 				if(Element->ConnectorName.IsEmpty())
 				{
-					return Element->ModulePath;
+					return Element->ModuleName;
 				}
-				return FString::Printf(TEXT("%s|%s"), *Element->ModulePath, *Element->ConnectorName);
+				return FRigHierarchyModulePath(Element->ModuleName.ToString(), Element->ConnectorName).GetPathFName();
 			}
-			return FString();
+			return FName(NAME_None);
 		});
-		HandleReresolveModules(SelectedPaths);
+		HandleReresolveModules(SelectedModuleNames);
 	}
 }
 
-void SModularRigModel::HandleReresolveModules(const TArray<FString>& InPaths)
+void SModularRigModel::HandleReresolveModules(const TArray<FName>& InModuleNames)
 {
 	if (ControlRigBlueprint.IsValid())
 	{
@@ -1011,16 +1003,16 @@ void SModularRigModel::HandleReresolveModules(const TArray<FString>& InPaths)
 		}
 
 		TArray<FRigElementKey> ConnectorKeys;
-		for (const FString& PathAndConnector : InPaths)
+		for (const FName& PathAndConnector : InModuleNames)
 		{
-			FString ModulePath = PathAndConnector;
+			FString ModuleNameString = PathAndConnector.ToString();
 			FString ConnectorName;
-			(void)PathAndConnector.Split(TEXT("|"), &ModulePath, &ConnectorName, ESearchCase::IgnoreCase, ESearchDir::FromEnd);
+			(void)FRigHierarchyModulePath(PathAndConnector).Split(&ModuleNameString, &ConnectorName);
 
-			const FRigModuleReference* Module = Controller->Model->FindModule(ModulePath);
+			const FRigModuleReference* Module = Controller->Model->FindModule(*ModuleNameString);
 			if (!Module)
 			{
-				UE_LOG(LogControlRig, Error, TEXT("Could not find module %s"), *ModulePath);
+				UE_LOG(LogControlRig, Error, TEXT("Could not find module %s"), *ModuleNameString);
 				return;
 			}
 
@@ -1081,7 +1073,7 @@ bool SModularRigModel::CanSwapModules() const
 				continue;
 			}
 			TSoftClassPtr<UControlRig> ModuleClass;
-			if (const FRigModuleReference* Module = ControlRigBlueprint->ModularRigModel.FindModule(SelectedItem->ModulePath))
+			if (const FRigModuleReference* Module = ControlRigBlueprint->ModularRigModel.FindModule(SelectedItem->ModuleName))
 			{
 				if (Module->Class.IsValid())
 				{
@@ -1117,31 +1109,31 @@ void SModularRigModel::HandleSwapClassForModules()
 	if (Rig)
 	{
 		const TArray<TSharedPtr<FModularRigTreeElement>> SelectedItems = GetSelectedItems();
-		TArray<FString> SelectedPaths;
-		Algo::Transform(SelectedItems, SelectedPaths, [](const TSharedPtr<FModularRigTreeElement>& Element)
+		TArray<FName> SelectedModuleNames;
+		Algo::Transform(SelectedItems, SelectedModuleNames, [](const TSharedPtr<FModularRigTreeElement>& Element)
 		{
 			if (Element.IsValid())
 			{
-				return Element->ModulePath;
+				return Element->ModuleName;
 			}
-			return FString();
+			return FName(NAME_None);
 		});
-		HandleSwapClassForModules(SelectedPaths);
+		HandleSwapClassForModules(SelectedModuleNames);
 	}
 }
 
-void SModularRigModel::HandleSwapClassForModules(const TArray<FString>& InPaths)
+void SModularRigModel::HandleSwapClassForModules(const TArray<FName>& InModuleNames)
 {
 	TArray<FSoftObjectPath> ModulePaths;
-	Algo::Transform(InPaths, ModulePaths, [this](const FString& Path)
+	Algo::Transform(InModuleNames, ModulePaths, [this](const FName& InModuleName)
 	{
 		FSoftObjectPath ModulePath(ControlRigBlueprint->GetPathName());
-		ModulePath.SetSubPathString(Path);
+		ModulePath.SetSubPathString(InModuleName.ToString());
 		return ModulePath;
 	});
 
 	TSoftClassPtr<UControlRig> SourceClass = nullptr;
-	if (FRigModuleReference* Module = ControlRigBlueprint->ModularRigModel.FindModule(InPaths[0]))
+	if (FRigModuleReference* Module = ControlRigBlueprint->ModularRigModel.FindModule(InModuleNames[0]))
 	{
 		SourceClass = Module->Class;
 	}
@@ -1178,7 +1170,7 @@ void SModularRigModel::HandleSwapClassForModules(const TArray<FString>& InPaths)
 			{
 				if (UControlRigBlueprint* RigBlueprint = Cast<UControlRigBlueprint>(ModulePath.GetWithoutSubPath().ResolveObject()))
 				{
-					return RigBlueprint->GetModularRigController()->SwapModuleClass(ModulePath.GetSubPathString(), NewModuleClass);
+					return RigBlueprint->GetModularRigController()->SwapModuleClass(*ModulePath.GetSubPathString(), NewModuleClass);
 				}
 			}
 			return false;
@@ -1238,7 +1230,7 @@ void SModularRigModel::HandleSelectionChanged(TSharedPtr<FModularRigTreeElement>
 		check(Controller);
 
 		const TGuardValue<bool> GuardSelection(bIsPerformingSelection, true);
-		const TArray<FString> NewSelection = TreeView->GetSelectedKeys();
+		const TArray<FName> NewSelection = TreeView->GetSelectedModuleNames();
 		Controller->SetModuleSelection(NewSelection);
 	}
 }
@@ -1260,11 +1252,11 @@ void SModularRigModel::OnModularRigModified(EModularRigNotification InNotif, con
 				const TGuardValue<bool> GuardSelection(bIsPerformingSelection, true);
 				if(const UModularRigController* ModularRigController = ControlRigBlueprint->GetModularRigController())
 				{
-					const TArray<FString> SelectedModulePaths = ModularRigController->GetSelectedModules();
+					const TArray<FName> SelectedModuleNames = ModularRigController->GetSelectedModules();
 					TArray<TSharedPtr<FModularRigTreeElement>> NewSelection;
-					for(const FString& SelectedModulePath : SelectedModulePaths)
+					for(const FName& SelectedModuleName : SelectedModuleNames)
 					{
-						if(const TSharedPtr<FModularRigTreeElement> Module = TreeView->FindElement(SelectedModulePath))
+						if(const TSharedPtr<FModularRigTreeElement> Module = TreeView->FindElement(SelectedModuleName.ToString()))
 						{
 							NewSelection.Add(Module);
 						}
@@ -1308,19 +1300,19 @@ void SModularRigModel::OnHierarchyModified(ERigHierarchyNotification InNotif, UR
 		case ERigHierarchyNotification::ElementSelected:
 		case ERigHierarchyNotification::ElementDeselected:
 		{
-			FString ModulePathOrConnectorName = InHierarchy->GetModulePath(InElement->GetKey());
+			FString ModuleOrConnectorName = InHierarchy->GetModuleName(InElement->GetKey());
 
 			if(const FRigConnectorElement* Connector = Cast<FRigConnectorElement>(InElement))
 			{
 				if(Connector->IsPrimary())
             	{
-					ModulePathOrConnectorName = Connector->GetName();
+					ModuleOrConnectorName = Connector->GetName();
             	}
 			}
 
-			if(!ModulePathOrConnectorName.IsEmpty())
+			if(!ModuleOrConnectorName.IsEmpty())
 			{
-				if(TSharedPtr<FModularRigTreeElement> Item = TreeView->FindElement(ModulePathOrConnectorName))
+				if(TSharedPtr<FModularRigTreeElement> Item = TreeView->FindElement(ModuleOrConnectorName))
 				{
 					const bool bSelected = InNotif == ERigHierarchyNotification::ElementSelected;
 					TreeView->SetItemHighlighted(Item, bSelected);
@@ -1396,13 +1388,13 @@ UModularRig* SModularRigModel::GetDefaultModularRig() const
 }
 
 
-void SModularRigModel::OnRequestDetailsInspection(const FString& InKey)
+void SModularRigModel::OnRequestDetailsInspection(const FName& InModuleName)
 {
 	if(!ControlRigEditor.IsValid())
 	{
 		return;
 	}
-	ControlRigEditor.Pin()->SetDetailViewForRigModules({InKey});
+	ControlRigEditor.Pin()->SetDetailViewForRigModules({InModuleName});
 }
 
 void SModularRigModel::PostRedo(bool bSuccess) 
@@ -1423,12 +1415,18 @@ void SModularRigModel::PostUndo(bool bSuccess)
 
 FReply SModularRigModel::OnDragDetected(const FGeometry& MyGeometry, const FPointerEvent& MouseEvent)
 {
-	TArray<FString> DraggedElements = GetSelectedKeys();
-	if (MouseEvent.IsMouseButtonDown(EKeys::LeftMouseButton) && DraggedElements.Num() > 0)
+	TArray<FString> DraggedKeys = GetSelectedKeys();
+	TArray<FName> ModuleNames;
+	for(const FString& DraggedKey : DraggedKeys)
+	{
+		ModuleNames.Add(*DraggedKey);
+	}
+	
+	if (MouseEvent.IsMouseButtonDown(EKeys::LeftMouseButton) && ModuleNames.Num() > 0)
 	{
 		if (ControlRigEditor.IsValid())
 		{
-			TSharedRef<FModularRigModuleDragDropOp> DragDropOp = FModularRigModuleDragDropOp::New(MoveTemp(DraggedElements));
+			TSharedRef<FModularRigModuleDragDropOp> DragDropOp = FModularRigModuleDragDropOp::New(MoveTemp(ModuleNames));
 			return FReply::Handled().BeginDragDrop(DragDropOp);
 		}
 	}
@@ -1478,7 +1476,7 @@ TOptional<EItemDropZone> SModularRigModel::OnCanAcceptDrop(const FDragDropEvent&
 		if(TargetItem.IsValid())
 		{
 			// we cannot drag a module onto itself
-			if(ModuleDragDropOperation->GetElements().Contains(TargetItem->ModulePath))
+			if(ModuleDragDropOperation->GetModules().Contains(TargetItem->ModuleName))
 			{
 				return InvalidDropZone;
 			}
@@ -1494,10 +1492,10 @@ TOptional<EItemDropZone> SModularRigModel::OnCanAcceptDrop(const FDragDropEvent&
 
 FReply SModularRigModel::OnAcceptDrop(const FDragDropEvent& DragDropEvent, EItemDropZone DropZone, TSharedPtr<FModularRigTreeElement> TargetItem)
 {
-	FString ParentPath;
+	FName ParentModuleName = NAME_None;
 	if (TargetItem.IsValid())
 	{
-		ParentPath = TargetItem->ModulePath;
+		ParentModuleName = TargetItem->ModuleName;
 	}
 
 	TSharedPtr<FAssetDragDropOp> AssetDragDropOperation = DragDropEvent.GetOperationAs<FAssetDragDropOp>();
@@ -1527,7 +1525,7 @@ FReply SModularRigModel::OnAcceptDrop(const FDragDropEvent& DragDropEvent, EItem
 
 			if(UControlRigBlueprint* AssetBlueprint = Cast<UControlRigBlueprint>(AssetData.GetAsset()))
 			{
-				HandleNewItem(AssetBlueprint->GetControlRigClass(), ParentPath);
+				HandleNewItem(AssetBlueprint->GetControlRigClass(), ParentModuleName);
 			}
 		}
 
@@ -1535,8 +1533,8 @@ FReply SModularRigModel::OnAcceptDrop(const FDragDropEvent& DragDropEvent, EItem
 	}
 	else if(ModuleDragDropOperation)
 	{
-		const TArray<FString> Paths = ModuleDragDropOperation->GetElements();
-		HandleReparentModules(Paths, ParentPath);
+		const TArray<FName> ModuleNames = ModuleDragDropOperation->GetModules();
+		HandleReparentModules(ModuleNames, ParentModuleName);
 	}
 	
 	return FReply::Unhandled();
diff --git a/Engine/Plugins/Animation/ControlRig/Source/ControlRigEditor/Private/Editor/SModularRigModel.h b/Engine/Plugins/Animation/ControlRig/Source/ControlRigEditor/Private/Editor/SModularRigModel.h
--- a/Engine/Plugins/Animation/ControlRig/Source/ControlRigEditor/Private/Editor/SModularRigModel.h
+++ b/Engine/Plugins/Animation/ControlRig/Source/ControlRigEditor/Private/Editor/SModularRigModel.h
@@ -63,33 +63,33 @@ private:
 
 	/** Create a new item */
 	void HandleNewItem();
-	void HandleNewItem(UClass* InClass, const FString& InParentPath);
+	void HandleNewItem(UClass* InClass, const FName& InParentModuleName);
 
 	/** Rename item */
 	bool CanRenameModule() const;
 	void HandleRenameModule();
-	FName HandleRenameModule(const FString& InOldPath, const FName& InNewName);
-	bool HandleVerifyNameChanged(const FString& InOldPath, const FName& InNewName, FText& OutErrorMessage);
+	FName HandleRenameModule(const FName& InOldModuleName, const FName& InNewName);
+	bool HandleVerifyNameChanged(const FName& InOldModuleName, const FName& InNewName, FText& OutErrorMessage);
 
 	/** Delete items */
 	void HandleDeleteModules();
-	void HandleDeleteModules(const TArray<FString>& InPaths);
+	void HandleDeleteModules(const TArray<FName>& InModuleNames);
 
 	/** Reparent items */
-	void HandleReparentModules(const TArray<FString>& InPaths, const FString& InParentPath);
+	void HandleReparentModules(const TArray<FName>& InModuleNames, const FName& InParentModuleName);
 
 	/** Mirror items */
 	void HandleMirrorModules();
-	void HandleMirrorModules(const TArray<FString>& InPaths);
+	void HandleMirrorModules(const TArray<FName>& InModuleNames);
 
 	/** Reresolve items */
 	void HandleReresolveModules();
-	void HandleReresolveModules(const TArray<FString>& InPaths);
+	void HandleReresolveModules(const TArray<FName>& InModuleNames);
 
 	/** Swap module class for items */
 	bool CanSwapModules() const;
 	void HandleSwapClassForModules();
-	void HandleSwapClassForModules(const TArray<FString>& InPaths);
+	void HandleSwapClassForModules(const TArray<FName>& InModuleNames);
 
 	/** Resolve connector */
 	void HandleConnectorResolved(const FRigElementKey& InConnector, const TArray<FRigElementKey>& InTargets);
@@ -141,7 +141,7 @@ private:
 	UModularRig* GetModularRig() const;
 	UModularRig* GetDefaultModularRig() const;
 	const UModularRig* GetModularRigForTreeView() const { return GetModularRig(); }
-	void OnRequestDetailsInspection(const FString& InKey);
+	void OnRequestDetailsInspection(const FName& InModuleName);
 	void HandlePreCompileModularRigs(URigVMBlueprint* InBlueprint);
 	void HandlePostCompileModularRigs(URigVMBlueprint* InBlueprint);
 	void OnModularRigModified(EModularRigNotification InNotif, const FRigModuleReference* InModule);
diff --git a/Engine/Plugins/Animation/ControlRig/Source/ControlRigEditor/Private/Editor/SModularRigTreeView.cpp b/Engine/Plugins/Animation/ControlRig/Source/ControlRigEditor/Private/Editor/SModularRigTreeView.cpp
--- a/Engine/Plugins/Animation/ControlRig/Source/ControlRigEditor/Private/Editor/SModularRigTreeView.cpp
+++ b/Engine/Plugins/Animation/ControlRig/Source/ControlRigEditor/Private/Editor/SModularRigTreeView.cpp
@@ -45,14 +45,15 @@ FModularRigTreeElement::FModularRigTreeElement(const FString& InKey, TWeakPtr<SM
 	Key = InKey;
 	bIsPrimary = bInIsPrimary;
 
-	FString ShortNameStr = Key;
-	(void)URigHierarchy::SplitNameSpace(ShortNameStr, &ModulePath, &ShortNameStr);
+	FString ModuleNameString;
+	FString ConnectorNameString = Key;
+	(void)FRigHierarchyModulePath(ConnectorNameString).Split(&ModuleNameString, &ConnectorNameString);
 	if (bIsPrimary)
 	{
-		ModulePath = Key;
+		ModuleName = *Key;
 		if(const UModularRig* ModularRig = InTreeView.Pin()->GetRigTreeDelegates().GetModularRig())
 		{
-			if (const FRigModuleInstance* Module = ModularRig->FindModule(ModulePath))
+			if (const FRigModuleInstance* Module = ModularRig->FindModule(ModuleName))
 			{
 				if (const UControlRig* Rig = Module->GetRig())
 				{
@@ -66,10 +67,11 @@ FModularRigTreeElement::FModularRigTreeElement(const FString& InKey, TWeakPtr<SM
 	}
 	else
 	{
-		ConnectorName = ShortNameStr;
+		ConnectorName = ConnectorNameString;
+		ModuleName = *ModuleNameString;
 	}
 	
-	ShortName = *ShortNameStr;
+	ShortName = *ConnectorNameString;
 	
 	if(InTreeView.IsValid())
 	{
@@ -118,9 +120,9 @@ void SModularRigModelItem::Construct(const FArguments& InArgs, const TSharedRef<
 		return;
 	}
 
-	const FString& ModulePath = InRigTreeElement->ModulePath;
-	FString ConnectorPath = FString::Printf(TEXT("%s:%s"), *ModulePath, *InRigTreeElement->ConnectorName);
-	ConnectorKey = FRigElementKey(*ConnectorPath, ERigElementType::Connector);
+	const FString& ModuleName = InRigTreeElement->ModuleName.ToString();
+	const FRigHierarchyModulePath ConnectorModulePath(ModuleName, InRigTreeElement->ConnectorName);
+	ConnectorKey = FRigElementKey(ConnectorModulePath.GetPathFName(), ERigElementType::Connector);
 
 	SMultiColumnTableRow<TSharedPtr<FModularRigTreeElement>>::Construct(
 		SMultiColumnTableRow<TSharedPtr<FModularRigTreeElement>>::FArguments()
@@ -144,17 +146,17 @@ void SModularRigModelItem::OnNameCommitted(const FText& InText, ETextCommit::Typ
 	if (InCommitType == ETextCommit::OnEnter)
 	{
 		FString NewName = InText.ToString();
-		const FString OldKey = WeakRigTreeElement.Pin()->ModulePath;
+		const FName OldModuleName = WeakRigTreeElement.Pin()->ModuleName;
 
-		Delegates.HandleRenameElement(OldKey, *NewName);
+		Delegates.HandleRenameElement(OldModuleName, *NewName);
 	}
 }
 
 bool SModularRigModelItem::OnVerifyNameChanged(const FText& InText, FText& OutErrorMessage)
 {
 	const FName NewName = *InText.ToString();
-	const FString OldPath = WeakRigTreeElement.Pin()->ModulePath;
-	return Delegates.HandleVerifyElementNameChanged(OldPath, NewName, OutErrorMessage);
+	const FName OldModuleName = WeakRigTreeElement.Pin()->ModuleName;
+	return Delegates.HandleVerifyElementNameChanged(OldModuleName, NewName, OutErrorMessage);
 }
 
 TSharedRef<SWidget> SModularRigModelItem::GenerateWidgetForColumn(const FName& ColumnName)
@@ -251,7 +253,7 @@ TSharedRef<SWidget> SModularRigModelItem::GenerateWidgetForColumn(const FName& C
 			
 				if (const UModularRig* ModularRig = Delegates.GetModularRig())
 				{
-					if (const FRigModuleInstance* Module = ModularRig->FindModule(WeakRigTreeElement.Pin()->ModulePath))
+					if (const FRigModuleInstance* Module = ModularRig->FindModule(WeakRigTreeElement.Pin()->ModuleName))
 					{
 						if (const UControlRigBlueprint* ModuleBlueprint = Cast<UControlRigBlueprint>(Module->GetRig()->GetClass()->ClassGeneratedBy))
 						{
@@ -275,7 +277,7 @@ TSharedRef<SWidget> SModularRigModelItem::GenerateWidgetForColumn(const FName& C
 			{
 				if (const UModularRig* ModularRig = Delegates.GetModularRig())
 				{
-					if (const FRigModuleInstance* Module = ModularRig->FindModule(WeakRigTreeElement.Pin()->ModulePath))
+					if (const FRigModuleInstance* Module = ModularRig->FindModule(WeakRigTreeElement.Pin()->ModuleName))
 					{
 						if (const UControlRigBlueprint* ModuleBlueprint = Cast<UControlRigBlueprint>(Module->GetRig()->GetClass()->ClassGeneratedBy))
 						{
@@ -459,7 +461,7 @@ FText SModularRigModelItem::GetName(bool bUseShortName) const
 	{
 		return (FText::FromName(WeakRigTreeElement.Pin()->ShortName));
 	}
-	return (FText::FromString(WeakRigTreeElement.Pin()->ModulePath));
+	return (FText::FromName(WeakRigTreeElement.Pin()->ModuleName));
 }
 
 FText SModularRigModelItem::GetItemTooltip() const
@@ -653,7 +655,7 @@ bool SModularRigTreeView::AddElement(FString InKey, FString InParentKey, bool bA
 		if(!FilterTextString.IsEmpty())
 		{
 			FString StringToSearch = InKey;
-			(void)URigHierarchy::SplitNameSpace(StringToSearch, nullptr, &StringToSearch);
+			(void)FRigHierarchyModulePath(StringToSearch).Split(nullptr, &StringToSearch);
 			
 			if(!StringToSearch.Contains(FilterTextString, ESearchCase::IgnoreCase))
 			{
@@ -661,12 +663,12 @@ bool SModularRigTreeView::AddElement(FString InKey, FString InParentKey, bool bA
 			}
 		}
 
-		TArray<FString> FilteredConnectors;
+		TArray<FRigHierarchyModulePath> FilteredConnectors;
 		if (const UModularRig* ModularRig = Delegates.GetModularRig())
 		{
 			const FModularRigModel& Model = ModularRig->GetModularRigModel();
 			
-			if (const FRigModuleInstance* Module = ModularRig->FindModule(ModulePath))
+			if (const FRigModuleInstance* Module = ModularRig->FindModule(*ModulePath))
 			{
 				if (const UControlRig* ModuleRig = Module->GetRig())
 				{
@@ -680,7 +682,7 @@ bool SModularRigTreeView::AddElement(FString InKey, FString InParentKey, bool bA
 							continue;
 						}
 
-						const FString Key = FString::Printf(TEXT("%s:%s"), *ModulePath, *Connector.Name);
+						const FRigHierarchyModulePath Key(ModulePath, Connector.Name);
 						bool bShouldFilterByConnectorType = true;
 
 						if(!FilterTextString.IsEmpty())
@@ -698,7 +700,7 @@ bool SModularRigTreeView::AddElement(FString InKey, FString InParentKey, bool bA
 
 						if(bShouldFilterByConnectorType)
 						{
-							const bool bIsConnected = Model.Connections.HasConnection(FRigElementKey(*Key, ERigElementType::Connector));
+							const bool bIsConnected = Model.Connections.HasConnection(FRigElementKey(Key.GetPathFName(), ERigElementType::Connector));
 							if(bIsConnected || !ShowUnresolvedConnectors.Get())
 							{
 								if(Connector.IsOptional())
@@ -763,33 +765,32 @@ bool SModularRigTreeView::AddElement(const FRigModuleInstance* InElement, bool b
 {
 	check(InElement);
 	
-	if (ElementMap.Contains(InElement->GetPath()))
+	if (ElementMap.Contains(InElement->Name.ToString()))
 	{
 		return false;
 	}
 
 	const UModularRig* ModularRig = Delegates.GetModularRig();
 
-	FString ElementPath = InElement->GetPath();
-	if(!AddElement(ElementPath, FString(), bApplyFilterText))
+	if(!AddElement(InElement->Name.ToString(), FString(), bApplyFilterText))
 	{
 		return false;
 	}
 
-	if (ElementMap.Contains(InElement->GetPath()))
+	if (ElementMap.Contains(InElement->Name.ToString()))
 	{
 		if(ModularRig)
 		{
-			FString ParentPath = ModularRig->GetParentPath(ElementPath);
-			if (!ParentPath.IsEmpty())
+			const FName ParentModuleName = ModularRig->GetParentModuleName(InElement->Name);
+			if (!ParentModuleName.IsNone())
 			{
-				if(const FRigModuleInstance* ParentElement = ModularRig->FindModule(ParentPath))
+				if(const FRigModuleInstance* ParentElement = ModularRig->FindModule(ParentModuleName))
 				{
 					AddElement(ParentElement, false);
 
-					if(ElementMap.Contains(ParentPath))
+					if(ElementMap.Contains(ParentModuleName.ToString()))
 					{
-						ReparentElement(ElementPath, ParentPath);
+						ReparentElement(InElement->Name.ToString(), ParentModuleName.ToString());
 					}
 				}
 			}
@@ -862,7 +863,7 @@ bool SModularRigTreeView::ReparentElement(const FString InKey, const FString InP
 void SModularRigTreeView::RefreshTreeView(bool bRebuildContent)
 {
 	TMap<FString, bool> ExpansionState;
-	TArray<FString> Selection;
+	TArray<FName> Selection;
 
 	if(bRebuildContent)
 	{
@@ -878,7 +879,7 @@ void SModularRigTreeView::RefreshTreeView(bool bRebuildContent)
 		ElementMap.Reset();
 		ParentMap.Reset();
 
-		Selection = GetSelectedKeys();
+		Selection = GetSelectedModuleNames();
 	}
 
 	if(bRebuildContent)
@@ -937,9 +938,9 @@ void SModularRigTreeView::RefreshTreeView(bool bRebuildContent)
 		if(!Selection.IsEmpty())
 		{
 			TArray<TSharedPtr<FModularRigTreeElement>> SelectedElements;
-			for(const FString& SelectedPath : Selection)
+			for(const FName& SelectedModuleName : Selection)
 			{
-				if(const TSharedPtr<FModularRigTreeElement> ElementToSelect = FindElement(SelectedPath))
+				if(const TSharedPtr<FModularRigTreeElement> ElementToSelect = FindElement(SelectedModuleName.ToString()))
 				{
 					SelectedElements.Add(ElementToSelect);
 				}
@@ -964,18 +965,18 @@ void SModularRigTreeView::HandleGetChildrenForTree(TSharedPtr<FModularRigTreeEle
 	OutChildren = InItem.Get()->Children;
 }
 
-TArray<FString> SModularRigTreeView::GetSelectedKeys() const
+TArray<FName> SModularRigTreeView::GetSelectedModuleNames() const
 {
-	TArray<FString> Keys;
+	TArray<FName> ModuleNames;
 	TArray<TSharedPtr<FModularRigTreeElement>> SelectedElements = GetSelectedItems();
 	for(const TSharedPtr<FModularRigTreeElement>& SelectedElement : SelectedElements)
 	{
 		if (SelectedElement.IsValid())
 		{
-			Keys.AddUnique(SelectedElement->ModulePath);
+			ModuleNames.AddUnique(SelectedElement->ModuleName);
 		}
 	}
-	return Keys;
+	return ModuleNames;
 }
 
 void SModularRigTreeView::SetSelection(const TArray<TSharedPtr<FModularRigTreeElement>>& InSelection) 
@@ -1017,11 +1018,11 @@ TPair<const FSlateBrush*, FSlateColor> FModularRigTreeElement::GetBrushAndColor(
 	FLinearColor Color = FSlateColor(EStyleColor::Foreground).GetColor(FWidgetStyle());
 	float Opacity = 1.f;
 
-	if (const FRigModuleInstance* ConnectorModule = InModularRig->FindModule(ModulePath))
+	if (const FRigModuleInstance* ConnectorModule = InModularRig->FindModule(ModuleName))
 	{
 		const FModularRigModel& Model = InModularRig->GetModularRigModel();
-		const FString ConnectorPath = FString::Printf(TEXT("%s:%s"), *ModulePath, *ConnectorName);
-		bool bIsConnected = Model.Connections.HasConnection(FRigElementKey(*ConnectorPath, ERigElementType::Connector));
+		const FRigHierarchyModulePath ConnectorPath(ModuleName.ToString(), ConnectorName);
+		bool bIsConnected = Model.Connections.HasConnection(FRigElementKey(ConnectorPath.GetPathFName(), ERigElementType::Connector));
 		bool bConnectionWarning = !bIsConnected;
 		
 		if (const UControlRig* ModuleRig = ConnectorModule->GetRig())
diff --git a/Engine/Plugins/Animation/ControlRig/Source/ControlRigEditor/Private/Editor/SModularRigTreeView.h b/Engine/Plugins/Animation/ControlRig/Source/ControlRigEditor/Private/Editor/SModularRigTreeView.h
--- a/Engine/Plugins/Animation/ControlRig/Source/ControlRigEditor/Private/Editor/SModularRigTreeView.h
+++ b/Engine/Plugins/Animation/ControlRig/Source/ControlRigEditor/Private/Editor/SModularRigTreeView.h
@@ -14,11 +14,11 @@ class FModularRigTreeElement;
 
 
 DECLARE_DELEGATE_RetVal(const UModularRig*, FOnGetModularRigTreeRig);
-DECLARE_DELEGATE_OneParam(FOnModularRigTreeRequestDetailsInspection, const FString&);
-DECLARE_DELEGATE_RetVal_TwoParams(FName, FOnModularRigTreeRenameElement, const FString& /*OldPath*/, const FName& /*NewName*/);
+DECLARE_DELEGATE_OneParam(FOnModularRigTreeRequestDetailsInspection, const FName&);
+DECLARE_DELEGATE_RetVal_TwoParams(FName, FOnModularRigTreeRenameElement, const FName& /*OldModuleName*/, const FName& /*NewName*/);
 DECLARE_DELEGATE_TwoParams(FOnModularRigTreeResolveConnector, const FRigElementKey& /*Connector*/, const TArray<FRigElementKey>& /*Targets*/);
 DECLARE_DELEGATE_OneParam(FOnModularRigTreeDisconnectConnector, const FRigElementKey& /*Connector*/);
-DECLARE_DELEGATE_RetVal_ThreeParams(bool, FOnModularRigTreeVerifyElementNameChanged, const FString& /*OldPath*/, const FName& /*NewName*/, FText& /*OutErrorMessage*/);
+DECLARE_DELEGATE_RetVal_ThreeParams(bool, FOnModularRigTreeVerifyElementNameChanged, const FName& /*OldModuleName*/, const FName& /*NewName*/, FText& /*OutErrorMessage*/);
 
 typedef STreeView<TSharedPtr<FModularRigTreeElement>>::FOnMouseButtonClick FOnModularRigTreeMouseButtonClick;
 typedef STreeView<TSharedPtr<FModularRigTreeElement>>::FOnMouseButtonDoubleClick FOnModularRigTreeMouseButtonDoubleClick;
@@ -57,20 +57,20 @@ public:
 		return nullptr;
 	}
 
-	FName HandleRenameElement(const FString& OldPath, const FName& NewName) const
+	FName HandleRenameElement(const FName& OldModuleName, const FName& NewName) const
 	{
 		if(OnRenameElement.IsBound())
 		{
-			return OnRenameElement.Execute(OldPath, NewName);
+			return OnRenameElement.Execute(OldModuleName, NewName);
 		}
-		return *OldPath;
+		return OldModuleName;
 	}
 
-	bool HandleVerifyElementNameChanged(const FString& OldPath, const FName& NewName, FText& OutErrorMessage) const
+	bool HandleVerifyElementNameChanged(const FName& OldModuleName, const FName& NewName, FText& OutErrorMessage) const
 	{
 		if(OnVerifyModuleNameChanged.IsBound())
 		{
-			return OnVerifyModuleNameChanged.Execute(OldPath, NewName, OutErrorMessage);
+			return OnVerifyModuleNameChanged.Execute(OldModuleName, NewName, OutErrorMessage);
 		}
 		return false;
 	}
@@ -123,7 +123,7 @@ public:
 	/** Element Data to display */
 	FString Key;
 	bool bIsPrimary;
-	FString ModulePath;
+	FName ModuleName;
 	FString ConnectorName;
 	FName ShortName;
 	TArray<TSharedPtr<FModularRigTreeElement>> Children;
@@ -250,7 +250,7 @@ public:
 	TSharedRef<ITableRow> MakeTableRowWidget(TSharedPtr<FModularRigTreeElement> InItem, const TSharedRef<STableViewBase>& OwnerTable, bool bPinned);
 	void HandleGetChildrenForTree(TSharedPtr<FModularRigTreeElement> InItem, TArray<TSharedPtr<FModularRigTreeElement>>& OutChildren);
 
-	TArray<FString> GetSelectedKeys() const;
+	TArray<FName> GetSelectedModuleNames() const;
 	void SetSelection(const TArray<TSharedPtr<FModularRigTreeElement>>& InSelection);
 	const TArray<TSharedPtr<FModularRigTreeElement>>& GetRootElements() const { return RootElements; }
 	FModularRigTreeDelegates& GetRigTreeDelegates() { return Delegates; }
diff --git a/Engine/Plugins/Animation/ControlRig/Source/ControlRigEditor/Private/Editor/SRigConnectorTargetWidget.cpp b/Engine/Plugins/Animation/ControlRig/Source/ControlRigEditor/Private/Editor/SRigConnectorTargetWidget.cpp
--- a/Engine/Plugins/Animation/ControlRig/Source/ControlRigEditor/Private/Editor/SRigConnectorTargetWidget.cpp
+++ b/Engine/Plugins/Animation/ControlRig/Source/ControlRigEditor/Private/Editor/SRigConnectorTargetWidget.cpp
@@ -151,9 +151,9 @@ void SRigConnectorTargetComboButton::PopulateButtonBox()
 					.Padding(0)
 					[
 						SNew( STextBlock )
-						.Text_Lambda([this]()
+						.Text_Lambda([this, Hierarchy]()
 						{
-							return FText::FromName(TargetKey.Get().Name);
+							return Hierarchy->GetDisplayNameForUI(TargetKey.Get(), true);
 						})
 						.Font( IDetailLayoutBuilder::GetDetailFont() )
 					]
diff --git a/Engine/Plugins/Animation/ControlRig/Source/ControlRigEditor/Private/Editor/SRigHierarchy.cpp b/Engine/Plugins/Animation/ControlRig/Source/ControlRigEditor/Private/Editor/SRigHierarchy.cpp
--- a/Engine/Plugins/Animation/ControlRig/Source/ControlRigEditor/Private/Editor/SRigHierarchy.cpp
+++ b/Engine/Plugins/Animation/ControlRig/Source/ControlRigEditor/Private/Editor/SRigHierarchy.cpp
@@ -402,10 +402,10 @@ void SRigHierarchy::BindCommands()
 		FIsActionChecked::CreateLambda([this]() { return DisplaySettings.bHideParentsOnFilter; }));
 
 	CommandList->MapAction(
-		Commands.ShowShortNames,
+		Commands.ShowFullPaths,
 		FExecuteAction::CreateLambda([this]() { DisplaySettings.bUseShortName = !DisplaySettings.bUseShortName; RefreshTreeView(); }),
 		FCanExecuteAction(),
-		FIsActionChecked::CreateLambda([this]() { return DisplaySettings.bUseShortName; }));
+		FIsActionChecked::CreateLambda([this]() { return !DisplaySettings.bUseShortName; }));
 
 	CommandList->MapAction(
 		Commands.ShowImportedBones,
@@ -1002,11 +1002,10 @@ void SRigHierarchy::OnModularRigModified(EModularRigNotification InNotif, const
 				{
 					if(URigHierarchy* Hierarchy = ControlRig->GetHierarchy())
 					{
-						const FString ModulePath = InModule->GetPath();
 						TArray<FRigElementKey> Keys = Hierarchy->GetAllKeys();
-						Keys = Keys.FilterByPredicate([Hierarchy, ModulePath](const FRigElementKey& InKey)
+						Keys = Keys.FilterByPredicate([Hierarchy, InModule](const FRigElementKey& InKey)
 						{
-							return ModulePath == Hierarchy->GetModulePath(InKey);
+							return InModule->Name == Hierarchy->GetModuleFName(InKey);
 						});
 
 						bool bScrollIntoView = true;
@@ -1193,7 +1192,7 @@ TSharedRef< SWidget > SRigHierarchy::CreateFilterMenu()
 	{
 		MenuBuilder.AddMenuEntry(Actions.FilteringFlattensHierarchy);
 		MenuBuilder.AddMenuEntry(Actions.HideParentsWhenFiltering);
-		MenuBuilder.AddMenuEntry(Actions.ShowShortNames);
+		MenuBuilder.AddMenuEntry(Actions.ShowFullPaths);
 	}
 	MenuBuilder.EndSection();
 
@@ -3411,20 +3410,20 @@ FReply SRigHierarchy::OnAcceptDrop(const FDragDropEvent& DragDropEvent, EItemDro
 			{
 				if (UModularRigController* Controller = ControlRigBlueprint->GetModularRigController())
 				{
-					static const FString ParentPath = FString();
-					const FRigName ModuleName = Controller->GetSafeNewName(ParentPath, FRigName(AssetBlueprint->RigModuleSettings.Identifier.Name));
-					const FString ModulePath = Controller->AddModule(ModuleName, AssetBlueprint->GetControlRigClass(), ParentPath);
-					if(TargetItem.IsValid() && !ModulePath.IsEmpty())
+					const FRigName DesiredModuleName = Controller->GetSafeNewName(FRigName(AssetBlueprint->RigModuleSettings.Identifier.Name));
+					const FName ModuleName = Controller->AddModule(DesiredModuleName.GetFName(), AssetBlueprint->GetControlRigClass(), NAME_None);
+					if(TargetItem.IsValid() && !ModuleName.IsNone())
 					{
 						FRigElementKey PrimaryConnectorKey;
 						TArray<FRigConnectorElement*> Connectors = GetHierarchy()->GetElementsOfType<FRigConnectorElement>();
+
+						const FString ModuleNameString = ModuleName.ToString();
 						for (FRigConnectorElement* Connector : Connectors)
 						{
 							if (Connector->IsPrimary())
 							{
-								FString Path, Name;
-								(void)URigHierarchy::SplitNameSpace(Connector->GetName(), &Path, &Name);
-								if (Path == ModulePath)
+								const FRigHierarchyModulePath ConnectorModulePath(Connector->GetName());
+								if (ConnectorModulePath.HasModuleName(ModuleNameString))
 								{
 									PrimaryConnectorKey = Connector->GetKey();
 									break;
@@ -3473,10 +3472,10 @@ void SRigHierarchy::UpdateConnectorMatchesOnDrag(const TArray<FRigHierarchyKey>&
 							{
 								if(const FRigConnectorElement* Connector = Hierarchy->Find<FRigConnectorElement>(DraggedElement.GetElement()))
 								{
-									const FName NameSpace = Hierarchy->GetNameSpaceFName(Connector->GetKey());
-									if(!NameSpace.IsNone())
+									const FName ModuleName = Hierarchy->GetModuleFName(Connector->GetKey());
+									if(!ModuleName.IsNone())
 									{
-										if(const FRigModuleInstance* Module = ControlRig->FindModule(NameSpace.ToString()))
+										if(const FRigModuleInstance* Module = ControlRig->FindModule(ModuleName))
 										{
 											const FModularRigResolveResult ResolveResult = RuleManager->FindMatches(Connector, Module, ControlRig->ElementKeyRedirector); 
 											DragRigResolveResults.Add(DraggedElement.GetElement(), ResolveResult);
diff --git a/Engine/Plugins/Animation/ControlRig/Source/ControlRigEditor/Private/Editor/SRigHierarchyTreeView.h b/Engine/Plugins/Animation/ControlRig/Source/ControlRigEditor/Private/Editor/SRigHierarchyTreeView.h
--- a/Engine/Plugins/Animation/ControlRig/Source/ControlRigEditor/Private/Editor/SRigHierarchyTreeView.h
+++ b/Engine/Plugins/Animation/ControlRig/Source/ControlRigEditor/Private/Editor/SRigHierarchyTreeView.h
@@ -20,7 +20,7 @@ struct CONTROLRIGEDITOR_API FRigTreeDisplaySettings
 
 		bFlattenHierarchyOnFilter = false;
 		bHideParentsOnFilter = false;
-		bUseShortName = true;
+		bUseShortName = false;
 		bShowImportedBones = true;
 		bShowBones = true;
 		bShowControls = true;
diff --git a/Engine/Plugins/Animation/ControlRig/Source/ControlRigEditor/Private/Graph/ControlRigGraphPanelPinFactory.cpp b/Engine/Plugins/Animation/ControlRig/Source/ControlRigEditor/Private/Graph/ControlRigGraphPanelPinFactory.cpp
--- a/Engine/Plugins/Animation/ControlRig/Source/ControlRigEditor/Private/Graph/ControlRigGraphPanelPinFactory.cpp
+++ b/Engine/Plugins/Animation/ControlRig/Source/ControlRigEditor/Private/Graph/ControlRigGraphPanelPinFactory.cpp
@@ -22,6 +22,7 @@
 #include "Units/Hierarchy/RigUnit_Metadata.h"
 #include "ControlRigElementDetails.h"
 #include "IPropertyAccessEditor.h"
+#include "ModularRig.h"
 
 TSharedPtr<SGraphPin> FControlRigGraphPanelPinFactory::CreatePin(UEdGraphPin* InPin) const
 {
@@ -493,7 +494,7 @@ TSharedPtr<SGraphPin> FControlRigGraphPanelPinFactory::CreatePin_Internal(UEdGra
 											if(bUseShortNames)
 											{
 												int32 Index = INDEX_NONE;
-												if(NameString.FindLastChar(TEXT(':'), Index))
+												if(NameString.FindLastChar(FRigHierarchyModulePath::ModuleNameSuffixChar, Index))
 												{
 													NameString.MidInline(Index + 1);
 												}
@@ -582,7 +583,7 @@ TSharedPtr<SGraphPin> FControlRigGraphPanelPinFactory::CreatePin_Internal(UEdGra
 											if(bUseShortNames)
 											{
 												int32 Index = INDEX_NONE;
-												if(TagString.FindLastChar(TEXT(':'), Index))
+												if(TagString.FindLastChar(FRigHierarchyModulePath::ModuleNameSuffixChar, Index))
 												{
 													TagString.MidInline(Index + 1);
 												}
diff --git a/Engine/Source/Runtime/MovieSceneTracks/Private/Sections/MovieSceneParameterSection.cpp b/Engine/Source/Runtime/MovieSceneTracks/Private/Sections/MovieSceneParameterSection.cpp
--- a/Engine/Source/Runtime/MovieSceneTracks/Private/Sections/MovieSceneParameterSection.cpp
+++ b/Engine/Source/Runtime/MovieSceneTracks/Private/Sections/MovieSceneParameterSection.cpp
@@ -35,36 +35,6 @@ void IMovieSceneParameterSectionExtender::ExtendEntity(UMovieSceneParameterSecti
 	ExtendEntityImpl(Section, EntityLinker, Params, OutImportedEntity);
 }
 
-FScalarParameterNameAndCurve::FScalarParameterNameAndCurve( FName InParameterName )
-{
-	ParameterName = InParameterName;
-}
-
-FBoolParameterNameAndCurve::FBoolParameterNameAndCurve(FName InParameterName)
-{
-	ParameterName = InParameterName;
-}
-
-FVector2DParameterNameAndCurves::FVector2DParameterNameAndCurves( FName InParameterName )
-{
-	ParameterName = InParameterName;
-}
-
-FVectorParameterNameAndCurves::FVectorParameterNameAndCurves(FName InParameterName)
-{
-	ParameterName = InParameterName;
-}
-
-FColorParameterNameAndCurves::FColorParameterNameAndCurves( FName InParameterName )
-{
-	ParameterName = InParameterName;
-}
-
-FTransformParameterNameAndCurves::FTransformParameterNameAndCurves(FName InParameterName)
-{
-	ParameterName = InParameterName;
-}
-
 UMovieSceneParameterSection::UMovieSceneParameterSection( const FObjectInitializer& ObjectInitializer )
 	: Super( ObjectInitializer )
 {
diff --git a/Engine/Source/Runtime/MovieSceneTracks/Public/Sections/MovieSceneParameterSection.h b/Engine/Source/Runtime/MovieSceneTracks/Public/Sections/MovieSceneParameterSection.h
--- a/Engine/Source/Runtime/MovieSceneTracks/Public/Sections/MovieSceneParameterSection.h
+++ b/Engine/Source/Runtime/MovieSceneTracks/Public/Sections/MovieSceneParameterSection.h
@@ -155,22 +155,34 @@ struct FTransformParameterNameAndValue
 };
 
 /**
- * Structure representing an bool  parameter and it's associated animation curve.
+ * Structure representing a parameter name without a value
  */
 USTRUCT()
-struct FBoolParameterNameAndCurve
+struct FBaseParameterNameAndValue
 {
 	GENERATED_USTRUCT_BODY()
 
-	FBoolParameterNameAndCurve()
+	FBaseParameterNameAndValue()
+	: ParameterName(NAME_None)
 	{}
 
 	/** Creates a new FScalarParameterNameAndCurve for a specific scalar parameter. */
-	MOVIESCENETRACKS_API FBoolParameterNameAndCurve(FName InParameterName);
+	MOVIESCENETRACKS_API FBaseParameterNameAndValue(FName InParameterName)
+	: ParameterName(InParameterName)
+	{}
 
-	/** The name of the scalar parameter which is being animated. */
+	/** The name of the parameter which is being animated. */
 	UPROPERTY()
 	FName ParameterName;
+};
+
+/**
+ * Structure representing a bool parameter and it's associated animation curve.
+ */
+USTRUCT()
+struct FBoolParameterNameAndCurve : public FBaseParameterNameAndValue
+{
+	GENERATED_USTRUCT_BODY()
 
 	/** The curve which contains the animation data for the scalar parameter. */
 	UPROPERTY()
@@ -181,20 +193,10 @@ struct FBoolParameterNameAndCurve
  * Structure representing an animated scalar parameter and it's associated animation curve.
  */
 USTRUCT()
-struct FScalarParameterNameAndCurve
+struct FScalarParameterNameAndCurve : public FBaseParameterNameAndValue
 {
 	GENERATED_USTRUCT_BODY()
 
-	FScalarParameterNameAndCurve()
-	{}
-
-	/** Creates a new FScalarParameterNameAndCurve for a specific scalar parameter. */
-	MOVIESCENETRACKS_API FScalarParameterNameAndCurve(FName InParameterName);
-
-	/** The name of the scalar parameter which is being animated. */
-	UPROPERTY()
-	FName ParameterName;
-
 	/** The curve which contains the animation data for the scalar parameter. */
 	UPROPERTY()
 	FMovieSceneFloatChannel ParameterCurve;
@@ -204,20 +206,10 @@ struct FScalarParameterNameAndCurve
  * Structure representing an animated vector2D parameter and it's associated animation curve.
  */
 USTRUCT()
-struct FVector2DParameterNameAndCurves
+struct FVector2DParameterNameAndCurves : public FBaseParameterNameAndValue
 {
 	GENERATED_USTRUCT_BODY()
 
-	FVector2DParameterNameAndCurves()
-	{}
-
-	/** Creates a new FVectorParameterNameAndCurve for a specific vector parameter. */
-	MOVIESCENETRACKS_API FVector2DParameterNameAndCurves(FName InParameterName);
-
-	/** The name of the vector parameter which is being animated. */
-	UPROPERTY()
-	FName ParameterName;
-
 	/** The curve which contains the animation data for the x component of the vector parameter. */
 	UPROPERTY()
 	FMovieSceneFloatChannel XCurve;
@@ -233,20 +225,10 @@ struct FVector2DParameterNameAndCurves
  * Structure representing an animated vector parameter and it's associated animation curve.
  */
 USTRUCT()
-struct FVectorParameterNameAndCurves
+struct FVectorParameterNameAndCurves : public FBaseParameterNameAndValue
 {
 	GENERATED_USTRUCT_BODY()
 
-	FVectorParameterNameAndCurves() 
-	{}
-
-	/** Creates a new FVectorParameterNameAndCurve for a specific vector parameter. */
-	MOVIESCENETRACKS_API FVectorParameterNameAndCurves(FName InParameterName);
-
-	/** The name of the vector parameter which is being animated. */
-	UPROPERTY()
-	FName ParameterName;
-
 	/** The curve which contains the animation data for the x component of the vector parameter. */
 	UPROPERTY()
 	FMovieSceneFloatChannel XCurve;
@@ -265,20 +247,10 @@ struct FVectorParameterNameAndCurves
 * Structure representing an animated color parameter and it's associated animation curve.
 */
 USTRUCT()
-struct FColorParameterNameAndCurves
+struct FColorParameterNameAndCurves : public FBaseParameterNameAndValue
 {
 	GENERATED_USTRUCT_BODY()
 
-	FColorParameterNameAndCurves()
-	{}
-
-	/** Creates a new FVectorParameterNameAndCurve for a specific color parameter. */
-	MOVIESCENETRACKS_API FColorParameterNameAndCurves(FName InParameterName);
-
-	/** The name of the color parameter which is being animated. */
-	UPROPERTY()
-	FName ParameterName;
-
 	/** The curve which contains the animation data for the red component of the color parameter. */
 	UPROPERTY()
 	FMovieSceneFloatChannel RedCurve;
@@ -300,20 +272,10 @@ struct FColorParameterNameAndCurves
 * Structure representing an animated transform parameter and it's associated animation curve.
 */
 USTRUCT()
-struct FTransformParameterNameAndCurves
+struct FTransformParameterNameAndCurves : public FBaseParameterNameAndValue
 {
 	GENERATED_USTRUCT_BODY()
 
-	FTransformParameterNameAndCurves()
-	{}
-
-	/** Creates a new FVectorParameterNameAndCurve for a specific color parameter. */
-	MOVIESCENETRACKS_API FTransformParameterNameAndCurves(FName InParameterName);
-
-	/** The name of the transform  parameter which is being animated. */
-	UPROPERTY()
-	FName ParameterName;
-
 	/** Translation curves */
 	UPROPERTY()
 	FMovieSceneFloatChannel Translation[3];
