Control Rig: Hierarchy imports to include mesh sockets

#rnx

diff --git a/Engine/Plugins/Animation/ControlRig/Source/ControlRig/Private/Rigs/RigHierarchyController.cpp b/Engine/Plugins/Animation/ControlRig/Source/ControlRig/Private/Rigs/RigHierarchyController.cpp
--- a/Engine/Plugins/Animation/ControlRig/Source/ControlRig/Private/Rigs/RigHierarchyController.cpp
+++ b/Engine/Plugins/Animation/ControlRig/Source/ControlRig/Private/Rigs/RigHierarchyController.cpp
@@ -19,6 +19,7 @@
 #endif
 
 #include "Algo/Count.h"
+#include "Engine/SkeletalMeshSocket.h"
 
 #include UE_INLINE_GENERATED_CPP_BY_NAME(RigHierarchyController)
 
@@ -1641,6 +1642,67 @@ TArray<FRigElementKey> URigHierarchyController::ImportBonesFromSkeletalMesh(
 	return BoneKeys;
 }
 
+TArray<FRigElementKey> URigHierarchyController::ImportSocketsFromSkeletalMesh(USkeletalMesh* InSkeletalMesh, const FName& InNameSpace,
+	bool bReplaceExistingSockets, bool bRemoveObsoleteSockets, bool bSelectSockets, bool bSetupUndo, bool bPrintPythonCommand)
+{
+	TArray<FRigElementKey> SocketKeys;
+	
+	if(URigHierarchy* Hierarchy = GetHierarchy())
+	{
+		for(int32 SocketIndex = 0; SocketIndex < InSkeletalMesh->NumSockets(); SocketIndex++)
+		{
+			if(const USkeletalMeshSocket* Socket = InSkeletalMesh->GetSocketByIndex(SocketIndex))
+			{
+				const FRigElementKey ParentKey(URigHierarchy::GetSanitizedName(Socket->BoneName).GetFName(), ERigElementType::Bone);
+				if(Hierarchy->Contains(ParentKey))
+				{
+					FRigElementKey SocketKey(URigHierarchy::GetSanitizedName(Socket->SocketName).GetFName(), ERigElementType::Null);
+					const FTransform SocketTransform = Socket->GetSocketLocalTransform();
+
+					if(bReplaceExistingSockets && Hierarchy->Contains(SocketKey))
+					{
+						// set the parent - doesn't do anything if the parent is already correct
+						SetParent(SocketKey, ParentKey, false, bSetupUndo, false);
+
+						// update the transforms
+						Hierarchy->SetLocalTransform(SocketKey, SocketTransform, true, true, bSetupUndo, false);
+						Hierarchy->SetLocalTransform(SocketKey, SocketTransform, false, true, bSetupUndo, false);
+					}
+					else
+					{
+						SocketKey = AddNull(SocketKey.Name, ParentKey, SocketTransform, false, bSetupUndo, false);
+					}
+
+					Hierarchy->SetTag(SocketKey, TEXT("MeshSocket"));
+
+					SocketKeys.Add(SocketKey);
+				}
+			}
+		}
+
+		if(bRemoveObsoleteSockets)
+		{
+			const TArray<FRigElementKey> NullKeys = Hierarchy->GetNullKeys();
+			for(const FRigElementKey& NullKey : NullKeys)
+			{
+				if(Hierarchy->HasTag(NullKey, TEXT("MeshSocket")))
+				{
+					if(InSkeletalMesh->FindSocket(NullKey.Name) == nullptr)
+					{
+						RemoveElement(NullKey, bSetupUndo, false);
+					}
+				}
+			}
+		}
+
+		if(bSelectSockets)
+		{
+			SetSelection(SocketKeys, false);
+		}
+	}
+	return SocketKeys;
+}
+
 #if WITH_EDITOR
 
 TArray<FRigElementKey> URigHierarchyController::ImportBonesFromAsset(FString InAssetPath, FName InNameSpace,
@@ -1700,6 +1762,17 @@ TArray<FRigElementKey> URigHierarchyController::ImportPreviewSkeletalMesh(USkele
 		}
 	}
 
+	// import the bones for the preview hierarchy
+	// use the ref skeleton so we'll only see the bones that are actually part of the mesh
+	const TArray<FRigElementKey> MeshSockets = ImportSocketsFromSkeletalMesh(InSkeletalMesh, NAME_None, bReplaceExistingBones, bRemoveObsoleteBones, bSelectBones, bSetupUndo);
+	for(const FRigElementKey& MeshSocket : MeshSockets)
+	{
+		if(FRigBaseElement* Element = Hierarchy->Find(MeshSocket))
+		{
+			Element->CreatedAtInstructionIndex = InstructionIndex;
+		}
+	}
+
 	// create a null to store controls under
 	static const FRigElementKey ControlParentKey(TEXT("Controls"), ERigElementType::Null);
 	if(!Hierarchy->Contains(ControlParentKey))
diff --git a/Engine/Plugins/Animation/ControlRig/Source/ControlRig/Private/Units/Execution/RigUnit_DynamicHierarchy.cpp b/Engine/Plugins/Animation/ControlRig/Source/ControlRig/Private/Units/Execution/RigUnit_DynamicHierarchy.cpp
--- a/Engine/Plugins/Animation/ControlRig/Source/ControlRig/Private/Units/Execution/RigUnit_DynamicHierarchy.cpp
+++ b/Engine/Plugins/Animation/ControlRig/Source/ControlRig/Private/Units/Execution/RigUnit_DynamicHierarchy.cpp
@@ -351,6 +351,13 @@ FRigUnit_HierarchyImportFromSkeleton_Execute()
 						Items.Append(Controller->ImportCurves(Skeleton, NameSpace, false, false, false));
 					}
 				}
+				if(bIncludeMeshSockets)
+				{
+					if(USkeletalMesh* SkeletalMesh = SkelMeshComponent->GetSkeletalMeshAsset())
+					{
+						Items.Append(Controller->ImportSocketsFromSkeletalMesh(SkeletalMesh, NameSpace, false, false, false));
+					}
+				}
 			}
 		}
 	}
diff --git a/Engine/Plugins/Animation/ControlRig/Source/ControlRig/Public/Rigs/RigHierarchyController.h b/Engine/Plugins/Animation/ControlRig/Source/ControlRig/Public/Rigs/RigHierarchyController.h
--- a/Engine/Plugins/Animation/ControlRig/Source/ControlRig/Public/Rigs/RigHierarchyController.h
+++ b/Engine/Plugins/Animation/ControlRig/Source/ControlRig/Public/Rigs/RigHierarchyController.h
@@ -495,6 +495,26 @@ public:
 		bool bSetupUndo = false,
 		bool bPrintPythonCommand = false);
 
+	/**
+	 * Imports the sockets from existing skeleton to the hierarchy as nulls
+	 * @param InSkeletalMesh The skeletal mesh asset to import
+	 * @param InNameSpace The namespace to prefix the socket names with
+	 * @param bReplaceExistingSockets If true existing sockets will be removed
+	 * @param bRemoveObsoleteSockets If true sockets non-existent in the skeleton will be removed from the hierarchy
+	 * @param bSelectSockets If true the sockets will be selected upon import
+	 * @param bSetupUndo If set to true the stack will record the change for undo / redo
+	 * @return The keys of the imported elements
+	 */
+	UFUNCTION(BlueprintCallable, Category = URigHierarchyController)
+	TArray<FRigElementKey> ImportSocketsFromSkeletalMesh(
+	   USkeletalMesh* InSkeletalMesh,
+	   const FName& InNameSpace,
+	   bool bReplaceExistingSockets = true,
+	   bool bRemoveObsoleteSockets = true,
+	   bool bSelectSockets = false,
+	   bool bSetupUndo = false,
+	   bool bPrintPythonCommand = false);
+
 #if WITH_EDITOR
 	/**
 	 * Imports an existing skeleton to the hierarchy
diff --git a/Engine/Plugins/Animation/ControlRig/Source/ControlRig/Public/Units/Execution/RigUnit_DynamicHierarchy.h b/Engine/Plugins/Animation/ControlRig/Source/ControlRig/Public/Units/Execution/RigUnit_DynamicHierarchy.h
--- a/Engine/Plugins/Animation/ControlRig/Source/ControlRig/Public/Units/Execution/RigUnit_DynamicHierarchy.h
+++ b/Engine/Plugins/Animation/ControlRig/Source/ControlRig/Public/Units/Execution/RigUnit_DynamicHierarchy.h
@@ -334,7 +334,7 @@ struct FRigUnit_HierarchyImportFromSkeleton : public FRigUnit_DynamicHierarchyBa
 	FRigUnit_HierarchyImportFromSkeleton()
 	{
 		NameSpace = NAME_None;
-		bIncludeCurves = false;
+		bIncludeCurves = bIncludeMeshSockets = false;
 	}
 
 	RIGVM_METHOD()
@@ -346,6 +346,9 @@ struct FRigUnit_HierarchyImportFromSkeleton : public FRigUnit_DynamicHierarchyBa
 	UPROPERTY(meta = (Input))
 	bool bIncludeCurves;
 
+	UPROPERTY(meta = (Input))
+	bool bIncludeMeshSockets;
+
 	UPROPERTY(meta = (Output))
 	TArray<FRigElementKey> Items;
 };
diff --git a/Engine/Plugins/Animation/ControlRig/Source/ControlRigDeveloper/Private/ControlRigBlueprint.cpp b/Engine/Plugins/Animation/ControlRig/Source/ControlRigDeveloper/Private/ControlRigBlueprint.cpp
--- a/Engine/Plugins/Animation/ControlRig/Source/ControlRigDeveloper/Private/ControlRigBlueprint.cpp
+++ b/Engine/Plugins/Animation/ControlRig/Source/ControlRigDeveloper/Private/ControlRigBlueprint.cpp
@@ -283,6 +283,7 @@ bool UControlRigBlueprint::TurnIntoControlRigModule(bool InAutoConvertHierarchy,
 		if(PreviewSkeletalMesh)
 		{
 			PreviewMeshHierarchy->GetController(true)->ImportBones(PreviewSkeletalMesh->GetSkeleton());
+			PreviewMeshHierarchy->GetController(true)->ImportSocketsFromSkeletalMesh(PreviewSkeletalMesh.Get(), NAME_None, false, false, false, false, false);;
 		}
 
 		// disable compilation
@@ -308,6 +309,14 @@ bool UControlRigBlueprint::TurnIntoControlRigModule(bool InAutoConvertHierarchy,
 						continue;
 					}
 				}
+				if(Key.Type == ERigElementType::Null)
+				{
+					// if this is a mesh socket based null
+					if(PreviewMeshHierarchy->Contains(Key))
+					{
+						continue;
+					}
+				}
 				KeysToSpawn.Add(Key);
 			}
 
@@ -352,6 +361,7 @@ bool UControlRigBlueprint::TurnIntoStandaloneRig(FString* OutErrorMessage)
 		if(PreviewSkeletalMesh)
 		{
 			Hierarchy->GetController(true)->ImportBones(PreviewSkeletalMesh->GetSkeleton());
+			Hierarchy->GetController(true)->ImportSocketsFromSkeletalMesh(PreviewSkeletalMesh.Get(), NAME_None, false, false, false, false, false);
 		}
 	}
 
diff --git a/Engine/Plugins/Animation/ControlRig/Source/ControlRigEditor/Private/Editor/ControlRigEditor.cpp b/Engine/Plugins/Animation/ControlRig/Source/ControlRigEditor/Private/Editor/ControlRigEditor.cpp
--- a/Engine/Plugins/Animation/ControlRig/Source/ControlRigEditor/Private/Editor/ControlRigEditor.cpp
+++ b/Engine/Plugins/Animation/ControlRig/Source/ControlRigEditor/Private/Editor/ControlRigEditor.cpp
@@ -3326,6 +3326,7 @@ void FControlRigEditor::HandlePreviewMeshChanged(USkeletalMesh* InOldSkeletalMes
 						if(InNewSkeletalMesh)
 						{
 							Controller->ImportCurvesFromSkeletalMesh(InNewSkeletalMesh, NAME_None, false, true, true);
+							Controller->ImportSocketsFromSkeletalMesh(InNewSkeletalMesh, NAME_None, true, true, false, true, true);
 						}
 						else
 						{
diff --git a/Engine/Plugins/Animation/ControlRig/Source/ControlRigEditor/Private/Editor/SRigHierarchy.cpp b/Engine/Plugins/Animation/ControlRig/Source/ControlRigEditor/Private/Editor/SRigHierarchy.cpp
--- a/Engine/Plugins/Animation/ControlRig/Source/ControlRigEditor/Private/Editor/SRigHierarchy.cpp
+++ b/Engine/Plugins/Animation/ControlRig/Source/ControlRigEditor/Private/Editor/SRigHierarchy.cpp
@@ -1896,6 +1896,7 @@ void SRigHierarchy::RefreshHierarchy(const FAssetData& InAssetData, bool bOnlyRe
 				Controller->ImportBones(Mesh->GetSkeleton(), NAME_None, true, true, bSelectBones, true, true);
 			}
 			Controller->ImportCurvesFromSkeletalMesh(Mesh, NAME_None, false, true, true);
+			Controller->ImportSocketsFromSkeletalMesh(Mesh, NAME_None, true, true, false, true, true);
 		}
 	}
 
@@ -2024,7 +2025,8 @@ void SRigHierarchy::ImportHierarchy(const FAssetData& InAssetData)
 		check(Controller);
 
 		const TArray<FRigElementKey> ImportedBones = Controller->ImportBones(Mesh->GetSkeleton(), NAME_None, false, false, bSelectBones, true, true);
-		Controller->ImportCurvesFromSkeletalMesh(Mesh, NAME_None, false, true);
+		Controller->ImportCurvesFromSkeletalMesh(Mesh, NAME_None, false, true, true);
+		Controller->ImportSocketsFromSkeletalMesh(Mesh, NAME_None, false, false, false, true, true);
 
 		ControlRigBlueprint->SourceHierarchyImport = Mesh->GetSkeleton();
 		ControlRigBlueprint->SourceCurveImport = Mesh->GetSkeleton();
