Control Rig: Re-enable proxies for layered control rigs

#rnx

diff --git a/Engine/Plugins/Animation/ControlRig/Source/ControlRig/Private/ControlRig.cpp b/Engine/Plugins/Animation/ControlRig/Source/ControlRig/Private/ControlRig.cpp
--- a/Engine/Plugins/Animation/ControlRig/Source/ControlRig/Private/ControlRig.cpp
+++ b/Engine/Plugins/Animation/ControlRig/Source/ControlRig/Private/ControlRig.cpp
@@ -495,6 +495,12 @@ void UControlRig::Evaluate_AnyThread()
 				}
 			}
 		}
+
+		// Interaction solve
+		if (InteractionType != (uint8)EControlRigInteractionType::None)
+		{
+			Execute(FRigUnit_InteractionExecution::EventName);
+		}
 		
 		// Forward solve
 		{
diff --git a/Engine/Plugins/Animation/ControlRig/Source/ControlRig/Public/Rigs/RigHierarchyElements.h b/Engine/Plugins/Animation/ControlRig/Source/ControlRig/Public/Rigs/RigHierarchyElements.h
--- a/Engine/Plugins/Animation/ControlRig/Source/ControlRig/Public/Rigs/RigHierarchyElements.h
+++ b/Engine/Plugins/Animation/ControlRig/Source/ControlRig/Public/Rigs/RigHierarchyElements.h
@@ -1775,6 +1775,10 @@ struct CONTROLRIG_API FRigControlElement final : public FRigMultiParentElement
 		{
 			return false;
 		}
+		if (Settings.AnimationType == ERigControlAnimationType::ProxyControl)
+		{
+			return false;
+		}
 		return true;
 	}
 
