Anim InEngine: Add ability to only load animation onto selected controls. We recently added time range so made sense to get this in also into the new API.
Also by default only load onto masked controls, and remove scale from additive tracks.

diff --git a/Engine/Plugins/Animation/ControlRig/Source/ControlRig/Private/Sequencer/MovieSceneControlRigParameterSection.cpp b/Engine/Plugins/Animation/ControlRig/Source/ControlRig/Private/Sequencer/MovieSceneControlRigParameterSection.cpp
--- a/Engine/Plugins/Animation/ControlRig/Source/ControlRig/Private/Sequencer/MovieSceneControlRigParameterSection.cpp
+++ b/Engine/Plugins/Animation/ControlRig/Source/ControlRig/Private/Sequencer/MovieSceneControlRigParameterSection.cpp
@@ -3190,7 +3190,7 @@ void UMovieSceneControlRigParameterSection::AutoSetTangents(const FName& Control
 }
 #if WITH_EDITOR
 
-void UMovieSceneControlRigParameterSection::RecordControlRigKey(FFrameNumber FrameNumber, bool bSetDefault, EMovieSceneKeyInterpolation InInterpMode)
+void UMovieSceneControlRigParameterSection::RecordControlRigKey(FFrameNumber FrameNumber, bool bSetDefault, EMovieSceneKeyInterpolation InInterpMode, bool bbOntoSelectedControls)
 {
 	if (ControlRig)
 	{
@@ -3233,6 +3233,13 @@ void UMovieSceneControlRigParameterSection::RecordControlRigKey(FFrameNumber Fra
 		TArray<FRigControlElement*> Controls;
 		ControlRig->GetControlsInOrder(Controls);
 		
+		//if additive zero out scale
+		const bool bIsAdditive = GetBlendType() == EMovieSceneBlendType::Additive;
+		TArray<FName> SelectedControls;
+		if (bbOntoSelectedControls)
+		{
+			SelectedControls = ControlRig->CurrentControlSelection();
+		}
 		for (FRigControlElement* ControlElement : Controls)
 		{
 			if (!ControlRig->GetHierarchy()->IsAnimatable(ControlElement))
@@ -3244,6 +3251,19 @@ void UMovieSceneControlRigParameterSection::RecordControlRigKey(FFrameNumber Fra
 			{
 				continue;
 			}
+			//if masked out don't do
+			if (GetControlNameMask(ControlElement->GetFName()) == false)
+			{
+				continue;
+			}
+
+			if (bbOntoSelectedControls)
+			{
+				if (SelectedControls.Contains(ControlElement->GetFName()) == false)
+				{
+					continue;
+				}
+			}
 			int32 ChannelIndex = pChannelIndex->ChannelIndex;
 
 	
@@ -3361,6 +3381,13 @@ void UMovieSceneControlRigParameterSection::RecordControlRigKey(FFrameNumber Fra
 						LastVal = FloatChannels[ChannelIndex + 2]->GetValues()[FloatChannels[ChannelIndex + 2]->GetNumKeys() - 1].Value;
 						FMath::WindRelativeAnglesDegrees(LastVal, Val.Z);
 					}
+					//if additive and scale subtract out unity scale
+					if (bIsAdditive && ControlElement->Settings.ControlType == ERigControlType::Scale)
+					{
+						Val.X -= 1.0f;
+						Val.Y -= 1.0f;
+						Val.Z -= 1.0f;
+					}
 					if (bSetDefault)
 					{
 						FloatChannels[ChannelIndex]->SetDefault(Val.X);
@@ -3426,6 +3453,12 @@ void UMovieSceneControlRigParameterSection::RecordControlRigKey(FFrameNumber Fra
 						ControlElement->Settings.ControlType == ERigControlType::EulerTransform)
 					{
 						CurrentVector = Val.GetScale3D();
+						if (bIsAdditive)
+						{
+							CurrentVector.X -= 1.0f;
+							CurrentVector.Y -= 1.0f;
+							CurrentVector.Z -= 1.0f;
+						}
 						if (bSetDefault)
 						{
 							FloatChannels[ChannelIndex]->SetDefault(CurrentVector.X);
@@ -3624,7 +3657,7 @@ bool UMovieSceneControlRigParameterSection::LoadAnimSequenceIntoThisSection(UAni
 		}
 		ControlRig->Execute(FRigUnit_InverseExecution::EventName);
 
-		RecordControlRigKey(FrameNumber, KeyIndex == 0, InInterpolation);
+		RecordControlRigKey(FrameNumber, KeyIndex == 0, InInterpolation, LoadData.bOntoSelectedControls);
 		Progress.EnterProgressFrame(1);
 		if (Progress.ShouldCancel())
 		{
diff --git a/Engine/Plugins/Animation/ControlRig/Source/ControlRig/Public/Sequencer/MovieSceneControlRigParameterSection.h b/Engine/Plugins/Animation/ControlRig/Source/ControlRig/Public/Sequencer/MovieSceneControlRigParameterSection.h
--- a/Engine/Plugins/Animation/ControlRig/Source/ControlRig/Public/Sequencer/MovieSceneControlRigParameterSection.h
+++ b/Engine/Plugins/Animation/ControlRig/Source/ControlRig/Public/Sequencer/MovieSceneControlRigParameterSection.h
@@ -376,7 +376,7 @@ public:
 
 #if WITH_EDITOR
 	//Function to save control rig key when recording.
-	void RecordControlRigKey(FFrameNumber FrameNumber, bool bSetDefault, EMovieSceneKeyInterpolation InInterpMode);
+	void RecordControlRigKey(FFrameNumber FrameNumber, bool bSetDefault, EMovieSceneKeyInterpolation InInterpMode, bool bOntoSelectedControls = false);
 
 	struct FLoadAnimSequenceData
 	{
@@ -386,6 +386,8 @@ public:
 		float Tolerance = 0.5;
 		//Whether to rest to default control states
 		bool bResetControls = true;
+		//onto selected controls only
+		bool bOntoSelectedControls = false;
 		//Frame to Insert at
 		FFrameNumber StartFrame;
 		//If set only load the animation from the specified range
diff --git a/Engine/Plugins/Animation/ControlRig/Source/ControlRigEditor/Private/ControlRigEditorModule.cpp b/Engine/Plugins/Animation/ControlRig/Source/ControlRigEditor/Private/ControlRigEditorModule.cpp
--- a/Engine/Plugins/Animation/ControlRig/Source/ControlRigEditor/Private/ControlRigEditorModule.cpp
+++ b/Engine/Plugins/Animation/ControlRig/Source/ControlRigEditor/Private/ControlRigEditorModule.cpp
@@ -1274,8 +1274,9 @@ void FControlRigEditorModule::BakeToControlRig(UClass* ControlRigClass, UAnimSeq
 							SmartReduce.TolerancePercentage = KeyReduceTolerance;
 							SmartReduce.SampleRate = BakeFrameRate;
 							TOptional<TRange<FFrameNumber>> AnimFrameRamge; //use whole range there
+							const bool bOntoSelectedControls = false;
 							FControlRigParameterTrackEditor::LoadAnimationIntoSection(SequencerPtr, AnimSequence, SkelMeshComp, FFrameNumber(0),
-								bKeyReduce, SmartReduce, bResetControls, AnimFrameRamge, ParamSection);
+								bKeyReduce, SmartReduce, bResetControls, AnimFrameRamge, bOntoSelectedControls, ParamSection);
 						}
 						SequencerPtr->EmptySelection();
 						SequencerPtr->SelectSection(ParamSection);
diff --git a/Engine/Plugins/Animation/ControlRig/Source/ControlRigEditor/Private/ControlRigSequencerEditorLibrary.cpp b/Engine/Plugins/Animation/ControlRig/Source/ControlRigEditor/Private/ControlRigSequencerEditorLibrary.cpp
--- a/Engine/Plugins/Animation/ControlRig/Source/ControlRigEditor/Private/ControlRigSequencerEditorLibrary.cpp
+++ b/Engine/Plugins/Animation/ControlRig/Source/ControlRigEditor/Private/ControlRigSequencerEditorLibrary.cpp
@@ -1416,12 +1416,12 @@ bool UControlRigSequencerEditorLibrary::LoadAnimSequenceIntoControlRigSection(UM
 {
 	return  UControlRigSequencerEditorLibrary::LoadAnimSequenceIntoControlRigSectionWithRange(MovieSceneSection, AnimSequence, SkelMeshComp, InStartFrame,
 		false, FFrameNumber(0), FFrameNumber(0),
-		TimeUnit, bKeyReduce, Tolerance, Interpolation, bResetControls);
+		TimeUnit, bKeyReduce, Tolerance, Interpolation, bResetControls, false);
 }
 
 bool UControlRigSequencerEditorLibrary::LoadAnimSequenceIntoControlRigSectionWithRange(UMovieSceneSection* MovieSceneSection, UAnimSequence* AnimSequence, USkeletalMeshComponent* SkelMeshComp,
 	FFrameNumber InStartFrame, bool bUseCustomAnimRange, FFrameNumber AnimStartRange, FFrameNumber AnimEndRange,
-	EMovieSceneTimeUnit TimeUnit,bool bKeyReduce, float Tolerance, EMovieSceneKeyInterpolation Interpolation, bool bResetControls)
+	EMovieSceneTimeUnit TimeUnit,bool bKeyReduce, float Tolerance, EMovieSceneKeyInterpolation Interpolation, bool bResetControls, bool bOntoSelectedControls)
 {
 	if (MovieSceneSection == nullptr || AnimSequence == nullptr || SkelMeshComp == nullptr)
 	{
@@ -1444,6 +1444,7 @@ bool UControlRigSequencerEditorLibrary::LoadAnimSequenceIntoControlRigSectionWit
 		Data.Tolerance = Tolerance;
 		Data.bResetControls = bResetControls;
 		Data.StartFrame = InStartFrame;
+		Data.bOntoSelectedControls = bOntoSelectedControls;
 		if (bUseCustomAnimRange)
 		{
 			TRange<FFrameNumber> Range{ 0 };
diff --git a/Engine/Plugins/Animation/ControlRig/Source/ControlRigEditor/Private/Sequencer/ControlRigParameterTrackEditor.cpp b/Engine/Plugins/Animation/ControlRig/Source/ControlRigEditor/Private/Sequencer/ControlRigParameterTrackEditor.cpp
--- a/Engine/Plugins/Animation/ControlRig/Source/ControlRigEditor/Private/Sequencer/ControlRigParameterTrackEditor.cpp
+++ b/Engine/Plugins/Animation/ControlRig/Source/ControlRigEditor/Private/Sequencer/ControlRigParameterTrackEditor.cpp
@@ -1167,7 +1167,8 @@ void FControlRigParameterTrackEditor::SmartReduce(const TSharedPtr<ISequencer>&
 }
 
 bool FControlRigParameterTrackEditor::LoadAnimationIntoSection(const TSharedPtr<ISequencer>& SequencerPtr, UAnimSequence* AnimSequence, USkeletalMeshComponent* SkelMeshComp,
-	FFrameNumber StartFrame, bool bReduceKeys, const FSmartReduceParams& ReduceParams, bool bResetControls, const TOptional<TRange<FFrameNumber>>& InAnimFrameRange, UMovieSceneControlRigParameterSection* ParamSection)
+	FFrameNumber StartFrame, bool bReduceKeys, const FSmartReduceParams& ReduceParams, bool bResetControls, const TOptional<TRange<FFrameNumber>>& InAnimFrameRange, 
+	bool bOntoSelectedControls, UMovieSceneControlRigParameterSection* ParamSection)
 {
 	EMovieSceneKeyInterpolation DefaultInterpolation = SequencerPtr->GetKeyInterpolation();
 	UMovieSceneSequence* OwnerSequence = SequencerPtr->GetFocusedMovieSceneSequence();
@@ -1178,6 +1179,7 @@ bool FControlRigParameterTrackEditor::LoadAnimationIntoSection(const TSharedPtr<
 	Data.bResetControls = bResetControls;
 	Data.StartFrame = StartFrame;
 	Data.AnimFrameRange = InAnimFrameRange;
+	Data.bOntoSelectedControls = bOntoSelectedControls;
 	if (ParamSection->LoadAnimSequenceIntoThisSection(AnimSequence, FFrameNumber(0), OwnerMovieScene, SkelMeshComp,
 		Data, DefaultInterpolation))
 	{
@@ -1337,9 +1339,9 @@ void FControlRigParameterTrackEditor::BakeToControlRig(UClass* InClass, FGuid Ob
 					TOptional<TRange<FFrameNumber>> OptionalRange = Sequencer->GetSubSequenceRange();
 					FFrameNumber StartFrame = OptionalRange.IsSet() ? OptionalRange.GetValue().GetLowerBoundValue() : OwnerMovieScene->GetPlaybackRange().GetLowerBoundValue();
 					TOptional<TRange<FFrameNumber>> AnimLoadFrameRamge; //do whole range
-
+					const bool bOntoSelectedControls = false;
 					LoadAnimationIntoSection(Sequencer, TempAnimSequence, SkelMeshComp, StartFrame,
-						BakeSettings->bReduceKeys, BakeSettings->SmartReduce, BakeSettings->bResetControls, AnimLoadFrameRamge, ParamSection);
+						BakeSettings->bReduceKeys, BakeSettings->SmartReduce, BakeSettings->bResetControls, AnimLoadFrameRamge, bOntoSelectedControls, ParamSection);
 
 					//Turn Off Any Skeletal Animation Tracks
 					TArray<UMovieSceneSkeletalAnimationTrack*> SkelAnimationTracks;
@@ -6498,9 +6500,9 @@ void FControlRigParameterSection::OnAnimationAssetSelected(const FAssetData& Ass
 		{
 			//if we get a new anim sequence we change the start and end range times so it can be the same as the anim sequence
 			static uint32 LastAnimSequenceID = (uint32)(-1);
-			if (AnimSequence->GetUniqueID() != LastAnimSequenceID)
+			ULoadAnimToControlRigSettings * LoadSettings = GetMutableDefault<ULoadAnimToControlRigSettings>();
+			if (LoadSettings->bUseCustomTimeRange == false || AnimSequence->GetUniqueID() != LastAnimSequenceID)
 			{
-				ULoadAnimToControlRigSettings* LoadSettings = GetMutableDefault<ULoadAnimToControlRigSettings>();
 				LoadSettings->StartFrame = 0;
 				LoadSettings->EndFrame = AnimSequence->GetDataModel()->GetNumberOfFrames();
 				LastAnimSequenceID = AnimSequence->GetUniqueID();
@@ -6523,7 +6525,8 @@ void FControlRigParameterSection::OnAnimationAssetSelected(const FAssetData& Ass
 							AnimLoadFrameRamge = Range;
 						}
 						if (!FControlRigParameterTrackEditor::LoadAnimationIntoSection(SequencerPtr, AnimSequence, SkelMeshComp, StartFrame,
-							LoadSettings->bReduceKeys, LoadSettings->SmartReduce, LoadSettings->bResetControls, AnimLoadFrameRamge, Section))
+							LoadSettings->bReduceKeys, LoadSettings->SmartReduce, LoadSettings->bResetControls, AnimLoadFrameRamge,
+							LoadSettings->bOntoSelectedControls, Section))
 						{
 							Transaction.Cancel();
 						}
diff --git a/Engine/Plugins/Animation/ControlRig/Source/ControlRigEditor/Private/Sequencer/ControlRigParameterTrackEditor.h b/Engine/Plugins/Animation/ControlRig/Source/ControlRigEditor/Private/Sequencer/ControlRigParameterTrackEditor.h
--- a/Engine/Plugins/Animation/ControlRig/Source/ControlRigEditor/Private/Sequencer/ControlRigParameterTrackEditor.h
+++ b/Engine/Plugins/Animation/ControlRig/Source/ControlRigEditor/Private/Sequencer/ControlRigParameterTrackEditor.h
@@ -313,7 +313,7 @@ public:
 	//** Function to load animation into a section, returns true if successful
 	static bool LoadAnimationIntoSection(const TSharedPtr<ISequencer>& InSequencer,  UAnimSequence* AnimSequence, USkeletalMeshComponent* SkelMeshComp,
 		FFrameNumber StartFrame, bool bReduceKeys, const FSmartReduceParams& ReduceParams, bool bResetControls,const TOptional<TRange<FFrameNumber>> &InFrameRamge,
-		UMovieSceneControlRigParameterSection* ParamSection);
+		bool bOntoSelectedControls, UMovieSceneControlRigParameterSection* ParamSection);
 
 	//** Function to smart reduce all keys on all constraints based on the parameters
 	static void SmartReduce(const TSharedPtr<ISequencer>& InSequencer, const FSmartReduceParams& InParams, UMovieSceneControlRigParameterSection* InSection);
diff --git a/Engine/Plugins/Animation/ControlRig/Source/ControlRigEditor/Private/Sequencer/LoadAnimToControlRigSettings.cpp b/Engine/Plugins/Animation/ControlRig/Source/ControlRigEditor/Private/Sequencer/LoadAnimToControlRigSettings.cpp
--- a/Engine/Plugins/Animation/ControlRig/Source/ControlRigEditor/Private/Sequencer/LoadAnimToControlRigSettings.cpp
+++ b/Engine/Plugins/Animation/ControlRig/Source/ControlRigEditor/Private/Sequencer/LoadAnimToControlRigSettings.cpp
@@ -6,12 +6,15 @@
 
 ULoadAnimToControlRigSettings::ULoadAnimToControlRigSettings(const FObjectInitializer& Initializer)
 	: Super(Initializer)
-{}
+{
+	Reset();
+}
 
 void ULoadAnimToControlRigSettings::Reset()
 {
 	bReduceKeys = false;
 	SmartReduce.Reset();
 	bUseCustomTimeRange = false;
+	bOntoSelectedControls = false;
 }
 
diff --git a/Engine/Plugins/Animation/ControlRig/Source/ControlRigEditor/Private/Sequencer/LoadAnimToControlRigSettings.h b/Engine/Plugins/Animation/ControlRig/Source/ControlRigEditor/Private/Sequencer/LoadAnimToControlRigSettings.h
--- a/Engine/Plugins/Animation/ControlRig/Source/ControlRigEditor/Private/Sequencer/LoadAnimToControlRigSettings.h
+++ b/Engine/Plugins/Animation/ControlRig/Source/ControlRigEditor/Private/Sequencer/LoadAnimToControlRigSettings.h
@@ -15,6 +15,10 @@ public:
 	ULoadAnimToControlRigSettings(const FObjectInitializer& Initializer);
 
 	GENERATED_BODY()
+	/** Load animation onto only selected controls */
+	UPROPERTY(EditAnywhere, Category = "Controls")
+	bool bOntoSelectedControls = false;
+
 	/** Reduce Keys */
 	UPROPERTY(EditAnywhere, Category = "Reduce Keys")
 	bool bReduceKeys = false;
diff --git a/Engine/Plugins/Animation/ControlRig/Source/ControlRigEditor/Private/Sequencer/SLoadAnimToControlRig.cpp b/Engine/Plugins/Animation/ControlRig/Source/ControlRigEditor/Private/Sequencer/SLoadAnimToControlRig.cpp
--- a/Engine/Plugins/Animation/ControlRig/Source/ControlRigEditor/Private/Sequencer/SLoadAnimToControlRig.cpp
+++ b/Engine/Plugins/Animation/ControlRig/Source/ControlRigEditor/Private/Sequencer/SLoadAnimToControlRig.cpp
@@ -97,7 +97,7 @@ void FLoadAnimToControlRigDialog::GetLoadAnimParams(FLoadAnimToControlRigDelegat
 		.Title(TitleText)
 		.HasCloseButton(true)
 		.SizingRule(ESizingRule::UserSized)
-		.ClientSize(FVector2D(400.0f, 200.0f))
+		.ClientSize(FVector2D(400.0f, 400.0f))
 		.AutoCenter(EAutoCenter::PreferredWorkArea)
 		.SupportsMinimize(false);
 
diff --git a/Engine/Plugins/Animation/ControlRig/Source/ControlRigEditor/Public/ControlRigSequencerEditorLibrary.h b/Engine/Plugins/Animation/ControlRig/Source/ControlRigEditor/Public/ControlRigSequencerEditorLibrary.h
--- a/Engine/Plugins/Animation/ControlRig/Source/ControlRigEditor/Public/ControlRigSequencerEditorLibrary.h
+++ b/Engine/Plugins/Animation/ControlRig/Source/ControlRigEditor/Public/ControlRigSequencerEditorLibrary.h
@@ -142,13 +142,15 @@ public:
 	* @param Tolerance If reducing keys, tolerance about which keys will be removed, smaller tolerance, more keys usually.
 	* @param Interpolation The key interpolation type to set the keys, defaults to EMovieSceneKeyInterpolation::SmartAuto
 	* @param bResetControls If true will reset all controls to initial value on every frame
+	* @param bOntoSelectedControls if true will only load onto selected controls on the Control Rig, otherwise will load onto all controls that aren't masked out
 	* @return returns True if successful, False otherwise
 	*/
 	UFUNCTION(BlueprintCallable, Category = "Editor Scripting | Sequencer Tools | Control Rig")
 	static bool LoadAnimSequenceIntoControlRigSectionWithRange(UMovieSceneSection* MovieSceneSection, UAnimSequence* AnimSequence, USkeletalMeshComponent* SkelMeshComp,
 		FFrameNumber InStartFrame, bool bUseCustomAnimRange, FFrameNumber AnimStartRange, FFrameNumber AnimEndRange,
 		EMovieSceneTimeUnit TimeUnit = EMovieSceneTimeUnit::DisplayRate, bool bKeyReduce = false, float Tolerance = 0.001f,
-		EMovieSceneKeyInterpolation Interpolation = EMovieSceneKeyInterpolation::SmartAuto, bool bResetControls = true);
+		EMovieSceneKeyInterpolation Interpolation = EMovieSceneKeyInterpolation::SmartAuto, bool bResetControls = true,
+		bool bOntoSelectedControls = false);
 
 	UE_DEPRECATED(5.6, "Use LoadAnimSequenceIntoControlRigSectionWithRange that takes UseCustomTimeRange")
 	UFUNCTION(BlueprintCallable, Category = "Editor Scripting | Sequencer Tools | Control Rig", meta = (DeprecatedFunction, DeprecationMessage = "Use LoadAnimSequenceIntoControlRigSectionWithRange that takes UseCustomTimeRange"))
