Constraint creation: store creation info before hiding the menus
- hiding menus destroys the widget, disconnecting delegates

#rnx

diff --git a/Engine/Plugins/Animation/ControlRig/Source/ControlRigEditor/Private/Editor/Constraints/SConstraintsWidget.cpp b/Engine/Plugins/Animation/ControlRig/Source/ControlRigEditor/Private/Editor/Constraints/SConstraintsWidget.cpp
--- a/Engine/Plugins/Animation/ControlRig/Source/ControlRigEditor/Private/Editor/Constraints/SConstraintsWidget.cpp
+++ b/Engine/Plugins/Animation/ControlRig/Source/ControlRigEditor/Private/Editor/Constraints/SConstraintsWidget.cpp
@@ -562,30 +562,38 @@ void SConstraintMenuEntry::NotifyPostChange(const FPropertyChangedEvent& InPrope
 
 FReply SConstraintMenuEntry::CreateFromSelectionOrPicker(const bool bUseDefault) const
 {
+	struct FCreationArgs
+	{
+		FOnConstraintCreated CreationDelegate;
+		ETransformConstraintType ConstraintType = ETransformConstraintType::Parent;
+		FConstrainable Parent;
+	};
+
+	FCreationArgs Args;
+	Args.CreationDelegate = OnConstraintCreated;
+	Args.ConstraintType = ConstraintType;
+	Args.Parent = OnGetParent.IsBound() ? OnGetParent.Execute() : FConstrainable();
+		
 	// hide the creation menu
 	FSlateApplication::Get().DismissAllMenus();
 	
 	// use the selected parent if any
-	if (OnGetParent.IsBound())
+	if (Args.Parent.Object.IsValid())
 	{
-		const FConstrainable Parent = OnGetParent.Execute();
-		if (Parent.Object.IsValid())
+		const TArray<AActor*> Children = UE::Private::GetCurrentSelection();
+		const bool bCreated = UE::Private::CreateConstraint(Children, Args.Parent, Args.ConstraintType, bUseDefault);
+		if (!bCreated)
 		{
-			const TArray<AActor*> Children = UE::Private::GetCurrentSelection();
-			const bool bCreated = UE::Private::CreateConstraint(Children, Parent, ConstraintType, bUseDefault);
-			if (!bCreated)
-			{
-				return FReply::Unhandled();
-			}
-
-			// update list
-			if (OnConstraintCreated.IsBound())
-			{
-				OnConstraintCreated.Execute();
-			}
+			return FReply::Unhandled();
+		}
 
-			return FReply::Handled();
+		// update list
+		if (Args.CreationDelegate.IsBound())
+		{
+			Args.CreationDelegate.Execute();
 		}
+
+		return FReply::Handled();
 	}
 
 	// pick the parent
