Control Rig: Backport backwards compatibility for Modular Rig UFUNCTIONS

#rnx

diff --git a/Engine/Plugins/Animation/ControlRig/Source/ControlRig/Private/ModularRig.cpp b/Engine/Plugins/Animation/ControlRig/Source/ControlRig/Private/ModularRig.cpp
--- a/Engine/Plugins/Animation/ControlRig/Source/ControlRig/Private/ModularRig.cpp
+++ b/Engine/Plugins/Animation/ControlRig/Source/ControlRig/Private/ModularRig.cpp
@@ -968,6 +968,19 @@ FRigModuleInstance* UModularRig::FindModule(const FName& InModuleName)
 	return Result;
 }
 
+FRigModuleInstance* UModularRig::FindModule_Deprecated(const FString& InModulePath)
+{
+	return Modules.FindByPredicate([InModulePath](const FRigModuleInstance& Module)
+	{
+		return Module.GetModulePath_Deprecated().Equals(InModulePath, ESearchCase::IgnoreCase);
+	});
+}
+
+const FRigModuleInstance* UModularRig::FindModule_Deprecated(const FString& InModulePath) const
+{
+	return const_cast<UModularRig*>(this)->FindModule_Deprecated(InModulePath);
+}
+
 const FRigModuleInstance* UModularRig::FindModule(const UControlRig* InModuleInstance) const
 {
 	return const_cast<UModularRig*>(this)->FindModule(InModuleInstance);
@@ -1022,12 +1035,33 @@ FRigModuleInstance* UModularRig::FindModule(const FRigElementKey& InElementKey)
 	return nullptr;
 }
 
-const UControlRig* UModularRig::GetModuleRig(FName InModuleName) const
+const UControlRig* UModularRig::GetModuleRig_Deprecated(FString InModulePath) const
+{
+	return const_cast<UModularRig*>(this)->GetModuleRig(InModulePath);
+}
+
+UControlRig* UModularRig::GetModuleRig(FString InModulePath)
+{
+	if(const FRigModuleInstance* Module = FindModule_Deprecated(InModulePath))
+	{
+		UpdateModuleVariables(Module);
+		if(UControlRig* ModuleRig = Module->GetRig())
+		{
+			return ModuleRig;
+		}
+		UE_LOG(LogControlRig, Error, TEXT("Module '%s' doesn't contain a rig instance."), *InModulePath);
+		return nullptr;
+	}
+	UE_LOG(LogControlRig, Error, TEXT("Module '%s' doesn't exist."), *InModulePath);
+	return nullptr;
+}
+
+const UControlRig* UModularRig::GetModuleRigByName(FName InModuleName) const
 {
-	return const_cast<UModularRig*>(this)->GetModuleRig(InModuleName);
+	return const_cast<UModularRig*>(this)->GetModuleRigByName(InModuleName);
 }
 
-UControlRig* UModularRig::GetModuleRig(FName InModuleName)
+UControlRig* UModularRig::GetModuleRigByName(FName InModuleName)
 {
 	if(const FRigModuleInstance* Module = FindModule(InModuleName))
 	{
@@ -1154,9 +1188,18 @@ TArray<FName> UModularRig::GetEventsForAllModules() const
 	return Events;
 }
 
-TArray<FName> UModularRig::GetEventsForModule(FName InModuleName) const
+TArray<FName> UModularRig::GetEventsForModule(FString InModulePath) const
 {
-	if(const UControlRig* ModuleRig = GetModuleRig(InModuleName))
+	if(const UControlRig* ModuleRig = GetModuleRig_Deprecated(InModulePath))
+	{
+		return ModuleRig->GetSupportedEvents();
+	}
+	return TArray<FName>();
+}
+
+TArray<FName> UModularRig::GetEventsForModuleByName(FName InModuleName) const
+{
+	if(const UControlRig* ModuleRig = GetModuleRigByName(InModuleName))
 	{
 		return ModuleRig->GetSupportedEvents();
 	}
@@ -1205,7 +1248,27 @@ TArray<FName> UModularRig::ExecuteEventOnAllModules(FName InEvent)
 	return ModulesWhichRanEvent;
 }
 
-bool UModularRig::ExecuteEventOnModuleForBP(FName InEvent, FName InModuleName)
+bool UModularRig::ExecuteEventOnModuleForBP(FName InEvent, FString InModulePath)
+{
+	if(FRigModuleInstance* Module = FindModule_Deprecated(InModulePath))
+	{
+		if(const UControlRig* ModuleRig = Module->GetRig())
+		{
+			if(ModuleRig->SupportsEvent(InEvent))
+			{
+				return ExecuteEventOnModule(InEvent, Module); 
+			}
+			UE_LOG(LogControlRig, Error, TEXT("Module '%s' doesn't support the event '%s'."), *InModulePath, *InEvent.ToString());
+			return false;
+		}
+		UE_LOG(LogControlRig, Error, TEXT("Module '%s' doesn't contain a rig instance."), *InModulePath);
+		return false;
+	}
+	UE_LOG(LogControlRig, Error, TEXT("Module '%s' doesn't exist."), *InModulePath);
+	return false;
+}
+
+bool UModularRig::ExecuteEventOnModuleByNameForBP(FName InEvent, FName InModuleName)
 {
 	if(FRigModuleInstance* Module = FindModule(InModuleName))
 	{
diff --git a/Engine/Plugins/Animation/ControlRig/Source/ControlRig/Public/ModularRig.h b/Engine/Plugins/Animation/ControlRig/Source/ControlRig/Public/ModularRig.h
--- a/Engine/Plugins/Animation/ControlRig/Source/ControlRig/Public/ModularRig.h
+++ b/Engine/Plugins/Animation/ControlRig/Source/ControlRig/Public/ModularRig.h
@@ -198,22 +198,32 @@ public:
 	TArray<FName> GetModuleNames() const;
 
 	const FRigModuleInstance* FindModule(const FName& InModuleName) const;
+	const FRigModuleInstance* FindModule_Deprecated(const FString& InModulePath) const;
 	const FRigModuleInstance* FindModule(const UControlRig* InModuleInstance) const;
 	const FRigModuleInstance* FindModule(const FRigBaseElement* InElement) const;
 	const FRigModuleInstance* FindModule(const FRigElementKey& InElementKey) const;
 
 	FRigModuleInstance* FindModule(const FName& InModuleName);
+	FRigModuleInstance* FindModule_Deprecated(const FString& InModulePath);
 	FRigModuleInstance* FindModule(const UControlRig* InModuleInstance);
 	FRigModuleInstance* FindModule(const FRigBaseElement* InElement);
 	FRigModuleInstance* FindModule(const FRigElementKey& InElementKey);
 
 	/**
-	 * @param InModuleName The name of the module to run the event on
+	 * @param InModulePath The path of the module to retrieve the rig for
+	 * @return Returns the rig instance for a given module name
+	 */
+	UFUNCTION(BlueprintCallable, Category=ModularRig, meta=(DeprecatedFunction, DeprecationMessage="Function has been deprecated, please rely on GetModuleRigByName instead."))
+	UControlRig* GetModuleRig(FString InModulePath);
+	const UControlRig* GetModuleRig_Deprecated(FString InModulePath) const;
+
+	/**
+	 * @param InModuleName The name of the module to retrieve the rig for
 	 * @return Returns the rig instance for a given module name
 	 */
 	UFUNCTION(BlueprintCallable, Category=ModularRig)
-	UControlRig* GetModuleRig(FName InModuleName);
-	const UControlRig* GetModuleRig(FName InModuleName) const;
+	UControlRig* GetModuleRigByName(FName InModuleName);
+	const UControlRig* GetModuleRigByName(FName InModuleName) const;
 
 	/**
 	 * @param InModulePath The path of the module to receive the parent path for
@@ -242,11 +252,18 @@ public:
 	TArray<FName> GetEventsForAllModules() const;
 
 	/**
-	 * @param InModuleName The name of the module to receive the events for
+	 * @param InModulePath The path of the module to receive the events for
 	 * @return Returns the names of all supported events for a given module path
 	 */
+	UFUNCTION(BlueprintPure, Category=ModularRig, meta=(DeprecatedFunction, DeprecationMessage="Function has been deprecated, please rely on GetEventsForModuleByName instead."))
+	TArray<FName> GetEventsForModule(FString InModulePath) const;
+
+	/**
+	 * @param InModuleName The name of the module to receive the events for
+	 * @return Returns the names of all supported events for a given module name
+	 */
 	UFUNCTION(BlueprintPure, Category=ModularRig)
-	TArray<FName> GetEventsForModule(FName InModuleName) const;
+	TArray<FName> GetEventsForModuleByName(FName InModuleName) const;
 
 	/**
 	 * @param InEvent The name of the event to run
@@ -255,13 +272,21 @@ public:
 	UFUNCTION(BlueprintCallable, Category=ModularRig)
 	TArray<FName> ExecuteEventOnAllModules(FName InEvent);
 
+	/**
+	 * @param InEvent The name of the event to run
+	 * @param InModulePath The name of the module to run the event on
+	 * @return Returns true if the event was run successfully
+	 */
+	UFUNCTION(BlueprintCallable, Category=ModularRig, meta=(ScriptName="ExecuteEventOnModule", DisplayName="Execute Event On Module", DeprecatedFunction, DeprecationMessage="Function has been deprecated, please rely on ExecuteEventOnModuleByName instead."))
+	bool ExecuteEventOnModuleForBP(FName InEvent, FString InModulePath);
+
 	/**
 	 * @param InEvent The name of the event to run
 	 * @param InModuleName The name of the module to run the event on
 	 * @return Returns true if the event was run successfully
 	 */
-	UFUNCTION(BlueprintCallable, Category=ModularRig, meta=(ScriptName="ExecuteEventOnModule", DisplayName="Execute Event On Module"))
-	bool ExecuteEventOnModuleForBP(FName InEvent, FName InModuleName);
+	UFUNCTION(BlueprintCallable, Category=ModularRig, meta=(ScriptName="ExecuteEventOnModuleByName", DisplayName="Execute Event On Module By Name"))
+	bool ExecuteEventOnModuleByNameForBP(FName InEvent, FName InModuleName);
 	bool ExecuteEventOnModule(const FName& InEvent, FRigModuleInstance* InModule);
 
 	/**
