Control Rig Edit Mode: evaluate the rig at least once before moving controls
- this ensures that the global transform of the first manipulated control reflects the current state of the rig's transform stack (this is particularly important when a control is parented to a bone)
- when manipulating multiple controls, the rig evaluation can't be deferred if a previous manipulation modified an element in the hierarchy of the control beeing moved. (the version of its parent's pose is used to track this)

#rnx

diff --git a/Engine/Plugins/Animation/ControlRig/Source/ControlRig/Private/Rigs/RigHierarchy.cpp b/Engine/Plugins/Animation/ControlRig/Source/ControlRig/Private/Rigs/RigHierarchy.cpp
--- a/Engine/Plugins/Animation/ControlRig/Source/ControlRig/Private/Rigs/RigHierarchy.cpp
+++ b/Engine/Plugins/Animation/ControlRig/Source/ControlRig/Private/Rigs/RigHierarchy.cpp
@@ -2873,11 +2873,13 @@ FText URigHierarchy::GetDisplayNameForUI(const FRigElementKey& InKey, bool bIncl
 
 int32 URigHierarchy::GetPoseVersion(const FRigElementKey& InKey) const
 {
-	if(const FRigTransformElement* TransformElement = Find<FRigTransformElement>(InKey))
-	{
-		return GetPoseVersion(TransformElement->Index);
-	}
-	return INDEX_NONE;
+	const FRigTransformElement* TransformElement = Find<FRigTransformElement>(InKey);
+	return GetPoseVersion(TransformElement);
+}
+
+int32 URigHierarchy::GetPoseVersion(const FRigTransformElement* InTransformElement) const
+{
+	return InTransformElement ? GetPoseVersion(InTransformElement->Index) : INDEX_NONE;
 }
 
 FEdGraphPinType URigHierarchy::GetControlPinType(FRigControlElement* InControlElement) const
diff --git a/Engine/Plugins/Animation/ControlRig/Source/ControlRig/Public/Rigs/RigHierarchy.h b/Engine/Plugins/Animation/ControlRig/Source/ControlRig/Public/Rigs/RigHierarchy.h
--- a/Engine/Plugins/Animation/ControlRig/Source/ControlRig/Public/Rigs/RigHierarchy.h
+++ b/Engine/Plugins/Animation/ControlRig/Source/ControlRig/Public/Rigs/RigHierarchy.h
@@ -2377,6 +2377,12 @@ public:
 	 */
 	int32 GetPoseVersion(const FRigElementKey& InKey) const;
 
+	/**
+	 * Returns the version of the transform / pose on the given element.
+	 * see int32 GetPoseVersion(const FRigElementKey& InKey) const for more info.
+	 */
+	int32 GetPoseVersion(const FRigTransformElement* InTransformElement) const;
+
 	/**
 	 * Returns the global offset transform for a given control element.
 	 * @param InKey The key of the control to retrieve the transform for
diff --git a/Engine/Plugins/Animation/ControlRig/Source/ControlRigEditor/Private/EditMode/ControlRigEditMode.cpp b/Engine/Plugins/Animation/ControlRig/Source/ControlRigEditor/Private/EditMode/ControlRigEditMode.cpp
--- a/Engine/Plugins/Animation/ControlRig/Source/ControlRigEditor/Private/EditMode/ControlRigEditMode.cpp
+++ b/Engine/Plugins/Animation/ControlRig/Source/ControlRigEditor/Private/EditMode/ControlRigEditMode.cpp
@@ -2062,9 +2062,11 @@ bool FControlRigEditMode::InputDelta(FEditorViewportClient* InViewportClient, FV
 		
 		for (auto& Pairs : ControlRigShapeActors)
 		{
-			if (AreRigElementsSelected(ValidControlTypeMask(), Pairs.Key))
+			UControlRig* ControlRig = Pairs.Key;
+			
+			if (AreRigElementsSelected(ValidControlTypeMask(), ControlRig))
 			{
-				FTransform ComponentTransform = GetHostingSceneComponentTransform(Pairs.Key);
+				FTransform ComponentTransform = GetHostingSceneComponentTransform(ControlRig);
 
 				if (bIsChangingControlShapeTransform)
 				{
@@ -2094,18 +2096,7 @@ bool FControlRigEditMode::InputDelta(FEditorViewportClient* InViewportClient, FV
 					FTransform InOutLocal = FTransform::Identity;
 					
 					bool const bJustStartedManipulation = !bManipulatorMadeChange;
-					bool bAnyAdditiveRig = false;
-					for (AControlRigShapeActor* ShapeActor : Pairs.Value)
-					{
-						if (ShapeActor->ControlRig.IsValid())
-						{
-							if (ShapeActor->ControlRig->IsAdditive())
-							{
-								bAnyAdditiveRig = true;
-								break;
-							}
-						}
-					}
+					const bool bAnyAdditiveRig = ControlRig->IsAdditive();
 
 					TMap<AControlRigShapeActor*, TArray<TFunction<void()>>> TasksPerActor;
 					for (AControlRigShapeActor* ShapeActor : Pairs.Value)
@@ -2118,6 +2109,13 @@ bool FControlRigEditMode::InputDelta(FEditorViewportClient* InViewportClient, FV
 								GEditor->BeginTransaction(LOCTEXT("MoveControlTransaction", "Move Control"));
 							}
 
+							// cache interaction dependencies + evaluate the rig at least once before manipulating anything
+							if (ControlRig->ElementsBeingInteracted.IsEmpty() && InteractionScopes.Contains(ControlRig))
+							{
+								GetInteractionDependencies(ControlRig);
+								EvaluateRig(ControlRig);
+							}
+
 							// Cannot benefit of same local transform when applying to additive rigs
 							if (!bAnyAdditiveRig)
 							{
@@ -2137,9 +2135,9 @@ bool FControlRigEditMode::InputDelta(FEditorViewportClient* InViewportClient, FV
 
 							if(bJustStartedManipulation)
 							{
-								if(const FRigControlElement* ControlElement = Pairs.Key->FindControl(ShapeActor->ControlName))
+								if(const FRigControlElement* ControlElement = ControlRig->FindControl(ShapeActor->ControlName))
 								{
-									ShapeActor->OffsetTransform = Pairs.Key->GetHierarchy()->GetGlobalControlOffsetTransform(ControlElement->GetKey(), false);
+									ShapeActor->OffsetTransform = ControlRig->GetHierarchy()->GetGlobalControlOffsetTransform(ControlElement->GetKey(), false);
 								}
 							}
 
@@ -3568,11 +3566,7 @@ void FControlRigEditMode::ZeroTransforms(bool bSelectionOnly)
 
 		if (bHasNonDefaultParent == true) //now we have the initial pose setup we need to get the global transforms as specified now then set them in the current parent space
 		{
-			{
-				const TGuardValue<float> AbsoluteTimeGuard(ControlRig->AbsoluteTime, ControlRig->AbsoluteTime);
-				const TGuardValue<bool> GuardEvaluationType(ControlRig->bEvaluationTriggeredFromInteraction, true);
-				ControlRig->Evaluate_AnyThread();
-			}
+			EvaluateRig(ControlRig);
 
 			//get global transforms
 			TMap<FRigElementKey, FTransform> GlobalTransforms;
@@ -3710,11 +3704,7 @@ void FControlRigEditMode::InvertInputPose(bool bSelectionOnly)
 		}
 
 		const TArray<FRigControlElement*> ModifiedElements = ControlRig->InvertInputPose(SelectedRigElements, EControlRigSetKey::Never);
-		{
-			const TGuardValue<float> AbsoluteTimeGuard(ControlRig->AbsoluteTime, ControlRig->AbsoluteTime);
-			const TGuardValue<bool> GuardEvaluationType(ControlRig->bEvaluationTriggeredFromInteraction, true);
-			ControlRig->Evaluate_AnyThread();
-		}
+		EvaluateRig(ControlRig);
 
 		for (FRigControlElement* ControlElement : ModifiedElements)
 		{
@@ -5218,8 +5208,24 @@ void FControlRigEditMode::MoveControlShape( AControlRigShapeActor* ShapeActor,
 			return;
 		}
 	}
+	// else: world, explicit or doing scale.
+
+	FInteractionDependencyCache& Dependencies = GetInteractionDependencies(ControlRig);
 
-	//not local or doing scale.
+	// for readability
+	const bool bSolveImmediately = (!bUseLocal && bCalcLocal) ||
+		Dependencies.HasDownwardDependencies(ShapeActor->GetElementKey()) ||
+		Dependencies.CheckAndUpdateParentsPoseVersion();
+	const bool bQueueTasks = !bSolveImmediately;
+
+	if(bQueueTasks)
+	{
+		RigsToEvaluateDuringThisTick.AddUnique(ControlRig);
+	}
+	else
+	{
+		EvaluateRig(ControlRig);
+	}
 
 	// Get the global transform from shape actor to avoid drifting
 	FTransform CurrentTransform;
@@ -5340,23 +5346,6 @@ void FControlRigEditMode::MoveControlShape( AControlRigShapeActor* ShapeActor,
 			ControlRig->GetHierarchy()->SetPose(ControlRig->InputPoseOnDebuggedRig);
 		}
 
-		// for readability
-		const FInteractionDependencyCache& Dependencies = GetInteractionDependencies(ControlRig);
-		const bool bSolveImmediately = (!bUseLocal && bCalcLocal) ||
-			Dependencies.HasDownwardDependencies(ShapeActor->GetElementKey());
-		const bool bQueueTasks = !bSolveImmediately;
-
-		if(bQueueTasks)
-		{
-			RigsToEvaluateDuringThisTick.AddUnique(ControlRig);
-		}
-		else
-		{
-			const TGuardValue<float> AbsoluteTimeGuard(ControlRig->AbsoluteTime, ControlRig->AbsoluteTime);
-			const TGuardValue<bool> GuardEvaluationType(ControlRig->bEvaluationTriggeredFromInteraction, true);
-			ControlRig->Evaluate_AnyThread();
-		}
-
 		const FTransform ToWorldTransformCopy = ToWorldTransform;
 		
 		const TFunction<void()> SetControlShapeTask = [this, ShapeActor, ControlRig, ControlElement, NewTransform, ToWorldTransformCopy, Context,
@@ -5375,9 +5364,7 @@ void FControlRigEditMode::MoveControlShape( AControlRigShapeActor* ShapeActor,
 				}
 				else
 				{
-					const TGuardValue<float> AbsoluteTimeGuard(ControlRig->AbsoluteTime, ControlRig->AbsoluteTime);
-					const TGuardValue<bool> GuardEvaluationType(ControlRig->bEvaluationTriggeredFromInteraction, true);
-					ControlRig->Evaluate_AnyThread();
+					EvaluateRig(ControlRig);
 				}
 			}
 		};
@@ -6324,7 +6311,7 @@ void FControlRigEditMode::UpdateRotationContext()
 	}
 }
 
-const FInteractionDependencyCache& FControlRigEditMode::GetInteractionDependencies(UControlRig* InControlRig)
+FInteractionDependencyCache& FControlRigEditMode::GetInteractionDependencies(UControlRig* InControlRig)
 {
 	if (FInteractionDependencyCache* DependencyCache = InteractionDependencies.Find(InControlRig))
 	{
@@ -6357,6 +6344,7 @@ const FInteractionDependencyCache& FControlRigEditMode::GetInteractionDependenci
 
 		// build dependencies between selected controls
 		FInteractionDependencyCache& NewDependency = InteractionDependencies.Add(InControlRig);
+		NewDependency.WeakHierarchy = Hierarchy;
 
 		// NOTE: this is not enough for modular rigs since there are several VMs
 		const URigHierarchy::TElementDependencyMap HierarchyDependencyMap = Hierarchy->GetDependenciesForVM(InControlRig->GetVM());
@@ -6378,6 +6366,17 @@ const FInteractionDependencyCache& FControlRigEditMode::GetInteractionDependenci
 					NewDependency.Children.Add(OtherControl->GetKey());
 				}
 			}
+
+			// store parents' pose versions
+			const FRigBaseElementParentArray Parents = Hierarchy->GetParents(Control);
+			NewDependency.ParentsPoseVersion.Reserve(Parents.Num());
+			for (const FRigBaseElement* Parent: Parents)
+			{
+				if (const FRigTransformElement* TransformParent = Cast<FRigTransformElement>(Parent))
+				{
+					NewDependency.ParentsPoseVersion.FindOrAdd(Parent->GetIndex()) = Hierarchy->GetPoseVersion(TransformParent);
+				}
+			}
 		}
 
 		return NewDependency;
@@ -6689,14 +6688,22 @@ FControlRigEditMode::FPendingControlRigEvaluator::~FPendingControlRigEvaluator()
 	{
 		for(UControlRig* ControlRig : EditMode->RigsToEvaluateDuringThisTick)
 		{
-			const TGuardValue<float> AbsoluteTimeGuard(ControlRig->AbsoluteTime, ControlRig->AbsoluteTime);
-			const TGuardValue<bool> GuardEvaluationType(ControlRig->bEvaluationTriggeredFromInteraction, true);
-			ControlRig->Evaluate_AnyThread();
+			FControlRigEditMode::EvaluateRig(ControlRig);
 		}
 		EditMode->RigsToEvaluateDuringThisTick.Reset();
 	}
 }
 
+void FControlRigEditMode::EvaluateRig(UControlRig* InControlRig)
+{
+	if (InControlRig)
+	{
+		const TGuardValue<float> AbsoluteTimeGuard(InControlRig->AbsoluteTime, InControlRig->AbsoluteTime);
+		const TGuardValue<bool> GuardEvaluationType(InControlRig->bEvaluationTriggeredFromInteraction, true);
+		InControlRig->Evaluate_AnyThread();
+	}
+}
+
 TOptional<FTransform> FControlRigEditMode::GetConstraintParentTransform(const UControlRig* InControlRig, const FName& InControlName) const
 {
 	static const TOptional<FTransform> Dummy;
diff --git a/Engine/Plugins/Animation/ControlRig/Source/ControlRigEditor/Public/EditMode/ControlRigEditMode.h b/Engine/Plugins/Animation/ControlRig/Source/ControlRigEditor/Public/EditMode/ControlRigEditMode.h
--- a/Engine/Plugins/Animation/ControlRig/Source/ControlRigEditor/Public/EditMode/ControlRigEditMode.h
+++ b/Engine/Plugins/Animation/ControlRig/Source/ControlRigEditor/Public/EditMode/ControlRigEditMode.h
@@ -450,7 +450,7 @@ protected:
 	TMap<UControlRig*, UE::ControlRigEditMode::FInteractionDependencyCache> InteractionDependencies;
 	
 	/** Returns the interaction dependencies of that ControlRig. */
-	const UE::ControlRigEditMode::FInteractionDependencyCache& GetInteractionDependencies(UControlRig* InControlRig);
+	UE::ControlRigEditMode::FInteractionDependencyCache& GetInteractionDependencies(UControlRig* InControlRig);
 	
 public: 
 	/** Clear all selected RigElements */
@@ -657,6 +657,9 @@ private:
 
 		FControlRigEditMode* EditMode;
 	};
+
+	static void EvaluateRig(UControlRig* InControlRig);
+	
 private:
 
 	// get default/mutable settings
diff --git a/Engine/Plugins/Animation/ControlRig/Source/ControlRigEditor/Public/EditMode/ControlRigEditModeUtil.h b/Engine/Plugins/Animation/ControlRig/Source/ControlRigEditor/Public/EditMode/ControlRigEditModeUtil.h
--- a/Engine/Plugins/Animation/ControlRig/Source/ControlRigEditor/Public/EditMode/ControlRigEditModeUtil.h
+++ b/Engine/Plugins/Animation/ControlRig/Source/ControlRigEditor/Public/EditMode/ControlRigEditModeUtil.h
@@ -5,6 +5,7 @@
 #include "Math/MathFwd.h"
 #include "Math/Vector.h"
 #include "Math/Rotator.h"
+#include "Rigs/RigHierarchy.h"
 
 #include "Rigs/RigHierarchyDefines.h"
 
@@ -54,7 +55,9 @@ struct FTransformContext
 
 /**
  * FInteractionDependencyCache provides a minimal "dependency graph" between the selected controls during interaction.
- * it simply stores information on who is a child and who is a parent in whatever is being manipulated.
+ * it stores information on who is a child and who is a parent in whatever is being manipulated, as well as
+ * the parents' pose versions (in the complete hierarchy), to know whether any previous manipulation has modified
+ * the parents' transform of an element being manipulated.
  */
 
 struct FInteractionDependencyCache
@@ -69,8 +72,36 @@ struct FInteractionDependencyCache
 		return Children.Contains(InKey);
 	}
 
+	bool CheckAndUpdateParentsPoseVersion()
+	{
+		if (ParentsPoseVersion.IsEmpty())
+		{
+			return false;
+		}
+		
+		if (URigHierarchy* Hierarchy = WeakHierarchy.Get())
+		{
+			bool bHasChanged = false;
+			
+			for (auto&[ElementIndex, PoseVersion]: ParentsPoseVersion)
+			{
+				const int32 NewPoseVersion = Hierarchy->GetPoseVersion(Hierarchy->Get<FRigTransformElement>(ElementIndex));
+				if (NewPoseVersion != PoseVersion)
+				{
+					PoseVersion = NewPoseVersion;
+					bHasChanged = true;
+				} 
+			}
+
+			return bHasChanged;
+		}
+		return false;
+	}
+
 	TSet<FRigElementKey> Parents;
 	TSet<FRigElementKey> Children;
+	TMap<int32, int32> ParentsPoseVersion;
+	TWeakObjectPtr<URigHierarchy> WeakHierarchy;
 };
 
 /**
