Anim In Engine: Speculative Linux test fix.

diff --git a/Engine/Plugins/Animation/ControlRig/Source/ControlRigEditor/Private/Sequencer/ControlRigParameterTrackEditor.cpp b/Engine/Plugins/Animation/ControlRig/Source/ControlRigEditor/Private/Sequencer/ControlRigParameterTrackEditor.cpp
--- a/Engine/Plugins/Animation/ControlRig/Source/ControlRigEditor/Private/Sequencer/ControlRigParameterTrackEditor.cpp
+++ b/Engine/Plugins/Animation/ControlRig/Source/ControlRigEditor/Private/Sequencer/ControlRigParameterTrackEditor.cpp
@@ -2774,55 +2774,58 @@ void FControlRigParameterTrackEditor::OnActivateSequenceChanged(FMovieSceneSeque
 		{
 			//we need to make sure pending deactivated edit modes, including a possible control rig edit mode
 			//get totally removed which only happens on a tick 
-			if (GEditor->GetActiveViewport() && GEditor->GetActiveViewport()->GetClient())
-			{
-				if (FEditorModeTools* EditorModeTools = GetEditorModeTools())
+				if (GLevelEditorModeTools().HasToolkitHost())
 				{
-					FViewport* ActiveViewport = GEditor->GetActiveViewport();
-					FEditorViewportClient* EditorViewClient = (FEditorViewportClient*)ActiveViewport->GetClient();
-					EditorModeTools->Tick(EditorViewClient, 0.033f);
-				}
-			}
-			
-			GEditor->GetTimerManager()->SetTimerForNextTick([this]()
-			{
-				//now we can recreate it
-				FControlRigEditMode* ControlRigEditMode = GetEditMode(true);
-				if (ControlRigEditMode)
-				{
-					bool bSequencerSet = false;
+					if (GEditor->GetActiveViewport() && GEditor->GetActiveViewport()->GetClient())
+					{
+						if (FEditorModeTools* EditorModeTools = GetEditorModeTools())
+						{
+							FViewport* ActiveViewport = GEditor->GetActiveViewport();
+							FEditorViewportClient* EditorViewClient = (FEditorViewportClient*)ActiveViewport->GetClient();
+							EditorModeTools->Tick(EditorViewClient, 0.033f);
+						}
+					}
 
-					TSharedPtr<ISequencer> Sequencer = GetSequencer();
-					for (TWeakObjectPtr<UControlRig>& ControlRig : BoundControlRigs)
+					GEditor->GetTimerManager()->SetTimerForNextTick([this]()
 					{
-						if (ControlRig.IsValid())
+						//now we can recreate it
+						FControlRigEditMode* ControlRigEditMode = GetEditMode(true);
+						if (ControlRigEditMode)
 						{
-							ControlRigEditMode->AddControlRigObject(ControlRig.Get(), Sequencer);
-							bSequencerSet = true;
-							
-							for (int32 Index = 0; Index < PreviousSelectedControlRigs.Num(); ++Index)
+							bool bSequencerSet = false;
+
+							TSharedPtr<ISequencer> Sequencer = GetSequencer();
+							for (TWeakObjectPtr<UControlRig>& ControlRig : BoundControlRigs)
 							{
-								if (ControlRig.Get()->GetClass() == PreviousSelectedControlRigs[Index].Key)
+								if (ControlRig.IsValid())
 								{
-									for (const FName& ControlName : PreviousSelectedControlRigs[Index].Value)
+									ControlRigEditMode->AddControlRigObject(ControlRig.Get(), Sequencer);
+									bSequencerSet = true;
+
+									for (int32 Index = 0; Index < PreviousSelectedControlRigs.Num(); ++Index)
 									{
-										ControlRig.Get()->SelectControl(ControlName, true);
+										if (ControlRig.Get()->GetClass() == PreviousSelectedControlRigs[Index].Key)
+										{
+											for (const FName& ControlName : PreviousSelectedControlRigs[Index].Value)
+											{
+												ControlRig.Get()->SelectControl(ControlName, true);
+											}
+											PreviousSelectedControlRigs.RemoveAt(Index);
+											break;
+										}
 									}
-									PreviousSelectedControlRigs.RemoveAt(Index);
-									break;
 								}
 							}
-						}
-					}
 
-					if (!bSequencerSet)
-					{
-						ControlRigEditMode->SetSequencer(Sequencer);
-					}
-				
+							if (!bSequencerSet)
+							{
+								ControlRigEditMode->SetSequencer(Sequencer);
+							}
+
+						}
+						PreviousSelectedControlRigs.Reset();
+					});
 				}
-				PreviousSelectedControlRigs.Reset();
-			});
 		});
 	}
 }
