Control Rig: Show module path in Anim Outliner

#rnx

diff --git a/Engine/Plugins/Animation/ControlRig/Source/ControlRigEditor/Private/EditMode/SControlRigOutliner.cpp b/Engine/Plugins/Animation/ControlRig/Source/ControlRigEditor/Private/EditMode/SControlRigOutliner.cpp
--- a/Engine/Plugins/Animation/ControlRig/Source/ControlRigEditor/Private/EditMode/SControlRigOutliner.cpp
+++ b/Engine/Plugins/Animation/ControlRig/Source/ControlRigEditor/Private/EditMode/SControlRigOutliner.cpp
@@ -85,7 +85,7 @@ FText FMultiRigData::GetDisplayName() const
 			{
 				const FRigElementKey& RigElementKey = Key.GetValue();
 
-				static constexpr bool bIncludeNameSpace = false;
+				static constexpr bool bIncludeNameSpace = true;
 				const FText DisplayNameForUI = Hierarchy->GetDisplayNameForUI(RigElementKey, bIncludeNameSpace);
 				if (!DisplayNameForUI.IsEmpty())
 				{
