Anim In Engine: Fix multiple enums parameter curves getting created. Regression from checking previous name, it was replacing the new parameter name with the old parameter name where we want the old parameter name replaced with the new parameter name.  Just reverted the old algorithm since not sure if the new one is needed for Helge's modular long/short path conversions,  so left that code alone.

diff --git a/Engine/Plugins/Animation/ControlRig/Source/ControlRig/Private/Sequencer/MovieSceneControlRigParameterSection.cpp b/Engine/Plugins/Animation/ControlRig/Source/ControlRig/Private/Sequencer/MovieSceneControlRigParameterSection.cpp
--- a/Engine/Plugins/Animation/ControlRig/Source/ControlRig/Private/Sequencer/MovieSceneControlRigParameterSection.cpp
+++ b/Engine/Plugins/Animation/ControlRig/Source/ControlRig/Private/Sequencer/MovieSceneControlRigParameterSection.cpp
@@ -2642,6 +2642,89 @@ void UMovieSceneControlRigParameterSection::KeyWeightValue(FFrameNumber InFrame,
 	Weight.AutoSetTangents();
 }
 
+bool UMovieSceneControlRigParameterSection::RenameParameterName(const FName& OldParameterName, const FName& NewParameterName, TOptional<ERigControlType> ControlType)
+{
+	bool bWasReplaced = false;
+
+	auto RenameParameterNameInner = [this, &bWasReplaced, OldParameterName, NewParameterName](auto& ParameterNamesAndCurves)
+		{
+			for (auto& ParameterNameAndCurve : ParameterNamesAndCurves)
+			{
+				if (ParameterNameAndCurve.ParameterName == OldParameterName)
+				{
+					if (!bWasReplaced)
+					{
+						Modify();
+						bWasReplaced = true;
+					}
+					ParameterNameAndCurve.ParameterName = NewParameterName;
+					break;
+				}
+			}
+		};
+	if (ControlType.IsSet())
+	{
+		switch (ControlType.GetValue())
+		{
+		case ERigControlType::Float:
+		case ERigControlType::ScaleFloat:
+		{
+			RenameParameterNameInner(ScalarParameterNamesAndCurves);
+			break;
+		}
+		case ERigControlType::Bool:
+		{
+			RenameParameterNameInner(BoolParameterNamesAndCurves);
+			break;
+		}
+		case ERigControlType::Integer:
+		{
+			RenameParameterNameInner(IntegerParameterNamesAndCurves);
+			RenameParameterNameInner(EnumParameterNamesAndCurves);
+			break;
+		}
+		case ERigControlType::Vector2D:
+		{
+			RenameParameterNameInner(Vector2DParameterNamesAndCurves);
+			break;
+		}
+		case ERigControlType::Position:
+		case ERigControlType::Rotator:
+		case ERigControlType::Scale:
+		{
+			RenameParameterNameInner(VectorParameterNamesAndCurves);
+			break;
+		}
+		break;
+		case ERigControlType::Transform:
+		case ERigControlType::EulerTransform:
+		case ERigControlType::TransformNoScale:
+		{
+			RenameParameterNameInner(TransformParameterNamesAndCurves);
+			break;
+		}
+
+		};
+	}
+	else
+	{
+		RenameParameterNameInner(ScalarParameterNamesAndCurves);
+		RenameParameterNameInner(BoolParameterNamesAndCurves);
+		RenameParameterNameInner(EnumParameterNamesAndCurves);
+		RenameParameterNameInner(IntegerParameterNamesAndCurves);
+		RenameParameterNameInner(Vector2DParameterNamesAndCurves);
+		RenameParameterNameInner(VectorParameterNamesAndCurves);
+		RenameParameterNameInner(ColorParameterNamesAndCurves);
+		RenameParameterNameInner(TransformParameterNamesAndCurves);
+	}
+
+	if (bWasReplaced)
+	{
+		ReconstructChannelProxy();
+	}
+	return bWasReplaced;
+}
+
 void UMovieSceneControlRigParameterSection::RecreateWithThisControlRig(UControlRig* InControlRig, bool bSetDefault)
 {
 	bool bSameControlRig = (ControlRig == InControlRig);
@@ -2701,25 +2784,9 @@ void UMovieSceneControlRigParameterSection::RecreateWithThisControlRig(UControlR
 			}
 		}
 	}
-
-	// rename all existing parameters based on previous name information.
-	ForEachParameter([this, Hierarchy](FBaseParameterNameAndValue* InOutParameter)
-	{
-		const FRigElementKey ControlKey = FRigElementKey(InOutParameter->ParameterName, ERigElementType::Control);
-		if(const FRigHierarchyKey* NewControlKey = Hierarchy->GetPreviousNameMap().Find(ControlKey))
-		{
-			InOutParameter->ParameterName = NewControlKey->GetFName();
-		}
-		else
-		{
-			const FRigElementKey CurveKey = FRigElementKey(InOutParameter->ParameterName, ERigElementType::Curve);
-			if(const FRigHierarchyKey* NewCurveKey = Hierarchy->GetPreviousNameMap().Find(CurveKey))
-			{
-				InOutParameter->ParameterName = NewCurveKey->GetFName();
-			}
-		}
-	});
-
+	
+	
+	//HELGE MAKE SURE THIS WORKS FOR THE MODULAR STUFF
 	// rename all existing parameters based on short name vs long name.
 	// this also recovers from controls being stored with an original long path
 	// which may now be outdated. UControlRig::FindControl has backwards compat for that.
@@ -2745,6 +2812,22 @@ void UMovieSceneControlRigParameterSection::RecreateWithThisControlRig(UControlR
 			continue;
 		}
 
+		FName PreviousName = Hierarchy->GetPreviousName(ControlElement->GetKey());
+		if (PreviousName != NAME_None && PreviousName != ControlElement->GetKey().Name)
+		{
+			TOptional<ERigControlType> ControlType = ControlElement->Settings.ControlType;
+			RenameParameterName(PreviousName, ControlElement->GetKey().Name, ControlType);
+		}
+		if (ControlElement->Settings.ControlType == ERigControlType::Float ||
+			ControlElement->Settings.ControlType == ERigControlType::ScaleFloat)
+		{
+			if (const FName* OldCurveControlName = CurveControlNameRemapping.Find(ControlElement->GetFName()))
+			{
+				TOptional<ERigControlType> ControlType = ControlElement->Settings.ControlType;
+				RenameParameterName(*OldCurveControlName, ControlElement->GetKey().Name, ControlType);
+			}
+		}
+
 		const FName& ControlName = ControlElement->GetFName();
 
 		switch (ControlElement->Settings.ControlType)
diff --git a/Engine/Plugins/Animation/ControlRig/Source/ControlRig/Public/Sequencer/MovieSceneControlRigParameterSection.h b/Engine/Plugins/Animation/ControlRig/Source/ControlRig/Public/Sequencer/MovieSceneControlRigParameterSection.h
--- a/Engine/Plugins/Animation/ControlRig/Source/ControlRig/Public/Sequencer/MovieSceneControlRigParameterSection.h
+++ b/Engine/Plugins/Animation/ControlRig/Source/ControlRig/Public/Sequencer/MovieSceneControlRigParameterSection.h
@@ -275,6 +275,10 @@ public:
 	virtual bool PopulateEvaluationFieldImpl(const TRange<FFrameNumber>& EffectiveRange, const FMovieSceneEvaluationFieldEntityMetaData& InMetaData, FMovieSceneEntityComponentFieldBuilder* OutFieldBuilder) override;
 	//mz todo virtual void InterrogateEntityImpl(UMovieSceneEntitySystemLinker* EntityLinker, const FEntityImportParams& Params, FImportedEntity* OutImportedEntity) override;
 
+
+	bool RenameParameterName(const FName& OldParameterName, const FName& NewParameterName, TOptional<ERigControlType> ControlType);
+
+
 private:
 
 	FSpaceChannelAddedEvent OnSpaceChannelAdded;
