Control Rig Edit Mode: autokey on release only

keyframing is currently executed for every mouse move, resulting in a huge fps drop.
the general idea is to defer keyframing to release (as sequencer does for scene components).
this is even more noticeable when dealing with spaces or constraints, which require sequencer to evaluate the previous frame in order to calculate compensation.

NOTE: this si currently enabled using AnimMode.AutoKeyOnRelease true but will be exposed as a preference based on user feedback in a following CL

diff --git a/Engine/Plugins/Animation/ControlRig/Source/ControlRig/Public/ControlRig.h b/Engine/Plugins/Animation/ControlRig/Source/ControlRig/Public/ControlRig.h
--- a/Engine/Plugins/Animation/ControlRig/Source/ControlRig/Public/ControlRig.h
+++ b/Engine/Plugins/Animation/ControlRig/Source/ControlRig/Public/ControlRig.h
@@ -1031,6 +1031,22 @@ public:
 		}
 	}
 
+	const TArray<FRigElementKey>& GetElementsBeingInteracted() const
+	{
+		if(ensure(ControlRig.IsValid()))
+		{
+			return ControlRig->ElementsBeingInteracted;
+		}
+
+		static const TArray<FRigElementKey> DummyElements;
+		return DummyElements;
+	}
+
+	UControlRig* GetControlRig() const
+	{
+		return ControlRig.Get();
+	}
+
 private:
 
 	TWeakObjectPtr<UControlRig> ControlRig;
diff --git a/Engine/Plugins/Animation/ControlRig/Source/ControlRigEditor/Private/EditMode/ControlRigEditMode.cpp b/Engine/Plugins/Animation/ControlRig/Source/ControlRigEditor/Private/EditMode/ControlRigEditMode.cpp
--- a/Engine/Plugins/Animation/ControlRig/Source/ControlRigEditor/Private/EditMode/ControlRigEditMode.cpp
+++ b/Engine/Plugins/Animation/ControlRig/Source/ControlRigEditor/Private/EditMode/ControlRigEditMode.cpp
@@ -97,9 +97,19 @@ namespace UE::ControlRigEditMode::Private
 	{
 		return IConsoleManager::Get().FindConsoleVariable(FocusModeName);
 	}
+
+	const TCHAR* AutoKeyOnReleaseName = TEXT("AnimMode.AutoKeyOnRelease");
+
+	static bool bAutoKeyOnRelease = false;
+	static FAutoConsoleVariableRef CVarAutoKeyOnRelease(
+	   AutoKeyOnReleaseName,
+	   bAutoKeyOnRelease,
+	   TEXT("Only add keys on release (or when sequencer is playing)")
+	   );
 }
 
 using namespace UE::ControlRigEditMode;
+using namespace UE::AnimationEditMode;
 
 void UControlRigEditModeDelegateHelper::OnPoseInitialized()
 {
@@ -1130,6 +1140,7 @@ bool FControlRigEditMode::HandleBeginTransform(const FEditorViewportClient* InVi
 	InteractionType = GetInteractionType(InViewportClient);
 	bIsTracking = true;
 	InteractionDependencies.Reset();
+	Keyframer.Init(DeferAutokeyOnMouseRelease());
 	
 	if (InteractionScopes.Num() == 0)
 	{
@@ -1223,7 +1234,13 @@ bool FControlRigEditMode::HandleEndTransform(FEditorViewportClient* InViewportCl
 	{
 		return false;
 	}
-	
+
+	const UE::Widget::EWidgetMode WidgetMode = InViewportClient->GetWidgetMode();
+	FTransformContext TransformContext;
+	TransformContext.bTranslation = WidgetMode == UE::Widget::WM_Translate || WidgetMode == UE::Widget::WM_TranslateRotateZ;
+	TransformContext.bRotation = WidgetMode == UE::Widget::WM_Rotate || WidgetMode == UE::Widget::WM_TranslateRotateZ;
+	TransformContext.bScale = WidgetMode == UE::Widget::WM_Scale;
+
 	const bool bWasInteracting = bManipulatorMadeChange && InteractionType != (uint8)EControlRigInteractionType::None;
 	
 	InteractionType = (uint8)EControlRigInteractionType::None;
@@ -1240,16 +1257,24 @@ bool FControlRigEditMode::HandleEndTransform(FEditorViewportClient* InViewportCl
 
 		TArray<TWeakObjectPtr<UControlRig>> RigsToTick;
 		RigsToTick.Reserve(InteractionScopes.Num());
-		
-		for (TPair<UControlRig*, FControlRigInteractionScope*>& InteractionScope : InteractionScopes)
+
+		for (TPair<UControlRig*, FControlRigInteractionScope*>& RigInteractionScope : InteractionScopes)
 		{
-			RigsToTick.Add(InteractionScope.Key);
-			if (InteractionScope.Value)
+			if (UControlRig* ControlRig = RigInteractionScope.Key)
 			{
-				delete InteractionScope.Value; 
+				RigsToTick.Add(ControlRig);
+			}
+
+			if (FControlRigInteractionScope* InteractionScope = RigInteractionScope.Value)
+			{
+				Keyframer.Apply(*InteractionScope, TransformContext);
+				delete InteractionScope; 
 			}
 		}
 
+		Keyframer.Finalize(InViewportClient->GetWorld());
+		Keyframer.Reset();
+		
 		InteractionScopes.Reset();
 
 		if (bWasInteracting && !AreEditingControlRigDirectly())
@@ -4992,19 +5017,25 @@ void FControlRigEditMode::SetControlShapeTransform(
 		}
 	}
 
+	FTransform LocalTransform = ControlRig->GetControlLocalTransform(InShapeActor->ControlName);
+
+	FControlKeyframeData KeyframeData;
+	KeyframeData.LocalTransform = LocalTransform;
+	
 	if (bNeedsConstraintPostProcess)
 	{
-		if (!bApplyExplicitRotation)	
+		if (!bApplyExplicitRotation)
 		{
 			// switch to constraint space
 			const FTransform WorldTransform = InGlobalTransform * InToWorldTransform;
-			FTransform LocalTransform = ControlRig->GetControlLocalTransform(InShapeActor->ControlName);
 
 			const TOptional<FTransform> RelativeTransform =
 				UE::TransformConstraintUtil::GetConstraintsRelativeTransform(Constraints, LocalTransform, WorldTransform);
 			if (RelativeTransform)
 			{
-				LocalTransform = *RelativeTransform; 
+				LocalTransform = *RelativeTransform;
+				KeyframeData.LocalTransform = LocalTransform;
+				KeyframeData.bConstraintSpace = true;
 			}
 
 			FRigControlModifiedContext Context = InContext;
@@ -5018,6 +5049,8 @@ void FControlRigEditMode::SetControlShapeTransform(
 		const FConstraintsManagerController& Controller = FConstraintsManagerController::Get(World);
 		Controller.EvaluateAllConstraints();
 	}
+
+	Keyframer.Store(ControlHash, MoveTemp(KeyframeData));
 }
 
 FTransform FControlRigEditMode::GetControlShapeTransform(const AControlRigShapeActor* ShapeActor)
@@ -5148,8 +5181,10 @@ bool FControlRigEditMode::MoveControlShapeLocally( AControlRigShapeActor* ShapeA
 	{
 		ControlRig->InteractionType = InteractionType;
 		ControlRig->ElementsBeingInteracted.AddUnique(ShapeActor->GetElementKey());
-		
-		ControlRig->SetControlLocalTransform(ShapeActor->ControlName, CurrentLocalTransform,true, FRigControlModifiedContext(), true, /*fix eulers*/ true);
+
+		const bool bDeferAutokey = DeferAutokeyOnMouseRelease();
+		const FRigControlModifiedContext Context(bDeferAutokey ? EControlRigSetKey::Never : EControlRigSetKey::DoNotCare);
+		ControlRig->SetControlLocalTransform(ShapeActor->ControlName, CurrentLocalTransform,true, Context, true, /*fix eulers*/ true);
 
 		FTransform CurrentTransform  = ControlRig->GetControlGlobalTransform(ShapeActor->ControlName);			// assumes it's attached to actor
 		CurrentTransform = ToWorldTransform * CurrentTransform;
@@ -5165,6 +5200,14 @@ bool FControlRigEditMode::MoveControlShapeLocally( AControlRigShapeActor* ShapeA
 		//ShapeActor->SetGlobalTransform(CurrentTransform);
 
 		RigsToEvaluateDuringThisTick.AddUnique(ControlRig);
+
+		if (bDeferAutokey)
+		{
+			FControlKeyframeData KeyframeData;
+			KeyframeData.LocalTransform = ControlRig->GetControlLocalTransform(ShapeActor->ControlName);
+			const uint32 ControlHash = UTransformableControlHandle::ComputeHash(ControlRig, ShapeActor->ControlName);
+			Keyframer.Store(ControlHash, MoveTemp(KeyframeData));
+		}
 	}
 
 	return bTransformChanged; 
@@ -5282,6 +5325,11 @@ void FControlRigEditMode::MoveControlShape( AControlRigShapeActor* ShapeActor,
 		FRigControlModifiedContext Context;
 		Context.EventName = FRigUnit_BeginExecution::EventName;
 		Context.bConstraintUpdate = true;
+		
+		if (DeferAutokeyOnMouseRelease())
+		{
+			Context.SetKey = EControlRigSetKey::Never;
+		}
 
 		FTransform TransformSpace = FTransform::Identity; 
 		if (bCalcLocal && InOutLocal)
@@ -6385,6 +6433,11 @@ FInteractionDependencyCache& FControlRigEditMode::GetInteractionDependencies(UCo
 	return DummyDependencies;
 }
 
+bool FControlRigEditMode::DeferAutokeyOnMouseRelease() const
+{
+	return Private::bAutoKeyOnRelease && !bSequencerPlaying && IsInLevelEditor();
+}
+
 /**
 * FDetailKeyFrameCacheAndHandler
 */
diff --git a/Engine/Plugins/Animation/ControlRig/Source/ControlRigEditor/Private/Sequencer/EditModeAnimationUtil.cpp b/Engine/Plugins/Animation/ControlRig/Source/ControlRigEditor/Private/Sequencer/EditModeAnimationUtil.cpp
new file mode 100644
--- /dev/null
+++ b/Engine/Plugins/Animation/ControlRig/Source/ControlRigEditor/Private/Sequencer/EditModeAnimationUtil.cpp
@@ -0,0 +1,118 @@
+﻿// Copyright Epic Games, Inc. All Rights Reserved.
+
+#include "Sequencer/EditModeAnimationUtil.h"
+
+#include "ControlRig.h"
+#include "ConstraintsManager.h"
+#include "Constraints/ControlRigTransformableHandle.h"
+#include "Constraints/MovieSceneConstraintChannelHelper.h"
+#include "EditMode/ControlRigEditModeUtil.h"
+#include "Rigs/RigHierarchyDefines.h"
+
+namespace UE::AnimationEditMode
+{
+
+namespace Private
+{
+
+void EvaluateRigIfAdditive(UControlRig* InControlRig)
+{
+	if (InControlRig->IsAdditive())
+	{
+		TGuardValue<bool> CompensateGuard(FMovieSceneConstraintChannelHelper::bDoNotCompensate, true);
+		InControlRig->Evaluate_AnyThread();
+	}
+}
+
+}
+
+using namespace ControlRigEditMode;
+using namespace Private;
+
+void FControlRigKeyframer::Init(const bool InEnabled)
+{
+	Reset();
+	bEnabled = InEnabled;
+}
+	
+void FControlRigKeyframer::Reset()
+{
+	KeyframeData.Reset();
+}
+
+void FControlRigKeyframer::Store(const uint32 InControlHash, FControlKeyframeData&& InData)
+{
+	if (bEnabled && InControlHash != 0)
+	{
+		FControlKeyframeData& ControlKeyframeData = KeyframeData.FindOrAdd(InControlHash);
+		ControlKeyframeData = MoveTemp(InData);
+	}
+}
+	
+void FControlRigKeyframer::Apply(const FControlRigInteractionScope& InInteractionScope, const FTransformContext& InTransformContext)
+{
+	if (!bEnabled)
+	{
+		return;
+	}
+	
+	const TArray<FRigElementKey>& InteractingControls = InInteractionScope.GetElementsBeingInteracted();
+	UControlRig* InteractingRig = InInteractionScope.GetControlRig();
+	if (!InteractingRig || InteractingControls.IsEmpty())
+	{
+		return;
+	}
+
+	static const FRigControlModifiedContext NoKeyContext(EControlRigSetKey::Never);
+	static constexpr bool bNotify = false, bSetupUndo = false;
+	
+	const bool bFixEulerFlips = !InteractingRig->IsAdditive() && InTransformContext.bRotation;
+	UControlRig::FControlModifiedEvent& AutoKeyEvent = InteractingRig->ControlModified();
+
+	for (const FRigElementKey& ControlKey: InteractingControls)
+	{
+		if (FRigControlElement* Control = InteractingRig->FindControl(ControlKey.Name))
+		{
+			const uint32 ControlHash = UTransformableControlHandle::ComputeHash(InteractingRig, ControlKey.Name);
+			if (const FControlKeyframeData* Data = KeyframeData.Find(ControlHash))
+			{
+				if (Data->bConstraintSpace)
+				{
+					// set the control's local transform withing its constraint space transform as it's the value that sequencer has to store
+					InteractingRig->SetControlLocalTransform(ControlKey.Name, Data->LocalTransform, bNotify, NoKeyContext, bSetupUndo, bFixEulerFlips);
+					EvaluateRigIfAdditive(InteractingRig);
+				}
+				
+				AutoKeyEvent.Broadcast(InteractingRig, Control, EControlRigSetKey::DoNotCare);				
+			}
+		}
+	}
+}
+
+void FControlRigKeyframer::Finalize(UWorld* InWorld)
+{
+	auto NeedsConstraintUpdate = [this]()
+	{
+		if (bEnabled)
+		{
+			for (const auto& [ControlHash, Data]: KeyframeData)
+			{
+				if (Data.bConstraintSpace)
+				{
+					return true;
+				}
+			}
+		}
+		
+		return false;
+	};
+
+	if (InWorld && NeedsConstraintUpdate())
+	{
+		TGuardValue<bool> CompensateGuard(FMovieSceneConstraintChannelHelper::bDoNotCompensate, true);
+		const FConstraintsManagerController& Controller = FConstraintsManagerController::Get(InWorld);
+		Controller.EvaluateAllConstraints();
+	}
+}
+	
+}
\ No newline at end of file
diff --git a/Engine/Plugins/Animation/ControlRig/Source/ControlRigEditor/Public/EditMode/ControlRigEditMode.h b/Engine/Plugins/Animation/ControlRig/Source/ControlRigEditor/Public/EditMode/ControlRigEditMode.h
--- a/Engine/Plugins/Animation/ControlRig/Source/ControlRigEditor/Public/EditMode/ControlRigEditMode.h
+++ b/Engine/Plugins/Animation/ControlRig/Source/ControlRigEditor/Public/EditMode/ControlRigEditMode.h
@@ -21,6 +21,7 @@
 #include "WidgetFocusUtils.h"
 #include "Transform/TransformConstraintUtil.h"
 #include "ControlRigEditModeUtil.h"
+#include "Sequencer/EditModeAnimationUtil.h"
 
 #include "ControlRigEditMode.generated.h"
 
@@ -679,6 +680,10 @@ private:
 	TOptional<FTransform> GetConstraintParentTransform(const UControlRig* InControlRig, const FName& InControlName) const;
 	mutable UE::TransformConstraintUtil::FConstraintsInteractionCache ConstraintsCache;
 
+	// Used to store and apply keyframes (if deferred)  
+	mutable UE::AnimationEditMode::FControlRigKeyframer Keyframer;
+	bool DeferAutokeyOnMouseRelease() const;
+
 	friend class FControlRigEditorModule;
 	friend class FControlRigEditor;
 	friend class FControlRigEditModeGenericDetails;
diff --git a/Engine/Plugins/Animation/ControlRig/Source/ControlRigEditor/Public/Sequencer/EditModeAnimationUtil.h b/Engine/Plugins/Animation/ControlRig/Source/ControlRigEditor/Public/Sequencer/EditModeAnimationUtil.h
new file mode 100644
--- /dev/null
+++ b/Engine/Plugins/Animation/ControlRig/Source/ControlRigEditor/Public/Sequencer/EditModeAnimationUtil.h
@@ -0,0 +1,67 @@
+﻿// Copyright Epic Games, Inc. All Rights Reserved.
+
+#pragma once
+
+#include "Containers/Map.h"
+#include "Math/Transform.h"
+
+class FControlRigInteractionScope;
+class UControlRig;
+struct FRigControlElement;
+
+namespace UE::ControlRigEditMode
+{
+	struct FTransformContext;
+}
+
+namespace UE::AnimationEditMode
+{
+
+/**
+ * FControlKeyframeData provides a way of passing the various keyframe parameters a control needs to set / know about.
+ * Extend it if necessary to pass in more data to the keyframer.
+ */
+
+struct FControlKeyframeData
+{
+	// Local transform data of the control to be keyed.
+	FTransform LocalTransform = FTransform::Identity;
+
+	// Whether this local transform represents a constraint space local transform.
+	bool bConstraintSpace = false;
+};
+
+/**
+ * FControlRigKeyframer enables the storage and application of controls' keyframe data.
+ * It stores keyframe data per control (represented as hash values) that can be applied on demand (on mouse release for example)
+ * This struct works in conjunction with FControlRigInteractionScope, and captures data from controls currently interacting (whether via the viewport
+ * or any other widget that would need to deffer keyframing. 
+ */
+	
+struct FControlRigKeyframer
+{
+	// Resets the data storage and enable/disable the keyframer.
+	void Init(const bool InEnabled);
+	
+	// Stores the keyframe data for a specific control.
+	void Store(const uint32 InControlHash, FControlKeyframeData&& InData);
+
+	// Does the actual work off adding keyframes to the controls currently interacting.
+	void Apply(const FControlRigInteractionScope& InInteractionScope, const ControlRigEditMode::FTransformContext& InTransformContext);
+
+	// Updates whatever needs to once the keyframes have been added (updating constraints is one of them).
+	void Finalize(UWorld* InWorld);
+
+	// Empties the storage.
+	void Reset();
+
+private:
+	
+	// Storage representing keyframe data per control.
+	TMap<uint32, FControlKeyframeData> KeyframeData;
+
+	// Whether the keyframer is enabled or not.
+	bool bEnabled = false;
+};
+	
+}
