Control Rig: Fix fallout from unique module names: renaming modules is broken

#rnx

diff --git a/Engine/Plugins/Animation/ControlRig/Source/ControlRig/Public/Rigs/RigHierarchyDefines.h b/Engine/Plugins/Animation/ControlRig/Source/ControlRig/Public/Rigs/RigHierarchyDefines.h
--- a/Engine/Plugins/Animation/ControlRig/Source/ControlRig/Public/Rigs/RigHierarchyDefines.h
+++ b/Engine/Plugins/Animation/ControlRig/Source/ControlRig/Public/Rigs/RigHierarchyDefines.h
@@ -1712,7 +1712,7 @@ inline bool FRigHierarchyModulePath::ReplaceModuleNameInline<FName>(const FName&
 {
 	if(HasModuleName(InOldModuleFName))
 	{
-		*this = Join(InOldModuleFName, GetElementFName());
+		*this = Join(InNewModuleFName, GetElementFName());
 		return true;
 	}
 	return false;
diff --git a/Engine/Plugins/Animation/ControlRig/Source/ControlRigDeveloper/Private/ControlRigBlueprint.cpp b/Engine/Plugins/Animation/ControlRig/Source/ControlRigDeveloper/Private/ControlRigBlueprint.cpp
--- a/Engine/Plugins/Animation/ControlRig/Source/ControlRigDeveloper/Private/ControlRigBlueprint.cpp
+++ b/Engine/Plugins/Animation/ControlRig/Source/ControlRigDeveloper/Private/ControlRigBlueprint.cpp
@@ -3097,7 +3097,7 @@ void UControlRigBlueprint::HandleRigModulesModified(EModularRigNotification InNo
 							FControlRigExecuteContext& PublicContext = Context.GetPublicDataSafe<FControlRigExecuteContext>();
 							for (TPair<FRigElementKey, ConnectionInfo>& Pair : RenamedConnectors)
 							{
-								const FName& ConnectorName = FRigHierarchyModulePath(Pair.Value.NewPath).GetElementFName();
+								const FName ConnectorName = FRigHierarchyModulePath(Pair.Value.NewPath).GetElementFName();
 								const FString ModulePrefix = InModule->GetElementPrefix();
 								const FString ParentModulePrefix = InModule->GetParentModule() ? InModule->GetParentModule()->GetElementPrefix() : ModulePrefix;
 								const FString RootModulePrefix = InModule->GetRootModule() ? InModule->GetRootModule()->GetElementPrefix() : ModulePrefix;
@@ -3108,7 +3108,33 @@ void UControlRigBlueprint::HandleRigModulesModified(EModularRigNotification InNo
 							}
 						}
 
-						// update the module table
+						// update the target connections
+						TMap<FRigElementKey, FRigElementKeyCollection> PreviousArrayConnectionMap;
+						Swap(PreviousArrayConnectionMap, ArrayConnectionMap);
+						ArrayConnectionMap.Reset();
+						for(TPair<FRigElementKey, FRigElementKeyCollection>& Connection : PreviousArrayConnectionMap)
+						{
+							FRigElementKey ConnectorKey = Connection.Key;
+							FRigElementKeyCollection& TargetKeys = Connection.Value;
+
+							FRigHierarchyModulePath ConnectorPath(ConnectorKey.Name);
+							if(ConnectorPath.ReplaceModuleNameInline(OldModuleName, NewModuleName))
+							{
+								ConnectorKey.Name = ConnectorPath.GetPathFName();
+							}
+							
+							for(FRigElementKey& TargetKey : TargetKeys.Keys)
+							{
+								FRigHierarchyModulePath TargetPath(TargetKey.Name);
+								if(TargetPath.ReplaceModuleNameInline(OldModuleName, NewModuleName))
+								{
+									TargetKey.Name = TargetPath.GetPathFName();
+								}
+							}
+							ArrayConnectionMap.Add(ConnectorKey, TargetKeys);
+						}
+
+						// update the previous module table
 						for(TPair<FRigHierarchyModulePath,FName>& ModulePathToName : ModularRigModel.PreviousModulePaths)
 						{
 							if(ModulePathToName.Value == OldModuleName)
