Control Rig: Only reset pose of modified elements after postconstruction. Compile in the editor will reset the modified elements array.
#rnx

diff --git a/Engine/Plugins/Animation/ControlRig/Source/ControlRigEditor/Private/Editor/ControlRigEditor.cpp b/Engine/Plugins/Animation/ControlRig/Source/ControlRigEditor/Private/Editor/ControlRigEditor.cpp
--- a/Engine/Plugins/Animation/ControlRig/Source/ControlRigEditor/Private/Editor/ControlRigEditor.cpp
+++ b/Engine/Plugins/Animation/ControlRig/Source/ControlRigEditor/Private/Editor/ControlRigEditor.cpp
@@ -1775,7 +1775,13 @@ void FControlRigEditor::Compile()
 		ControlRigBlueprint->ClearTransientControls();
 
 		// default to always reset all bone modifications 
-		ResetAllBoneModification(); 
+		ResetAllBoneModification();
+
+		// remove all cached transforms from modified constrols
+		if (FControlRigEditorEditMode* EditoMode = GetEditMode())
+		{
+			EditoMode->ModifiedRigElements.Reset();
+		}
 
 		{
 			FRigVMEditorBase::Compile();
@@ -6146,8 +6152,12 @@ void FControlRigEditor::OnPreConstructionForUI_AnyThread(UControlRig* InRig, con
 
 	if(ShouldExecuteControlRig(InRig))
 	{
-		const TArrayView<const FRigElementKey> Elements;
-		PreConstructionPose = InRig->GetHierarchy()->GetPose(false, ERigElementType::ToResetAfterConstructionEvent, Elements);
+		FControlRigEditorEditMode* EditMode = GetEditMode();
+		PreConstructionPose.Reset();
+		if (!EditMode->ModifiedRigElements.IsEmpty())
+		{
+			PreConstructionPose = InRig->GetHierarchy()->GetPose(false, ERigElementType::ToResetAfterConstructionEvent, TArrayView<FRigElementKey>(EditMode->ModifiedRigElements));
+		}
 
 		if(const UControlRigBlueprint* RigBlueprint = GetControlRigBlueprint())
 		{
@@ -6226,7 +6236,7 @@ void FControlRigEditor::OnPostConstruction_AnyThread(UControlRig* InRig, const F
 		}
 	}
 	
-	const int32 HierarchyHash = InRig->GetHierarchy()->GetTopologyHash(false);
+	const uint32 HierarchyHash = InRig->GetHierarchy()->GetTopologyHash(false);
 	if(LastHierarchyHash != HierarchyHash)
 	{
 		LastHierarchyHash = HierarchyHash;
diff --git a/Engine/Plugins/Animation/ControlRig/Source/ControlRigEditor/Private/Editor/ControlRigEditorEditMode.cpp b/Engine/Plugins/Animation/ControlRig/Source/ControlRigEditor/Private/Editor/ControlRigEditorEditMode.cpp
--- a/Engine/Plugins/Animation/ControlRig/Source/ControlRigEditor/Private/Editor/ControlRigEditorEditMode.cpp
+++ b/Engine/Plugins/Animation/ControlRig/Source/ControlRigEditor/Private/Editor/ControlRigEditorEditMode.cpp
@@ -264,3 +264,23 @@ void FControlRigEditorEditMode::GetOnScreenDebugInfo(TArray<FText>& OutDebugInfo
 {
 
 }
+
+bool FControlRigEditorEditMode::HandleBeginTransform(const FEditorViewportClient* InViewportClient)
+{
+	const bool bResult = FControlRigEditMode::HandleBeginTransform(InViewportClient);
+
+	for (TWeakObjectPtr<UControlRig>& RuntimeRigPtr : RuntimeControlRigs)
+	{
+		if (UControlRig* ControlRig = RuntimeRigPtr.Get())
+		{
+			TArray<FRigElementKey> SelectedRigElements = GetSelectedRigElements(ControlRig);
+			SelectedRigElements = SelectedRigElements.FilterByPredicate([](FRigElementKey Key) { return Key.Type == ERigElementType::Control; });
+			for (const FRigElementKey& Key : SelectedRigElements)
+			{
+				ModifiedRigElements.AddUnique(Key);
+			}
+		}
+	}
+	
+	return bResult;
+}
diff --git a/Engine/Plugins/Animation/ControlRig/Source/ControlRigEditor/Private/Editor/ControlRigEditorEditMode.h b/Engine/Plugins/Animation/ControlRig/Source/ControlRigEditor/Private/Editor/ControlRigEditorEditMode.h
--- a/Engine/Plugins/Animation/ControlRig/Source/ControlRigEditor/Private/Editor/ControlRigEditorEditMode.h
+++ b/Engine/Plugins/Animation/ControlRig/Source/ControlRigEditor/Private/Editor/ControlRigEditorEditMode.h
@@ -22,11 +22,17 @@ public:
 	virtual class IPersonaPreviewScene& GetAnimPreviewScene() const override;
 	virtual void GetOnScreenDebugInfo(TArray<FText>& OutDebugInfo) const override;
 
+	/* FControlRigEditMode interface */
+	virtual bool HandleBeginTransform(const FEditorViewportClient* InViewportClient) override;
+
 	/** If set to true the edit mode will additionally render all bones */
 	bool bDrawHierarchyBones = false;
 
 	/** Drawing options */
 	UPersonaOptions* ConfigOption = nullptr;
+
+	// Elements which have been modified (e.g. controls moved)
+	TArray<FRigElementKey> ModifiedRigElements;
 };
 
 class FModularRigEditorEditMode : public FControlRigEditorEditMode
diff --git a/Engine/Plugins/Animation/ControlRig/Source/ControlRigEditor/Public/EditMode/ControlRigEditMode.h b/Engine/Plugins/Animation/ControlRig/Source/ControlRigEditor/Public/EditMode/ControlRigEditMode.h
--- a/Engine/Plugins/Animation/ControlRig/Source/ControlRigEditor/Public/EditMode/ControlRigEditMode.h
+++ b/Engine/Plugins/Animation/ControlRig/Source/ControlRigEditor/Public/EditMode/ControlRigEditMode.h
@@ -276,8 +276,8 @@ protected:
 	/** Returns true if all control rig shapes are visible in the viewport */
 	bool AreControlsVisible() const;
 
-	bool HandleBeginTransform(const FEditorViewportClient* InViewportClient);
-	bool HandleEndTransform(FEditorViewportClient* InViewportClient);
+	virtual bool HandleBeginTransform(const FEditorViewportClient* InViewportClient);
+	virtual bool HandleEndTransform(FEditorViewportClient* InViewportClient);
 
 public:
 	
@@ -621,8 +621,11 @@ private:
     };
 
 	FMarqueeDragTool DragToolHandler;
-	
+
+protected:
 	TArray<TWeakObjectPtr<UControlRig>> RuntimeControlRigs;
+
+private:
 	TMap<UControlRig*,TStrongObjectPtr<UControlRigEditModeDelegateHelper>> DelegateHelpers;
 
 	TArray<FRigElementKey> DeferredItemsToFrame;
