Control Rig: Prevent crash when opening Modular rig
#rnx

diff --git a/Engine/Plugins/Animation/ControlRig/Source/ControlRigEditor/Private/Editor/ControlRigEditor.cpp b/Engine/Plugins/Animation/ControlRig/Source/ControlRigEditor/Private/Editor/ControlRigEditor.cpp
--- a/Engine/Plugins/Animation/ControlRig/Source/ControlRigEditor/Private/Editor/ControlRigEditor.cpp
+++ b/Engine/Plugins/Animation/ControlRig/Source/ControlRigEditor/Private/Editor/ControlRigEditor.cpp
@@ -4904,6 +4904,10 @@ void FControlRigEditor::ResetAllBoneModification()
 
 FControlRigEditorEditMode* FControlRigEditor::GetEditMode() const
 {
+	if(IsModularRig())
+	{
+		return static_cast<FModularRigEditorEditMode*>(GetEditorModeManager().GetActiveMode(GetEditorModeName()));
+	}
 	return static_cast<FControlRigEditorEditMode*>(GetEditorModeManager().GetActiveMode(GetEditorModeName()));
 }
 
@@ -6152,11 +6156,13 @@ void FControlRigEditor::OnPreConstructionForUI_AnyThread(UControlRig* InRig, con
 
 	if(ShouldExecuteControlRig(InRig))
 	{
-		FControlRigEditorEditMode* EditMode = GetEditMode();
 		PreConstructionPose.Reset();
-		if (!EditMode->ModifiedRigElements.IsEmpty())
+		if (FControlRigEditorEditMode* EditMode = GetEditMode())
 		{
-			PreConstructionPose = InRig->GetHierarchy()->GetPose(false, ERigElementType::ToResetAfterConstructionEvent, TArrayView<FRigElementKey>(EditMode->ModifiedRigElements));
+			if (!EditMode->ModifiedRigElements.IsEmpty())
+			{
+				PreConstructionPose = InRig->GetHierarchy()->GetPose(false, ERigElementType::ToResetAfterConstructionEvent, TArrayView<FRigElementKey>(EditMode->ModifiedRigElements));
+			}
 		}
 
 		if(const UControlRigBlueprint* RigBlueprint = GetControlRigBlueprint())
