Control Rig Edit Mode: improved selection of multiple controls

- disable sequencer selection notification when selecting multiple controls in a row
- wrap ClearSelection logic into a single transaction to avoid recording one deselection transaction per control

NOTE: this CL also paves the way for a future 'single notification for multiple controls' scheme (as sending one notification per control often leads to slow behaviors)

diff --git a/Engine/Plugins/Animation/ControlRig/Source/ControlRigEditor/Private/EditMode/ControlRigEditMode.cpp b/Engine/Plugins/Animation/ControlRig/Source/ControlRigEditor/Private/EditMode/ControlRigEditMode.cpp
--- a/Engine/Plugins/Animation/ControlRig/Source/ControlRigEditor/Private/EditMode/ControlRigEditMode.cpp
+++ b/Engine/Plugins/Animation/ControlRig/Source/ControlRigEditor/Private/EditMode/ControlRigEditMode.cpp
@@ -1770,42 +1770,31 @@ bool FControlRigEditMode::BoxSelect(FBox& InBox, bool InSelect)
 	}
 
 	// Select all actors that are within the selection box area.  Be aware that certain modes do special processing below.	
-	bool bSomethingSelected = false;
-	UWorld* IteratorWorld = GWorld;
-	for( FActorIterator It(IteratorWorld); It; ++It )
+	TMap<UControlRig*, TArray<FRigElementKey>> ControlsToSelect;
+	for (const auto& [ControlRig, ShapeActors]: ControlRigShapeActors)
 	{
-		AActor* Actor = *It;
-
-		if (!Actor->IsA<AControlRigShapeActor>())
-		{
-			continue;
-		}
-
-		AControlRigShapeActor* ShapeActor = CastChecked<AControlRigShapeActor>(Actor);
-		if (!ShapeActor->IsSelectable() || ShapeActor->ControlRig.IsValid () == false || ShapeActor->ControlRig->GetControlsVisible() == false)
-		{
-			continue;
-		}
-
-		if (IntersectsBox(*Actor, InBox, LevelViewportClient, bStrictDragSelection))
+		if (ControlRig && ControlRig->GetControlsVisible())
 		{
-			bSomethingSelected = true;
-			const FName& ControlName = ShapeActor->ControlName;
-			SetRigElementSelection(ShapeActor->ControlRig.Get(),ERigElementType::Control, ControlName, true);
-
-			if (bShiftDown)
-			{
-			}
-			else
+			for (AControlRigShapeActor* ShapeActor : ShapeActors)
 			{
-				SetRigElementSelection(ShapeActor->ControlRig.Get(),ERigElementType::Control, ControlName, true);
+				const bool bTreatShape = ShapeActor && ShapeActor->IsSelectable() && !ShapeActor->IsTemporarilyHiddenInEditor();
+				if (bTreatShape && IntersectsBox(*ShapeActor, InBox, LevelViewportClient, bStrictDragSelection))
+				{
+					TArray<FRigElementKey>& Controls = ControlsToSelect.FindOrAdd(ControlRig);
+					Controls.Add(ShapeActor->GetElementKey());
+				}
 			}
 		}
 	}
-	if (bSomethingSelected == true)
+
+	const bool bSomethingSelected = !ControlsToSelect.IsEmpty();
+	if (bSomethingSelected)
 	{
+		static constexpr bool bSelected = true;
+		SetRigElementsSelectionInternal(ControlsToSelect, bSelected);
 		return true;
-	}	
+	}
+	
 	ScopedTransaction.Cancel();
 	//if only selecting controls return true to stop any more selections
 	if (Settings && Settings->bOnlySelectRigControls)
@@ -1834,31 +1823,33 @@ bool FControlRigEditMode::FrustumSelect(const FConvexVolume& InFrustum, FEditorV
 	}
 
 	FScopedTransaction ScopedTransaction(LOCTEXT("SelectControlTransaction", "Select Control"), !AreEditingControlRigDirectly() && !GIsTransacting);
-	bool bSomethingSelected(false);
 	const bool bShiftDown = InViewportClient->Viewport->KeyState(EKeys::LeftShift) || InViewportClient->Viewport->KeyState(EKeys::RightShift);
 	if (!bShiftDown)
 	{
 		ClearRigElementSelection(ValidControlTypeMask());
 	}
 
-	for (auto& Pairs : ControlRigShapeActors)
+	TMap<UControlRig*, TArray<FRigElementKey>> ControlsToSelect;
+	
+	for (const auto& [ControlRig, ShapeActors]: ControlRigShapeActors)
 	{
-		for (AControlRigShapeActor* ShapeActor : Pairs.Value)
+		if (ControlRig && ControlRig->GetControlsVisible())
 		{
-			const bool bTreatShape = ShapeActor && ShapeActor->ControlRig.IsValid() && ShapeActor->ControlRig->GetControlsVisible() &&
-				ShapeActor->IsSelectable() && !ShapeActor->IsTemporarilyHiddenInEditor();
-			if (bTreatShape)
+			for (AControlRigShapeActor* ShapeActor : ShapeActors)
 			{
-				for (UActorComponent* Component : ShapeActor->GetComponents())
+				const bool bTreatShape = ShapeActor && ShapeActor->IsSelectable() && !ShapeActor->IsTemporarilyHiddenInEditor();
+				if (bTreatShape)
 				{
-					UPrimitiveComponent* PrimitiveComponent = Cast<UPrimitiveComponent>(Component);
-					if (PrimitiveComponent && PrimitiveComponent->IsRegistered() && PrimitiveComponent->IsVisibleInEditor())
+					for (UActorComponent* Component : ShapeActor->GetComponents())
 					{
-						if (PrimitiveComponent->IsShown(InViewportClient->EngineShowFlags) && PrimitiveComponent->ComponentIsTouchingSelectionFrustum(InFrustum, false /*only bsp*/, false/*encompass entire*/))
+						UPrimitiveComponent* PrimitiveComponent = Cast<UPrimitiveComponent>(Component);
+						if (PrimitiveComponent && PrimitiveComponent->IsRegistered() && PrimitiveComponent->IsVisibleInEditor())
 						{
-							bSomethingSelected = true;
-							const FName& ControlName = ShapeActor->ControlName;
-							SetRigElementSelection(Pairs.Key,ERigElementType::Control, ControlName, true);
+							if (PrimitiveComponent->IsShown(InViewportClient->EngineShowFlags) && PrimitiveComponent->ComponentIsTouchingSelectionFrustum(InFrustum, false /*only bsp*/, false/*encompass entire*/))
+							{
+								TArray<FRigElementKey>& Controls = ControlsToSelect.FindOrAdd(ControlRig);
+								Controls.Add(ShapeActor->GetElementKey());
+							}
 						}
 					}
 				}
@@ -1866,6 +1857,13 @@ bool FControlRigEditMode::FrustumSelect(const FConvexVolume& InFrustum, FEditorV
 		}
 	}
 
+	bool bSomethingSelected = !ControlsToSelect.IsEmpty();
+	if (bSomethingSelected)
+	{
+		static constexpr bool bSelected = true;
+		SetRigElementsSelectionInternal(ControlsToSelect, bSelected);
+	}
+
 	EWorldType::Type WorldType = InViewportClient->GetWorld()->WorldType;
 	const bool bIsAssetEditor =
 		(WorldType == EWorldType::Editor || WorldType == EWorldType::EditorPreview) &&
@@ -2322,22 +2320,45 @@ void FControlRigEditMode::AddReferencedObjects( FReferenceCollector& Collector )
 
 void FControlRigEditMode::ClearRigElementSelection(uint32 InTypes)
 {
-	for (TWeakObjectPtr<UControlRig>& RuntimeRigPtr : RuntimeControlRigs)
+	auto GetController = [bControlRigEditor = AreEditingControlRigDirectly()](UControlRig* InControlRig) -> URigHierarchyController*
 	{
-		if (UControlRig* ControlRig = RuntimeRigPtr.Get())
+		if (InControlRig)
 		{
-			if (!AreEditingControlRigDirectly())
+			if (!bControlRigEditor)
 			{
-				if (URigHierarchyController* Controller = ControlRig->GetHierarchy()->GetController())
+				if (URigHierarchy* Hierarchy = InControlRig->GetHierarchy())
 				{
-					Controller->ClearSelection();
+					return Hierarchy->GetController();
 				}
 			}
-			else if (UControlRigBlueprint* Blueprint = Cast<UControlRigBlueprint>(ControlRig->GetClass()->ClassGeneratedBy))
+			else if (UControlRigBlueprint* Blueprint = Cast<UControlRigBlueprint>(InControlRig->GetClass()->ClassGeneratedBy))
 			{
-				Blueprint->GetHierarchyController()->ClearSelection();
+				return Blueprint->GetHierarchyController();
 			}
 		}
+
+		return nullptr;
+	};
+
+	// put sequencer's selection changed listener on hold during selection to avoid a notifications storm
+	// and only send it once the full selection has been done
+	using namespace UE::Sequencer;
+	TUniquePtr<FSelectionEventSuppressor> SequencerSelectionGuard;
+	if (TSharedPtr<ISequencer> Sequencer = WeakSequencer.IsValid() ? WeakSequencer.Pin() : nullptr)
+	{
+		TSharedPtr<FSequencerEditorViewModel> SequencerViewModel = Sequencer->GetViewModel();
+		if (TSharedPtr<FSequencerSelection> SequencerSelection = SequencerViewModel ? SequencerViewModel->GetSelection() : nullptr)
+		{
+			SequencerSelectionGuard.Reset(new FSelectionEventSuppressor(SequencerSelection.Get()));
+		}
+	}
+
+	for (TWeakObjectPtr<UControlRig>& RuntimeRigPtr : RuntimeControlRigs)
+	{
+		if (URigHierarchyController* Controller = GetController(RuntimeRigPtr.Get()))
+		{
+			Controller->ClearSelection();
+		}
 	}
 }
 
@@ -2364,17 +2385,55 @@ void FControlRigEditMode::SetRigElementSelection(UControlRig* ControlRig, ERigEl
 
 void FControlRigEditMode::SetRigElementSelection(UControlRig* ControlRig, ERigElementType Type, const TArray<FName>& InRigElementNames, bool bSelected)
 {
-	if (!bSelecting)
+	if (!bSelecting && ControlRig)
 	{
-		TGuardValue<bool> ReentrantGuard(bSelecting, true);
+		TMap<UControlRig*, TArray<FRigElementKey>> RigElementsToSelect;
+		TArray<FRigElementKey>& ElementsToSelect = RigElementsToSelect.FindOrAdd(ControlRig);
+		ElementsToSelect.Reserve(InRigElementNames.Num());
+		Algo::Transform(InRigElementNames, ElementsToSelect, [Type](const FName& ElementName)
+		{
+			return FRigElementKey(ElementName, Type);
+		});
+
+		return SetRigElementsSelectionInternal(RigElementsToSelect, bSelected);
+	}
+}
 
-		for (const FName& ElementName : InRigElementNames)
+void FControlRigEditMode::SetRigElementsSelectionInternal(const TMap<UControlRig*, TArray<FRigElementKey>>& InRigElementsToSelect, const bool bSelected)
+{
+	if (bSelecting || InRigElementsToSelect.IsEmpty())
+	{
+		return;
+	}
+
+	// put sequencer's selection changed listener on hold during selection to avoid a notifications storm
+	// and only send it once the full selection has been done
+	using namespace UE::Sequencer;
+	TUniquePtr<FSelectionEventSuppressor> SequencerSelectionGuard;
+	if (TSharedPtr<ISequencer> Sequencer = WeakSequencer.IsValid() ? WeakSequencer.Pin() : nullptr)
+	{
+		TSharedPtr<FSequencerEditorViewModel> SequencerViewModel = Sequencer->GetViewModel();
+		if (TSharedPtr<FSequencerSelection> SequencerSelection = SequencerViewModel ? SequencerViewModel->GetSelection() : nullptr)
 		{
-			SetRigElementSelectionInternal(ControlRig, Type, ElementName, bSelected);
+			SequencerSelectionGuard.Reset(new FSelectionEventSuppressor(SequencerSelection.Get()));
 		}
+	}
 
-		HandleSelectionChanged();
+	TGuardValue<bool> ReentrantGuard(bSelecting, true);
+	
+	for (const auto& [ControlRig, Elements]: InRigElementsToSelect)
+	{
+		URigHierarchy* Hierarchy = ControlRig ? ControlRig->GetHierarchy() : nullptr;
+		if (URigHierarchyController* Controller = Hierarchy ? Hierarchy->GetController() : nullptr)
+		{
+			for (const FRigElementKey& Element : Elements)
+			{
+				Controller->SelectElement(Element, bSelected);
+			}					
+		}
 	}
+	
+	HandleSelectionChanged();
 }
 
 TArray<FRigElementKey> FControlRigEditMode::GetSelectedRigElements(UControlRig* ControlRig) const
@@ -3024,7 +3083,11 @@ bool FControlRigEditMode::CanFrameSelection()
 
 void FControlRigEditMode::ClearSelection()
 {
+	const bool bShouldActuallyTransact = !AreEditingControlRigDirectly() && !GIsTransacting;
+	const FScopedTransaction ScopedTransaction(LOCTEXT("AnimMode_ClearSelectionTransaction", "Clear Selection"), bShouldActuallyTransact);
+	
 	ClearRigElementSelection(FRigElementTypeHelper::ToMask(ERigElementType::All));
+	
 	if (GEditor)
 	{
 		GEditor->Exec(GetWorld(), TEXT("SELECT NONE"));
diff --git a/Engine/Plugins/Animation/ControlRig/Source/ControlRigEditor/Public/EditMode/ControlRigEditMode.h b/Engine/Plugins/Animation/ControlRig/Source/ControlRigEditor/Public/EditMode/ControlRigEditMode.h
--- a/Engine/Plugins/Animation/ControlRig/Source/ControlRigEditor/Public/EditMode/ControlRigEditMode.h
+++ b/Engine/Plugins/Animation/ControlRig/Source/ControlRigEditor/Public/EditMode/ControlRigEditMode.h
@@ -489,10 +489,15 @@ public:
 	void RequestTransformWidgetMode(UE::Widget::EWidgetMode InWidgetMode);
 	
 private:
-	/** Whether or not Pivot Transforms have changed, in which case we need to redraw viewport*/
-	bool HasPivotTransformsChanged() const;
-	/** Set a RigElement's selection state */
+	/** Set a RigElement's selection state. */
 	void SetRigElementSelectionInternal(UControlRig* ControlRig, ERigElementType Type, const FName& InRigElementName, bool bSelected);
+
+	/** Set multiple RigElements' selection states. */
+	void SetRigElementsSelectionInternal(const TMap<UControlRig*, TArray<FRigElementKey>>& InRigElementsToSelect, const bool bSelected);
+
+	/** Whether or not Pivot Transforms have changed, in which case we need to redraw viewport. */
+	bool HasPivotTransformsChanged() const;
+
 	/** Updates the pivot transforms before ticking to ensure that they are up-to-date when needed. */
 	void UpdatePivotTransformsIfNeeded(UControlRig* InControlRig, FTransform& InOutTransform) const;
 	
