Control Rig: Prettify module based paths in the UI

#rnx

diff --git a/Engine/Plugins/Animation/ControlRig/Source/ControlRig/Private/Rigs/RigHierarchy.cpp b/Engine/Plugins/Animation/ControlRig/Source/ControlRig/Private/Rigs/RigHierarchy.cpp
--- a/Engine/Plugins/Animation/ControlRig/Source/ControlRig/Private/Rigs/RigHierarchy.cpp
+++ b/Engine/Plugins/Animation/ControlRig/Source/ControlRig/Private/Rigs/RigHierarchy.cpp
@@ -2853,9 +2853,8 @@ FText URigHierarchy::GetDisplayNameForUI(const FRigBaseElement* InElement, bool
 		const FName ModuleName = GetNameMetadata(InElement->Key, ModuleMetadataName, NAME_None);
 		if(!ModuleName.IsNone())
 		{
-			const FString ModuleNameString = ModuleName.ToString();
-			const FString ModuleDisplayName = FRigHierarchyModulePath(ModuleNameString, DisplayNameString);
-			return FText::FromString(ModuleDisplayName);
+			static const FText ModulePathDisplayFormat = NSLOCTEXT("URigHierarchy", "ModulePathDisplayFormat", "{0} / {1}");
+			return FText::Format(ModulePathDisplayFormat, FText::FromName(ModuleName), FText::FromString(DisplayNameString));
 		}
 	}
 
diff --git a/Engine/Plugins/Animation/ControlRig/Source/ControlRigEditor/Private/Editor/SRigHierarchyTreeView.cpp b/Engine/Plugins/Animation/ControlRig/Source/ControlRigEditor/Private/Editor/SRigHierarchyTreeView.cpp
--- a/Engine/Plugins/Animation/ControlRig/Source/ControlRigEditor/Private/Editor/SRigHierarchyTreeView.cpp
+++ b/Engine/Plugins/Animation/ControlRig/Source/ControlRigEditor/Private/Editor/SRigHierarchyTreeView.cpp
@@ -36,7 +36,7 @@ FRigTreeDisplaySettings FRigTreeDelegates::DefaultDisplaySettings;
 FRigTreeElement::FRigTreeElement(const FRigHierarchyKey& InKey, TWeakPtr<SRigHierarchyTreeView> InTreeView, bool InSupportsRename, ERigTreeFilterResult InFilterResult)
 {
 	Key = InKey;
-	ShortName = InKey.GetFName();
+	ShortName = LongName = FText::FromName(InKey.GetFName());
 	ChannelName = NAME_None;
 	bIsTransient = false;
 	bIsAnimationChannel = false;
@@ -51,7 +51,8 @@ FRigTreeElement::FRigTreeElement(const FRigHierarchyKey& InKey, TWeakPtr<SRigHie
 		{
 			if(InKey.IsElement())
 			{
-				ShortName = *Hierarchy->GetDisplayNameForUI(InKey.GetElement(), false).ToString();
+				LongName = Hierarchy->GetDisplayNameForUI(InKey.GetElement(), true);
+				ShortName = Hierarchy->GetDisplayNameForUI(InKey.GetElement(), false);
 			}
 			
 			const FRigTreeDisplaySettings& Settings = InTreeView.Pin()->GetRigTreeDelegates().GetDisplaySettings();
@@ -271,9 +272,9 @@ FText SRigHierarchyItem::GetName(bool bUseShortName) const
 	}
 	if(bUseShortName)
 	{
-		return (FText::FromName(WeakRigTreeElement.Pin()->ShortName));
+		return WeakRigTreeElement.Pin()->ShortName;
 	}
-	return (FText::FromName(WeakRigTreeElement.Pin()->Key.GetFName()));
+	return WeakRigTreeElement.Pin()->LongName;
 }
 
 FText SRigHierarchyItem::GetItemTooltip() const
@@ -655,12 +656,9 @@ bool SRigHierarchyTreeView::AddElement(const FRigBaseElement* InElement)
 								}
 							}
 
-							FName Name = ConnectorKey.Name;
-							if (GetRigTreeDelegates().GetDisplaySettings().bUseShortName)
-							{
-								Name = *Hierarchy->GetDisplayNameForUI(ConnectorKey).ToString();
-							}
-							TagArguments.Text(FText::FromName(Name));
+							const bool bUseShortName = GetRigTreeDelegates().GetDisplaySettings().bUseShortName;
+							const FText Label = Hierarchy->GetDisplayNameForUI(ConnectorKey, !bUseShortName);
+							TagArguments.Text(Label);
 							TagArguments.TooltipText(FText::FromName(ConnectorKey.Name));
 							TagArguments.Color(BackgroundColor);
 							TagArguments.IconColor(IconColor);
diff --git a/Engine/Plugins/Animation/ControlRig/Source/ControlRigEditor/Private/Editor/SRigHierarchyTreeView.h b/Engine/Plugins/Animation/ControlRig/Source/ControlRigEditor/Private/Editor/SRigHierarchyTreeView.h
--- a/Engine/Plugins/Animation/ControlRig/Source/ControlRigEditor/Private/Editor/SRigHierarchyTreeView.h
+++ b/Engine/Plugins/Animation/ControlRig/Source/ControlRigEditor/Private/Editor/SRigHierarchyTreeView.h
@@ -230,7 +230,8 @@ public:
 public:
 	/** Element Data to display */
 	FRigHierarchyKey Key;
-	FName ShortName;
+	FText LongName;
+	FText ShortName;
 	FName ChannelName;
 	bool bIsTransient;
 	bool bIsAnimationChannel;
