Fixed a memory use after delete at RigVMHierarchyController

diff --git a/Engine/Plugins/Animation/ControlRig/Source/ControlRig/Private/Rigs/RigHierarchyController.cpp b/Engine/Plugins/Animation/ControlRig/Source/ControlRig/Private/Rigs/RigHierarchyController.cpp
--- a/Engine/Plugins/Animation/ControlRig/Source/ControlRig/Private/Rigs/RigHierarchyController.cpp
+++ b/Engine/Plugins/Animation/ControlRig/Source/ControlRig/Private/Rigs/RigHierarchyController.cpp
@@ -2285,6 +2285,9 @@ TArray<FRigElementKey> URigHierarchyController::ImportFromText(FString InContent
 		{
 			if(FRigBaseElement* ExistingElement = Hierarchy->Find(NewElement->GetKey()))
 			{
+				// as we have created a new element in the same hierarchy, we have to update the storage link of the existing element
+				ExistingElement->LinkStorage(Hierarchy->ElementTransforms.GetStorage(), Hierarchy->ElementDirtyStates.GetStorage(), Hierarchy->ElementCurves.GetStorage());
+
 				ExistingElement->CopyPose(NewElement, true, true, false);
 
 				if(FRigControlElement* ControlElement = Cast<FRigControlElement>(ExistingElement))
