Animation Mode: improved selection of multiple controls
- the space picker is now updated lazily (using an active timer) once the selection is complete. this has the advantage of avoiding rebuilding it N times when selecting N controls.
- disable sequencer selection notification when selecting driven or nested controls

diff --git a/Engine/Plugins/Animation/ControlRig/Source/ControlRigEditor/Private/Editor/SRigSpacePickerWidget.cpp b/Engine/Plugins/Animation/ControlRig/Source/ControlRigEditor/Private/Editor/SRigSpacePickerWidget.cpp
--- a/Engine/Plugins/Animation/ControlRig/Source/ControlRigEditor/Private/Editor/SRigSpacePickerWidget.cpp
+++ b/Engine/Plugins/Animation/ControlRig/Source/ControlRigEditor/Private/Editor/SRigSpacePickerWidget.cpp
@@ -234,6 +234,8 @@ void SRigSpacePickerWidget::Construct(const FArguments& InArgs)
 
 SRigSpacePickerWidget::~SRigSpacePickerWidget()
 {
+	UnregisterPendingSelection();
+	
 	GEditor->UnregisterForUndo(this);
 
 	if(HierarchyModifiedHandle.IsValid())
@@ -246,25 +248,52 @@ SRigSpacePickerWidget::~SRigSpacePickerWidget()
 	}
 }
 
-void SRigSpacePickerWidget::SetControls(
-	URigHierarchy* InHierarchy,
-	const TArray<FRigElementKey>& InControls)
+void SRigSpacePickerWidget::SetControls(URigHierarchy* InHierarchy, const TArray<FRigElementKey>& InControls)
+{
+	UnregisterPendingSelection();
+	
+	PendingSelectionHandle = RegisterActiveTimer(0.f,
+		FWidgetActiveTimerDelegate::CreateLambda([this, WeakHierarchy = TWeakObjectPtr<URigHierarchy>(InHierarchy), InControls](double, float)
+		{
+			if(URigHierarchy* Hierarchy = WeakHierarchy.Get())
+			{
+				UpdateSelection(Hierarchy, InControls);
+			}
+			return EActiveTimerReturnType::Stop;
+		})
+	);
+}
+
+void SRigSpacePickerWidget::UnregisterPendingSelection()
+{
+	if (PendingSelectionHandle.IsValid())
+	{
+		if (TSharedPtr<FActiveTimerHandle> ActiveTimerHandle = PendingSelectionHandle.Pin())
+		{
+			UnRegisterActiveTimer(ActiveTimerHandle.ToSharedRef());
+		}
+		PendingSelectionHandle.Reset();
+	}
+}
+
+void SRigSpacePickerWidget::UpdateSelection(URigHierarchy* InHierarchy, const TArray<FRigElementKey>& InControls)
 {
-	if(Hierarchy.IsValid())
+	if (Hierarchy.IsValid())
 	{
 		URigHierarchy* StrongHierarchy = Hierarchy.Get();
 		if (StrongHierarchy != InHierarchy)
 		{
-			if(HierarchyModifiedHandle.IsValid())
+			if (HierarchyModifiedHandle.IsValid())
 			{
 				StrongHierarchy->OnModified().Remove(HierarchyModifiedHandle);
 				HierarchyModifiedHandle.Reset();
 			}
 		}
 	}
-	
+		
 	Hierarchy = InHierarchy;
 	ControlKeys.SetNum(0);
+		
 	for (const FRigElementKey& Key : InControls)
 	{
 		if (const FRigControlElement* ControlElement = Hierarchy->FindChecked<FRigControlElement>(Key))
@@ -288,10 +317,12 @@ void SRigSpacePickerWidget::SetControls(
 		}
 		ControlKeys.Add(Key);
 	}
+			
 	if (Hierarchy.IsValid() && HierarchyModifiedHandle.IsValid() == false)
 	{
 		HierarchyModifiedHandle = InHierarchy->OnModified().AddSP(this, &SRigSpacePickerWidget::OnHierarchyModified);
 	}
+			
 	UpdateActiveSpaces();
 	RepopulateItemSpaces();
 }
diff --git a/Engine/Plugins/Animation/ControlRig/Source/ControlRigEditor/Private/Editor/SRigSpacePickerWidget.h b/Engine/Plugins/Animation/ControlRig/Source/ControlRigEditor/Private/Editor/SRigSpacePickerWidget.h
--- a/Engine/Plugins/Animation/ControlRig/Source/ControlRigEditor/Private/Editor/SRigSpacePickerWidget.h
+++ b/Engine/Plugins/Animation/ControlRig/Source/ControlRigEditor/Private/Editor/SRigSpacePickerWidget.h
@@ -178,6 +178,15 @@ private:
 	FRigTreeDisplaySettings HierarchyDisplaySettings;
 
 	static FRigElementKey InValidKey;
+
+	/** Unregisters any current pending selection active timer. */
+	void UnregisterPendingSelection();
+	
+	/** Updates the selected controls array and rebuilds the spaces list. */
+	void UpdateSelection(URigHierarchy* InHierarchy, const TArray<FRigElementKey>& InControls);
+
+	/** Controls selection currently pending. */
+	TWeakPtr<FActiveTimerHandle> PendingSelectionHandle;
 };
 
 class ISequencer;
diff --git a/Engine/Plugins/Animation/ControlRig/Source/ControlRigEditor/Private/Sequencer/ControlRigParameterTrackEditor.cpp b/Engine/Plugins/Animation/ControlRig/Source/ControlRigEditor/Private/Sequencer/ControlRigParameterTrackEditor.cpp
--- a/Engine/Plugins/Animation/ControlRig/Source/ControlRigEditor/Private/Sequencer/ControlRigParameterTrackEditor.cpp
+++ b/Engine/Plugins/Animation/ControlRig/Source/ControlRigEditor/Private/Sequencer/ControlRigParameterTrackEditor.cpp
@@ -3890,6 +3890,16 @@ void FControlRigParameterTrackEditor::HandleControlSelected(UControlRig* Subject
 		return;
 	}
 
+	TUniquePtr<FSelectionEventSuppressor> SequencerSelectionGuard;
+	if (TSharedPtr<ISequencer> Sequencer = GetSequencer())
+	{
+		TSharedPtr<FSequencerEditorViewModel> SequencerViewModel = Sequencer->GetViewModel();
+		if (TSharedPtr<FSequencerSelection> SequencerSelection = SequencerViewModel ? SequencerViewModel->GetSelection() : nullptr)
+		{
+			SequencerSelectionGuard.Reset(new FSelectionEventSuppressor(SequencerSelection.Get()));
+		}
+	}
+	
 	if(ControlElement->CanDriveControls())
 	{
 		const TArray<FRigElementKey>& DrivenControls = ControlElement->Settings.DrivenControls;
@@ -3915,40 +3925,43 @@ void FControlRigParameterTrackEditor::HandleControlSelected(UControlRig* Subject
 	//if parent selected we select child here if it's a bool,integer or single float
 	TArray<FRigControl> Controls;
 
-	if(!bIsSelectingIndirectControl)
+	if (!bIsSelectingIndirectControl)
 	{
 		if (URigHierarchyController* Controller = Hierarchy->GetController())
-		{
+		{			
 			Hierarchy->ForEach<FRigControlElement>([Hierarchy, ControlElement, Controller, bSelected](FRigControlElement* OtherControlElement) -> bool
+			{
+				const FRigControlSettings& OtherSettings = OtherControlElement->Settings;
+				
+				const ERigControlType OtherControlType = OtherSettings.ControlType;
+				if (OtherControlType == ERigControlType::Bool ||
+					OtherControlType == ERigControlType::Float ||
+					OtherControlType == ERigControlType::ScaleFloat ||
+					OtherControlType == ERigControlType::Integer)
 				{
-					if (OtherControlElement->Settings.ControlType == ERigControlType::Bool ||
-						OtherControlElement->Settings.ControlType == ERigControlType::Float ||
-						OtherControlElement->Settings.ControlType == ERigControlType::ScaleFloat ||
-						OtherControlElement->Settings.ControlType == ERigControlType::Integer)
+					if(OtherControlElement->Settings.SupportsShape() || !Hierarchy->IsAnimatable(OtherControlElement))
 					{
-						if(OtherControlElement->Settings.SupportsShape() || !Hierarchy->IsAnimatable(OtherControlElement))
+						return true;
+					}
+					
+					for (const FRigElementParentConstraint& ParentConstraint : OtherControlElement->ParentConstraints)
+					{
+						if (ParentConstraint.ParentElement == ControlElement)
 						{
+							Controller->SelectElement(OtherControlElement->GetKey(), bSelected);
 							return true;
 						}
-						
-						for (const FRigElementParentConstraint& ParentConstraint : OtherControlElement->ParentConstraints)
-						{
-							if (ParentConstraint.ParentElement == ControlElement)
-							{
-								Controller->SelectElement(OtherControlElement->GetKey(), bSelected);
-								return true;
-							}
-						}
-					}
-
-					if(OtherControlElement->IsAnimationChannel() && OtherControlElement->Settings.Customization.AvailableSpaces.FindByKey(ControlElement->GetKey()))
-					{
-						Controller->SelectElement(OtherControlElement->GetKey(), bSelected);
-						return true;
 					}
+				}
 
+				if (OtherControlElement->IsAnimationChannel() && OtherSettings.Customization.AvailableSpaces.FindByKey(ControlElement->GetKey()))
+				{
+					Controller->SelectElement(OtherControlElement->GetKey(), bSelected);
 					return true;
-				});
+				}
+
+				return true;
+			});
 		}
 	}
 	
@@ -3976,8 +3989,6 @@ void FControlRigParameterTrackEditor::HandleControlSelected(UControlRig* Subject
 		const TSharedPtr<ISequencer> Sequencer = GetSequencer();
 		if (Track && Sequencer)
 		{
-			FSelectionEventSuppressor SuppressSelectionEvents = Sequencer->GetViewModel()->GetSelection()->SuppressEvents();
-
 			//Just select in section to key, if deselecting makes sure deselected everywhere
 			if (bSelected == true)
 			{
diff --git a/Engine/Source/Editor/CurveEditor/Private/Tree/CurveEditorTree.cpp b/Engine/Source/Editor/CurveEditor/Private/Tree/CurveEditorTree.cpp
--- a/Engine/Source/Editor/CurveEditor/Private/Tree/CurveEditorTree.cpp
+++ b/Engine/Source/Editor/CurveEditor/Private/Tree/CurveEditorTree.cpp
@@ -557,29 +557,42 @@ void FCurveEditorTree::RemoveFromSelection(TArrayView<const FCurveEditorTreeItem
 	}
 }
 
- TArray<FCurveEditorTreeItemID> FCurveEditorTree::GetCachedExpandedItems() const
- {
-	 TArray<FCurveEditorTreeItemID> ExpandedItems;
-	 for (const TPair<FCurveEditorTreeItemID, FCurveEditorTreeItem>& Item : Items)
-	 {
-		 const TOptional<FString> PathName = Item.Value.GetUniquePathName();
-		 if (PathName.IsSet())
-		 {
-			 int32 Hash = GetTypeHash(PathName.GetValue());
-			 if (CachedExpandedItems.Contains(Hash))
-			 {
-				 ExpandedItems.Add(Item.Key);
-			 }
-		 }
-	 }
-	 return ExpandedItems;
- }
+TArray<FCurveEditorTreeItemID> FCurveEditorTree::GetCachedExpandedItems() const
+{
+	TArray<FCurveEditorTreeItemID> ExpandedItems;
+
+	const int32 NumCachedExpandedItems = CachedExpandedItems.Num();
+	if (NumCachedExpandedItems > 0)
+	{
+		ExpandedItems.Reserve(NumCachedExpandedItems);
+		for (const TPair<FCurveEditorTreeItemID, FCurveEditorTreeItem>& Item : Items)
+		{
+			const TOptional<FString>& PathName = Item.Value.GetUniquePathName();
+			if (PathName.IsSet())
+			{
+				int32 Hash = GetTypeHash(PathName.GetValue());
+				if (CachedExpandedItems.Contains(Hash))
+				{
+					ExpandedItems.Add(Item.Key);
+
+					// early exit
+					if (ExpandedItems.Num() == NumCachedExpandedItems)
+					{
+						return ExpandedItems;					
+					}
+				}
+			}
+		}
+	}
+	
+	return ExpandedItems;
+}
 
 void FCurveEditorTree::SetItemExpansion(FCurveEditorTreeItemID InTreeItemID, bool bInExpansion)
 {
 	if (FCurveEditorTreeItem* TreeItem = FindItem(InTreeItemID))
 	{
-		TOptional<FString> OptString = TreeItem->GetUniquePathName();
+		const TOptional<FString>& OptString = TreeItem->GetUniquePathName();
 		if (OptString.IsSet())
 		{
 			int32 Hash = GetTypeHash(OptString.GetValue());
diff --git a/Engine/Source/Editor/CurveEditor/Public/Tree/CurveEditorTree.h b/Engine/Source/Editor/CurveEditor/Public/Tree/CurveEditorTree.h
--- a/Engine/Source/Editor/CurveEditor/Public/Tree/CurveEditorTree.h
+++ b/Engine/Source/Editor/CurveEditor/Public/Tree/CurveEditorTree.h
@@ -153,7 +153,7 @@ struct CURVEEDITOR_API FCurveEditorTreeItem
 	/**
 	*Get optional unique path name for the tree editor item
 	*/
-	TOptional<FString> GetUniquePathName() const { return UniquePathName; } 
+	const TOptional<FString>& GetUniquePathName() const { return UniquePathName; } 
 
 	/**
 	*Set optional unique path name for the tree editor item
