fixed clang compilation

#rnx

diff --git a/Engine/Plugins/Animation/ControlRig/Source/ControlRigEditor/Private/Editor/SRigSpacePickerWidget.cpp b/Engine/Plugins/Animation/ControlRig/Source/ControlRigEditor/Private/Editor/SRigSpacePickerWidget.cpp
--- a/Engine/Plugins/Animation/ControlRig/Source/ControlRigEditor/Private/Editor/SRigSpacePickerWidget.cpp
+++ b/Engine/Plugins/Animation/ControlRig/Source/ControlRigEditor/Private/Editor/SRigSpacePickerWidget.cpp
@@ -255,9 +255,9 @@ void SRigSpacePickerWidget::SetControls(URigHierarchy* InHierarchy, const TArray
 	PendingSelectionHandle = RegisterActiveTimer(0.f,
 		FWidgetActiveTimerDelegate::CreateLambda([this, WeakHierarchy = TWeakObjectPtr<URigHierarchy>(InHierarchy), InControls](double, float)
 		{
-			if(URigHierarchy* Hierarchy = WeakHierarchy.Get())
+			if(URigHierarchy* HierarchyToUpdate = WeakHierarchy.Get())
 			{
-				UpdateSelection(Hierarchy, InControls);
+				UpdateSelection(HierarchyToUpdate, InControls);
 			}
 			return EActiveTimerReturnType::Stop;
 		})
