Control Rig: Fixes to animating integer parameters with Control Enums

Three fixes here, two related to the ecs, one that was generic:
1) We were previously attempting to blend any parameter type, where only scalars, vectors, and transforms should be blended. This is fixed by constraining the overlapping entity tracker to those component types.
2) FControlRigParameterValues::AddCurrentValue was not correctly handling integer controls that had Control Enums. This resulted in an erroneous static_cast on the channel type later on where we assumed an integer channel was present when it was actually a byte channel.
3) FControlRigSequencerHelpers::GetByteChannels and GetIntegerChannels could crash if the control channel map was not up-to-date due to the control type not matching. This happened if an integer control was changed to use a Control Enum when it was open in Sequencer.

diff --git a/Engine/Plugins/Animation/ControlRig/Source/ControlRig/Private/Sequencer/ControlRigSequencerHelpers.cpp b/Engine/Plugins/Animation/ControlRig/Source/ControlRig/Private/Sequencer/ControlRigSequencerHelpers.cpp
--- a/Engine/Plugins/Animation/ControlRig/Source/ControlRig/Private/Sequencer/ControlRigSequencerHelpers.cpp
+++ b/Engine/Plugins/Animation/ControlRig/Source/ControlRig/Private/Sequencer/ControlRigSequencerHelpers.cpp
@@ -125,7 +125,7 @@ TArrayView<FMovieSceneByteChannel*> FControlRigSequencerHelpers::GetByteChannels
 	{
 		int32 NumChannels = 1;
 		const FChannelMapInfo* ChannelInfo = CRSection ? CRSection->ControlChannelMap.Find(InControlName) : nullptr;
-		if (ChannelInfo == nullptr)
+		if (ChannelInfo == nullptr || ChannelInfo->ChannelTypeName != FMovieSceneByteChannel::StaticStruct()->GetFName())
 		{
 			return EmptyChannelsView;
 		}
@@ -155,7 +155,7 @@ TArrayView<FMovieSceneIntegerChannel*> FControlRigSequencerHelpers::GetIntegerCh
 	{
 		int32 NumChannels = 1;
 		const FChannelMapInfo* ChannelInfo = CRSection ? CRSection->ControlChannelMap.Find(InControlName) : nullptr;
-		if (ChannelInfo == nullptr)
+		if (ChannelInfo == nullptr || ChannelInfo->ChannelTypeName != FMovieSceneIntegerChannel::StaticStruct()->GetFName())
 		{
 			return EmptyChannelsView;
 		}
diff --git a/Engine/Plugins/Animation/ControlRig/Source/ControlRig/Private/Sequencer/MovieSceneControlRigParameterBuffer.cpp b/Engine/Plugins/Animation/ControlRig/Source/ControlRig/Private/Sequencer/MovieSceneControlRigParameterBuffer.cpp
--- a/Engine/Plugins/Animation/ControlRig/Source/ControlRig/Private/Sequencer/MovieSceneControlRigParameterBuffer.cpp
+++ b/Engine/Plugins/Animation/ControlRig/Source/ControlRig/Private/Sequencer/MovieSceneControlRigParameterBuffer.cpp
@@ -603,8 +603,16 @@ void FControlRigParameterValues::AddCurrentValue(UControlRig* Rig, FRigControlEl
 
 	case ERigControlType::Integer:
 	{
-		const int32 Val = Rig->GetControlValue(ControlElement, ERigControlValueType::Current).Get<int32>();
-		Add(ControlName, Val);
+		if (ControlElement->Settings.ControlEnum != nullptr)
+		{
+			const uint8 Val = Rig->GetControlValue(ControlElement, ERigControlValueType::Current).Get<uint8>();
+			Add(ControlName, Val);
+		}
+		else
+		{
+			const int32 Val = Rig->GetControlValue(ControlElement, ERigControlValueType::Current).Get<int32>();
+			Add(ControlName, Val);
+		}
 		break;
 	}
 
diff --git a/Engine/Plugins/Animation/ControlRig/Source/ControlRig/Private/Sequencer/MovieSceneControlRigSystem.cpp b/Engine/Plugins/Animation/ControlRig/Source/ControlRig/Private/Sequencer/MovieSceneControlRigSystem.cpp
--- a/Engine/Plugins/Animation/ControlRig/Source/ControlRig/Private/Sequencer/MovieSceneControlRigSystem.cpp
+++ b/Engine/Plugins/Animation/ControlRig/Source/ControlRig/Private/Sequencer/MovieSceneControlRigSystem.cpp
@@ -1693,6 +1693,7 @@ void UMovieSceneControlRigParameterEvaluatorSystem::OnInstantiation()
 		.Read(BuiltInComponents->Group)
 		.PassthroughFilter(FEntityComponentFilter().All({ BuiltInComponents->Tags.NeedsLink }))
 		.FilterAny({ BuiltInComponents->Tags.NeedsLink, BuiltInComponents->Tags.NeedsUnlink })
+		.FilterAny({ TracksComponents->Parameters.Transform.PropertyTag, TracksComponents->Parameters.Scalar.PropertyTag, TracksComponents->Parameters.Vector3.PropertyTag })
 		.FilterAll({ ControlRigComponents->ControlRigSource, TracksComponents->GenericParameterName, ControlRigComponents->Tags.ControlRigParameter })
 		.Iterate_PerAllocation(&Linker->EntityManager, [this](const FEntityAllocation* Allocation, TRead<FEntityGroupID> Group, bool bIsBeingLinked)
 			{
