Control Rig: Fix display names of rig elements not allowing same name controls from different module paths
#rnx

diff --git a/Engine/Plugins/Animation/ControlRig/Source/ControlRig/Private/Rigs/RigHierarchyController.cpp b/Engine/Plugins/Animation/ControlRig/Source/ControlRig/Private/Rigs/RigHierarchyController.cpp
--- a/Engine/Plugins/Animation/ControlRig/Source/ControlRig/Private/Rigs/RigHierarchyController.cpp
+++ b/Engine/Plugins/Animation/ControlRig/Source/ControlRig/Private/Rigs/RigHierarchyController.cpp
@@ -515,7 +515,7 @@ FRigElementKey URigHierarchyController::AddControl(
 			{
 				// avoid self name collision
 				NewElement->Settings.DisplayName = NAME_None;
-				NewElement->Settings.DisplayName = Hierarchy->GetSafeNewDisplayName(InParent, InName);
+				NewElement->Settings.DisplayName = Hierarchy->GetSafeNewDisplayName(InParent, NewElement->Key.Name);
 			}
 		}
 		
