Control Rig Edit Mode: occlusion based selection for perspective viewports
- ControlRig's  marquee selection now respects ULevelEditorViewportSettings::bTransparentBoxSelection (won't select occluded controls if bTransparentBoxSelection is false)
- TODO: update the orthographic selection so it uses the same code path (occlusion, null & bones selection, etc.)

diff --git a/Engine/Plugins/Animation/ControlRig/Source/ControlRigEditor/Private/EditMode/ControlRigEditMode.cpp b/Engine/Plugins/Animation/ControlRig/Source/ControlRigEditor/Private/EditMode/ControlRigEditMode.cpp
--- a/Engine/Plugins/Animation/ControlRig/Source/ControlRigEditor/Private/EditMode/ControlRigEditMode.cpp
+++ b/Engine/Plugins/Animation/ControlRig/Source/ControlRigEditor/Private/EditMode/ControlRigEditMode.cpp
@@ -1674,48 +1674,6 @@ void FControlRigEditMode::OpenContextMenu(FEditorViewportClient* InViewportClien
 	}
 }
 
-bool FControlRigEditMode::IntersectSelect(bool InSelect, const TFunctionRef<bool(const AControlRigShapeActor*, const FTransform&)>& Intersects)
-{
-	bool bSelected = false;
-
-	for (auto& Pairs : ControlRigShapeActors)
-	{
-		FTransform ComponentTransform = GetHostingSceneComponentTransform(Pairs.Key);
-		for (AControlRigShapeActor* ShapeActor : Pairs.Value)
-		{
-			if (ShapeActor->IsHiddenEd())
-			{
-				continue;
-			}
-
-			const FTransform ControlTransform = ShapeActor->GetGlobalTransform() * ComponentTransform;
-			if (Intersects(ShapeActor, ControlTransform))
-			{
-				SetRigElementSelection(Pairs.Key,ERigElementType::Control, ShapeActor->ControlName, InSelect);
-				bSelected = true;
-			}
-		}
-	}
-
-	return bSelected;
-}
-
-static FConvexVolume GetVolumeFromBox(const FBox& InBox)
-{
-	FConvexVolume ConvexVolume;
-	ConvexVolume.Planes.Empty(6);
-
-	ConvexVolume.Planes.Add(FPlane(FVector::LeftVector, -InBox.Min.Y));
-	ConvexVolume.Planes.Add(FPlane(FVector::RightVector, InBox.Max.Y));
-	ConvexVolume.Planes.Add(FPlane(FVector::UpVector, InBox.Max.Z));
-	ConvexVolume.Planes.Add(FPlane(FVector::DownVector, -InBox.Min.Z));
-	ConvexVolume.Planes.Add(FPlane(FVector::ForwardVector, InBox.Max.X));
-	ConvexVolume.Planes.Add(FPlane(FVector::BackwardVector, -InBox.Min.X));
-
-	ConvexVolume.Init();
-
-	return ConvexVolume;
-}
 
 bool IntersectsBox( AActor& InActor, const FBox& InBox, FLevelEditorViewportClient* LevelViewportClient, bool bUseStrictSelection )
 {
@@ -1831,31 +1789,8 @@ bool FControlRigEditMode::FrustumSelect(const FConvexVolume& InFrustum, FEditorV
 
 	TMap<UControlRig*, TArray<FRigElementKey>> ControlsToSelect;
 	
-	for (const auto& [ControlRig, ShapeActors]: ControlRigShapeActors)
-	{
-		if (ControlRig && ControlRig->GetControlsVisible())
-		{
-			for (AControlRigShapeActor* ShapeActor : ShapeActors)
-			{
-				const bool bTreatShape = ShapeActor && ShapeActor->IsSelectable() && !ShapeActor->IsTemporarilyHiddenInEditor();
-				if (bTreatShape)
-				{
-					for (UActorComponent* Component : ShapeActor->GetComponents())
-					{
-						UPrimitiveComponent* PrimitiveComponent = Cast<UPrimitiveComponent>(Component);
-						if (PrimitiveComponent && PrimitiveComponent->IsRegistered() && PrimitiveComponent->IsVisibleInEditor())
-						{
-							if (PrimitiveComponent->IsShown(InViewportClient->EngineShowFlags) && PrimitiveComponent->ComponentIsTouchingSelectionFrustum(InFrustum, false /*only bsp*/, false/*encompass entire*/))
-							{
-								TArray<FRigElementKey>& Controls = ControlsToSelect.FindOrAdd(ControlRig);
-								Controls.Add(ShapeActor->GetElementKey());
-							}
-						}
-					}
-				}
-			}
-		}
-	}
+	FSelectionHelper SelectionHelper(InViewportClient, ControlRigShapeActors, ControlsToSelect);
+	SelectionHelper.GetFromFrustum(InFrustum);
 
 	bool bSomethingSelected = !ControlsToSelect.IsEmpty();
 	if (bSomethingSelected)
diff --git a/Engine/Plugins/Animation/ControlRig/Source/ControlRigEditor/Private/EditMode/ControlRigEditModeUtil.cpp b/Engine/Plugins/Animation/ControlRig/Source/ControlRigEditor/Private/EditMode/ControlRigEditModeUtil.cpp
--- a/Engine/Plugins/Animation/ControlRig/Source/ControlRigEditor/Private/EditMode/ControlRigEditModeUtil.cpp
+++ b/Engine/Plugins/Animation/ControlRig/Source/ControlRigEditor/Private/EditMode/ControlRigEditModeUtil.cpp
@@ -3,9 +3,15 @@
 #include "EditMode/ControlRigEditModeUtil.h"
 
 #include "ControlRig.h"
+#include "EditorViewportClient.h"
+#include "EngineUtils.h"
+#include "LevelEditorViewport.h"
 #include "Rigs/RigHierarchy.h"
+#include "SceneView.h"
+#include "Settings/LevelEditorViewportSettings.h"
 #include "Transform/TransformConstraint.h"
 #include "Transform/TransformConstraintUtil.h"
+#include "UnrealClient.h"
 
 namespace UE::ControlRigEditMode
 {
@@ -109,5 +115,187 @@ void FExplicitRotationInteraction::Apply(
 	}
 }
 
+FSelectionHelper::FSelectionHelper(
+	FEditorViewportClient* InViewportClient,
+	const TMap<UControlRig*, TArray<TObjectPtr<AControlRigShapeActor>>>& InControlRigShapeActors,
+	TMap<UControlRig*, TArray<FRigElementKey>>& OutElements)
+	: ViewportClient(InViewportClient)
+	, ControlRigShapeActors(InControlRigShapeActors)
+	, Elements(OutElements)
+	, LevelEditorViewportSettings(GetDefault<ULevelEditorViewportSettings>())
+{
+	if (ensure(InViewportClient))
+	{
+		if (InViewportClient->IsLevelEditorClient())
+		{
+			FLevelEditorViewportClient* LevelViewportClient = static_cast<FLevelEditorViewportClient*>(InViewportClient);
+			HiddenLayers = LevelViewportClient->ViewHiddenLayers;
+		}
+	}
+}
+
+bool FSelectionHelper::IsValid() const
+{
+	return ViewportClient && ViewportClient->Viewport;
+}
+	
+void FSelectionHelper::GetFromFrustum(const FConvexVolume& InFrustum) const
+{
+	if (!IsValid())
+	{
+		return;
+	}
+
+	// NOTE: occlusion based selection is a level editor property but should probably be by viewport client.
+	bool bTransparentBoxSelection = ViewportClient->IsLevelEditorClient() ? LevelEditorViewportSettings->bTransparentBoxSelection : true;
+	if (!bTransparentBoxSelection)
+	{
+		if (const TOptional<FIntRect> Rect = RectangleFromFrustum(InFrustum))
+		{
+			GetNonOccludedElements(*Rect);
+		}
+		else
+		{
+			bTransparentBoxSelection = true;
+		}
+	}
+
+	if (bTransparentBoxSelection)
+	{
+		for (const auto& [ControlRig, ShapeActors]: ControlRigShapeActors)
+		{
+			if (ControlRig && ControlRig->GetControlsVisible())
+			{
+				for (AControlRigShapeActor* ShapeActor : ShapeActors)
+				{
+					const bool bTreatShape = ShapeActor && ShapeActor->IsSelectable() && !ShapeActor->IsTemporarilyHiddenInEditor();
+					if (bTreatShape)
+					{
+						for (UActorComponent* Component : ShapeActor->GetComponents())
+						{
+							UPrimitiveComponent* PrimitiveComponent = Cast<UPrimitiveComponent>(Component);
+							if (PrimitiveComponent && PrimitiveComponent->IsRegistered() && PrimitiveComponent->IsVisibleInEditor())
+							{
+								if (PrimitiveComponent->IsShown(ViewportClient->EngineShowFlags) && PrimitiveComponent->ComponentIsTouchingSelectionFrustum(InFrustum, false /*only bsp*/, false/*encompass entire*/))
+								{
+									TArray<FRigElementKey>& Controls = Elements.FindOrAdd(ControlRig);
+									Controls.Add(ShapeActor->GetElementKey());
+								}
+							}
+						}
+					}
+				}
+			}
+		}
+	}
+}
+	
+TOptional<FIntRect> FSelectionHelper::RectangleFromFrustum(const FConvexVolume& InFrustum) const
+{
+	static TOptional<FIntRect> InvalidBox;
+
+	if (!IsValid())
+	{
+		return InvalidBox;
+	}
+	
+	const FConvexVolume::FPlaneArray& Planes = InFrustum.Planes;
+	const int32 NumPlanes = Planes.Num();
+	if (NumPlanes < 4)
+	{
+		return InvalidBox;
+	}
+
+	FSceneInterface* Scene = ViewportClient->GetScene();
+	if (!Scene)
+	{
+		return InvalidBox;
+	}
+	
+	FSceneViewFamilyContext ViewFamily(FSceneViewFamily::ConstructionValues(ViewportClient->Viewport, Scene, ViewportClient->EngineShowFlags ));
+	const FSceneView* SceneView = ViewportClient->CalcSceneView(&ViewFamily);
+	if (!SceneView)
+	{
+		return InvalidBox;
+	}
+
+	FPlane NearPlane;
+	if (!SceneView->ViewMatrices.GetViewProjectionMatrix().GetFrustumNearPlane(NearPlane))
+	{
+		return InvalidBox;
+	}
+
+	const FVector& ViewLocation = ViewportClient->GetViewLocation();
+
+	// compute the intersections of side planes with the near plane
+	TArray<FVector> Intersections;
+	Intersections.Reserve(4);
+	constexpr double Threshold = FMath::Square(0.001); // cf. IntersectPlanes2 for threshold
+	for (uint8 Index = 0; Index < 4; Index++)
+	{
+		FVector Direction = FVector::CrossProduct(Planes[Index], Planes[(Index+1)%4]);
+		const double DD = Direction.SizeSquared();
+		if (DD >= Threshold)
+		{
+			// planes intersect, compute the intersection with the near plane
+			Direction.Normalize();
+			Intersections.Add(FMath::RayPlaneIntersection(ViewLocation, Direction, NearPlane));
+		}
+	}
+
+	if (Intersections.Num() == 4)
+	{
+		// compute the screen space & pixel projections of those intersections
+		FVector2D ScreenPos[4];
+		bool bIsProjectionValid = true;
+		for (uint8 Index = 0; Index < 4; Index++)
+		{
+			bIsProjectionValid &= SceneView->ScreenToPixel(SceneView->WorldToScreen(Intersections[Index]), ScreenPos[Index]);
+		}
+
+		if (bIsProjectionValid)
+		{
+			const FVector2D& TopLeft = ScreenPos[3], BottomRight = ScreenPos[1];
+
+			const int32 MinX = FMath::Max<int32>(0, FMath::FloorToInt32(TopLeft.X));
+			const int32 MinY = FMath::Max<int32>(0, FMath::FloorToInt32(TopLeft.Y));
+
+			const FIntPoint ViewportSize = ViewportClient->Viewport->GetSizeXY();
+			const int32 MaxX = FMath::Clamp(FMath::CeilToInt32(BottomRight.X), MinX+1, ViewportSize.X);
+			const int32 MaxY = FMath::Clamp(FMath::CeilToInt32(BottomRight.Y), MinY+1, ViewportSize.Y);
+
+			const FIntRect Rect(MinX, MinY, MaxX, MaxY);
+			return Rect;
+		}
+	}
+
+	return InvalidBox;
+}
+
+void FSelectionHelper::GetNonOccludedElements(const FIntRect& InRect) const
+{
+	if (!IsValid())
+	{
+		return;
+	}
+
+	// extend that predicate to filter more hit proxies
+	ViewportClient->Viewport->EnumerateHitProxiesInRect(InRect, [this](HHitProxy* HitProxy)
+	{
+		const HActor* ActorHitProxy = HitProxyCast<HActor>(HitProxy);
+		if (const AControlRigShapeActor* ShapeActor = ActorHitProxy ? Cast<AControlRigShapeActor>(ActorHitProxy->Actor) : nullptr)
+		{
+			if (ShapeActor->IsSelectable() && !ShapeActor->IsTemporarilyHiddenInEditor())
+			{
+				if (UControlRig* ControlRig = ShapeActor->ControlRig.Get())
+				{
+					TArray<FRigElementKey>& Controls = Elements.FindOrAdd(ControlRig);
+					Controls.AddUnique(ShapeActor->GetElementKey());
+				}
+			}
+		}
+		return true;
+	});
+}
 	
 }
\ No newline at end of file
diff --git a/Engine/Plugins/Animation/ControlRig/Source/ControlRigEditor/Public/EditMode/ControlRigEditMode.h b/Engine/Plugins/Animation/ControlRig/Source/ControlRigEditor/Public/EditMode/ControlRigEditMode.h
--- a/Engine/Plugins/Animation/ControlRig/Source/ControlRigEditor/Public/EditMode/ControlRigEditMode.h
+++ b/Engine/Plugins/Animation/ControlRig/Source/ControlRigEditor/Public/EditMode/ControlRigEditMode.h
@@ -244,7 +244,6 @@ protected:
 	/** Get the node name from the property path */
 	AControlRigShapeActor* GetControlShapeFromControlName(UControlRig* InControlRig,const FName& InControlName) const;
 
-protected:
 	/** Helper function: set ControlRigs array to the details panel */
 	void SetObjects_Internal();
 
@@ -261,9 +260,6 @@ protected:
 	/** Get the current coordinate system space */
 	ECoordSystem GetCoordSystemSpace() const;
 	
-	/** Helper function for box/frustum intersection */
-	bool IntersectSelect(bool InSelect, const TFunctionRef<bool(const AControlRigShapeActor*, const FTransform&)>& Intersects);
-
 	/** Handle selection internally */
 	void HandleSelectionChanged();
 
diff --git a/Engine/Plugins/Animation/ControlRig/Source/ControlRigEditor/Public/EditMode/ControlRigEditModeUtil.h b/Engine/Plugins/Animation/ControlRig/Source/ControlRigEditor/Public/EditMode/ControlRigEditModeUtil.h
--- a/Engine/Plugins/Animation/ControlRig/Source/ControlRigEditor/Public/EditMode/ControlRigEditModeUtil.h
+++ b/Engine/Plugins/Animation/ControlRig/Source/ControlRigEditor/Public/EditMode/ControlRigEditModeUtil.h
@@ -2,19 +2,21 @@
 
 #pragma once
 
+#include "ControlRigGizmoActor.h"
 #include "Math/MathFwd.h"
 #include "Math/Vector.h"
 #include "Math/Rotator.h"
 #include "Rigs/RigHierarchy.h"
-
 #include "Rigs/RigHierarchyDefines.h"
-
 #include "UObject/WeakObjectPtr.h"
 
+class FEditorViewportClient;
 class UControlRig;
+class ULevelEditorViewportSettings;
 class URigHierarchy;
 class UTickableConstraint;
 
+struct FConvexVolume;
 struct FRigControlModifiedContext;
 struct FRigControlElement;
 
@@ -128,6 +130,50 @@ private:
 	FRigControlElement* ControlElement = nullptr;
 	const FTransform& ComponentWorldTransform;
 };
+
+/**
+* FSelectionHelper is a wrapper struct to handle control rig related viewport selection.
+*/
+
+struct FSelectionHelper
+{
+	FSelectionHelper(
+		FEditorViewportClient* InViewportClient,
+		const TMap<UControlRig*, TArray<TObjectPtr<AControlRigShapeActor>>>& InControlRigShapeActors,
+		TMap<UControlRig*, TArray<FRigElementKey>>& OutElements);
+
+	/**
+	 * Gets the elements contained in the frustum.
+	 * Note that the function will actually use a screen space rectangle selection if occluded objects must be skipped.
+	 */
+	void GetFromFrustum(const FConvexVolume& InFrustum) const;
+
+private:
+
+	/** Returns true if the viewport to select into is not null. */
+	bool IsValid() const;
+	
+	/** Returns a screen space rectangle based of the frustum. (it assumes that there's a valid near plane) */
+	TOptional<FIntRect> RectangleFromFrustum(const FConvexVolume& InFrustum) const;
+
+	/** Gets the no-occluded elements contained in the screen space rectangle. */
+	void GetNonOccludedElements(const FIntRect& InRect) const;
+
+	/** The viewport client being interacted with. */
+	FEditorViewportClient* ViewportClient = nullptr;
+
+	/** A reference to the edit mode's control shapes. */
+	const TMap<UControlRig*, TArray<TObjectPtr<AControlRigShapeActor>>>& ControlRigShapeActors;
+
+	/** The elements to be selected. */
+	TMap<UControlRig*, TArray<FRigElementKey>>& Elements;
+
+	/** Current viewport settings. (used for strict box & transparent selection) */
+	const ULevelEditorViewportSettings* LevelEditorViewportSettings = nullptr;
+
+	/** List of layers that are hidden in this view. (only valid for level editor vpc) */
+	TArray<FName> HiddenLayers;
+};
 	
 }
 
diff --git a/Engine/Source/Runtime/Engine/Public/UnrealClient.h b/Engine/Source/Runtime/Engine/Public/UnrealClient.h
--- a/Engine/Source/Runtime/Engine/Public/UnrealClient.h
+++ b/Engine/Source/Runtime/Engine/Public/UnrealClient.h
@@ -591,6 +591,13 @@ public:
 	 */
 	ENGINE_API void GetElementHandlesInRect(FIntRect InRect, FTypedElementListRef OutElementHandles);
 
+	/**
+	 * Enumerate all valid hit proxies found within a specified region.
+	 * InRect must be entirely within the viewport's client area.
+	 * If the hit proxies are not cached, this will call ViewportClient->Draw with a hit-testing canvas.
+	 */
+	ENGINE_API void EnumerateHitProxiesInRect(FIntRect InRect, TFunctionRef<bool(HHitProxy*)> InCallback);
+	
 	/**
 	 * Retrieves the interface to the viewport's frame, if it has one.
 	 * @return The viewport's frame interface.
@@ -735,17 +742,6 @@ protected:
 	 */
 	void HighResScreenshot();
 
-private:
-
-	/**
-	 * Enumerate all valid hit proxies found within a specified region.
-	 * InRect must be entirely within the viewport's client area.
-	 * If the hit proxies are not cached, this will call ViewportClient->Draw with a hit-testing canvas.
-	 */
-	void EnumerateHitProxiesInRect(FIntRect InRect, TFunctionRef<bool(HHitProxy*)> InCallback);
-
-protected:
-
 	/** A map from 2D coordinates to cached hit proxies. */
 	class FHitProxyMap : public FHitProxyConsumer, public FRenderTarget, public FGCObject
 	{
