Animation Mode: improved selection of multiple controls (follow-up)
- force selection guard reset to ensure sequencer selection change is broadcast while bIsDoingSelection is still true

#rnx

diff --git a/Engine/Plugins/Animation/ControlRig/Source/ControlRigEditor/Private/Sequencer/ControlRigParameterTrackEditor.cpp b/Engine/Plugins/Animation/ControlRig/Source/ControlRigEditor/Private/Sequencer/ControlRigParameterTrackEditor.cpp
--- a/Engine/Plugins/Animation/ControlRig/Source/ControlRigEditor/Private/Sequencer/ControlRigParameterTrackEditor.cpp
+++ b/Engine/Plugins/Animation/ControlRig/Source/ControlRigEditor/Private/Sequencer/ControlRigParameterTrackEditor.cpp
@@ -4016,6 +4016,9 @@ void FControlRigParameterTrackEditor::HandleControlSelected(UControlRig* Subject
 			Sequencer->NotifyMovieSceneDataChanged(EMovieSceneDataChangeType::RefreshTree);
 		}
 	}
+
+	// force selection guard reset to ensure sequencer selection change is broadcast while bIsDoingSelection is still true
+	SequencerSelectionGuard.Reset();
 }
 
 void FControlRigParameterTrackEditor::HandleOnPostConstructed(UControlRig* Subject, const FName& InEventName)
@@ -4258,6 +4261,12 @@ void FControlRigParameterTrackEditor::GetControlRigKeys(
 			case ERigControlType::Bool:
 			{
 				bool Val = ControlValue.Get<bool>();
+
+				if (ControlElement->GetName() == "Arm R/IK Solve")
+				{
+					UE_LOG(LogControlRigEditor, Warning, TEXT("Arm R/IK Solve = %d"), Val);
+				}
+					
 				pChannelIndex->GeneratedKeyIndex = OutGeneratedKeys.Num();
 				OutGeneratedKeys.Add(FMovieSceneChannelValueSetter::Create<FMovieSceneBoolChannel>(ChannelIndex, Val, bSetKey));
 				break;
