Control Rig: Fix Bad Data from bad Rename fix.  A popup will be shown to fix any bad sections with multiple parameters names of the same type.

diff --git a/Engine/Plugins/Animation/ControlRig/Source/ControlRig/Private/Sequencer/MovieSceneControlRigParameterSection.cpp b/Engine/Plugins/Animation/ControlRig/Source/ControlRig/Private/Sequencer/MovieSceneControlRigParameterSection.cpp
--- a/Engine/Plugins/Animation/ControlRig/Source/ControlRig/Private/Sequencer/MovieSceneControlRigParameterSection.cpp
+++ b/Engine/Plugins/Animation/ControlRig/Source/ControlRig/Private/Sequencer/MovieSceneControlRigParameterSection.cpp
@@ -32,6 +32,10 @@
 #include "Sequencer/MovieSceneControlRigSystem.h"
 #include "ControlRigObjectBinding.h"
 
+#if WITH_EDITOR
+#include "Misc/MessageDialog.h"
+#endif
+
 #include UE_INLINE_GENERATED_CPP_BY_NAME(MovieSceneControlRigParameterSection)
 
 namespace UE::MovieScene
@@ -1715,8 +1719,49 @@ void UMovieSceneControlRigParameterSection::ReconstructChannelProxy()
 	BroadcastChanged();
 }
 
+//hopefully remove after a bit
+void UMovieSceneControlRigParameterSection::HACK_FixMultipleParamsWithSameName()
+{
+#if WITH_EDITOR
+	bool bHasDup = false;
+	bHasDup = HACK_CheckForDupParameters(ScalarParameterNamesAndCurves) || bHasDup;
+	bHasDup = HACK_CheckForDupParameters(BoolParameterNamesAndCurves) || bHasDup;
+	bHasDup = HACK_CheckForDupParameters(EnumParameterNamesAndCurves) || bHasDup;
+	bHasDup = HACK_CheckForDupParameters(IntegerParameterNamesAndCurves) || bHasDup;
+	bHasDup = HACK_CheckForDupParameters(Vector2DParameterNamesAndCurves) || bHasDup;
+	bHasDup = HACK_CheckForDupParameters(VectorParameterNamesAndCurves) || bHasDup;
+	bHasDup = HACK_CheckForDupParameters(ColorParameterNamesAndCurves) || bHasDup;
+	bHasDup = HACK_CheckForDupParameters(TransformParameterNamesAndCurves) || bHasDup;
+	if (bHasDup)
+	{
+
+		const EAppReturnType::Type Choice = FMessageDialog::Open(EAppMsgType::YesNo, LOCTEXT("Sequencer", "Duplicate corrupted controls found in level sequence, just keep first set of original Control curves (recommended but may revert to an older version)?"));
+
+		if (Choice == EAppReturnType::Yes)
+		{
+			bool bRemovedSomething = false;
+			bRemovedSomething = HACK_CleanForEachParameter(ScalarParameterNamesAndCurves) || bRemovedSomething;
+			bRemovedSomething = HACK_CleanForEachParameter(BoolParameterNamesAndCurves) || bRemovedSomething;
+			bRemovedSomething = HACK_CleanForEachParameter(EnumParameterNamesAndCurves) || bRemovedSomething;
+			bRemovedSomething = HACK_CleanForEachParameter(IntegerParameterNamesAndCurves) || bRemovedSomething;
+			bRemovedSomething = HACK_CleanForEachParameter(Vector2DParameterNamesAndCurves) || bRemovedSomething;
+			bRemovedSomething = HACK_CleanForEachParameter(VectorParameterNamesAndCurves) || bRemovedSomething;
+			bRemovedSomething = HACK_CleanForEachParameter(ColorParameterNamesAndCurves) || bRemovedSomething;
+			bRemovedSomething = HACK_CleanForEachParameter(TransformParameterNamesAndCurves) || bRemovedSomething;
+
+			if (bRemovedSomething)
+			{
+				UE_LOG(LogControlRig, Warning, TEXT("There were duplicated curves found possibly by 40400084, resave to remove duplicates and keep first one"));
+				Modify(); //we removed something so mark as dirty so animators can resave
+			}
+		}
+	}
+#endif
+}
+
 EMovieSceneChannelProxyType UMovieSceneControlRigParameterSection::CacheChannelProxy()
 {
+	HACK_FixMultipleParamsWithSameName();
 	const FName UIMin("UIMin");
 	const FName UIMax("UIMax");
 
diff --git a/Engine/Plugins/Animation/ControlRig/Source/ControlRig/Public/Sequencer/MovieSceneControlRigParameterSection.h b/Engine/Plugins/Animation/ControlRig/Source/ControlRig/Public/Sequencer/MovieSceneControlRigParameterSection.h
--- a/Engine/Plugins/Animation/ControlRig/Source/ControlRig/Public/Sequencer/MovieSceneControlRigParameterSection.h
+++ b/Engine/Plugins/Animation/ControlRig/Source/ControlRig/Public/Sequencer/MovieSceneControlRigParameterSection.h
@@ -669,4 +669,56 @@ private:
 	//Last set of Controls used to reconstruct the channel proxies, used to make sure controls really changed if we want to reconstruct
 	//only care to check name and type
 	TArray<TPair<FName, ERigControlType>> LastControlsUsedToReconstruct;
+
+private:
+
+	//there was a regression that caused certain controls present in the Hierarchy->GetPreviousName map to be applied backwards so it would incorrectly try to
+	//replace controls with new names with their old name s instead of vice versa.
+	//hope we can get rid of this function since it's only a few files with bad data hopefully
+	//this will go through the parameter names and if there are duplicates it will remove them, the first one is always the one to keep
+	template<typename T>
+	bool  HACK_CheckForDupParameters(TArray<T>& InParameterArray)
+	{
+		TSet<FName> ExistingNames;
+		TArray<int32> IndicesToRemove;
+		for (T& Parameter : InParameterArray)
+		{
+			if (ExistingNames.Contains(Parameter.ParameterName))
+			{
+				return true;
+			}
+			else
+			{
+				ExistingNames.Add(Parameter.ParameterName);
+			}
+		}
+		return false;
+	}
+	
+	
+	template<typename T>
+	bool  HACK_CleanForEachParameter(TArray<T>& InParameterArray)
+	{
+		TSet<FName> ExistingNames;
+		int32 Index = 0;
+		TArray<int32> IndicesToRemove; 
+		for (T& Parameter : InParameterArray)
+		{
+			if (ExistingNames.Contains(Parameter.ParameterName))
+			{
+				IndicesToRemove.Insert(Index,0);  //put first since we want to remove last
+			}
+			else
+			{
+				ExistingNames.Add(Parameter.ParameterName);
+			}
+			++Index;
+		}
+		for (int32 RemoveIndex : IndicesToRemove)
+		{
+			InParameterArray.RemoveAt(RemoveIndex);
+		}
+		return IndicesToRemove.Num() > 0;
+	}
+	void HACK_FixMultipleParamsWithSameName();
 };
