Fixed array allignment error in Anim Next Control Rig latent pins
Added support for byte variables in RigVM variable mappings, as these are used for UUserStructs

#ForUE5

diff --git a/Engine/Plugins/Animation/ControlRig/Source/ControlRig/Private/Tools/ControlRigVariableMappings.cpp b/Engine/Plugins/Animation/ControlRig/Source/ControlRig/Private/Tools/ControlRigVariableMappings.cpp
--- a/Engine/Plugins/Animation/ControlRig/Source/ControlRig/Private/Tools/ControlRigVariableMappings.cpp
+++ b/Engine/Plugins/Animation/ControlRig/Source/ControlRig/Private/Tools/ControlRigVariableMappings.cpp
@@ -480,6 +480,13 @@ void FControlRigVariableMappings::PropagateInputPropertiesNoCache(const UObject*
 					EnumProperty->CopyCompleteValue(Variable.Memory, SrcPtr);
 				}
 			}
+			else if (FByteProperty* ByteProperty = CastField<FByteProperty>(CallerProperty))
+			{
+				if (ensure(ByteProperty->SameType(Variable.Property)))
+				{
+					ByteProperty->CopyCompleteValue(Variable.Memory, SrcPtr);
+				}
+			}
 			else
 			{
 				ensureMsgf(false, TEXT("Property %s type %s not recognized"), *CallerProperty->GetName(), *CallerProperty->GetCPPType());
@@ -747,6 +754,17 @@ void FControlRigVariableMappings::AddVariableFunction(const FRigVMExternalVariab
 			});
 		}
 	}
+	else if (FByteProperty* ByteProperty = CastField<FByteProperty>(InSourceProperty))
+	{
+		if (ensure(ByteProperty->SameType(InVariable.Property)))
+		{
+			return AddUpdateFunction([ByteProperty, this, VariableIndex](const UObject* InSourceInstance)
+				{
+					const uint8* SrcPtr = ByteProperty->ContainerPtrToValuePtr<uint8>(InSourceInstance);
+					ByteProperty->CopyCompleteValue(Variables[VariableIndex].Memory, SrcPtr);
+				});
+		}
+	}
 	else
 	{
 		ensureMsgf(false, TEXT("Property %s type %s not recognized"), *InSourceProperty->GetName(), *InSourceProperty->GetCPPType());
@@ -950,6 +968,16 @@ void FControlRigVariableMappings::AddCustomVariableFunction(const FRigVMExternal
 				});
 		}
 	}
+	else if (const FByteProperty* ByteProperty = CastField<FByteProperty>(InSourceProperty))
+	{
+		if (ensure(ByteProperty->SameType(InVariable.Property)))
+		{
+			AddCustomUpdateFunction([ByteProperty, InSourcePropertyMemory, Variable = InVariable]()
+				{
+					ByteProperty->CopyCompleteValue(Variable.Memory, InSourcePropertyMemory);
+				});
+		}
+	}
 	else
 	{
 		ensureMsgf(false, TEXT("Property %s type %s not recognized"), *InSourceProperty->GetName(), *InSourceProperty->GetCPPType());
diff --git a/Engine/Plugins/Experimental/AnimNextControlRig/Source/AnimNextControlRig/Private/ControlRigTrait.cpp b/Engine/Plugins/Experimental/AnimNextControlRig/Source/AnimNextControlRig/Private/ControlRigTrait.cpp
--- a/Engine/Plugins/Experimental/AnimNextControlRig/Source/AnimNextControlRig/Private/ControlRigTrait.cpp
+++ b/Engine/Plugins/Experimental/AnimNextControlRig/Source/AnimNextControlRig/Private/ControlRigTrait.cpp
@@ -780,8 +780,8 @@ bool FControlRigTrait::GetVariableSizeAndAlignment(const FRigVMExternalVariable&
 		}
 		else if (const UEnum* Enum = Cast<UEnum>(Variable.TypeObject))
 		{
-			PropertySize = Variable.Size;
-			PropertyAlignment = Variable.Size;
+			PropertySize = Variable.Property->GetSize();
+			PropertyAlignment = Variable.Property->GetMinAlignment();
 			bValidType = true;
 		}
 		else if (const UClass* Class = Cast<UClass>(Variable.TypeObject))
@@ -793,14 +793,14 @@ bool FControlRigTrait::GetVariableSizeAndAlignment(const FRigVMExternalVariable&
 		else
 		{
 			ensureMsgf(false, TEXT("Unsupported ControlRig public variable type: %s"), *Variable.TypeName.ToString());
-			PropertySize = Variable.Size;
-			PropertyAlignment = Variable.Size;
+			PropertySize = Variable.Property->GetSize();
+			PropertyAlignment = Variable.Property->GetMinAlignment();
 		}
 	}
 	else
 	{
-		PropertySize = Variable.Size;
-		PropertyAlignment = Variable.Size;
+		PropertySize = Variable.Property->GetSize();
+		PropertyAlignment = Variable.Property->GetMinAlignment();
 		bValidType = true;
 	}
 
