Control Rig: Fix layered rigs not running pre/post forward solve
#rnx

diff --git a/Engine/Plugins/Animation/ControlRig/Source/ControlRig/Private/ControlRig.cpp b/Engine/Plugins/Animation/ControlRig/Source/ControlRig/Private/ControlRig.cpp
--- a/Engine/Plugins/Animation/ControlRig/Source/ControlRig/Private/ControlRig.cpp
+++ b/Engine/Plugins/Animation/ControlRig/Source/ControlRig/Private/ControlRig.cpp
@@ -504,7 +504,9 @@ void UControlRig::Evaluate_AnyThread()
 		
 		// Forward solve
 		{
+			Execute(FRigUnit_PreBeginExecution::EventName);
 			Execute(FRigUnit_BeginExecution::EventName);
+			Execute(FRigUnit_PostBeginExecution::EventName);
 		}
 	}
 	else
