Control Rig: Outliner work to arrange by modules
#rnx

diff --git a/Engine/Plugins/Animation/ControlRig/Source/ControlRigEditor/Private/ControlRigElementDetails.cpp b/Engine/Plugins/Animation/ControlRig/Source/ControlRigEditor/Private/ControlRigElementDetails.cpp
--- a/Engine/Plugins/Animation/ControlRig/Source/ControlRigEditor/Private/ControlRigElementDetails.cpp
+++ b/Engine/Plugins/Animation/ControlRig/Source/ControlRigEditor/Private/ControlRigElementDetails.cpp
@@ -4258,6 +4258,9 @@ void FRigControlElementDetails::CustomizeAvailableSpaces(IDetailLayoutBuilder& D
 	DisplaySettings.bHideParentsOnFilter = true;
 	DisplaySettings.bFlattenHierarchyOnFilter = true;
 	DisplaySettings.bShowIconColors = true;
+	DisplaySettings.bShowAnimationChannels = true;
+	DisplaySettings.bArrangeByModules = false;
+	DisplaySettings.bUseShortName = false;
 
 	const TSharedRef<SHorizontalBox> HeaderContentWidget = SNew(SHorizontalBox);
 	HeaderContentWidget->AddSlot()
diff --git a/Engine/Plugins/Animation/ControlRig/Source/ControlRigEditor/Private/ControlRigHierarchyCommands.cpp b/Engine/Plugins/Animation/ControlRig/Source/ControlRigEditor/Private/ControlRigHierarchyCommands.cpp
--- a/Engine/Plugins/Animation/ControlRig/Source/ControlRigEditor/Private/ControlRigHierarchyCommands.cpp
+++ b/Engine/Plugins/Animation/ControlRig/Source/ControlRigEditor/Private/ControlRigHierarchyCommands.cpp
@@ -33,10 +33,12 @@ void FControlRigHierarchyCommands::RegisterCommands()
 	UI_COMMAND(Unparent, "Unparent", "Unparents the selected elements from the hierarchy", EUserInterfaceActionType::Button, FInputChord(EKeys::P, EModifierKey::Shift));
 	UI_COMMAND(FilteringFlattensHierarchy, "Filtering Flattens Hierarchy", "Whether to keep the hierarchy or flatten it when searching for tree items", EUserInterfaceActionType::ToggleButton, FInputChord());
 	UI_COMMAND(HideParentsWhenFiltering, "Hide Parents When Filtering", "Whether to show parent items grayed out, or hide them entirely when filtering", EUserInterfaceActionType::ToggleButton, FInputChord());
+	UI_COMMAND(ArrangeByModules, "Arrange by Modules", "Arrange the hierarchy by the modules that spawn each control", EUserInterfaceActionType::ToggleButton, FInputChord());
 	UI_COMMAND(ShowFullPaths, "Show Full Paths", "Whether to show the full path or the short name of each element", EUserInterfaceActionType::ToggleButton, FInputChord());
 	UI_COMMAND(ShowImportedBones, "Show Imported Bones", "Whether to show or hide imported bones", EUserInterfaceActionType::ToggleButton, FInputChord());
 	UI_COMMAND(ShowBones, "Show Bones", "Whether to show or hide bones", EUserInterfaceActionType::ToggleButton, FInputChord());
 	UI_COMMAND(ShowControls, "Show Controls", "Whether to show or hide controls", EUserInterfaceActionType::ToggleButton, FInputChord());
+	UI_COMMAND(ShowAnimationChannels, "Show Animation Channels", "Whether to show or hide animation channels", EUserInterfaceActionType::ToggleButton, FInputChord());
 	UI_COMMAND(ShowNulls, "Show Nulls", "Whether to show or hide nulls", EUserInterfaceActionType::ToggleButton, FInputChord());
 	UI_COMMAND(ShowReferences, "Show References", "Whether to show or hide references", EUserInterfaceActionType::ToggleButton, FInputChord());
 	UI_COMMAND(ShowSockets, "Show Sockets", "Whether to show or hide sockets", EUserInterfaceActionType::ToggleButton, FInputChord());
diff --git a/Engine/Plugins/Animation/ControlRig/Source/ControlRigEditor/Private/ControlRigHierarchyCommands.h b/Engine/Plugins/Animation/ControlRig/Source/ControlRigEditor/Private/ControlRigHierarchyCommands.h
--- a/Engine/Plugins/Animation/ControlRig/Source/ControlRigEditor/Private/ControlRigHierarchyCommands.h
+++ b/Engine/Plugins/Animation/ControlRig/Source/ControlRigEditor/Private/ControlRigHierarchyCommands.h
@@ -99,6 +99,9 @@ public:
 	/* Hide parents on filter */
 	TSharedPtr< FUICommandInfo > HideParentsWhenFiltering;
 
+	/* Arrange the controls by modules in a modular rig */
+	TSharedPtr< FUICommandInfo > ArrangeByModules;
+
 	/* Show full paths */
 	TSharedPtr< FUICommandInfo > ShowFullPaths;
 
@@ -111,6 +114,9 @@ public:
 	/* Show controls */
 	TSharedPtr< FUICommandInfo > ShowControls;
 
+	/* Show animation channels */
+	TSharedPtr< FUICommandInfo > ShowAnimationChannels;
+
 	/* Show spaces */
 	TSharedPtr< FUICommandInfo > ShowNulls;
 
diff --git a/Engine/Plugins/Animation/ControlRig/Source/ControlRigEditor/Private/EditMode/ControlRigEditMode.cpp b/Engine/Plugins/Animation/ControlRig/Source/ControlRigEditor/Private/EditMode/ControlRigEditMode.cpp
--- a/Engine/Plugins/Animation/ControlRig/Source/ControlRigEditor/Private/EditMode/ControlRigEditMode.cpp
+++ b/Engine/Plugins/Animation/ControlRig/Source/ControlRigEditor/Private/EditMode/ControlRigEditMode.cpp
@@ -70,6 +70,7 @@
 #include "AnimationEditorViewportClient.h"
 #include "ControlShapeActorHelper.h"
 #include "EditorInteractiveGizmoManager.h"
+#include "ModularRig.h"
 #include "Constraints/ControlRigTransformableHandle.h"
 #include "Editor/ControlRigViewportToolbarExtensions.h"
 #include "Editor/Sequencer/Private/SSequencer.h"
@@ -924,6 +925,21 @@ void FControlRigEditMode::Render(const FSceneView* View, FViewport* Viewport, FP
 							{
 								return true;
 							}
+
+							if (UModularRig* ModularRig = Cast<UModularRig>(ControlRig))
+							{
+								const FString ModuleName = Hierarchy->GetModuleName(ControlElement->GetKey());
+								if (FRigModuleInstance* Module = ModularRig->FindModule(*ModuleName))
+								{
+									if (UControlRig* ModuleRig = Module->GetRig())
+									{
+										if (!ModuleRig->GetControlsVisible())
+										{
+											return true;
+										}
+									}
+								}
+							}
 						}
 				
 						const FTransform Transform = Hierarchy->GetTransform(TransformElement, ERigTransformType::CurrentGlobal);
@@ -3421,6 +3437,14 @@ void FControlRigEditMode::ToggleAllManipulators()
 			if (UControlRig* ControlRig = RuntimeRigPtr.Get())
 			{
 				ControlRig->SetControlsVisible(true);
+				if (UModularRig* ModularRig = Cast<UModularRig>(ControlRig))
+				{
+					ModularRig->ForEachModule([](FRigModuleInstance* Module)
+					{
+						Module->GetRig()->SetControlsVisible(true);
+						return true;
+					});
+				}
 			}
 		}
 	}
@@ -5590,6 +5614,14 @@ void FControlRigEditMode::AddControlRigInternal(UControlRig* InControlRig)
 	LastShapeLibraryHash.Add(INDEX_NONE);
 
 	InControlRig->SetControlsVisible(true);
+	if (UModularRig* ModularRig = Cast<UModularRig>(InControlRig))
+	{
+		ModularRig->ForEachModule([](FRigModuleInstance* Module)
+		{
+			Module->GetRig()->SetControlsVisible(true);
+			return true;
+		});
+	}
 	InControlRig->PostInitInstanceIfRequired();
 
 	InControlRig->GetHierarchy()->OnModified().RemoveAll(this);
diff --git a/Engine/Plugins/Animation/ControlRig/Source/ControlRigEditor/Private/EditMode/ControlShapeActorHelper.cpp b/Engine/Plugins/Animation/ControlRig/Source/ControlRigEditor/Private/EditMode/ControlShapeActorHelper.cpp
--- a/Engine/Plugins/Animation/ControlRig/Source/ControlRigEditor/Private/EditMode/ControlShapeActorHelper.cpp
+++ b/Engine/Plugins/Animation/ControlRig/Source/ControlRigEditor/Private/EditMode/ControlShapeActorHelper.cpp
@@ -5,6 +5,7 @@
 #include "ControlRig.h"
 #include "ControlRigEditModeSettings.h"
 #include "ControlRigGizmoActor.h"
+#include "ModularRig.h"
 #include "Rigs/RigHierarchyElements.h"
 
 namespace ControlRigEditMode::Shapes
@@ -203,6 +204,14 @@ void UpdateControlShape(AControlRigShapeActor* InShapeActor, FRigControlElement*
 
 	// update visibility & color
 	bool bIsVisible = ControlSettings.IsVisible();
+	if (const UModularRig* ModularRig = Cast<UModularRig>(InUpdateParams.ControlRig))
+	{
+		const FString ModuleName = InUpdateParams.Hierarchy->GetModuleName(InControlElement->GetKey());
+		if (const FRigModuleInstance* Module = ModularRig->FindModule(*ModuleName))
+		{
+			bIsVisible &= Module->GetRig()->GetControlsVisible();
+		}
+	}
 	bool bRespectVisibilityForSelection = !InUpdateParams.bIsInLevelEditor;
 
 	const bool bControlsHiddenInViewport =
diff --git a/Engine/Plugins/Animation/ControlRig/Source/ControlRigEditor/Private/EditMode/SControlRigEditModeTools.cpp b/Engine/Plugins/Animation/ControlRig/Source/ControlRigEditor/Private/EditMode/SControlRigEditModeTools.cpp
--- a/Engine/Plugins/Animation/ControlRig/Source/ControlRigEditor/Private/EditMode/SControlRigEditModeTools.cpp
+++ b/Engine/Plugins/Animation/ControlRig/Source/ControlRigEditor/Private/EditMode/SControlRigEditModeTools.cpp
@@ -220,6 +220,9 @@ void SControlRigEditModeTools::Construct(const FArguments& InArgs, TSharedPtr<FC
 	DisplaySettings.bHideParentsOnFilter = true;
 	DisplaySettings.bFlattenHierarchyOnFilter = true;
 	DisplaySettings.bShowIconColors = true;
+	DisplaySettings.bShowAnimationChannels = true;
+	DisplaySettings.bArrangeByModules = false;
+	DisplaySettings.bUseShortName = false;
 #if USE_LOCAL_DETAILS
 	FRigTreeDelegates RigTreeDelegates;
 	RigTreeDelegates.OnGetHierarchy = FOnGetRigTreeHierarchy::CreateSP(this, &SControlRigEditModeTools::GetHierarchy);
diff --git a/Engine/Plugins/Animation/ControlRig/Source/ControlRigEditor/Private/EditMode/SControlRigOutliner.cpp b/Engine/Plugins/Animation/ControlRig/Source/ControlRigEditor/Private/EditMode/SControlRigOutliner.cpp
--- a/Engine/Plugins/Animation/ControlRig/Source/ControlRigEditor/Private/EditMode/SControlRigOutliner.cpp
+++ b/Engine/Plugins/Animation/ControlRig/Source/ControlRigEditor/Private/EditMode/SControlRigOutliner.cpp
@@ -5,12 +5,14 @@
 #include "Widgets/Input/SButton.h"
 #include "Widgets/Layout/SExpandableArea.h"
 #include "Widgets/Layout/SScrollBox.h"
+#include "Widgets/Input/SComboButton.h"
 #include "AssetRegistry/AssetData.h"
 #include "Widgets/Text/SInlineEditableTextBlock.h"
 #include "Styling/AppStyle.h"
 #include "Styling/CoreStyle.h"
 #include "ScopedTransaction.h"
 #include "ControlRig.h"
+#include "ControlRigHierarchyCommands.h"
 #include "EditMode/ControlRigEditMode.h"
 #include "EditorModeManager.h"
 #include "ISequencer.h"
@@ -23,6 +25,7 @@
 #include "Editor/SRigHierarchyTreeView.h"
 #include "Rigs/RigHierarchyController.h"
 #include "ControlRigObjectBinding.h"
+#include "ModularRig.h"
 #include "Settings/ControlRigSettings.h"
 #include "Widgets/Layout/SSpacer.h"
 #include "Widgets/Images/SImage.h"
@@ -30,6 +33,7 @@
 #include "Framework/Application/SlateApplication.h"
 #include "Editor/EditorEngine.h"
 #include "ScopedTransaction.h"
+#include "Editor/SRigHierarchy.h"
 #include "Rigs/FKControlRig.h"
 
 #define LOCTEXT_NAMESPACE "ControlRigOutliner"
@@ -52,6 +56,11 @@ FText FMultiRigData::GetName() const
 	{
 		return FText::FromName(Key->Name);
 	}
+
+	if (ModuleName.IsSet())
+	{
+		return FText::FromName(ModuleName.GetValue());
+	}
 	
 	if (UControlRig* ControlRigPtr = ControlRig.Get())
 	{
@@ -75,7 +84,7 @@ FText FMultiRigData::GetName() const
 	return FText::FromName(NAME_None);
 }
 
-FText FMultiRigData::GetDisplayName() const
+FText FMultiRigData::GetDisplayName(const FRigTreeDisplaySettings& InSettings) const
 {
 	if (!CachedDisplayName.IsSet())
 	{
@@ -84,9 +93,8 @@ FText FMultiRigData::GetDisplayName() const
 			if (const URigHierarchy* Hierarchy = GetHierarchy())
 			{
 				const FRigElementKey& RigElementKey = Key.GetValue();
-
-				static constexpr bool bIncludeNameSpace = true;
-				const FText DisplayNameForUI = Hierarchy->GetDisplayNameForUI(RigElementKey, bIncludeNameSpace);
+				
+				const FText DisplayNameForUI = Hierarchy->GetDisplayNameForUI(RigElementKey, !InSettings.bUseShortName);
 				if (!DisplayNameForUI.IsEmpty())
 				{
 					CachedDisplayName = DisplayNameForUI;
@@ -101,6 +109,11 @@ FText FMultiRigData::GetDisplayName() const
 			}
 		}
 
+		if (ModuleName.IsSet())
+		{
+			CachedDisplayName = FText::FromName(ModuleName.GetValue());
+		}
+
 		if (!CachedDisplayName.IsSet())
 		{
 			CachedDisplayName = GetName();
@@ -122,7 +135,14 @@ bool FMultiRigData::operator== (const FMultiRigData & Other) const
 		{
 			if (Key.IsSet() == false && Other.Key.IsSet() == false)
 			{
-				return true;
+				if (ModuleName.IsSet() == false && Other.ModuleName.IsSet() == false)
+				{
+					return true;
+				}
+				else if (ModuleName.IsSet() == true && Other.ModuleName.IsSet() == true)
+				{
+					return ModuleName.GetValue() == Other.ModuleName.GetValue();
+				}
 			}
 			else if (Key.IsSet() == true && Other.Key.IsSet() == true)
 			{
@@ -135,7 +155,19 @@ bool FMultiRigData::operator== (const FMultiRigData & Other) const
 
 bool FMultiRigData::IsValid() const
 {
-	return (ControlRig.IsValid() && (Key.IsSet() == false || Key.GetValue().IsValid()));
+	if (ControlRig.IsValid())
+	{
+		if (Key.IsSet())
+		{
+			return Key.GetValue().IsValid();
+		}
+		if (ModuleName.IsSet())
+		{
+			return !ModuleName->IsNone();
+		}
+		return true;
+	}
+	return false;
 }
 
 URigHierarchy* FMultiRigData::GetHierarchy() const
@@ -238,21 +270,40 @@ void SMultiRigHierarchyItem::Construct(const FArguments& InArgs, const TSharedRe
 					SNew(SImage)
 					.Image_Lambda([this]() -> const FSlateBrush*
 					{
-						//if no key is set then it's the control rig so we get that based upon it's state
-						if (WeakRigTreeElement.Pin()->Data.Key.IsSet() == false)
+						if (TSharedPtr<FMultiRigTreeElement> SharedElement = WeakRigTreeElement.Pin())
 						{
-							if (UControlRig* ControlRig = WeakRigTreeElement.Pin()->Data.ControlRig.Get())
+							//if no key is set then it's the control rig or a module so we get that based upon it's state
+							if (!SharedElement->Data.Key.IsSet())
 							{
-								if (ControlRig->GetControlsVisible())
+								UControlRig* Rig = nullptr;
+								// If we have a module name, we are dealing with a module, get the rig from there
+								if (SharedElement->Data.ModuleName.IsSet())
 								{
-									return (FAppStyle::GetBrush("Level.VisibleIcon16x"));
+									if (UModularRig* ModularRig = Cast<UModularRig>(SharedElement->Data.ControlRig))
+									{
+										if (FRigModuleInstance* Module = ModularRig->FindModule(SharedElement->Data.ModuleName.GetValue()))
+										{
+											Rig = Module->GetRig();
+										}
+									}
 								}
 								else
 								{
-									return  (FAppStyle::GetBrush("Level.NotVisibleIcon16x"));
+									Rig = SharedElement->Data.ControlRig.Get();
 								}
-							}
 
+								if (Rig)
+								{
+									if (Rig->GetControlsVisible())
+									{
+										return (FAppStyle::GetBrush("Level.VisibleIcon16x"));
+									}
+									else
+									{
+										return  (FAppStyle::GetBrush("Level.NotVisibleIcon16x"));
+									}
+								}
+							}
 						}
 						return WeakRigTreeElement.Pin()->IconBrush;
 					})
@@ -285,15 +336,31 @@ void SMultiRigHierarchyItem::Construct(const FArguments& InArgs, const TSharedRe
 
 FReply SMultiRigHierarchyItem::OnGetSelectedClicked()
 {
-	if (FMultiRigTreeElement* Element = WeakRigTreeElement.Pin().Get())
+	if (TSharedPtr<FMultiRigTreeElement> Element = WeakRigTreeElement.Pin())
 	{
 		if (Element->Data.Key.IsSet() == false)
 		{
-			if (UControlRig* ControlRig = Element->Data.ControlRig.Get())
+			UControlRig* Rig = nullptr;
+			if (Element->Data.ModuleName.IsSet())
+			{
+				if (UModularRig* ModularRig = Cast<UModularRig>(Element->Data.ControlRig))
+				{
+					if (FRigModuleInstance* Module = ModularRig->FindModule(Element->Data.ModuleName.GetValue()))
+					{
+						Rig = Module->GetRig();
+					}
+				}
+			}
+			else
+			{
+				Rig  = Element->Data.ControlRig.Get();
+			}
+			
+			if (Rig)
 			{
 				FScopedTransaction ScopedTransaction(LOCTEXT("ToggleControlsVisibility", "Toggle Controls Visibility"), !GIsTransacting);
-				ControlRig->Modify();
-				ControlRig->ToggleControlsVisible();
+				Rig->Modify();
+				Rig->ToggleControlsVisible();
 				return FReply::Handled();
 			}
 		}
@@ -303,7 +370,8 @@ FReply SMultiRigHierarchyItem::OnGetSelectedClicked()
 
 FText SMultiRigHierarchyItem::GetDisplayName() const
 {
-	return (WeakRigTreeElement.Pin()->Data.GetDisplayName());
+	const FRigTreeDisplaySettings Settings = Delegates.GetDisplaySettings();
+	return (WeakRigTreeElement.Pin()->Data.GetDisplayName(Settings));
 }
 
 TPair<const FSlateBrush*, FSlateColor> SMultiRigHierarchyItem::GetBrushForElementType(const URigHierarchy* InHierarchy, const FMultiRigData& InData)
@@ -521,6 +589,10 @@ bool SMultiRigHierarchyTreeView::AddElement(UControlRig* InControlRig, const FRi
 				{
 					return false;
 				}
+				if (!Settings.bShowAnimationChannels && ControlElement->Settings.AnimationType == ERigControlAnimationType::AnimationChannel)
+				{
+					return false;
+				}
 				break;
 			}
 			case ERigElementType::Physics:
@@ -582,54 +654,96 @@ bool SMultiRigHierarchyTreeView::AddElement(UControlRig* InControlRig, const FRi
 	{
 		if (const URigHierarchy* Hierarchy = InControlRig->GetHierarchy())
 		{
-			FRigElementKey ParentKey = Hierarchy->GetFirstParent(InElement->GetKey());
-
-			TArray<FRigElementWeight> ParentWeights = Hierarchy->GetParentWeightArray(InElement->GetKey());
-			if (ParentWeights.Num() > 0)
+			if (InControlRig->IsModularRig() && Settings.bArrangeByModules)
 			{
-				TArray<FRigElementKey> ParentKeys = Hierarchy->GetParents(InElement->GetKey());
-				check(ParentKeys.Num() == ParentWeights.Num());
-				for (int32 ParentIndex = 0; ParentIndex < ParentKeys.Num(); ParentIndex++)
+				ParentData.ModuleName = *Hierarchy->GetModuleName(InElement->GetKey());
+
+				FMultiRigData NewData;
+				NewData.ControlRig = InControlRig;
+				NewData.Key = InElement->GetKey();
+
+				if (TSharedPtr<FMultiRigTreeElement>* ParentElementPtr = ElementMap.Find(ParentData))
 				{
-					if (ParentWeights[ParentIndex].IsAlmostZero())
+					if (Settings.FilterText.IsEmpty() || !Settings.bFlattenHierarchyOnFilter)
 					{
-						continue;
+						if (ReparentElement(Data, ParentData))
+						{
+							// Move any rig element to be inserted before any other module
+							{
+								TSharedPtr<FMultiRigTreeElement> Parent = *ParentElementPtr;
+								int32 InsertIndex = INDEX_NONE;
+								for (int32 i=Parent->Children.Num()-2; i >= 0; i--)
+								{
+									if (!Parent->Children[i]->Data.ModuleName.IsSet())
+									{
+										InsertIndex = i;
+										break;
+									}
+								}
+								if (InsertIndex+1 != Parent->Children.Num() - 1)
+								{
+									TSharedPtr<FMultiRigTreeElement> Element = Parent->Children.Last();
+									Parent->Children.RemoveAt(Parent->Children.Num() - 1);
+									Parent->Children.Insert(Element, InsertIndex+1);
+								}
+							}
+						}
 					}
-					ParentKey = ParentKeys[ParentIndex];
-					break;
 				}
+				
+				AddElement(NewData, ParentData);
 			}
-
-			if (ParentKey.IsValid())
+			else
 			{
-				if(FKControlRig && ParentKey != URigHierarchy::GetWorldSpaceReferenceKey())
+				FRigElementKey ParentKey = Hierarchy->GetFirstParent(InElement->GetKey());
+
+				TArray<FRigElementWeight> ParentWeights = Hierarchy->GetParentWeightArray(InElement->GetKey());
+				if (ParentWeights.Num() > 0)
 				{
-					if(const FRigControlElement* ControlElement = Cast<FRigControlElement>(InElement))
+					TArray<FRigElementKey> ParentKeys = Hierarchy->GetParents(InElement->GetKey());
+					check(ParentKeys.Num() == ParentWeights.Num());
+					for (int32 ParentIndex = 0; ParentIndex < ParentKeys.Num(); ParentIndex++)
 					{
-						if(ControlElement->Settings.AnimationType == ERigControlAnimationType::AnimationControl)
+						if (ParentWeights[ParentIndex].IsAlmostZero())
 						{
-							const FRigElementKey& ElementKey = InElement->GetKey();
-							const FName BoneName = FKControlRig->GetControlTargetName(ElementKey.Name, ParentKey.Type);
-							const FRigElementKey ParentBoneKey = Hierarchy->GetFirstParent(FRigElementKey(BoneName, ERigElementType::Bone));
-							if(ParentBoneKey.IsValid())
-							{
-								ParentKey = FRigElementKey(FKControlRig->GetControlName(ParentBoneKey.Name, ParentKey.Type), ElementKey.Type);
-							}
+							continue;
 						}
+						ParentKey = ParentKeys[ParentIndex];
+						break;
 					}
 				}
 
-				if (const FRigBaseElement* ParentElement = Hierarchy->Find(ParentKey))
+				if (ParentKey.IsValid())
 				{
-					if(ParentElement != nullptr)
+					if(FKControlRig && ParentKey != URigHierarchy::GetWorldSpaceReferenceKey())
 					{
-						AddElement(InControlRig,ParentElement);
-
-						FMultiRigData NewParentData(InControlRig, ParentKey);
+						if(const FRigControlElement* ControlElement = Cast<FRigControlElement>(InElement))
+						{
+							if(ControlElement->Settings.AnimationType == ERigControlAnimationType::AnimationControl)
+							{
+								const FRigElementKey& ElementKey = InElement->GetKey();
+								const FName BoneName = FKControlRig->GetControlTargetName(ElementKey.Name, ParentKey.Type);
+								const FRigElementKey ParentBoneKey = Hierarchy->GetFirstParent(FRigElementKey(BoneName, ERigElementType::Bone));
+								if(ParentBoneKey.IsValid())
+								{
+									ParentKey = FRigElementKey(FKControlRig->GetControlName(ParentBoneKey.Name, ParentKey.Type), ElementKey.Type);
+								}
+							}
+						}
+					}
 
-						if (ElementMap.Contains(NewParentData))
+					if (const FRigBaseElement* ParentElement = Hierarchy->Find(ParentKey))
+					{
+						if(ParentElement != nullptr)
 						{
-							ReparentElement(Data, NewParentData);
+							AddElement(InControlRig,ParentElement);
+
+							FMultiRigData NewParentData(InControlRig, ParentKey);
+
+							if (ElementMap.Contains(NewParentData))
+							{
+								ReparentElement(Data, NewParentData);
+							}
 						}
 					}
 				}
@@ -782,6 +896,33 @@ void SMultiRigHierarchyTreeView::RefreshTreeView()
 			// add root element
 			AddElement(CRData, EmptyParentData);
 
+			const FRigTreeDisplaySettings& Settings = Delegates.GetDisplaySettings();
+			if (ControlRig->IsModularRig() && Settings.bArrangeByModules)
+			{
+				if (UModularRig* ModularRig = Cast<UModularRig>(ControlRig))
+				{
+					ModularRig->ForEachModule([&](FRigModuleInstance* Module)
+					{
+						FMultiRigData ParentData;
+						ParentData.ControlRig = ControlRig;
+						if (Module->ParentModuleName.IsNone())
+						{
+							ParentData = CRData;
+						}
+						else
+						{
+							ParentData.ModuleName = Module->ParentModuleName;
+						}
+
+						FMultiRigData NewData;
+						NewData.ControlRig = ControlRig;
+						NewData.ModuleName = Module->Name;
+						AddElement(NewData, ParentData);
+						return true;
+					});
+				}
+			}
+
 			// add children
 			if (const URigHierarchy* Hierarchy = ControlRig->GetHierarchy())
 			{
@@ -926,6 +1067,15 @@ void SMultiRigHierarchyTreeView::HandleGetChildrenForTree(TSharedPtr<FMultiRigTr
 	OutChildren = InItem.Get()->Children;
 }
 
+TSharedPtr<FMultiRigTreeElement> SMultiRigHierarchyTreeView::FindElement(const FMultiRigData& InData)
+{
+	if (TSharedPtr<FMultiRigTreeElement>* Element = ElementMap.Find(InData))
+	{
+		return *Element;
+	}
+	return nullptr;
+}
+
 TArray<FMultiRigData> SMultiRigHierarchyTreeView::GetSelectedData() const
 {
 	TArray<FMultiRigData> Keys;
@@ -973,6 +1123,9 @@ void SSearchableMultiRigHierarchyTreeView::Construct(const FArguments& InArgs)
 	FMultiRigTreeDelegates TreeDelegates = InArgs._RigTreeDelegates;
 	SuperGetRigTreeDisplaySettings = TreeDelegates.OnGetDisplaySettings;
 
+	CommandList = MakeShared<FUICommandList>();
+	BindCommands();
+
 	TreeDelegates.OnGetDisplaySettings.BindSP(this, &SSearchableMultiRigHierarchyTreeView::GetDisplaySettings);
 
 	ChildSlot
@@ -984,9 +1137,32 @@ void SSearchableMultiRigHierarchyTreeView::Construct(const FArguments& InArgs)
 				.HAlign(HAlign_Fill)
 				.Padding(0.0f)
 				[
-					SNew(SSearchBox)
-					.InitialText(InArgs._InitialFilterText)
-				.OnTextChanged(this, &SSearchableMultiRigHierarchyTreeView::OnFilterTextChanged)
+					SNew(SHorizontalBox)
+			
+					+SHorizontalBox::Slot()
+					.Padding(0.0f, 0.0f)
+					.HAlign(HAlign_Left)
+					.AutoWidth()
+					[
+						SNew(SComboButton)
+					   .ComboButtonStyle(&FAppStyle::Get().GetWidgetStyle<FComboButtonStyle>("SimpleComboButtonWithIcon"))
+					   .ForegroundColor(FSlateColor::UseStyle())
+					   .ToolTipText(LOCTEXT("OptionsToolTip", "Open the Options Menu ."))
+					   .OnGetMenuContent(this, &SSearchableMultiRigHierarchyTreeView::OnGetOptionsMenu)
+					   .ContentPadding(FMargin(1, 0))
+					   .ButtonContent()
+					   [
+							SNew(SImage)
+							.Image(FAppStyle::Get().GetBrush("Icons.Filter"))
+							.ColorAndOpacity(FSlateColor::UseForeground())
+					   ]
+					]
+					+SHorizontalBox::Slot()
+					[
+						SNew(SSearchBox)
+						.InitialText(InArgs._InitialFilterText)
+						.OnTextChanged(this, &SSearchableMultiRigHierarchyTreeView::OnFilterTextChanged)
+					]
 				]
 
 			+ SVerticalBox::Slot()
@@ -1017,9 +1193,86 @@ const FRigTreeDisplaySettings& SSearchableMultiRigHierarchyTreeView::GetDisplayS
 		Settings = SuperGetRigTreeDisplaySettings.Execute();
 	}
 	Settings.FilterText = FilterText;
+	Settings.bShowAnimationChannels = bShowAnimationChannels;
+	Settings.bArrangeByModules = bArrangeByModules;
+	Settings.bUseShortName = bUseShortName;
 	return Settings;
 }
 
+TSharedRef<SWidget> SSearchableMultiRigHierarchyTreeView::OnGetOptionsMenu()
+{
+	
+	FMenuBuilder MenuBuilder(true, CommandList);
+
+	MenuBuilder.BeginSection("FilterOptions", LOCTEXT("FilterOptions", "Filter Options"));
+	{
+		const FControlRigHierarchyCommands& Actions = FControlRigHierarchyCommands::Get();
+		MenuBuilder.AddMenuEntry(Actions.ArrangeByModules);
+		MenuBuilder.AddMenuEntry(Actions.ShowAnimationChannels);
+		MenuBuilder.AddMenuEntry(Actions.ShowFullPaths);
+	}
+	MenuBuilder.EndSection();
+
+	return MenuBuilder.MakeWidget();
+}
+
+void SSearchableMultiRigHierarchyTreeView::BindCommands()
+{
+	const FControlRigHierarchyCommands& Commands = FControlRigHierarchyCommands::Get();
+
+	CommandList->MapAction(
+		Commands.ArrangeByModules,
+		FExecuteAction::CreateSP(this, &SSearchableMultiRigHierarchyTreeView::ToggleArrangeByModules),
+		FCanExecuteAction(),
+		FIsActionChecked::CreateSP(this, &SSearchableMultiRigHierarchyTreeView::IsArrangedByModules));
+	
+	CommandList->MapAction(
+		Commands.ShowAnimationChannels,
+		FExecuteAction::CreateSP(this, &SSearchableMultiRigHierarchyTreeView::ToggleShowAnimationChannels),
+		FCanExecuteAction(),
+		FIsActionChecked::CreateSP(this, &SSearchableMultiRigHierarchyTreeView::IsShowingAnimationChannels));
+
+	CommandList->MapAction(
+		Commands.ShowFullPaths,
+		FExecuteAction::CreateSP(this, &SSearchableMultiRigHierarchyTreeView::ToggleShortNames),
+		FCanExecuteAction(),
+		FIsActionChecked::CreateSP(this, &SSearchableMultiRigHierarchyTreeView::IsShowingFullPaths));
+
+}
+
+bool SSearchableMultiRigHierarchyTreeView::IsArrangedByModules() const
+{
+	return bArrangeByModules;
+}
+
+void SSearchableMultiRigHierarchyTreeView::ToggleArrangeByModules()
+{
+	bArrangeByModules = !bArrangeByModules;
+	TreeView->RequestTreeViewRefresh();
+}
+
+bool SSearchableMultiRigHierarchyTreeView::IsShowingAnimationChannels() const
+{
+	return bShowAnimationChannels;
+}
+
+void SSearchableMultiRigHierarchyTreeView::ToggleShowAnimationChannels()
+{
+	bShowAnimationChannels = !bShowAnimationChannels;
+	TreeView->RequestTreeViewRefresh();
+}
+
+bool SSearchableMultiRigHierarchyTreeView::IsShowingFullPaths() const
+{
+	return !bUseShortName;
+}
+
+void SSearchableMultiRigHierarchyTreeView::ToggleShortNames()
+{
+	bUseShortName = !bUseShortName;
+	TreeView->RequestTreeViewRefresh();
+}
+
 void SSearchableMultiRigHierarchyTreeView::OnFilterTextChanged(const FText& SearchText)
 {
 	FilterText = SearchText;
@@ -1043,6 +1296,9 @@ void SControlRigOutliner::Construct(const FArguments& InArgs, FControlRigEditMod
 	DisplaySettings.bHideParentsOnFilter = true;
 	DisplaySettings.bFlattenHierarchyOnFilter = true;
 	DisplaySettings.bShowConnectors = false;
+	DisplaySettings.bShowAnimationChannels = true;
+	DisplaySettings.bArrangeByModules = true;
+	DisplaySettings.bUseShortName = true;
 
 	FMultiRigTreeDelegates RigTreeDelegates;
 	RigTreeDelegates.OnGetDisplaySettings = FOnGetRigTreeDisplaySettings::CreateSP(this, &SControlRigOutliner::GetDisplaySettings);
@@ -1158,13 +1414,42 @@ void SControlRigOutliner::HandleSelectionChanged(TSharedPtr<FMultiRigTreeElement
 	TMap<UControlRig*, TArray<FRigElementKey>> SelectedRigAndKeys;
 	for (const FMultiRigData& Data : NewSelection)
 	{
-		if (Data.ControlRig.IsValid() && (Data.Key.IsSet() && Data.Key.GetValue() != FRigElementKey()))
+		if (Data.ControlRig.IsValid())
 		{
-			if (SelectedRigAndKeys.Find(Data.ControlRig.Get()) == nullptr)
+			if (Data.Key.IsSet() && Data.Key.GetValue() != FRigElementKey())
+			{
+				SelectedRigAndKeys.FindOrAdd(Data.ControlRig.Get()).Add(Data.Key.GetValue());
+			}
+
+			if (Data.ModuleName.IsSet())
 			{
-				SelectedRigAndKeys.Add(Data.ControlRig.Get());
+				TArray<FRigElementKey>& Keys = SelectedRigAndKeys.FindOrAdd(Data.ControlRig.Get());
+				const FRigTreeDisplaySettings& Settings = HierarchyTreeView->GetDisplaySettings();
+				if (Settings.FilterText.IsEmpty() || !Settings.bFlattenHierarchyOnFilter)
+				{
+					TSharedPtr<FMultiRigTreeElement> Element = HierarchyTreeView->GetTreeView()->FindElement(Data);
+					for (TSharedPtr<FMultiRigTreeElement> Child : Element->Children)
+					{
+						if (!Child->Data.Key.IsSet())
+						{
+							break;
+						}
+						Keys.Add(Child->Data.Key.GetValue());
+					}
+				}
+				else
+				{
+					// If we have flatten the hierarchy due to a search, we cannot rely on the children of the module
+					Data.GetHierarchy()->ForEach<FRigControlElement>([&](FRigControlElement* ControlElement)
+					{
+						if (Data.GetHierarchy()->GetModuleName(ControlElement->GetKey()) == Data.ModuleName.GetValue())
+						{
+							Keys.Add(ControlElement->GetKey());
+						}
+						return true;
+					});
+				}
 			}
-			SelectedRigAndKeys[Data.ControlRig.Get()].Add(Data.Key.GetValue());
 		}
 	}
 	TGuardValue<bool> GuardRigHierarchyChanges(bIsChangingRigHierarchy, true);
diff --git a/Engine/Plugins/Animation/ControlRig/Source/ControlRigEditor/Private/EditMode/SControlRigOutliner.h b/Engine/Plugins/Animation/ControlRig/Source/ControlRigEditor/Private/EditMode/SControlRigOutliner.h
--- a/Engine/Plugins/Animation/ControlRig/Source/ControlRigEditor/Private/EditMode/SControlRigOutliner.h
+++ b/Engine/Plugins/Animation/ControlRig/Source/ControlRigEditor/Private/EditMode/SControlRigOutliner.h
@@ -78,14 +78,16 @@ class FMultiRigData
 public:
 	FMultiRigData() {};
 	FMultiRigData(UControlRig* InControlRig, FRigElementKey InKey) : ControlRig(InControlRig), Key(InKey) {};
+	FMultiRigData(UControlRig* InControlRig, FName InModuleName) : ControlRig(InControlRig), ModuleName(InModuleName) {};
 	FText GetName() const;
-	FText GetDisplayName() const;
+	FText GetDisplayName(const FRigTreeDisplaySettings& DisplaySettings) const;
 	bool operator == (const FMultiRigData & Other) const;
 	bool IsValid() const;
 	URigHierarchy* GetHierarchy() const;
 
 	TWeakObjectPtr<UControlRig> ControlRig;
 	TOptional<FRigElementKey> Key;
+	TOptional<FName> ModuleName;
 
 private:
 	mutable TOptional<FText> CachedDisplayName;
@@ -185,6 +187,7 @@ public:
 	void SetExpansionRecursive(TSharedPtr<FMultiRigTreeElement> InElement, bool bTowardsParent, bool bShouldBeExpanded);
 	TSharedRef<ITableRow> MakeTableRowWidget(TSharedPtr<FMultiRigTreeElement> InItem, const TSharedRef<STableViewBase>& OwnerTable, bool bPinned);
 	void HandleGetChildrenForTree(TSharedPtr<FMultiRigTreeElement> InItem, TArray<TSharedPtr<FMultiRigTreeElement>>& OutChildren);
+	TSharedPtr<FMultiRigTreeElement> FindElement(const FMultiRigData& InData);
 
 	TArray<FMultiRigData> GetSelectedData() const;
 	const TArray<TSharedPtr<FMultiRigTreeElement>>& GetRootElements() const { return RootElements; }
@@ -237,13 +240,28 @@ public:
 	virtual ~SSearchableMultiRigHierarchyTreeView() {}
 	TSharedRef<SMultiRigHierarchyTreeView> GetTreeView() const { return TreeView.ToSharedRef(); }
 	const FRigTreeDisplaySettings& GetDisplaySettings();
+	TSharedRef<SWidget> OnGetOptionsMenu();
+	void BindCommands();
+
+	bool IsArrangedByModules() const;
+	void ToggleArrangeByModules();
+	bool IsShowingAnimationChannels() const;
+	void ToggleShowAnimationChannels();
+	bool IsShowingFullPaths() const;
+	void ToggleShortNames();
 
 private:
 
 	void OnFilterTextChanged(const FText& SearchText);
 
+	/** Command list we bind to */
+	TSharedPtr<FUICommandList> CommandList;
+
 	FOnGetRigTreeDisplaySettings SuperGetRigTreeDisplaySettings;
 	FText FilterText;
+	bool bArrangeByModules = true;
+	bool bShowAnimationChannels = true;
+	bool bUseShortName = false;
 	FRigTreeDisplaySettings Settings;
 	TSharedPtr<SMultiRigHierarchyTreeView> TreeView;
 };
diff --git a/Engine/Plugins/Animation/ControlRig/Source/ControlRigEditor/Private/Editor/SRigHierarchyTreeView.h b/Engine/Plugins/Animation/ControlRig/Source/ControlRigEditor/Private/Editor/SRigHierarchyTreeView.h
--- a/Engine/Plugins/Animation/ControlRig/Source/ControlRigEditor/Private/Editor/SRigHierarchyTreeView.h
+++ b/Engine/Plugins/Animation/ControlRig/Source/ControlRigEditor/Private/Editor/SRigHierarchyTreeView.h
@@ -30,6 +30,8 @@ struct CONTROLRIGEDITOR_API FRigTreeDisplaySettings
 		bShowConnectors = true;
 		bShowComponents = false;
 		bShowIconColors = true;
+		bShowAnimationChannels = true;
+		bArrangeByModules = false;
 	}
 	
 	FText FilterText;
@@ -69,6 +71,10 @@ struct CONTROLRIGEDITOR_API FRigTreeDisplaySettings
 
 	/** Whether to tint the icons with the element color */
 	bool bShowIconColors;
+
+	bool bShowAnimationChannels;
+
+	bool bArrangeByModules;
 };
 
 DECLARE_DELEGATE_RetVal(const URigHierarchy*, FOnGetRigTreeHierarchy);
