Control Rig Edit Mode: autokey on release only (follow up)

- this CL removes the previous AnimMode.AutoKeyOnRelease cvar and exposes the option in the editor settings via (Animation Authoring Settings > Auto Key On Release)

diff --git a/Engine/Plugins/Animation/ControlRig/Source/ControlRigEditor/ControlRigEditor.Build.cs b/Engine/Plugins/Animation/ControlRig/Source/ControlRigEditor/ControlRigEditor.Build.cs
--- a/Engine/Plugins/Animation/ControlRig/Source/ControlRigEditor/ControlRigEditor.Build.cs
+++ b/Engine/Plugins/Animation/ControlRig/Source/ControlRigEditor/ControlRigEditor.Build.cs
@@ -85,7 +85,8 @@ namespace UnrealBuildTool.Rules
                     "ActorPickerMode",
                     "Constraints",
                     "AnimationEditMode",
-					"SequencerWidgets"
+					"SequencerWidgets",
+					"DeveloperSettings"
 				}
             );
 
diff --git a/Engine/Plugins/Animation/ControlRig/Source/ControlRigEditor/Private/EditMode/ControlRigEditMode.cpp b/Engine/Plugins/Animation/ControlRig/Source/ControlRigEditor/Private/EditMode/ControlRigEditMode.cpp
--- a/Engine/Plugins/Animation/ControlRig/Source/ControlRigEditor/Private/EditMode/ControlRigEditMode.cpp
+++ b/Engine/Plugins/Animation/ControlRig/Source/ControlRigEditor/Private/EditMode/ControlRigEditMode.cpp
@@ -98,15 +98,6 @@ namespace UE::ControlRigEditMode::Private
 	{
 		return IConsoleManager::Get().FindConsoleVariable(FocusModeName);
 	}
-
-	const TCHAR* AutoKeyOnReleaseName = TEXT("AnimMode.AutoKeyOnRelease");
-
-	static bool bAutoKeyOnRelease = false;
-	static FAutoConsoleVariableRef CVarAutoKeyOnRelease(
-	   AutoKeyOnReleaseName,
-	   bAutoKeyOnRelease,
-	   TEXT("Only add keys on release (or when sequencer is playing)")
-	   );
 }
 
 using namespace UE::ControlRigEditMode;
@@ -517,6 +508,7 @@ void FControlRigEditMode::Enter()
 	if (IsInLevelEditor())
 	{
 		ConstraintsCache.RegisterNotifications(WorldPtr);
+		Keyframer.Initialize();
 	}
 }
 
@@ -1156,7 +1148,9 @@ bool FControlRigEditMode::HandleBeginTransform(const FEditorViewportClient* InVi
 	InteractionType = GetInteractionType(InViewportClient);
 	bIsTracking = true;
 	InteractionDependencies.Reset();
-	Keyframer.Init(DeferAutokeyOnMouseRelease());
+
+	const bool bDeferAutokeyOnMouseRelease = !bSequencerPlaying && IsInLevelEditor();
+	Keyframer.Enable(bDeferAutokeyOnMouseRelease);
 	
 	if (InteractionScopes.Num() == 0)
 	{
@@ -5204,7 +5198,7 @@ bool FControlRigEditMode::MoveControlShapeLocally( AControlRigShapeActor* ShapeA
 		ControlRig->InteractionType = InteractionType;
 		ControlRig->ElementsBeingInteracted.AddUnique(ShapeActor->GetElementKey());
 
-		const bool bDeferAutokey = DeferAutokeyOnMouseRelease();
+		const bool bDeferAutokey = Keyframer.IsEnabled();
 		const FRigControlModifiedContext Context(bDeferAutokey ? EControlRigSetKey::Never : EControlRigSetKey::DoNotCare);
 		ControlRig->SetControlLocalTransform(ShapeActor->ControlName, CurrentLocalTransform,true, Context, true, /*fix eulers*/ true);
 
@@ -5348,7 +5342,8 @@ void FControlRigEditMode::MoveControlShape( AControlRigShapeActor* ShapeActor,
 		Context.EventName = FRigUnit_BeginExecution::EventName;
 		Context.bConstraintUpdate = true;
 		
-		if (DeferAutokeyOnMouseRelease())
+		const bool bDeferAutokey = Keyframer.IsEnabled();
+		if (bDeferAutokey)
 		{
 			Context.SetKey = EControlRigSetKey::Never;
 		}
@@ -6463,11 +6458,6 @@ FInteractionDependencyCache& FControlRigEditMode::GetInteractionDependencies(UCo
 	return DummyDependencies;
 }
 
-bool FControlRigEditMode::DeferAutokeyOnMouseRelease() const
-{
-	return Private::bAutoKeyOnRelease && !bSequencerPlaying && IsInLevelEditor();
-}
-
 /**
 * FDetailKeyFrameCacheAndHandler
 */
diff --git a/Engine/Plugins/Animation/ControlRig/Source/ControlRigEditor/Private/Sequencer/AnimationAuthoringSettings.cpp b/Engine/Plugins/Animation/ControlRig/Source/ControlRigEditor/Private/Sequencer/AnimationAuthoringSettings.cpp
new file mode 100644
--- /dev/null
+++ b/Engine/Plugins/Animation/ControlRig/Source/ControlRigEditor/Private/Sequencer/AnimationAuthoringSettings.cpp
@@ -0,0 +1,16 @@
+// Copyright Epic Games, Inc. All Rights Reserved.
+
+#include "Sequencer/AnimationAuthoringSettings.h"
+
+UAnimationAuthoringSettings::UAnimationAuthoringSettings(const FObjectInitializer& ObjectInitializer)
+	: UDeveloperSettings(ObjectInitializer)
+{}
+
+UAnimationAuthoringSettings::FOnUpdateSettings UAnimationAuthoringSettings::OnSettingsChange;
+
+void UAnimationAuthoringSettings::PostEditChangeProperty(FPropertyChangedEvent& PropertyChangedEvent)
+{
+	Super::PostEditChangeProperty(PropertyChangedEvent);
+	
+	OnSettingsChange.Broadcast(this);
+}
diff --git a/Engine/Plugins/Animation/ControlRig/Source/ControlRigEditor/Private/Sequencer/EditModeAnimationUtil.cpp b/Engine/Plugins/Animation/ControlRig/Source/ControlRigEditor/Private/Sequencer/EditModeAnimationUtil.cpp
--- a/Engine/Plugins/Animation/ControlRig/Source/ControlRigEditor/Private/Sequencer/EditModeAnimationUtil.cpp
+++ b/Engine/Plugins/Animation/ControlRig/Source/ControlRigEditor/Private/Sequencer/EditModeAnimationUtil.cpp
@@ -8,6 +8,7 @@
 #include "Constraints/MovieSceneConstraintChannelHelper.h"
 #include "EditMode/ControlRigEditModeUtil.h"
 #include "Rigs/RigHierarchyDefines.h"
+#include "Sequencer/AnimationAuthoringSettings.h"
 
 namespace UE::AnimationEditMode
 {
@@ -29,10 +30,45 @@ void EvaluateRigIfAdditive(UControlRig* InControlRig)
 using namespace ControlRigEditMode;
 using namespace Private;
 
-void FControlRigKeyframer::Init(const bool InEnabled)
+FControlRigKeyframer::~FControlRigKeyframer()
+{
+	if (OnAnimSettingsChanged.IsValid())
+	{
+		if (const UAnimationAuthoringSettings* Settings = GetDefault<UAnimationAuthoringSettings>())
+		{
+			Settings->OnSettingsChange.Remove(OnAnimSettingsChanged);
+		}
+		OnAnimSettingsChanged.Reset();
+	}
+}
+	
+void FControlRigKeyframer::Initialize()
+{
+	EnableState = EEnableState::Disabled;
+	
+	if (const UAnimationAuthoringSettings* Settings = GetDefault<UAnimationAuthoringSettings>())
+	{
+		OnSettingsChanged(Settings);
+
+		if (!OnAnimSettingsChanged.IsValid())
+		{
+			OnAnimSettingsChanged = Settings->OnSettingsChange.AddRaw(this, &FControlRigKeyframer::OnSettingsChanged);
+		}
+	}
+}
+	
+void FControlRigKeyframer::Enable(const bool InEnabled)
 {
 	Reset();
-	bEnabled = InEnabled;
+
+	if (InEnabled)
+	{
+		EnumAddFlags(EnableState, EEnableState::EnabledDirectly);
+	}
+	else
+	{
+		EnumRemoveFlags(EnableState, EEnableState::EnabledDirectly);
+	}
 }
 	
 void FControlRigKeyframer::Reset()
@@ -42,7 +78,7 @@ void FControlRigKeyframer::Reset()
 
 void FControlRigKeyframer::Store(const uint32 InControlHash, FControlKeyframeData&& InData)
 {
-	if (bEnabled && InControlHash != 0)
+	if (IsEnabled() && InControlHash != 0)
 	{
 		FControlKeyframeData& ControlKeyframeData = KeyframeData.FindOrAdd(InControlHash);
 		ControlKeyframeData = MoveTemp(InData);
@@ -51,7 +87,7 @@ void FControlRigKeyframer::Store(const uint32 InControlHash, FControlKeyframeDat
 	
 void FControlRigKeyframer::Apply(const FControlRigInteractionScope& InInteractionScope, const FTransformContext& InTransformContext)
 {
-	if (!bEnabled)
+	if (!IsEnabled())
 	{
 		return;
 	}
@@ -93,7 +129,7 @@ void FControlRigKeyframer::Finalize(UWorld* InWorld)
 {
 	auto NeedsConstraintUpdate = [this]()
 	{
-		if (bEnabled)
+		if (IsEnabled())
 		{
 			for (const auto& [ControlHash, Data]: KeyframeData)
 			{
@@ -114,5 +150,26 @@ void FControlRigKeyframer::Finalize(UWorld* InWorld)
 		Controller.EvaluateAllConstraints();
 	}
 }
+
+bool FControlRigKeyframer::IsEnabled() const
+{
+	return OnAnimSettingsChanged.IsValid() ?
+		EnumHasAllFlags(EnableState, EEnableState::FullyEnabled) : EnumHasAllFlags(EnableState, EEnableState::EnabledDirectly);
+}
 	
-}
\ No newline at end of file
+void FControlRigKeyframer::OnSettingsChanged(const UAnimationAuthoringSettings* InSettings)
+{
+	if (InSettings)
+	{
+		if (InSettings->bAutoKeyOnRelease)
+		{
+			EnumAddFlags(EnableState, EEnableState::EnabledBySettings);
+		}
+		else
+		{
+			EnumRemoveFlags(EnableState, EEnableState::EnabledBySettings);
+		}
+	}
+}
+	
+}
diff --git a/Engine/Plugins/Animation/ControlRig/Source/ControlRigEditor/Public/EditMode/ControlRigEditMode.h b/Engine/Plugins/Animation/ControlRig/Source/ControlRigEditor/Public/EditMode/ControlRigEditMode.h
--- a/Engine/Plugins/Animation/ControlRig/Source/ControlRigEditor/Public/EditMode/ControlRigEditMode.h
+++ b/Engine/Plugins/Animation/ControlRig/Source/ControlRigEditor/Public/EditMode/ControlRigEditMode.h
@@ -686,7 +686,6 @@ private:
 
 	// Used to store and apply keyframes (if deferred)  
 	mutable UE::AnimationEditMode::FControlRigKeyframer Keyframer;
-	bool DeferAutokeyOnMouseRelease() const;
 
 	friend class FControlRigEditorModule;
 	friend class FControlRigEditor;
diff --git a/Engine/Plugins/Animation/ControlRig/Source/ControlRigEditor/Public/Sequencer/AnimationAuthoringSettings.h b/Engine/Plugins/Animation/ControlRig/Source/ControlRigEditor/Public/Sequencer/AnimationAuthoringSettings.h
new file mode 100644
--- /dev/null
+++ b/Engine/Plugins/Animation/ControlRig/Source/ControlRigEditor/Public/Sequencer/AnimationAuthoringSettings.h
@@ -0,0 +1,25 @@
+// Copyright Epic Games, Inc. All Rights Reserved.
+
+#pragma once
+
+#include "CoreMinimal.h"
+#include "Engine/DeveloperSettings.h"
+
+#include "AnimationAuthoringSettings.generated.h"
+
+UCLASS(config=EditorPerProjectUserSettings, meta = (DisplayName = "Animation Authoring Settings"))
+class CONTROLRIGEDITOR_API UAnimationAuthoringSettings : public UDeveloperSettings
+{
+	GENERATED_UCLASS_BODY()
+
+public:
+
+	// Whether controls' keyframes should be added on release only.
+	UPROPERTY(config, EditAnywhere, Category = Interaction)
+	bool bAutoKeyOnRelease = false;
+    
+	DECLARE_MULTICAST_DELEGATE_OneParam(FOnUpdateSettings,  const UAnimationAuthoringSettings*);
+	static FOnUpdateSettings OnSettingsChange;
+
+	virtual void PostEditChangeProperty(FPropertyChangedEvent& PropertyChangedEvent) override;
+};
\ No newline at end of file
diff --git a/Engine/Plugins/Animation/ControlRig/Source/ControlRigEditor/Public/Sequencer/EditModeAnimationUtil.h b/Engine/Plugins/Animation/ControlRig/Source/ControlRigEditor/Public/Sequencer/EditModeAnimationUtil.h
--- a/Engine/Plugins/Animation/ControlRig/Source/ControlRigEditor/Public/Sequencer/EditModeAnimationUtil.h
+++ b/Engine/Plugins/Animation/ControlRig/Source/ControlRigEditor/Public/Sequencer/EditModeAnimationUtil.h
@@ -4,9 +4,12 @@
 
 #include "Containers/Map.h"
 #include "Math/Transform.h"
+#include "Misc/EnumClassFlags.h"
 
 class FControlRigInteractionScope;
+class UAnimationAuthoringSettings;
 class UControlRig;
+
 struct FRigControlElement;
 
 namespace UE::ControlRigEditMode
@@ -40,8 +43,16 @@ struct FControlKeyframeData
 	
 struct FControlRigKeyframer
 {
+	~FControlRigKeyframer();
+	
+	// Initializes this keyframer and bounds it to the animation authoring settings.
+	void Initialize();
+
 	// Resets the data storage and enable/disable the keyframer.
-	void Init(const bool InEnabled);
+	void Enable(const bool InEnabled);
+
+	// Returns true if enabled. 
+	bool IsEnabled() const;
 	
 	// Stores the keyframe data for a specific control.
 	void Store(const uint32 InControlHash, FControlKeyframeData&& InData);
@@ -60,8 +71,25 @@ private:
 	// Storage representing keyframe data per control.
 	TMap<uint32, FControlKeyframeData> KeyframeData;
 
+	// Current state of the keyframer
+	enum class EEnableState : uint8
+	{
+		Disabled = 0x000,			// Keyframing disabled
+		EnabledDirectly = 0x001,	// Keyframing enabled via code
+		EnabledBySettings = 0x002,	// Keyframing enabled by settings
+		
+		FullyEnabled = EnabledDirectly | EnabledBySettings
+	};
+	FRIEND_ENUM_CLASS_FLAGS(EEnableState)
+	
 	// Whether the keyframer is enabled or not.
-	bool bEnabled = false;
+	EEnableState EnableState = EEnableState::Disabled;
+
+	// Handle to UAnimationAuthoringSettings::OnSettingsChange delegate.
+	FDelegateHandle OnAnimSettingsChanged;
+	
+	// Used to track changes to animation authoring settings.
+	void OnSettingsChanged(const UAnimationAuthoringSettings* InSettings);
 };
 	
 }
