fixed useless log

#rnx

diff --git a/Engine/Plugins/Animation/ControlRig/Source/ControlRigEditor/Private/Sequencer/ControlRigParameterTrackEditor.cpp b/Engine/Plugins/Animation/ControlRig/Source/ControlRigEditor/Private/Sequencer/ControlRigParameterTrackEditor.cpp
--- a/Engine/Plugins/Animation/ControlRig/Source/ControlRigEditor/Private/Sequencer/ControlRigParameterTrackEditor.cpp
+++ b/Engine/Plugins/Animation/ControlRig/Source/ControlRigEditor/Private/Sequencer/ControlRigParameterTrackEditor.cpp
@@ -4261,12 +4261,6 @@ void FControlRigParameterTrackEditor::GetControlRigKeys(
 			case ERigControlType::Bool:
 			{
 				bool Val = ControlValue.Get<bool>();
-
-				if (ControlElement->GetName() == "Arm R/IK Solve")
-				{
-					UE_LOG(LogControlRigEditor, Warning, TEXT("Arm R/IK Solve = %d"), Val);
-				}
-					
 				pChannelIndex->GeneratedKeyIndex = OutGeneratedKeys.Num();
 				OutGeneratedKeys.Add(FMovieSceneChannelValueSetter::Create<FMovieSceneBoolChannel>(ChannelIndex, Val, bSetKey));
 				break;
