Anim In Engine : Fix IK/FK not working with certain rigs.

diff --git a/Engine/Plugins/Animation/ControlRig/Source/ControlRigEditor/Private/Sequencer/ControlRigParameterTrackEditor.cpp b/Engine/Plugins/Animation/ControlRig/Source/ControlRigEditor/Private/Sequencer/ControlRigParameterTrackEditor.cpp
--- a/Engine/Plugins/Animation/ControlRig/Source/ControlRigEditor/Private/Sequencer/ControlRigParameterTrackEditor.cpp
+++ b/Engine/Plugins/Animation/ControlRig/Source/ControlRigEditor/Private/Sequencer/ControlRigParameterTrackEditor.cpp
@@ -2439,7 +2439,7 @@ static void EvaluateThisControl(UMovieSceneControlRigParameterSection* Section,
 					TOptional <bool> Value = Section->EvaluateBoolParameter(FrameTime, ControlName);
 					if (Value.IsSet())
 					{
-						ControlRig->SetControlValue<bool>(ControlName, Value.GetValue(), true, EControlRigSetKey::Never, bSetupUndo);
+						ControlRig->SetControlValue<bool>(ControlName, Value.GetValue(), true, EControlRigSetKey::DoNotCare, bSetupUndo);
 					}
 				}
 				break;
@@ -2454,7 +2454,7 @@ static void EvaluateThisControl(UMovieSceneControlRigParameterSection* Section,
 						if (Value.IsSet())
 						{
 							int32 IVal = (int32)Value.GetValue();
-							ControlRig->SetControlValue<int32>(ControlName, IVal, true, EControlRigSetKey::Never, bSetupUndo);
+							ControlRig->SetControlValue<int32>(ControlName, IVal, true, EControlRigSetKey::DoNotCare, bSetupUndo);
 						}
 					}
 					else
@@ -2462,7 +2462,7 @@ static void EvaluateThisControl(UMovieSceneControlRigParameterSection* Section,
 						TOptional <int32> Value = Section->EvaluateIntegerParameter(FrameTime, ControlName);
 						if (Value.IsSet())
 						{
-							ControlRig->SetControlValue<int32>(ControlName, Value.GetValue(), true, EControlRigSetKey::Never, bSetupUndo);
+							ControlRig->SetControlValue<int32>(ControlName, Value.GetValue(), true, EControlRigSetKey::DoNotCare, bSetupUndo);
 						}
 					}
 				}
@@ -2475,7 +2475,7 @@ static void EvaluateThisControl(UMovieSceneControlRigParameterSection* Section,
 				TOptional <float> Value = Section->EvaluateScalarParameter(FrameTime, ControlName);
 				if (Value.IsSet())
 				{
-					ControlRig->SetControlValue<float>(ControlName, Value.GetValue(), true, EControlRigSetKey::Never, bSetupUndo);
+					ControlRig->SetControlValue<float>(ControlName, Value.GetValue(), true, EControlRigSetKey::DoNotCare, bSetupUndo);
 				}
 				break;
 			}
@@ -2484,7 +2484,7 @@ static void EvaluateThisControl(UMovieSceneControlRigParameterSection* Section,
 				TOptional <FVector2D> Value = Section->EvaluateVector2DParameter(FrameTime, ControlName);
 				if (Value.IsSet())
 				{
-					ControlRig->SetControlValue<FVector2D>(ControlName, Value.GetValue(), true, EControlRigSetKey::Never, bSetupUndo);
+					ControlRig->SetControlValue<FVector2D>(ControlName, Value.GetValue(), true, EControlRigSetKey::DoNotCare, bSetupUndo);
 				}
 				break;
 			}
@@ -2496,7 +2496,7 @@ static void EvaluateThisControl(UMovieSceneControlRigParameterSection* Section,
 				if (Value.IsSet())
 				{
 					FVector3f FloatVal = (FVector3f)Value.GetValue();
-					ControlRig->SetControlValue<FVector3f>(ControlName, FloatVal, true, EControlRigSetKey::Never, bSetupUndo);
+					ControlRig->SetControlValue<FVector3f>(ControlName, FloatVal, true, EControlRigSetKey::DoNotCare, bSetupUndo);
 				}
 				break;
 			}
@@ -2516,14 +2516,14 @@ static void EvaluateThisControl(UMovieSceneControlRigParameterSection* Section,
 					{
 						FVector EulerAngle(Value.GetValue().Rotation.Roll, Value.GetValue().Rotation.Pitch, Value.GetValue().Rotation.Yaw);
 						RigHierarchy->SetControlSpecifiedEulerAngle(ControlElement, EulerAngle);
-						ControlRig->SetControlValue<FRigControlValue::FTransform_Float>(ControlName, Value.GetValue().ToFTransform(), true, EControlRigSetKey::Never, bSetupUndo);
+						ControlRig->SetControlValue<FRigControlValue::FTransform_Float>(ControlName, Value.GetValue().ToFTransform(), true, EControlRigSetKey::DoNotCare, bSetupUndo);
 					}
 					else if (ControlElement->Settings.ControlType == ERigControlType::TransformNoScale)
 					{
 						FTransformNoScale NoScale = Value.GetValue().ToFTransform();
 						FVector EulerAngle(Value.GetValue().Rotation.Roll, Value.GetValue().Rotation.Pitch, Value.GetValue().Rotation.Yaw);
 						RigHierarchy->SetControlSpecifiedEulerAngle(ControlElement, EulerAngle);
-						ControlRig->SetControlValue<FRigControlValue::FTransformNoScale_Float>(ControlName, NoScale, true, EControlRigSetKey::Never, bSetupUndo);
+						ControlRig->SetControlValue<FRigControlValue::FTransformNoScale_Float>(ControlName, NoScale, true, EControlRigSetKey::DoNotCare, bSetupUndo);
 					}
 					else if (ControlElement->Settings.ControlType == ERigControlType::EulerTransform)
 					{
@@ -2534,7 +2534,7 @@ static void EvaluateThisControl(UMovieSceneControlRigParameterSection* Section,
 						FRotator UERotator(Quat);
 						FEulerTransform Transform = Euler;
 						Transform.Rotation = UERotator;
-						ControlRig->SetControlValue<FRigControlValue::FEulerTransform_Float>(ControlName, Transform, true, EControlRigSetKey::Never, bSetupUndo);
+						ControlRig->SetControlValue<FRigControlValue::FEulerTransform_Float>(ControlName, Transform, true, EControlRigSetKey::DoNotCare, bSetupUndo);
 
 					}
 				}
@@ -2542,7 +2542,8 @@ static void EvaluateThisControl(UMovieSceneControlRigParameterSection* Section,
 
 			}
 		}
-		//note we don't need to evaluate the control rig, setting the value is enough
+		//note we do need to evaluate the rig, within the interaction scope now
+		ControlRig->Evaluate_AnyThread();
 	}
 }
 
@@ -2557,10 +2558,13 @@ struct FChannelChangedStruct
 	TMap<UMovieSceneControlRigParameterSection*, TSet<FName>> SectionControlNames;
 };
 
+
+//this function used to set up a set of control names per control and evaluate them on next tick, but this doesn't work with 
+//certain rigs that send events immediately on an evaluation which would happen before the tick. In that case we would get the set key events
+//for the fk/ik but on a non-game thread so we couldn't set any keys.
+//the fix is to just immediately set the control value on the channel(bool) and then that will send the key events on this game thread.
 void FControlRigParameterTrackEditor::OnChannelChanged(const FMovieSceneChannelMetaData* MetaData, UMovieSceneSection* InSection)
 {
-	static FChannelChangedStruct ChannelChanged;
-
 	const TSharedPtr<ISequencer> Sequencer = GetSequencer();
 	if (!Sequencer)
 	{
@@ -2576,62 +2580,13 @@ void FControlRigParameterTrackEditor::OnChannelChanged(const FMovieSceneChannelM
 		if (StringArray.Num() > 0)
 		{
 			FName ControlName(*StringArray[0]);
-			if (ChannelChanged.bWasSetAlready == false)
+			if (FRigControlElement* ControlElement = Section->GetControlRig()->FindControl(ControlName))
 			{
-				ChannelChanged.bWasSetAlready = true;
-				ChannelChanged.SectionControlNames.Reset();
-				auto ChannelChangedDelegate = [this]()
+				if (ControlElement->Settings.ControlType == ERigControlType::Bool)
 				{
-					const TSharedPtr<ISequencer> Sequencer = GetSequencer();
-					if (!Sequencer)
-					{
-						GEditor->GetTimerManager()->ClearTimer(ChannelChanged.TimerHandle);
-						return;
-					}
-
-					if (!(FSlateApplication::Get().HasAnyMouseCaptor() || GUnrealEd->IsUserInteracting()))
-					{
-						if (ChannelChanged.TimerHandle.IsValid())
-						{
-							FFrameTime Time = Sequencer->GetLocalTime().Time;
-							GEditor->GetTimerManager()->ClearTimer(ChannelChanged.TimerHandle);
-							ChannelChanged.TimerHandle.Invalidate();
-							ChannelChanged.bWasSetAlready = false;
-							TOptional<FFrameNumber> Optional;	
-
-							for (TPair <UMovieSceneControlRigParameterSection*, TSet<FName>>& Pair : ChannelChanged.SectionControlNames)
-							{
-								for (FName& ControlName : Pair.Value)
-								{
-									Pair.Key->ControlsToSet.Empty();
-									Pair.Key->ControlsToSet.Add(ControlName);
-									//only do the fk/ik hack for bool's since that's only where it's needed
-									//otherwise we would incorrectly set values on an addtive section which has scale values of 0.0.
-									if (FRigControlElement* ControlElement = Pair.Key->GetControlRig()->FindControl(ControlName))
-									{
-										if (ControlElement->Settings.ControlType == ERigControlType::Bool)
-										{
-											EvaluateThisControl(Pair.Key, ControlName, Time);
-										}
-									}
-									Pair.Key->ControlsToSet.Empty();
-								}
-							}
-							ChannelChanged.SectionControlNames.Reset();
-						}
-					}
-				};
-				GEditor->GetTimerManager()->SetTimer(ChannelChanged.TimerHandle, ChannelChangedDelegate, 0.01f, true);			
-			}
-			if (TSet<FName>* SetOfNames = ChannelChanged.SectionControlNames.Find(Section))
-			{
-				SetOfNames->Add(ControlName);
-			}
-			else
-			{
-				TSet<FName> Names;
-				Names.Add(ControlName);
-				ChannelChanged.SectionControlNames.Add(Section, Names);
+					FFrameTime Time = Sequencer->GetLocalTime().Time;
+					EvaluateThisControl(Section, ControlName, Time);
+				}
 			}
 		}
 	}
@@ -6620,5 +6575,4 @@ FControlRigEditMode* FControlRigParameterTrackEditor::GetEditMode(bool bForceAct
 	return nullptr;
 }
 
-
 #undef LOCTEXT_NAMESPACE
