Control Rig: Create settings to handle ouliner options in order to make them persistent.
#rnx

diff --git a/Engine/Plugins/Animation/ControlRig/Source/ControlRig/Private/Settings/ControlRigSettings.cpp b/Engine/Plugins/Animation/ControlRig/Source/ControlRig/Private/Settings/ControlRigSettings.cpp
--- a/Engine/Plugins/Animation/ControlRig/Source/ControlRig/Private/Settings/ControlRigSettings.cpp
+++ b/Engine/Plugins/Animation/ControlRig/Source/ControlRig/Private/Settings/ControlRigSettings.cpp
@@ -32,6 +32,9 @@ UControlRigEditorSettings::UControlRigEditorSettings(const FObjectInitializer& O
 	bShowStackedHierarchy = false;
 	MaxStackSize = 16;
 	bLeftMouseDragDoesMarquee = true;
+	bUseShortName = false;
+	bArrangeByModules = true;
+	bShowAnimationChannels = true;
 #endif
 }
 
diff --git a/Engine/Plugins/Animation/ControlRig/Source/ControlRig/Public/Settings/ControlRigSettings.h b/Engine/Plugins/Animation/ControlRig/Source/ControlRig/Public/Settings/ControlRigSettings.h
--- a/Engine/Plugins/Animation/ControlRig/Source/ControlRig/Public/Settings/ControlRigSettings.h
+++ b/Engine/Plugins/Animation/ControlRig/Source/ControlRig/Public/Settings/ControlRigSettings.h
@@ -132,8 +132,39 @@ public:
 	UPROPERTY(EditAnywhere, config, Category = Hierarchy)
 	bool bLeftMouseDragDoesMarquee;
 
+	/**
+	 * If turned on the controls in the Anim Outliner will be arranged by modules in Modular Rigs.
+	 */
+	UPROPERTY(EditAnywhere, config, Category = Outliner)
+	bool bArrangeByModules;
+
+	/**
+	 * If turned on the animation channel controls in the Anim Outliner will be displayed.
+	 */
+	UPROPERTY(EditAnywhere, config, Category = Outliner)
+	bool bShowAnimationChannels;
+
+	/**
+	 * If turned on the controls in the Anim Outliner will be displayed using their short name, without their NameSpace.
+	 */
+	UPROPERTY(EditAnywhere, config, Category = Outliner)
+	bool bUseShortName;
+
 #endif
 
+	void PostEditChangeProperty(struct FPropertyChangedEvent& PropertyChangedEvent)
+	{
+		Super::PostEditChangeProperty(PropertyChangedEvent);
+
+		if (PropertyChangedEvent.ChangeType != EPropertyChangeType::Interactive)
+		{
+			// Dragging spinboxes causes this to be called every frame so we wait until they've finished dragging before saving.
+			SaveConfig();
+		}
+		
+		OnSettingChanged().Broadcast(this, PropertyChangedEvent);
+	}
+
 	static UControlRigEditorSettings * Get() { return GetMutableDefault<UControlRigEditorSettings>(); }
 };
 
diff --git a/Engine/Plugins/Animation/ControlRig/Source/ControlRigEditor/Private/EditMode/SControlRigOutliner.cpp b/Engine/Plugins/Animation/ControlRig/Source/ControlRigEditor/Private/EditMode/SControlRigOutliner.cpp
--- a/Engine/Plugins/Animation/ControlRig/Source/ControlRigEditor/Private/EditMode/SControlRigOutliner.cpp
+++ b/Engine/Plugins/Animation/ControlRig/Source/ControlRigEditor/Private/EditMode/SControlRigOutliner.cpp
@@ -1127,6 +1127,7 @@ void SSearchableMultiRigHierarchyTreeView::Construct(const FArguments& InArgs)
 	BindCommands();
 
 	TreeDelegates.OnGetDisplaySettings.BindSP(this, &SSearchableMultiRigHierarchyTreeView::GetDisplaySettings);
+	GetMutableDefault<UControlRigEditorSettings>()->OnSettingChanged().AddSP(this, &SSearchableMultiRigHierarchyTreeView::OnSettingChanged);
 
 	ChildSlot
 		[
@@ -1193,9 +1194,9 @@ const FRigTreeDisplaySettings& SSearchableMultiRigHierarchyTreeView::GetDisplayS
 		Settings = SuperGetRigTreeDisplaySettings.Execute();
 	}
 	Settings.FilterText = FilterText;
-	Settings.bShowAnimationChannels = bShowAnimationChannels;
-	Settings.bArrangeByModules = bArrangeByModules;
-	Settings.bUseShortName = bUseShortName;
+	Settings.bShowAnimationChannels = UControlRigEditorSettings::Get()->bShowAnimationChannels;
+	Settings.bArrangeByModules = UControlRigEditorSettings::Get()->bArrangeByModules;
+	Settings.bUseShortName = UControlRigEditorSettings::Get()->bUseShortName;
 	return Settings;
 }
 
@@ -1240,36 +1241,63 @@ void SSearchableMultiRigHierarchyTreeView::BindCommands()
 
 }
 
+void SSearchableMultiRigHierarchyTreeView::OnSettingChanged(UObject* Object, FPropertyChangedEvent& PropertyChangedEvent)
+{
+	TreeView->RequestTreeViewRefresh();
+}
+
 bool SSearchableMultiRigHierarchyTreeView::IsArrangedByModules() const
 {
-	return bArrangeByModules;
+	return UControlRigEditorSettings::Get()->bArrangeByModules;
 }
 
 void SSearchableMultiRigHierarchyTreeView::ToggleArrangeByModules()
 {
-	bArrangeByModules = !bArrangeByModules;
+	UControlRigEditorSettings* EditorSettings = UControlRigEditorSettings::Get();
+	EditorSettings->bArrangeByModules = !EditorSettings->bArrangeByModules;
+
+	if (FProperty* Property = UControlRigEditorSettings::StaticClass()->FindPropertyByName(GET_MEMBER_NAME_CHECKED(UControlRigEditorSettings, bArrangeByModules)))
+	{
+		FPropertyChangedEvent PropertyChangedEvent(Property, EPropertyChangeType::ValueSet);
+		EditorSettings->PostEditChangeProperty(PropertyChangedEvent);
+	}
 	TreeView->RequestTreeViewRefresh();
 }
 
 bool SSearchableMultiRigHierarchyTreeView::IsShowingAnimationChannels() const
 {
-	return bShowAnimationChannels;
+	return UControlRigEditorSettings::Get()->bShowAnimationChannels;
 }
 
 void SSearchableMultiRigHierarchyTreeView::ToggleShowAnimationChannels()
 {
-	bShowAnimationChannels = !bShowAnimationChannels;
+	UControlRigEditorSettings* EditorSettings = UControlRigEditorSettings::Get();
+	EditorSettings->bShowAnimationChannels = !EditorSettings->bShowAnimationChannels;
+
+	if (FProperty* Property = UControlRigEditorSettings::StaticClass()->FindPropertyByName(GET_MEMBER_NAME_CHECKED(UControlRigEditorSettings, bShowAnimationChannels)))
+	{
+		FPropertyChangedEvent PropertyChangedEvent(Property, EPropertyChangeType::ValueSet);
+		EditorSettings->PostEditChangeProperty(PropertyChangedEvent);
+	}
 	TreeView->RequestTreeViewRefresh();
 }
 
 bool SSearchableMultiRigHierarchyTreeView::IsShowingFullPaths() const
 {
-	return !bUseShortName;
+	return !UControlRigEditorSettings::Get()->bUseShortName;
 }
 
 void SSearchableMultiRigHierarchyTreeView::ToggleShortNames()
 {
-	bUseShortName = !bUseShortName;
+	UControlRigEditorSettings* EditorSettings = UControlRigEditorSettings::Get();
+	EditorSettings->bUseShortName = !EditorSettings->bUseShortName;
+
+	if (FProperty* Property = UControlRigEditorSettings::StaticClass()->FindPropertyByName(GET_MEMBER_NAME_CHECKED(UControlRigEditorSettings, bUseShortName)))
+	{
+		FPropertyChangedEvent PropertyChangedEvent(Property, EPropertyChangeType::ValueSet);
+		EditorSettings->PostEditChangeProperty(PropertyChangedEvent);
+	}
+	
 	TreeView->RequestTreeViewRefresh();
 }
 
diff --git a/Engine/Plugins/Animation/ControlRig/Source/ControlRigEditor/Private/EditMode/SControlRigOutliner.h b/Engine/Plugins/Animation/ControlRig/Source/ControlRigEditor/Private/EditMode/SControlRigOutliner.h
--- a/Engine/Plugins/Animation/ControlRig/Source/ControlRigEditor/Private/EditMode/SControlRigOutliner.h
+++ b/Engine/Plugins/Animation/ControlRig/Source/ControlRigEditor/Private/EditMode/SControlRigOutliner.h
@@ -242,6 +242,7 @@ public:
 	const FRigTreeDisplaySettings& GetDisplaySettings();
 	TSharedRef<SWidget> OnGetOptionsMenu();
 	void BindCommands();
+	void OnSettingChanged(UObject* Object, FPropertyChangedEvent& PropertyChangedEvent);
 
 	bool IsArrangedByModules() const;
 	void ToggleArrangeByModules();
@@ -259,9 +260,6 @@ private:
 
 	FOnGetRigTreeDisplaySettings SuperGetRigTreeDisplaySettings;
 	FText FilterText;
-	bool bArrangeByModules = true;
-	bool bShowAnimationChannels = true;
-	bool bUseShortName = false;
 	FRigTreeDisplaySettings Settings;
 	TSharedPtr<SMultiRigHierarchyTreeView> TreeView;
 };
