Control Rig: Fix compile error
#rnx

diff --git a/Engine/Plugins/Animation/ControlRig/Source/ControlRig/Public/Settings/ControlRigSettings.h b/Engine/Plugins/Animation/ControlRig/Source/ControlRig/Public/Settings/ControlRigSettings.h
--- a/Engine/Plugins/Animation/ControlRig/Source/ControlRig/Public/Settings/ControlRigSettings.h
+++ b/Engine/Plugins/Animation/ControlRig/Source/ControlRig/Public/Settings/ControlRigSettings.h
@@ -151,7 +151,7 @@ public:
 	bool bUseShortName;
 
 #endif
-
+#if WITH_EDITOR
 	void PostEditChangeProperty(struct FPropertyChangedEvent& PropertyChangedEvent)
 	{
 		Super::PostEditChangeProperty(PropertyChangedEvent);
@@ -164,6 +164,7 @@ public:
 		
 		OnSettingChanged().Broadcast(this, PropertyChangedEvent);
 	}
+#endif
 
 	static UControlRigEditorSettings * Get() { return GetMutableDefault<UControlRigEditorSettings>(); }
 };
