Control Rig: Enable shift+click expansion/collapse in outliner tree
#rnx

diff --git a/Engine/Plugins/Animation/ControlRig/Source/ControlRigEditor/Private/EditMode/SControlRigOutliner.cpp b/Engine/Plugins/Animation/ControlRig/Source/ControlRigEditor/Private/EditMode/SControlRigOutliner.cpp
--- a/Engine/Plugins/Animation/ControlRig/Source/ControlRigEditor/Private/EditMode/SControlRigOutliner.cpp
+++ b/Engine/Plugins/Animation/ControlRig/Source/ControlRigEditor/Private/EditMode/SControlRigOutliner.cpp
@@ -427,7 +427,7 @@ void SMultiRigHierarchyTreeView::Construct(const FArguments& InArgs)
 	SuperArgs.OnContextMenuOpening(Delegates.OnContextMenuOpening);
 	SuperArgs.OnMouseButtonClick(Delegates.OnMouseButtonClick);
 	SuperArgs.OnMouseButtonDoubleClick(Delegates.OnMouseButtonDoubleClick);
-	SuperArgs.OnSetExpansionRecursive(Delegates.OnSetExpansionRecursive);
+	SuperArgs.OnSetExpansionRecursive(FOnMultiRigTreeSetExpansionRecursive::CreateSP(this, &SMultiRigHierarchyTreeView::SetExpansionRecursive));
 	SuperArgs.HighlightParentNodesForSelection(true);
 	SuperArgs.AllowInvisibleItemSelection(true);  //without this we deselect everything when we filter or we collapse
 
@@ -1030,6 +1030,11 @@ void SMultiRigHierarchyTreeView::RefreshTreeView()
 	}
 }
 
+void SMultiRigHierarchyTreeView::SetExpansionRecursive(TSharedPtr<FMultiRigTreeElement> InElement, bool bShouldBeExpanded)
+{
+	SetExpansionRecursive(InElement, false, bShouldBeExpanded);
+}
+
 void SMultiRigHierarchyTreeView::SetExpansionRecursive(TSharedPtr<FMultiRigTreeElement> InElement, bool bTowardsParent,
 	bool bShouldBeExpanded)
 {
diff --git a/Engine/Plugins/Animation/ControlRig/Source/ControlRigEditor/Private/EditMode/SControlRigOutliner.h b/Engine/Plugins/Animation/ControlRig/Source/ControlRigEditor/Private/EditMode/SControlRigOutliner.h
--- a/Engine/Plugins/Animation/ControlRig/Source/ControlRigEditor/Private/EditMode/SControlRigOutliner.h
+++ b/Engine/Plugins/Animation/ControlRig/Source/ControlRigEditor/Private/EditMode/SControlRigOutliner.h
@@ -39,7 +39,6 @@ struct CONTROLRIGEDITOR_API FMultiRigTreeDelegates
 	FOnContextMenuOpening OnContextMenuOpening;
 	FOnMultiRigTreeMouseButtonClick OnMouseButtonClick;
 	FOnMultiRigTreeMouseButtonDoubleClick OnMouseButtonDoubleClick;
-	FOnMultiRigTreeSetExpansionRecursive OnSetExpansionRecursive;
 	FOnRigTreeCompareKeys OnCompareKeys;
 
 	FMultiRigTreeDelegates()
@@ -184,6 +183,7 @@ public:
 	bool AddElement(UControlRig* InControlRig, const FRigBaseElement* InElement);
 	bool ReparentElement(const FMultiRigData& InData, const FMultiRigData& InParentData);
 	bool RemoveElement(const FMultiRigData& InData);
+	void SetExpansionRecursive(TSharedPtr<FMultiRigTreeElement> InElement, bool bShouldBeExpanded);
 	void SetExpansionRecursive(TSharedPtr<FMultiRigTreeElement> InElement, bool bTowardsParent, bool bShouldBeExpanded);
 	TSharedRef<ITableRow> MakeTableRowWidget(TSharedPtr<FMultiRigTreeElement> InItem, const TSharedRef<STableViewBase>& OwnerTable, bool bPinned);
 	void HandleGetChildrenForTree(TSharedPtr<FMultiRigTreeElement> InItem, TArray<TSharedPtr<FMultiRigTreeElement>>& OutChildren);
