Control Rig: Allow show/hide of module controls in anim outliner
#rnx

diff --git a/Engine/Plugins/Animation/ControlRig/Source/ControlRigEditor/Private/EditMode/ControlRigEditMode.cpp b/Engine/Plugins/Animation/ControlRig/Source/ControlRigEditor/Private/EditMode/ControlRigEditMode.cpp
--- a/Engine/Plugins/Animation/ControlRig/Source/ControlRigEditor/Private/EditMode/ControlRigEditMode.cpp
+++ b/Engine/Plugins/Animation/ControlRig/Source/ControlRigEditor/Private/EditMode/ControlRigEditMode.cpp
@@ -2905,6 +2905,9 @@ void FControlRigEditMode::BindCommands()
 	CommandBindings->MapAction(
 		Commands.ToggleManipulators,
 		FExecuteAction::CreateRaw(this, &FControlRigEditMode::ToggleManipulators));
+	CommandBindings->MapAction(
+		Commands.ToggleModuleManipulators,
+		FExecuteAction::CreateRaw(this, &FControlRigEditMode::ToggleModuleManipulators));
 	CommandBindings->MapAction(
 		Commands.ToggleAllManipulators,
 		FExecuteAction::CreateRaw(this, &FControlRigEditMode::ToggleAllManipulators));
@@ -3418,6 +3421,44 @@ void FControlRigEditMode::ToggleManipulators()
 	}
 }
 
+void FControlRigEditMode::ToggleModuleManipulators()
+{
+	UControlRigEditModeSettings* Settings = GetMutableSettings();
+
+	if (!AreEditingControlRigDirectly() && !Settings->bHideControlShapes)
+	{
+		for (TWeakObjectPtr<UControlRig>& RuntimeRigPtr : RuntimeControlRigs)
+		{
+			if (UControlRig* ControlRig = RuntimeRigPtr.Get())
+			{
+				if (ControlRig->GetControlsVisible())
+				{
+					if (UModularRig* ModularRig = Cast<UModularRig>(ControlRig))
+					{
+						TArray<FString> ModuleNames;
+						TArray<FRigElementKey> Selected = GetSelectedRigElements(ControlRig);
+						for (const FRigElementKey& Key : Selected)
+						{
+							const FString ModuleName = ControlRig->GetHierarchy()->GetModuleName(Key);
+							ModuleNames.AddUnique(ModuleName);
+						}
+						for (const FString& ModuleName : ModuleNames)
+						{
+							if (FRigModuleInstance* Module = ModularRig->FindModule(*ModuleName))
+							{
+								if (UControlRig* Rig = Module->GetRig())
+								{
+									Rig->ToggleControlsVisible();
+								}
+							}
+						}
+					}
+				}
+			}
+		}
+	}
+}
+
 void FControlRigEditMode::ToggleAllManipulators()
 {	
 	UControlRigEditModeSettings* Settings = GetMutableSettings();
diff --git a/Engine/Plugins/Animation/ControlRig/Source/ControlRigEditor/Private/EditMode/ControlRigEditModeCommands.cpp b/Engine/Plugins/Animation/ControlRig/Source/ControlRigEditor/Private/EditMode/ControlRigEditModeCommands.cpp
--- a/Engine/Plugins/Animation/ControlRig/Source/ControlRigEditor/Private/EditMode/ControlRigEditModeCommands.cpp
+++ b/Engine/Plugins/Animation/ControlRig/Source/ControlRigEditor/Private/EditMode/ControlRigEditModeCommands.cpp
@@ -11,6 +11,7 @@ void FControlRigEditModeCommands::RegisterCommands()
 	UI_COMMAND(ZeroTransforms, "Zeros Transforms", "Resets Control Transforms to Zero", EUserInterfaceActionType::Button, FInputChord(EKeys::I, EModifierKey::Control));
 	UI_COMMAND(ZeroAllTransforms, "Zeros All Transforms", "Resets all of the Control Transforms to Zero", EUserInterfaceActionType::Button, FInputChord(EKeys::I, EModifierKey::Control | EModifierKey::Shift));
 	UI_COMMAND(ToggleManipulators, "Toggle Shapes", "Toggles visibility of active control rig shapes in the viewport", EUserInterfaceActionType::Button, FInputChord(EKeys::T));
+	UI_COMMAND(ToggleModuleManipulators, "Toggle Module Shapes", "Toggles visibility of active control rig shapes belonging to selected modules in the viewport", EUserInterfaceActionType::Button, FInputChord(EKeys::T, EModifierKey::Control));
 	UI_COMMAND(ToggleAllManipulators, "Toggle All Shapes", "Toggles visibility of all control rig shapes in the viewport", EUserInterfaceActionType::Button, FInputChord(EKeys::T, EModifierKey::Control | EModifierKey::Alt));
 	UI_COMMAND(FrameSelection, "Frame Selection", "Focus the viewport on the current selection", EUserInterfaceActionType::Button, FInputChord(EKeys::F));
 	UI_COMMAND(ClearSelection, "Clear Selection", "Clear Selection", EUserInterfaceActionType::Button, FInputChord(EKeys::Escape));
diff --git a/Engine/Plugins/Animation/ControlRig/Source/ControlRigEditor/Private/EditMode/ControlRigEditModeCommands.h b/Engine/Plugins/Animation/ControlRig/Source/ControlRigEditor/Private/EditMode/ControlRigEditModeCommands.h
--- a/Engine/Plugins/Animation/ControlRig/Source/ControlRigEditor/Private/EditMode/ControlRigEditModeCommands.h
+++ b/Engine/Plugins/Animation/ControlRig/Source/ControlRigEditor/Private/EditMode/ControlRigEditModeCommands.h
@@ -23,6 +23,9 @@ public:
 	/** Toggles hiding all manipulators on active control rig in the viewport */
 	TSharedPtr< FUICommandInfo > ToggleManipulators;
 
+	/** Toggles hiding all manipulators belonging to the selected module on active control rig in the viewport */
+    TSharedPtr< FUICommandInfo > ToggleModuleManipulators;
+
 	/** Toggles hiding all manipulators on all control rigs in the viewport */
 	TSharedPtr< FUICommandInfo > ToggleAllManipulators;
 
diff --git a/Engine/Plugins/Animation/ControlRig/Source/ControlRigEditor/Private/EditMode/SControlRigOutliner.cpp b/Engine/Plugins/Animation/ControlRig/Source/ControlRigEditor/Private/EditMode/SControlRigOutliner.cpp
--- a/Engine/Plugins/Animation/ControlRig/Source/ControlRigEditor/Private/EditMode/SControlRigOutliner.cpp
+++ b/Engine/Plugins/Animation/ControlRig/Source/ControlRigEditor/Private/EditMode/SControlRigOutliner.cpp
@@ -12,6 +12,7 @@
 #include "Styling/CoreStyle.h"
 #include "ScopedTransaction.h"
 #include "ControlRig.h"
+#include "ControlRigEditModeCommands.h"
 #include "ControlRigHierarchyCommands.h"
 #include "EditMode/ControlRigEditMode.h"
 #include "EditorModeManager.h"
@@ -1127,6 +1128,7 @@ void SSearchableMultiRigHierarchyTreeView::Construct(const FArguments& InArgs)
 {
 	FMultiRigTreeDelegates TreeDelegates = InArgs._RigTreeDelegates;
 	SuperGetRigTreeDisplaySettings = TreeDelegates.OnGetDisplaySettings;
+	GetEditMode = TreeDelegates.OnGetEditMode;
 
 	CommandList = MakeShared<FUICommandList>();
 	BindCommands();
@@ -1213,9 +1215,11 @@ TSharedRef<SWidget> SSearchableMultiRigHierarchyTreeView::OnGetOptionsMenu()
 	MenuBuilder.BeginSection("FilterOptions", LOCTEXT("FilterOptions", "Filter Options"));
 	{
 		const FControlRigHierarchyCommands& Actions = FControlRigHierarchyCommands::Get();
+		const FControlRigEditModeCommands& EditModeActions = FControlRigEditModeCommands::Get();
 		MenuBuilder.AddMenuEntry(Actions.ArrangeByModules);
 		MenuBuilder.AddMenuEntry(Actions.ShowAnimationChannels);
 		MenuBuilder.AddMenuEntry(Actions.ShowFullPaths);
+		MenuBuilder.AddMenuEntry(EditModeActions.ToggleModuleManipulators);
 	}
 	MenuBuilder.EndSection();
 
@@ -1244,6 +1248,10 @@ void SSearchableMultiRigHierarchyTreeView::BindCommands()
 		FCanExecuteAction(),
 		FIsActionChecked::CreateSP(this, &SSearchableMultiRigHierarchyTreeView::IsShowingFullPaths));
 
+	CommandList->MapAction(
+		FControlRigEditModeCommands::Get().ToggleModuleManipulators,
+		FExecuteAction::CreateSP(this, &SSearchableMultiRigHierarchyTreeView::ToggleModuleManipulators),
+		FCanExecuteAction::CreateSP(this, &SSearchableMultiRigHierarchyTreeView::CanToggleModuleManipulators));
 }
 
 void SSearchableMultiRigHierarchyTreeView::OnSettingChanged(UObject* Object, FPropertyChangedEvent& PropertyChangedEvent)
@@ -1306,6 +1314,27 @@ void SSearchableMultiRigHierarchyTreeView::ToggleShortNames()
 	TreeView->RequestTreeViewRefresh();
 }
 
+void SSearchableMultiRigHierarchyTreeView::ToggleModuleManipulators()
+{
+	if (FControlRigEditMode* EditMode = GetEditMode.Execute())
+	{
+		EditMode->ToggleModuleManipulators();
+	}
+}
+
+bool SSearchableMultiRigHierarchyTreeView::CanToggleModuleManipulators() const
+{
+	const TArray<FMultiRigData> Selection = TreeView->GetSelectedData();
+	for (const FMultiRigData& Data : Selection)
+	{
+		if (Data.ControlRig->IsModularRig())
+		{
+			return true;
+		}
+	}
+	return false;
+}
+
 void SSearchableMultiRigHierarchyTreeView::OnFilterTextChanged(const FText& SearchText)
 {
 	FilterText = SearchText;
@@ -1336,7 +1365,7 @@ void SControlRigOutliner::Construct(const FArguments& InArgs, FControlRigEditMod
 	FMultiRigTreeDelegates RigTreeDelegates;
 	RigTreeDelegates.OnGetDisplaySettings = FOnGetRigTreeDisplaySettings::CreateSP(this, &SControlRigOutliner::GetDisplaySettings);
 	RigTreeDelegates.OnSelectionChanged = FOnMultiRigTreeSelectionChanged::CreateSP(this, &SControlRigOutliner::HandleSelectionChanged);
-
+	RigTreeDelegates.OnGetEditMode = FOnMultiRigTreeGetEditMode::CreateSP(this, &SControlRigOutliner::GetEditMode);
 	
 	ChildSlot
 		[
@@ -1426,6 +1455,10 @@ void SControlRigOutliner::HandleControlSelected(UControlRig* Subject, FRigContro
 			TArray<TSharedPtr<FMultiRigTreeElement>> SelectedItems = HierarchyTreeView->GetTreeView()->GetSelectedItems();
 			for (TSharedPtr<FMultiRigTreeElement> SelectedItem : SelectedItems)
 			{
+				if (SelectedItem->Data.ModuleName.IsSet())
+				{
+					continue;
+				}
 				HierarchyTreeView->GetTreeView()->SetExpansionRecursive(SelectedItem, false, true);
 			}
 
@@ -1451,7 +1484,7 @@ void SControlRigOutliner::HandleSelectionChanged(TSharedPtr<FMultiRigTreeElement
 		{
 			if (Data.Key.IsSet() && Data.Key.GetValue() != FRigElementKey())
 			{
-				SelectedRigAndKeys.FindOrAdd(Data.ControlRig.Get()).Add(Data.Key.GetValue());
+				SelectedRigAndKeys.FindOrAdd(Data.ControlRig.Get()).AddUnique(Data.Key.GetValue());
 			}
 
 			if (Data.ModuleName.IsSet())
@@ -1467,7 +1500,7 @@ void SControlRigOutliner::HandleSelectionChanged(TSharedPtr<FMultiRigTreeElement
 						{
 							break;
 						}
-						Keys.Add(Child->Data.Key.GetValue());
+						Keys.AddUnique(Child->Data.Key.GetValue());
 					}
 				}
 				else
@@ -1477,7 +1510,7 @@ void SControlRigOutliner::HandleSelectionChanged(TSharedPtr<FMultiRigTreeElement
 					{
 						if (Data.GetHierarchy()->GetModuleName(ControlElement->GetKey()) == Data.ModuleName.GetValue())
 						{
-							Keys.Add(ControlElement->GetKey());
+							Keys.AddUnique(ControlElement->GetKey());
 						}
 						return true;
 					});
@@ -1532,14 +1565,7 @@ void SControlRigOutliner::HandleSelectionChanged(TSharedPtr<FMultiRigTreeElement
 		{
 			URigHierarchyController* Controller = ((URigHierarchy*)Hierarchy)->GetController(true);
 			check(Controller);
-			if (!Controller->SetSelection(RigAndKeys.Value))
-			{
-				if (bEndTransaction)
-				{
-					GEditor->EndTransaction();
-				}
-				return;
-			}
+			Controller->SetSelection(RigAndKeys.Value);
 		}
 	}
 	if (bEndTransaction)
diff --git a/Engine/Plugins/Animation/ControlRig/Source/ControlRigEditor/Private/EditMode/SControlRigOutliner.h b/Engine/Plugins/Animation/ControlRig/Source/ControlRigEditor/Private/EditMode/SControlRigOutliner.h
--- a/Engine/Plugins/Animation/ControlRig/Source/ControlRigEditor/Private/EditMode/SControlRigOutliner.h
+++ b/Engine/Plugins/Animation/ControlRig/Source/ControlRigEditor/Private/EditMode/SControlRigOutliner.h
@@ -24,6 +24,7 @@ class FMultiRigTreeElement;
 class FMultiRigData;
 
 DECLARE_DELEGATE_RetVal_TwoParams(bool, FOnMultiRigTreeCompareKeys, const FMultiRigData& /*A*/, const FMultiRigData& /*B*/);
+DECLARE_DELEGATE_RetVal(FControlRigEditMode*, FOnMultiRigTreeGetEditMode);
 
 typedef STreeView<TSharedPtr<FMultiRigTreeElement>>::FOnSelectionChanged FOnMultiRigTreeSelectionChanged;
 typedef STreeView<TSharedPtr<FMultiRigTreeElement>>::FOnMouseButtonClick FOnMultiRigTreeMouseButtonClick;
@@ -40,6 +41,8 @@ struct CONTROLRIGEDITOR_API FMultiRigTreeDelegates
 	FOnMultiRigTreeMouseButtonClick OnMouseButtonClick;
 	FOnMultiRigTreeMouseButtonDoubleClick OnMouseButtonDoubleClick;
 	FOnRigTreeCompareKeys OnCompareKeys;
+	FOnMultiRigTreeGetEditMode OnGetEditMode;
+	
 
 	FMultiRigTreeDelegates()
 	{
@@ -56,6 +59,15 @@ struct CONTROLRIGEDITOR_API FMultiRigTreeDelegates
 		return DefaultDisplaySettings;
 	}
 
+	FControlRigEditMode* GetEditMode() const
+	{
+		if (OnGetEditMode.IsBound())
+		{
+			return OnGetEditMode.Execute();
+		}
+		return nullptr;
+	}
+
 	void HandleSelectionChanged(TSharedPtr<FMultiRigTreeElement> Selection, ESelectInfo::Type SelectInfo)
 	{
 		if (bIsChangingRigHierarchy)
@@ -250,6 +262,8 @@ public:
 	void ToggleShowAnimationChannels();
 	bool IsShowingFullPaths() const;
 	void ToggleShortNames();
+	void ToggleModuleManipulators();
+	bool CanToggleModuleManipulators() const;
 
 private:
 
@@ -259,6 +273,7 @@ private:
 	TSharedPtr<FUICommandList> CommandList;
 
 	FOnGetRigTreeDisplaySettings SuperGetRigTreeDisplaySettings;
+	FOnMultiRigTreeGetEditMode GetEditMode;
 	FText FilterText;
 	FRigTreeDisplaySettings Settings;
 	TSharedPtr<SMultiRigHierarchyTreeView> TreeView;
diff --git a/Engine/Plugins/Animation/ControlRig/Source/ControlRigEditor/Public/EditMode/ControlRigEditMode.h b/Engine/Plugins/Animation/ControlRig/Source/ControlRigEditor/Public/EditMode/ControlRigEditMode.h
--- a/Engine/Plugins/Animation/ControlRig/Source/ControlRigEditor/Public/EditMode/ControlRigEditMode.h
+++ b/Engine/Plugins/Animation/ControlRig/Source/ControlRigEditor/Public/EditMode/ControlRigEditMode.h
@@ -276,6 +276,9 @@ protected:
 	virtual bool HandleEndTransform(FEditorViewportClient* InViewportClient);
 
 public:
+
+	/** Toggles visibility of acive control rig shapes inside the selected module in the viewport */
+	void ToggleModuleManipulators();
 	
 	/** Clear Selection*/
 	void ClearSelection();
