Anim In Engine: Can Mirror translation without rotation change
#trivial

diff --git a/Engine/Plugins/Animation/ControlRig/Source/ControlRig/Private/Tools/ControlRigPoseMirrorTable.cpp b/Engine/Plugins/Animation/ControlRig/Source/ControlRig/Private/Tools/ControlRigPoseMirrorTable.cpp
--- a/Engine/Plugins/Animation/ControlRig/Source/ControlRig/Private/Tools/ControlRigPoseMirrorTable.cpp
+++ b/Engine/Plugins/Animation/ControlRig/Source/ControlRig/Private/Tools/ControlRigPoseMirrorTable.cpp
@@ -82,7 +82,7 @@ void FControlRigPoseMirrorTable::GetMirrorTransform(const FRigControlCopy& Contr
 	OutLocalTransform = LocalTransform;
 	if (Settings)
 	{
-		if (!bIsMatched  && (OutGlobalTransform.GetRotation().IsIdentity() == false || OutLocalTransform.GetRotation().IsIdentity() == false))
+		if (!bIsMatched) //still do it if not a matching, need to mirror translation
 		{
 			FRigVMMirrorSettings MirrorSettings;
 			MirrorSettings.MirrorAxis = Settings->MirrorAxis;
