GetControlOffset rigUnit using Hierarchy->GetControlOffsetTransform

diff --git a/Engine/Plugins/Animation/ControlRig/Source/ControlRig/Private/Units/Hierarchy/RigUnit_GetControlOffset.cpp b/Engine/Plugins/Animation/ControlRig/Source/ControlRig/Private/Units/Hierarchy/RigUnit_GetControlOffset.cpp
--- a/Engine/Plugins/Animation/ControlRig/Source/ControlRig/Private/Units/Hierarchy/RigUnit_GetControlOffset.cpp
+++ b/Engine/Plugins/Animation/ControlRig/Source/ControlRig/Private/Units/Hierarchy/RigUnit_GetControlOffset.cpp
@@ -17,16 +17,17 @@ FRigUnit_GetControlOffset_Execute()
 		}
 		else
 		{
+			FRigControlElement* ControlElement = Hierarchy->Get<FRigControlElement>(CachedIndex);
 			switch (Space)
 			{
 				case ERigVMTransformSpace::GlobalSpace:
 				{
-					OffsetTransform = Hierarchy->GetInitialGlobalTransform(CachedIndex);
+					OffsetTransform = Hierarchy->GetControlOffsetTransform(ControlElement, ERigTransformType::InitialGlobal);
 					break;
 				}
 				case ERigVMTransformSpace::LocalSpace:
 				{
-					OffsetTransform = Hierarchy->GetInitialLocalTransform(CachedIndex);
+					OffsetTransform = Hierarchy->GetControlOffsetTransform(ControlElement, ERigTransformType::InitialLocal);
 					break;
 				}
 				default:
