Control Rig: New Control Rig Editor
#rnx

diff --git a/Engine/Plugins/Animation/ControlRig/Source/ControlRig/Public/ControlRig.h b/Engine/Plugins/Animation/ControlRig/Source/ControlRig/Public/ControlRig.h
--- a/Engine/Plugins/Animation/ControlRig/Source/ControlRig/Public/ControlRig.h
+++ b/Engine/Plugins/Animation/ControlRig/Source/ControlRig/Public/ControlRig.h
@@ -932,7 +932,7 @@ private:
 
 	friend class FControlRigBlueprintCompilerContext;
 	friend struct FRigHierarchyRef;
-	friend class FControlRigEditor;
+	friend class FControlRigBaseEditor;
 	friend class SRigCurveContainer;
 	friend class SRigHierarchy;
 	friend class SControlRigAnimAttributeView;
diff --git a/Engine/Plugins/Animation/ControlRig/Source/ControlRig/Public/ControlRigReplay.h b/Engine/Plugins/Animation/ControlRig/Source/ControlRig/Public/ControlRigReplay.h
--- a/Engine/Plugins/Animation/ControlRig/Source/ControlRig/Public/ControlRigReplay.h
+++ b/Engine/Plugins/Animation/ControlRig/Source/ControlRig/Public/ControlRigReplay.h
@@ -313,5 +313,5 @@ protected:
 	static const FText GroundTruthStatus;
 	static const FText GroundTruthStatusTooltip;
 
-	friend class FControlRigEditor;
+	friend class FControlRigBaseEditor;
 };
\ No newline at end of file
diff --git a/Engine/Plugins/Animation/ControlRig/Source/ControlRig/Public/Rigs/RigHierarchy.h b/Engine/Plugins/Animation/ControlRig/Source/ControlRig/Public/Rigs/RigHierarchy.h
--- a/Engine/Plugins/Animation/ControlRig/Source/ControlRig/Public/Rigs/RigHierarchy.h
+++ b/Engine/Plugins/Animation/ControlRig/Source/ControlRig/Public/Rigs/RigHierarchy.h
@@ -5481,7 +5481,7 @@ private:
 	friend class UControlRig;
 	friend class UFKControlRig;
 	friend class UModularRig;
-	friend class FControlRigEditor;
+	friend class FControlRigBaseEditor;
 	friend struct FRigBaseElement;
 	friend struct FRigHierarchyValidityBracket;
 	friend struct FRigHierarchyGlobalValidityBracket;
diff --git a/Engine/Plugins/Animation/ControlRig/Source/ControlRig/Public/Rigs/RigHierarchyComponents.h b/Engine/Plugins/Animation/ControlRig/Source/ControlRig/Public/Rigs/RigHierarchyComponents.h
--- a/Engine/Plugins/Animation/ControlRig/Source/ControlRig/Public/Rigs/RigHierarchyComponents.h
+++ b/Engine/Plugins/Animation/ControlRig/Source/ControlRig/Public/Rigs/RigHierarchyComponents.h
@@ -220,7 +220,7 @@ public:
 
 protected:
 
-	friend class FControlRigEditor;
+	friend class FControlRigBaseEditor;
 	friend class URigHierarchy;
 	friend class URigHierarchyController;
 	friend struct FControlRigReplayTracks;
diff --git a/Engine/Plugins/Animation/ControlRig/Source/ControlRig/Public/Rigs/RigHierarchyElements.h b/Engine/Plugins/Animation/ControlRig/Source/ControlRig/Public/Rigs/RigHierarchyElements.h
--- a/Engine/Plugins/Animation/ControlRig/Source/ControlRig/Public/Rigs/RigHierarchyElements.h
+++ b/Engine/Plugins/Animation/ControlRig/Source/ControlRig/Public/Rigs/RigHierarchyElements.h
@@ -1025,7 +1025,7 @@ protected:
 	virtual void LinkStorage(const TArrayView<FTransform>& InTransforms, const TArrayView<bool>& InDirtyStates, const TArrayView<float>& InCurves) {}
 	virtual void UnlinkStorage(FRigReusableElementStorage<FTransform>& InTransforms, FRigReusableElementStorage<bool>& InDirtyStates, FRigReusableElementStorage<float>& InCurves) {}
 
-	friend class FControlRigEditor;
+	friend class FControlRigBaseEditor;
 	friend class URigHierarchy;
 	friend class URigHierarchyController;
 	friend struct FRigElementKeyAndIndex;
@@ -1117,7 +1117,7 @@ protected:
 	friend class URigHierarchy;
 	friend class URigHierarchyController;
 	friend struct FRigBaseElement;
-	friend class FControlRigEditor;
+	friend class FControlRigBaseEditor;
 };
 
 USTRUCT(BlueprintType)
@@ -1810,7 +1810,7 @@ private:
 	friend class URigHierarchy;
 	friend class URigHierarchyController;
 	friend struct FRigBaseElement;
-	friend class FControlRigEditor;
+	friend class FControlRigBaseEditor;
 };
 
 USTRUCT(BlueprintType)
diff --git a/Engine/Plugins/Animation/ControlRig/Source/ControlRig/Public/Units/Hierarchy/RigUnit_Component.h b/Engine/Plugins/Animation/ControlRig/Source/ControlRig/Public/Units/Hierarchy/RigUnit_Component.h
--- a/Engine/Plugins/Animation/ControlRig/Source/ControlRig/Public/Units/Hierarchy/RigUnit_Component.h
+++ b/Engine/Plugins/Animation/ControlRig/Source/ControlRig/Public/Units/Hierarchy/RigUnit_Component.h
@@ -42,7 +42,7 @@ protected:
 	static inline const FLazyName ComponentArgName = FLazyName(TEXT("Component"));
 	static inline const FLazyName SuccessArgName = FLazyName(TEXT("Success"));
 
-	friend class FControlRigEditor;
+	friend class FControlRigBaseEditor;
 };
 
 /**
diff --git a/Engine/Plugins/Animation/ControlRig/Source/ControlRigDeveloper/Public/ControlRigBlueprint.h b/Engine/Plugins/Animation/ControlRig/Source/ControlRigDeveloper/Public/ControlRigBlueprint.h
--- a/Engine/Plugins/Animation/ControlRig/Source/ControlRigDeveloper/Public/ControlRigBlueprint.h
+++ b/Engine/Plugins/Animation/ControlRig/Source/ControlRigDeveloper/Public/ControlRigBlueprint.h
@@ -438,6 +438,10 @@ private:
 	friend class FControlRigBlueprintCompilerContext;
 	friend class SRigHierarchy;
 	friend class SRigCurveContainer;
+	friend class FControlRigBaseEditor;
+#if WITH_RIGVMLEGACYEDITOR
+	friend class FControlRigLegacyEditor;
+#endif
 	friend class FControlRigEditor;
 	friend class UEngineTestControlRig;
 	friend class FControlRigEditMode;
diff --git a/Engine/Plugins/Animation/ControlRig/Source/ControlRigDeveloper/Public/Graph/ControlRigGraph.h b/Engine/Plugins/Animation/ControlRig/Source/ControlRigDeveloper/Public/Graph/ControlRigGraph.h
--- a/Engine/Plugins/Animation/ControlRig/Source/ControlRigDeveloper/Public/Graph/ControlRigGraph.h
+++ b/Engine/Plugins/Animation/ControlRig/Source/ControlRigDeveloper/Public/Graph/ControlRigGraph.h
@@ -160,7 +160,6 @@ private:
 #endif
 
 	friend class UControlRigGraphNode;
-	friend class FControlRigEditor;
 	friend class SRigVMGraphNode;
 	friend class UControlRigBlueprint;
 };
diff --git a/Engine/Plugins/Animation/ControlRig/Source/ControlRigDeveloper/Public/Graph/ControlRigGraphNode.h b/Engine/Plugins/Animation/ControlRig/Source/ControlRigDeveloper/Public/Graph/ControlRigGraphNode.h
--- a/Engine/Plugins/Animation/ControlRig/Source/ControlRigDeveloper/Public/Graph/ControlRigGraphNode.h
+++ b/Engine/Plugins/Animation/ControlRig/Source/ControlRigDeveloper/Public/Graph/ControlRigGraphNode.h
@@ -35,7 +35,6 @@ class CONTROLRIGDEVELOPER_API UControlRigGraphNode : public URigVMEdGraphNode
 	friend class UControlRigBlueprint;
 	friend class FControlRigGraphTraverser;
 	friend class FControlRigGraphPanelPinFactory;
-	friend class FControlRigEditor;
 	friend class SControlRigGraphPinCurveFloat;
 
 public:
diff --git a/Engine/Plugins/Animation/ControlRig/Source/ControlRigEditor/ControlRigEditor.Build.cs b/Engine/Plugins/Animation/ControlRig/Source/ControlRigEditor/ControlRigEditor.Build.cs
--- a/Engine/Plugins/Animation/ControlRig/Source/ControlRigEditor/ControlRigEditor.Build.cs
+++ b/Engine/Plugins/Animation/ControlRig/Source/ControlRigEditor/ControlRigEditor.Build.cs
@@ -35,7 +35,6 @@ namespace UnrealBuildTool.Rules
                     "BlueprintGraph",
                     "ControlRig",
                     "ControlRigDeveloper",
-                    "Kismet",
 					"KismetCompiler",
                     "EditorStyle",
 					"EditorWidgets",
@@ -90,6 +89,12 @@ namespace UnrealBuildTool.Rules
 				}
             );
 
+            bool bAddBlueprintEditorDependency = true;
+            if (bAddBlueprintEditorDependency)
+            {
+	            PublicDependencyModuleNames.Add("Kismet");
+            }
+
             AddEngineThirdPartyPrivateStaticDependencies(Target, "FBX");
         }
     }
diff --git a/Engine/Plugins/Animation/ControlRig/Source/ControlRigEditor/Private/ControlRigBlueprintActions.cpp b/Engine/Plugins/Animation/ControlRig/Source/ControlRigEditor/Private/ControlRigBlueprintActions.cpp
--- a/Engine/Plugins/Animation/ControlRig/Source/ControlRigEditor/Private/ControlRigBlueprintActions.cpp
+++ b/Engine/Plugins/Animation/ControlRig/Source/ControlRigEditor/Private/ControlRigBlueprintActions.cpp
@@ -134,7 +134,9 @@ void FControlRigBlueprintActions::PerformAssetDiff(UObject* OldAsset, UObject* N
 		AssetEditorSubsystem->CloseAllEditorsForAsset(NewAsset);
 	}
 
+#if WITH_RIGVMLEGACYEDITOR
 	SBlueprintDiff::CreateDiffWindow(OldBlueprint, NewBlueprint, OldRevision, NewRevision, GetSupportedClass());
+#endif
 }
 
 void FControlRigBlueprintActions::ExtendSketalMeshToolMenu()
diff --git a/Engine/Plugins/Animation/ControlRig/Source/ControlRigEditor/Private/ControlRigEditorModule.cpp b/Engine/Plugins/Animation/ControlRig/Source/ControlRigEditor/Private/ControlRigEditorModule.cpp
--- a/Engine/Plugins/Animation/ControlRig/Source/ControlRigEditor/Private/ControlRigEditorModule.cpp
+++ b/Engine/Plugins/Animation/ControlRig/Source/ControlRigEditor/Private/ControlRigEditorModule.cpp
@@ -98,6 +98,10 @@
 #include "AssetRegistry/IAssetRegistry.h"
 #include "SequencerUtilities.h"
 #include "Bindings/MovieSceneSpawnableActorBinding.h"
+#if WITH_RIGVMLEGACYEDITOR
+#include "Editor/ControlRigLegacyEditor.h"
+#endif
+#include "Editor/ControlRigNewEditor.h"
 
 #if WITH_RIGVMLEGACYEDITOR
 #include "SKismetInspector.h"
@@ -1505,9 +1509,37 @@ void FControlRigEditorModule::ExtendAnimSequenceMenu()
 
 TSharedRef<IControlRigEditor> FControlRigEditorModule::CreateControlRigEditor(const EToolkitMode::Type Mode, const TSharedPtr< IToolkitHost >& InitToolkitHost, class UControlRigBlueprint* InBlueprint)
 {
+
+#if WITH_RIGVMLEGACYEDITOR
+	if (CVarRigVMUseDualEditor.GetValueOnAnyThread())
+	{
+		TSharedRef< FControlRigLegacyEditor > LegacyRigVMExampleEditor(new FControlRigLegacyEditor());
+		LegacyRigVMExampleEditor->InitRigVMEditor(Mode, InitToolkitHost, InBlueprint);
+
+		TSharedRef< FControlRigEditor > NewControlRigEditor(new FControlRigEditor());
+		NewControlRigEditor->InitRigVMEditor(Mode, InitToolkitHost, InBlueprint);
+
+		return NewControlRigEditor;
+	}
+	
+	if(CVarRigVMUseNewEditor.GetValueOnAnyThread())
+	{
+		TSharedRef< FControlRigEditor > NewControlRigEditor(new FControlRigEditor());
+		NewControlRigEditor->InitRigVMEditor(Mode, InitToolkitHost, InBlueprint);
+		return NewControlRigEditor;
+	}
+
+	TSharedRef< FControlRigLegacyEditor > LegacyRigVMExampleEditor(new FControlRigLegacyEditor());
+	LegacyRigVMExampleEditor->InitRigVMEditor(Mode, InitToolkitHost, InBlueprint);
+	return LegacyRigVMExampleEditor;
+
+#else
+	
 	TSharedRef< FControlRigEditor > NewControlRigEditor(new FControlRigEditor());
 	NewControlRigEditor->InitRigVMEditor(Mode, InitToolkitHost, InBlueprint);
 	return NewControlRigEditor;
+#endif
+	
 }
 
 FControlRigClassFilter::FControlRigClassFilter(bool bInCheckSkeleton, bool bInCheckAnimatable, bool bInCheckInversion, USkeleton* InSkeleton) :	bFilterAssetBySkeleton(bInCheckSkeleton),
diff --git a/Engine/Plugins/Animation/ControlRig/Source/ControlRigEditor/Private/ControlRigEditorModule.h b/Engine/Plugins/Animation/ControlRig/Source/ControlRigEditor/Private/ControlRigEditorModule.h
--- a/Engine/Plugins/Animation/ControlRig/Source/ControlRigEditor/Private/ControlRigEditorModule.h
+++ b/Engine/Plugins/Animation/ControlRig/Source/ControlRigEditor/Private/ControlRigEditorModule.h
@@ -9,6 +9,7 @@
 #include "IAnimationEditor.h"
 #include "IAnimationEditorModule.h"
 
+class IControlRigBaseEditor;
 class UBlueprint;
 class IAssetTypeActions;
 class UMaterial;
diff --git a/Engine/Plugins/Animation/ControlRig/Source/ControlRigEditor/Private/ControlRigElementDetails.cpp b/Engine/Plugins/Animation/ControlRig/Source/ControlRigEditor/Private/ControlRigElementDetails.cpp
--- a/Engine/Plugins/Animation/ControlRig/Source/ControlRigEditor/Private/ControlRigElementDetails.cpp
+++ b/Engine/Plugins/Animation/ControlRig/Source/ControlRigEditor/Private/ControlRigElementDetails.cpp
@@ -1746,7 +1746,7 @@ void FRigBaseElementDetails::CustomizeMetadata(IDetailLayoutBuilder& DetailBuild
 			UAssetEditorSubsystem* AssetEditorSubsystem = GEditor->GetEditorSubsystem<UAssetEditorSubsystem>();
 			if (IAssetEditorInstance* EditorInstance = AssetEditorSubsystem->FindEditorForAsset(Blueprint, false))
 			{
-				if (FControlRigEditor* Editor = static_cast<FControlRigEditor*>(EditorInstance))
+				if (IControlRigBaseEditor* Editor = FControlRigBaseEditor::GetFromAssetEditorInstance(EditorInstance))
 				{
 					if (Editor->GetReplayPlaybackMode() != EControlRigReplayPlaybackMode::Live)
 					{
diff --git a/Engine/Plugins/Animation/ControlRig/Source/ControlRigEditor/Private/Editor/ControlRigContextMenuContext.cpp b/Engine/Plugins/Animation/ControlRig/Source/ControlRigEditor/Private/Editor/ControlRigContextMenuContext.cpp
--- a/Engine/Plugins/Animation/ControlRig/Source/ControlRigEditor/Private/Editor/ControlRigContextMenuContext.cpp
+++ b/Engine/Plugins/Animation/ControlRig/Source/ControlRigEditor/Private/Editor/ControlRigContextMenuContext.cpp
@@ -23,7 +23,7 @@ FString FControlRigRigHierarchyToGraphDragAndDropContext::GetSectionTitle() cons
 	return SectionTitle;
 }
 
-void UControlRigContextMenuContext::Init(TWeakPtr<FControlRigEditor> InControlRigEditor, const FControlRigMenuSpecificContext& InMenuSpecificContext)
+void UControlRigContextMenuContext::Init(TWeakPtr<IControlRigBaseEditor> InControlRigEditor, const FControlRigMenuSpecificContext& InMenuSpecificContext)
 {
 	WeakControlRigEditor = InControlRigEditor;
 	MenuSpecificContext = InMenuSpecificContext;
@@ -31,7 +31,7 @@ void UControlRigContextMenuContext::Init(TWeakPtr<FControlRigEditor> InControlRi
 
 UControlRigBlueprint* UControlRigContextMenuContext::GetControlRigBlueprint() const
 {
-	if (const TSharedPtr<FControlRigEditor> Editor = WeakControlRigEditor.Pin())
+	if (const TSharedPtr<IControlRigBaseEditor> Editor = WeakControlRigEditor.Pin())
 	{
 		return Editor->GetControlRigBlueprint();
 	}
@@ -89,9 +89,9 @@ SModularRigModel* UControlRigContextMenuContext::GetModularRigModelPanel() const
 	return nullptr;
 }
 
-FControlRigEditor* UControlRigContextMenuContext::GetControlRigEditor() const
+IControlRigBaseEditor* UControlRigContextMenuContext::GetControlRigEditor() const
 {
-	if (const TSharedPtr<FControlRigEditor> Editor = WeakControlRigEditor.Pin())
+	if (const TSharedPtr<IControlRigBaseEditor> Editor = WeakControlRigEditor.Pin())
 	{
 		return Editor.Get();
 	}
diff --git a/Engine/Plugins/Animation/ControlRig/Source/ControlRigEditor/Private/Editor/ControlRigContextMenuContext.h b/Engine/Plugins/Animation/ControlRig/Source/ControlRigEditor/Private/Editor/ControlRigContextMenuContext.h
--- a/Engine/Plugins/Animation/ControlRig/Source/ControlRigEditor/Private/Editor/ControlRigContextMenuContext.h
+++ b/Engine/Plugins/Animation/ControlRig/Source/ControlRigEditor/Private/Editor/ControlRigContextMenuContext.h
@@ -10,7 +10,7 @@
 
 class UControlRig;
 class UControlRigBlueprint;
-class FControlRigEditor;
+class IControlRigBaseEditor;
 class URigVMGraph;
 class URigVMNode;
 class URigVMPin;
@@ -116,7 +116,7 @@ public:
 	 * @param InControlRigEditor 	    The Control Rig Editor hosting the menus
 	 * @param InMenuSpecificContext 	Additional context for specific menus
 	*/
-	void Init(TWeakPtr<FControlRigEditor> InControlRigEditor, const FControlRigMenuSpecificContext& InMenuSpecificContext = FControlRigMenuSpecificContext());
+	void Init(TWeakPtr<IControlRigBaseEditor> InControlRigEditor, const FControlRigMenuSpecificContext& InMenuSpecificContext = FControlRigMenuSpecificContext());
 
 	/** Get the control rig blueprint that we are editing */
 	UFUNCTION(BlueprintCallable, Category = ControlRigEditorExtensions)
@@ -146,10 +146,10 @@ public:
 
 	SModularRigModel* GetModularRigModelPanel() const;
 	
-	FControlRigEditor* GetControlRigEditor() const;
+	IControlRigBaseEditor* GetControlRigEditor() const;
 
 private:
-	TWeakPtr<FControlRigEditor> WeakControlRigEditor;
+	TWeakPtr<IControlRigBaseEditor> WeakControlRigEditor;
 
 	FControlRigMenuSpecificContext MenuSpecificContext;
 };
\ No newline at end of file
diff --git a/Engine/Plugins/Animation/ControlRig/Source/ControlRigEditor/Private/Editor/ControlRigEditor.cpp b/Engine/Plugins/Animation/ControlRig/Source/ControlRigEditor/Private/Editor/ControlRigEditor.cpp
--- a/Engine/Plugins/Animation/ControlRig/Source/ControlRigEditor/Private/Editor/ControlRigEditor.cpp
+++ b/Engine/Plugins/Animation/ControlRig/Source/ControlRigEditor/Private/Editor/ControlRigEditor.cpp
@@ -123,6 +123,7 @@
 #include "ViewportToolbar/UnrealEdViewportToolbarContext.h"
 #include "ITransportControl.h"
 #include "PropertyPath.h"
+#include "Editor/SRigVMDetailsInspector.h"
 
 #if WITH_RIGVMLEGACYEDITOR
 #include "SKismetInspector.h"
@@ -273,14 +274,13 @@ TAutoConsoleVariable<bool> CVarControlRigShowTestingToolbar(TEXT("ControlRig.Tes
 TAutoConsoleVariable<bool> CVarShowSchematicPanelOverlay(TEXT("ControlRig.Preview.ShowSchematicPanelOverlay"), true, TEXT("When true we'll add an overlay to the persona viewport to show modular rig information."));
 
 const FName FControlRigEditorModes::ControlRigEditorMode = TEXT("Rigging");
-const TArray<FName> FControlRigEditor::ForwardsSolveEventQueue = {FRigUnit_BeginExecution::EventName};
-const TArray<FName> FControlRigEditor::BackwardsSolveEventQueue = {FRigUnit_InverseExecution::EventName};
-const TArray<FName> FControlRigEditor::ConstructionEventQueue = {FRigUnit_PrepareForExecution::EventName};
-const TArray<FName> FControlRigEditor::BackwardsAndForwardsSolveEventQueue = {FRigUnit_InverseExecution::EventName, FRigUnit_BeginExecution::EventName};
-
-FControlRigEditor::FControlRigEditor()
-	: IControlRigEditor()
-	, PreviewInstance(nullptr)
+const TArray<FName> FControlRigBaseEditor::ForwardsSolveEventQueue = {FRigUnit_BeginExecution::EventName};
+const TArray<FName> FControlRigBaseEditor::BackwardsSolveEventQueue = {FRigUnit_InverseExecution::EventName};
+const TArray<FName> FControlRigBaseEditor::ConstructionEventQueue = {FRigUnit_PrepareForExecution::EventName};
+const TArray<FName> FControlRigBaseEditor::BackwardsAndForwardsSolveEventQueue = {FRigUnit_InverseExecution::EventName, FRigUnit_BeginExecution::EventName};
+
+FControlRigBaseEditor::FControlRigBaseEditor()
+	: PreviewInstance(nullptr)
 	, ActiveController(nullptr)
 	, bExecutionControlRig(true)
 	, RigHierarchyTabCount(0)
@@ -290,194 +290,53 @@ FControlRigEditor::FControlRigEditor()
 	, bRefreshDirectionManipulationTargetsRequired(false)
 	, bSchematicViewPortIsHidden(false)
 {
-	LastEventQueue = ConstructionEventQueue;
+	
 }
 
-FControlRigEditor::~FControlRigEditor()
+FControlRigBaseEditor* FControlRigBaseEditor::GetFromAssetEditorInstance(IAssetEditorInstance* Instance)
 {
-	UControlRigBlueprint* RigBlueprint = GetControlRigBlueprint();
+	FControlRigBaseEditor* Editor = nullptr;
+	FWorkflowCentricApplication* App = static_cast<FWorkflowCentricApplication*>(Instance);
 
-	if (RigBlueprint)
+	TSharedRef<FAssetEditorToolkit> SharedApp = App->AsShared();
+#if WITH_RIGVMLEGACYEDITOR
+	if (SharedApp->IsBlueprintEditor())
 	{
-		UControlRigBlueprint::sCurrentlyOpenedRigBlueprints.Remove(RigBlueprint);
-
-		RigBlueprint->OnHierarchyModified().RemoveAll(this);
-		if (FControlRigEditMode* EditMode = GetEditMode())
-		{
-			RigBlueprint->OnHierarchyModified().RemoveAll(EditMode);
-			EditMode->OnEditorClosed();
-		}
-
-		RigBlueprint->OnRigTypeChanged().RemoveAll(this);
-		if (RigBlueprint->IsModularRig())
-		{
-			RigBlueprint->GetModularRigController()->OnModified().RemoveAll(this);
-			RigBlueprint->OnModularRigCompiled().RemoveAll(this);
-
-			RigBlueprint->OnSetObjectBeingDebugged().RemoveAll(&SchematicModel);
-			RigBlueprint->OnHierarchyModified().RemoveAll(&SchematicModel);
-			RigBlueprint->GetModularRigController()->OnModified().RemoveAll(&SchematicModel);
-		}
+		TSharedPtr<FControlRigLegacyEditor> LegacyEditor = StaticCastSharedPtr<FControlRigLegacyEditor>(App->AsShared().ToSharedPtr());
+		Editor = LegacyEditor.Get();
 	}
-
-	if (PersonaToolkit.IsValid())
+	else
 	{
-		constexpr bool bSetPreviewMeshInAsset = false;
-		PersonaToolkit->SetPreviewMesh(nullptr, bSetPreviewMeshInAsset);
+		TSharedPtr<FControlRigEditor> NewEditor = StaticCastSharedPtr<FControlRigEditor>(App->AsShared().ToSharedPtr());
+		Editor = NewEditor.Get();
 	}
+#else
+	TSharedPtr<FControlRigEditor> NewEditor = StaticCastSharedPtr<FControlRigEditor>(App->AsShared().ToSharedPtr());
+	Editor = NewEditor.Get();
+#endif
+	return Editor;
 }
 
-UObject* FControlRigEditor::GetOuterForHost() const
+UObject* FControlRigBaseEditor::GetOuterForHostImpl() const
 {
 	UControlRigSkeletalMeshComponent* EditorSkelComp = Cast<UControlRigSkeletalMeshComponent>(GetPersonaToolkit()->GetPreviewScene()->GetPreviewMeshComponent());
 	if(EditorSkelComp)
 	{
 		return EditorSkelComp;
 	}
-	return FRigVMEditorBase::GetOuterForHost();
+	return GetOuterForHostSuper();
 }
 
-UClass* FControlRigEditor::GetDetailWrapperClass() const
+UClass* FControlRigBaseEditor::GetDetailWrapperClassImpl() const
 {
 	return UControlRigWrapperObject::StaticClass();
 }
 
-#if WITH_RIGVMLEGACYEDITOR
-bool FControlRigEditor::IsSectionVisible(NodeSectionID::Type InSectionID) const
-{
-	if(!IControlRigEditor::IsSectionVisible(InSectionID))
-	{
-		return false;
-	}
-
-	if(const UControlRigBlueprint* RigBlueprint = GetControlRigBlueprint())
-	{
-		if(IsModularRig())
-		{
-			switch (InSectionID)
-			{
-				case NodeSectionID::GRAPH:
-				{
-					return RigBlueprint->SupportsEventGraphs();
-				}
-				case NodeSectionID::FUNCTION:
-				{
-					return RigBlueprint->SupportsFunctions();
-				}
-				default:
-				{
-					break;
-				}
-			}
-		}
-	}
-	return true;
-}
-
-#else
-bool FControlRigEditor::IsSectionVisible(RigVMNodeSectionID::Type InSectionID) const
-{
-	if(!IControlRigEditor::IsSectionVisible(InSectionID))
-    {
-    	return false;
-    }
-
-	if(const UControlRigBlueprint* RigBlueprint = GetControlRigBlueprint())
-	{
-		if(IsModularRig())
-		{
-			switch (InSectionID)
-			{
-				case RigVMNodeSectionID::GRAPH:
-				{
-					return RigBlueprint->SupportsEventGraphs();
-				}
-				case RigVMNodeSectionID::FUNCTION:
-				{
-					return RigBlueprint->SupportsFunctions();
-				}
-				default:
-				{
-					break;
-				}
-			}
-		}
-	}
-	return true;
-}
-#endif
-
-
-#if WITH_RIGVMLEGACYEDITOR
-bool FControlRigEditor::NewDocument_IsVisibleForType(FBlueprintEditor::ECreatedDocumentType GraphType) const
-{
-	if(!IControlRigEditor::NewDocument_IsVisibleForType(GraphType))
-	{
-		return false;
-	}
-
-	if(const UControlRigBlueprint* RigBlueprint = GetControlRigBlueprint())
-	{
-		if(IsModularRig())
-		{
-			switch(GraphType)
-			{
-				case FBlueprintEditor::CGT_NewEventGraph:
-				{
-					return RigBlueprint->SupportsEventGraphs();
-				}
-				case FBlueprintEditor::CGT_NewFunctionGraph:
-				{
-					return RigBlueprint->SupportsFunctions();
-				}
-				default:
-				{
-					break;
-				}
-			}
-		}
-	}
-
-	return true;
-}
-#else
-bool FControlRigEditor::NewDocument_IsVisibleForType(FRigVMEditorBase::ECreatedDocumentType GraphType) const
-{
-	if(!IControlRigEditor::NewDocument_IsVisibleForType(GraphType))
-	{
-		return false;
-	}
-
-	if(const UControlRigBlueprint* RigBlueprint = GetControlRigBlueprint())
-	{
-		if(IsModularRig())
-		{
-			switch(GraphType)
-			{
-				case FRigVMEditorBase::CGT_NewEventGraph:
-				{
-					return RigBlueprint->SupportsEventGraphs();
-				}
-				case FRigVMEditorBase::CGT_NewFunctionGraph:
-				{
-					return RigBlueprint->SupportsFunctions();
-				}
-				default:
-				{
-					break;
-				}
-			}
-		}
-	}
 
-	return true;
-}
-#endif
 
-
-FReply FControlRigEditor::OnViewportDrop(const FGeometry& MyGeometry, const FDragDropEvent& DragDropEvent)
+FReply FControlRigBaseEditor::OnViewportDropImpl(const FGeometry& MyGeometry, const FDragDropEvent& DragDropEvent)
 {
-	const FReply SuperReply = IControlRigEditor::OnViewportDrop(MyGeometry, DragDropEvent);
+	const FReply SuperReply = OnViewportDropSuper(MyGeometry, DragDropEvent);
 	if(SuperReply.IsEventHandled())
 	{
 		return SuperReply;
@@ -521,7 +380,7 @@ FReply FControlRigEditor::OnViewportDrop(const FGeometry& MyGeometry, const FDra
 	return FReply::Unhandled();
 }
 
-void FControlRigEditor::CreateEmptyGraphContent(URigVMController* InController)
+void FControlRigBaseEditor::CreateEmptyGraphContentImpl(URigVMController* InController)
 {
 	URigVMNode* Node = InController->AddUnitNode(FRigUnit_BeginExecution::StaticStruct(), FRigUnit::GetMethodName(), FVector2D::ZeroVector, FString(), false);
 	if (Node)
@@ -532,17 +391,17 @@ void FControlRigEditor::CreateEmptyGraphContent(URigVMController* InController)
 	}
 }
 
-UControlRigBlueprint* FControlRigEditor::GetControlRigBlueprint() const
+UControlRigBlueprint* FControlRigBaseEditor::GetControlRigBlueprint() const
 {
 	return Cast<UControlRigBlueprint>(GetRigVMBlueprint());
 }
 
-UControlRig* FControlRigEditor::GetControlRig() const
+UControlRig* FControlRigBaseEditor::GetControlRig() const
 {
 	return Cast<UControlRig>(GetRigVMHost());
 }
 
-URigHierarchy* FControlRigEditor::GetHierarchyBeingDebugged() const
+URigHierarchy* FControlRigBaseEditor::GetHierarchyBeingDebugged() const
 {
 	if(UControlRigBlueprint* RigBlueprint = GetControlRigBlueprint())
 	{
@@ -555,11 +414,11 @@ URigHierarchy* FControlRigEditor::GetHierarchyBeingDebugged() const
 	return nullptr;
 }
 
-void FControlRigEditor::InitRigVMEditor(const EToolkitMode::Type Mode, const TSharedPtr< class IToolkitHost >& InitToolkitHost, URigVMBlueprint* InRigVMBlueprint)
+void FControlRigBaseEditor::InitRigVMEditorImpl(const EToolkitMode::Type Mode, const TSharedPtr<class IToolkitHost>& InitToolkitHost, class URigVMBlueprint* InRigVMBlueprint)
 {
 	DECLARE_SCOPE_HIERARCHICAL_COUNTER_FUNC()
 
-	IControlRigEditor::InitRigVMEditor(Mode, InitToolkitHost, InRigVMBlueprint);
+	InitRigVMEditorSuper(Mode, InitToolkitHost, InRigVMBlueprint);
 
 	UControlRigBlueprint* ControlRigBlueprint = CastChecked<UControlRigBlueprint>(InRigVMBlueprint);
 
@@ -568,11 +427,11 @@ void FControlRigEditor::InitRigVMEditor(const EToolkitMode::Type Mode, const TSh
 
 	if (FControlRigEditMode* EditMode = GetEditMode())
 	{
-		EditMode->OnGetRigElementTransform() = FOnGetRigElementTransform::CreateSP(this, &FControlRigEditor::GetRigElementTransform);
-		EditMode->OnSetRigElementTransform() = FOnSetRigElementTransform::CreateSP(this, &FControlRigEditor::SetRigElementTransform);
-		EditMode->OnGetContextMenu() = FOnGetContextMenu::CreateSP(this, &FControlRigEditor::HandleOnGetViewportContextMenuDelegate);
-		EditMode->OnContextMenuCommands() = FNewMenuCommandsDelegate::CreateSP(this, &FControlRigEditor::HandleOnViewportContextMenuCommandsDelegate);
-		EditMode->OnAnimSystemInitialized().Add(FSimpleMulticastDelegate::FDelegate::CreateRaw(this, &FControlRigEditor::OnAnimInitialized));
+		EditMode->OnGetRigElementTransform() = FOnGetRigElementTransform::CreateSP(SharedRef(), &FControlRigBaseEditor::GetRigElementTransform);
+		EditMode->OnSetRigElementTransform() = FOnSetRigElementTransform::CreateSP(SharedRef(), &FControlRigBaseEditor::SetRigElementTransform);
+		EditMode->OnGetContextMenu() = FOnGetContextMenu::CreateSP(SharedRef(), &FControlRigBaseEditor::HandleOnGetViewportContextMenuDelegate);
+		EditMode->OnContextMenuCommands() = FNewMenuCommandsDelegate::CreateSP(SharedRef(), &FControlRigBaseEditor::HandleOnViewportContextMenuCommandsDelegate);
+		EditMode->OnAnimSystemInitialized().Add(FSimpleMulticastDelegate::FDelegate::CreateRaw(this, &FControlRigBaseEditor::OnAnimInitialized));
 	
 		PersonaToolkit->GetPreviewScene()->SetRemoveAttachedComponentFilter(FOnRemoveAttachedComponentFilter::CreateSP(EditMode, &FControlRigEditMode::CanRemoveFromPreviewScene));
 	}
@@ -590,7 +449,7 @@ void FControlRigEditor::InitRigVMEditor(const EToolkitMode::Type Mode, const TSh
 		// 3. BP performs propagation again in UControlRigBlueprint::HandleHierarchyModified, invalidates the rig element
 		//    that the UI is observing
 		// 4. Editor UI shows an invalid rig element
-		ControlRigBlueprint->OnHierarchyModified().AddSP(this, &FControlRigEditor::OnHierarchyModified);
+		ControlRigBlueprint->OnHierarchyModified().AddSP(SharedRef(), &FControlRigBaseEditor::OnHierarchyModified);
 
 		if (FControlRigEditMode* EditMode = GetEditMode())
 		{
@@ -599,16 +458,16 @@ void FControlRigEditor::InitRigVMEditor(const EToolkitMode::Type Mode, const TSh
 
 		if (ControlRigBlueprint->IsModularRig())
 		{
-			SchematicModel.SetEditor(SharedThis(this));
+			SchematicModel.SetEditor(SharedRef());
 			ControlRigBlueprint->OnSetObjectBeingDebugged().AddRaw(&SchematicModel, &FControlRigSchematicModel::OnSetObjectBeingDebugged);
 			ControlRigBlueprint->GetModularRigController()->OnModified().AddRaw(&SchematicModel, &FControlRigSchematicModel::HandleModularRigModified);
 		}
 
-		ControlRigBlueprint->OnRigTypeChanged().AddSP(this, &FControlRigEditor::HandleRigTypeChanged);
+		ControlRigBlueprint->OnRigTypeChanged().AddSP(SharedRef(), &FControlRigBaseEditor::HandleRigTypeChanged);
 		if (ControlRigBlueprint->IsModularRig())
 		{
-			ControlRigBlueprint->GetModularRigController()->OnModified().AddSP(this, &FControlRigEditor::HandleModularRigModified);
-			ControlRigBlueprint->OnModularRigCompiled().AddRaw(this, &FControlRigEditor::HandlePostCompileModularRigs);
+			ControlRigBlueprint->GetModularRigController()->OnModified().AddSP(SharedRef(), &FControlRigBaseEditor::HandleModularRigModified);
+			ControlRigBlueprint->OnModularRigCompiled().AddSP(SharedRef(), &FControlRigBaseEditor::HandlePostCompileModularRigs);
 		}
 	}
 
@@ -620,7 +479,7 @@ void FControlRigEditor::InitRigVMEditor(const EToolkitMode::Type Mode, const TSh
 	}
 }
 
-void FControlRigEditor::CreatePersonaToolKitIfRequired()
+void FControlRigBaseEditor::CreatePersonaToolKitIfRequired()
 {
 	if(PersonaToolkit.IsValid())
 	{
@@ -632,7 +491,7 @@ void FControlRigEditor::CreatePersonaToolKitIfRequired()
 	FPersonaModule& PersonaModule = FModuleManager::GetModuleChecked<FPersonaModule>("Persona");
 
 	FPersonaToolkitArgs PersonaToolkitArgs;
-	PersonaToolkitArgs.OnPreviewSceneCreated = FOnPreviewSceneCreated::FDelegate::CreateSP(this, &FControlRigEditor::HandlePreviewSceneCreated);
+	PersonaToolkitArgs.OnPreviewSceneCreated = FOnPreviewSceneCreated::FDelegate::CreateSP(SharedRef(), &FControlRigBaseEditor::HandlePreviewSceneCreated);
 	PersonaToolkitArgs.bPreviewMeshCanUseDifferentSkeleton = true;
 	USkeleton* Skeleton = nullptr;
 	if(USkeletalMesh* PreviewMesh = ControlRigBlueprint->GetPreviewMesh())
@@ -643,20 +502,20 @@ void FControlRigEditor::CreatePersonaToolKitIfRequired()
 
 	// set delegate prior to setting mesh
 	// otherwise, you don't get delegate
-	PersonaToolkit->GetPreviewScene()->RegisterOnPreviewMeshChanged(FOnPreviewMeshChanged::CreateSP(this, &FControlRigEditor::HandlePreviewMeshChanged));
+	PersonaToolkit->GetPreviewScene()->RegisterOnPreviewMeshChanged(FOnPreviewMeshChanged::CreateSP(SharedRef(), &FControlRigBaseEditor::HandlePreviewMeshChanged));
 	
 	// Set a default preview mesh, if any
 	TGuardValue<bool> AutoResolveGuard(ControlRigBlueprint->ModularRigSettings.bAutoResolve, false);
 	PersonaToolkit->SetPreviewMesh(ControlRigBlueprint->GetPreviewMesh(), false);
 }
 
-const FName FControlRigEditor::GetEditorAppName() const
+const FName FControlRigBaseEditor::GetEditorAppNameImpl() const
 {
 	static const FName ControlRigEditorAppName(TEXT("ControlRigEditorApp"));
 	return ControlRigEditorAppName;
 }
 
-const FName FControlRigEditor::GetEditorModeName() const
+const FName FControlRigBaseEditor::GetEditorModeNameImpl() const
 {
 	if(IsModularRig())
 	{
@@ -665,24 +524,13 @@ const FName FControlRigEditor::GetEditorModeName() const
 	return FControlRigEditorEditMode::ModeName;
 }
 
-TSharedPtr<FApplicationMode> FControlRigEditor::CreateEditorMode()
-{
-	CreatePersonaToolKitIfRequired();
-
-	if(IsModularRig())
-	{
-		return MakeShareable(new FModularRigEditorMode(SharedThis(this)));
-	}
-	return MakeShareable(new FControlRigEditorMode(SharedThis(this)));
-}
-
-const FSlateBrush* FControlRigEditor::GetDefaultTabIcon() const
+const FSlateBrush* FControlRigBaseEditor::GetDefaultTabIconImpl() const
 {
 	static const FSlateIcon TabIcon = FSlateIcon(FControlRigEditorStyle::Get().GetStyleSetName(), "ControlRig.Editor.TabIcon");
 	return TabIcon.GetIcon();
 }
 
-FText FControlRigEditor::GetReplayAssetName() const
+FText FControlRigBaseEditor::GetReplayAssetName() const
 {
 	if(ReplayStrongPtr.IsValid())
 	{
@@ -693,7 +541,7 @@ FText FControlRigEditor::GetReplayAssetName() const
 	return NoReplayAsset;
 }
 
-FText FControlRigEditor::GetReplayAssetTooltip() const
+FText FControlRigBaseEditor::GetReplayAssetTooltip() const
 {
 	if(ReplayStrongPtr.IsValid())
 	{
@@ -703,7 +551,7 @@ FText FControlRigEditor::GetReplayAssetTooltip() const
 	return NoReplayTooltip;
 }
 
-bool FControlRigEditor::SetReplayAssetPath(const FString& InAssetPath)
+bool FControlRigBaseEditor::SetReplayAssetPath(const FString& InAssetPath)
 {
 	if(ReplayStrongPtr.IsValid())
 	{
@@ -729,7 +577,7 @@ bool FControlRigEditor::SetReplayAssetPath(const FString& InAssetPath)
 	return true;
 }
 
-TSharedRef<SWidget> FControlRigEditor::GenerateReplayAssetModeMenuContent()
+TSharedRef<SWidget> FControlRigBaseEditor::GenerateReplayAssetModeMenuContent()
 {
 	FMenuBuilder MenuBuilder(true, GetToolkitCommands());
 
@@ -784,7 +632,7 @@ TSharedRef<SWidget> FControlRigEditor::GenerateReplayAssetModeMenuContent()
 	return MenuBuilder.MakeWidget();
 }
 
-TSharedRef<SWidget> FControlRigEditor::GenerateReplayAssetRecordMenuContent()
+TSharedRef<SWidget> FControlRigBaseEditor::GenerateReplayAssetRecordMenuContent()
 {
 	FMenuBuilder MenuBuilder(true, GetToolkitCommands());
 
@@ -838,7 +686,7 @@ TSharedRef<SWidget> FControlRigEditor::GenerateReplayAssetRecordMenuContent()
 	return MenuBuilder.MakeWidget();
 }
 
-TSharedRef<SWidget> FControlRigEditor::GenerateReplayAssetPlaybackMenuContent()
+TSharedRef<SWidget> FControlRigBaseEditor::GenerateReplayAssetPlaybackMenuContent()
 {
 	FMenuBuilder MenuBuilder(true, GetToolkitCommands());
 
@@ -890,7 +738,7 @@ TSharedRef<SWidget> FControlRigEditor::GenerateReplayAssetPlaybackMenuContent()
 	return MenuBuilder.MakeWidget();
 }
 
-bool FControlRigEditor::RecordReplay(double InRecordingDuration)
+bool FControlRigBaseEditor::RecordReplay(double InRecordingDuration)
 {
 	if(GetControlRig() == nullptr)
 	{
@@ -946,7 +794,7 @@ bool FControlRigEditor::RecordReplay(double InRecordingDuration)
 	return true;
 }
 
-void FControlRigEditor::ToggleReplay()
+void FControlRigBaseEditor::ToggleReplay()
 {
 	if(ReplayStrongPtr.IsValid())
 	{
@@ -971,7 +819,7 @@ void FControlRigEditor::ToggleReplay()
 	}
 }
 
-EControlRigReplayPlaybackMode FControlRigEditor::GetReplayPlaybackMode() const
+EControlRigReplayPlaybackMode FControlRigBaseEditor::GetReplayPlaybackMode() const
 {
 	if (ReplayStrongPtr.IsValid())
 	{
@@ -980,9 +828,9 @@ EControlRigReplayPlaybackMode FControlRigEditor::GetReplayPlaybackMode() const
 	return EControlRigReplayPlaybackMode::Live;
 }
 
-void FControlRigEditor::FillToolbar(FToolBarBuilder& ToolbarBuilder, bool bEndSection)
+void FControlRigBaseEditor::FillToolbarImpl(FToolBarBuilder& ToolbarBuilder, bool bEndSection)
 {
-	IControlRigEditor::FillToolbar(ToolbarBuilder, false);
+	FillToolbarSuper(ToolbarBuilder, false);
 	
 	{
 		if(CVarControlRigHierarchyEnableModules.GetValueOnAnyThread())
@@ -1106,9 +954,9 @@ void FControlRigEditor::FillToolbar(FToolBarBuilder& ToolbarBuilder, bool bEndSe
 
 			ToolbarBuilder.AddComboButton(
 				EmptyOnlyWhenNotRecordingAction,
-				FOnGetContent::CreateSP(this, &FControlRigEditor::GenerateReplayAssetModeMenuContent),
-				TAttribute<FText>::CreateSP(this, &FControlRigEditor::GetReplayAssetName),
-				TAttribute<FText>::CreateSP(this, &FControlRigEditor::GetReplayAssetTooltip),
+				FOnGetContent::CreateSP(SharedRef(), &FControlRigBaseEditor::GenerateReplayAssetModeMenuContent),
+				TAttribute<FText>::CreateRaw(this, &FControlRigBaseEditor::GetReplayAssetName),
+				TAttribute<FText>::CreateRaw(this, &FControlRigBaseEditor::GetReplayAssetTooltip),
 				FSlateIcon(FAppStyle::GetAppStyleSetName(), "AutomationTools.TestAutomation"),
 				false);
 
@@ -1126,7 +974,7 @@ void FControlRigEditor::FillToolbar(FToolBarBuilder& ToolbarBuilder, bool bEndSe
 				);
 			ToolbarBuilder.AddComboButton(
 				EmptyOnlyWhenNotRecordingAction,
-				FOnGetContent::CreateSP(this, &FControlRigEditor::GenerateReplayAssetRecordMenuContent),
+				FOnGetContent::CreateSP(SharedRef(), &FControlRigBaseEditor::GenerateReplayAssetRecordMenuContent),
 				LOCTEXT("ReplayRecordMenu_Label", "Recording Modes"),
 				LOCTEXT("ReplayRecordMenu_ToolTip", "Pick between different modes for recording"),
 				FSlateIcon(FAppStyle::GetAppStyleSetName(), "LevelEditor.Recompile"),
@@ -1199,7 +1047,7 @@ void FControlRigEditor::FillToolbar(FToolBarBuilder& ToolbarBuilder, bool bEndSe
 
 			ToolbarBuilder.AddComboButton(
 				EmptyOnlyWithValidReplayAction,
-				FOnGetContent::CreateSP(this, &FControlRigEditor::GenerateReplayAssetPlaybackMenuContent),
+				FOnGetContent::CreateSP(SharedRef(), &FControlRigBaseEditor::GenerateReplayAssetPlaybackMenuContent),
 				LOCTEXT("ReplayPlaybackModeMenu_Label", "Playback Modes"),
 				LOCTEXT("ReplayPlaybackModeMenu_ToolTip", "Pick between different modes for playback"),
 				FSlateIcon(FAppStyle::GetAppStyleSetName(), "LevelEditor.Recompile"),
@@ -1213,12 +1061,12 @@ void FControlRigEditor::FillToolbar(FToolBarBuilder& ToolbarBuilder, bool bEndSe
 	}
 }
 
-TArray<FName> FControlRigEditor::GetDefaultEventQueue() const
+TArray<FName> FControlRigBaseEditor::GetDefaultEventQueueImpl() const
 {
 	return ForwardsSolveEventQueue;
 }
 
-void FControlRigEditor::SetEventQueue(TArray<FName> InEventQueue, bool bCompile)
+void FControlRigBaseEditor::SetEventQueueImpl(TArray<FName> InEventQueue, bool bCompile)
 {
 	if (GetEventQueue() == InEventQueue)
 	{
@@ -1226,7 +1074,7 @@ void FControlRigEditor::SetEventQueue(TArray<FName> InEventQueue, bool bCompile)
 	}
 
 	TArray<FRigElementKey> PreviousSelection;
-	if (UControlRigBlueprint* RigBlueprint = Cast<UControlRigBlueprint>(FRigVMEditorBase::GetBlueprintObj()))
+	if (UControlRigBlueprint* RigBlueprint = Cast<UControlRigBlueprint>(GetRigVMBlueprint()))
 	{
 		if(bCompile)
 		{
@@ -1249,17 +1097,17 @@ void FControlRigEditor::SetEventQueue(TArray<FName> InEventQueue, bool bCompile)
 		}
 	}
 
-	IControlRigEditor::SetEventQueue(InEventQueue, bCompile);
+	SetEventQueueSuper(InEventQueue, bCompile);
 
 	if (UControlRig* ControlRig = GetControlRig())
 	{
 		if (InEventQueue.Num() > 0)
 		{
-			if (UControlRigBlueprint* RigBlueprint = Cast<UControlRigBlueprint>(FRigVMEditorBase::GetBlueprintObj()))
+			if (UControlRigBlueprint* RigBlueprint = Cast<UControlRigBlueprint>(GetRigVMBlueprint()))
 			{
 				RigBlueprint->Validator->SetControlRig(ControlRig);
 
-				if (LastEventQueue == ConstructionEventQueue)
+				if (GetLastEventQueue() == ConstructionEventQueue)
 				{
 					// This will propagate any user bone transformation done during construction to the preview instance
 					ResetAllBoneModification();
@@ -1293,7 +1141,7 @@ void FControlRigEditor::SetEventQueue(TArray<FName> InEventQueue, bool bCompile)
 	}
 }
 
-int32 FControlRigEditor::GetEventQueueComboValue() const
+int32 FControlRigBaseEditor::GetEventQueueComboValueImpl() const
 {
 	const TArray<FName> EventQueue = GetEventQueue();
 	if(EventQueue == ForwardsSolveEventQueue)
@@ -1312,10 +1160,10 @@ int32 FControlRigEditor::GetEventQueueComboValue() const
 	{
 		return 3;
 	}
-	return IControlRigEditor::GetEventQueueComboValue();
+	return GetEventQueueComboValueSuper();
 }
 
-FText FControlRigEditor::GetEventQueueLabel() const
+FText FControlRigBaseEditor::GetEventQueueLabelImpl() const
 {
 	TArray<FName> EventQueue = GetEventQueue();
 
@@ -1351,7 +1199,7 @@ FText FControlRigEditor::GetEventQueueLabel() const
 	return LOCTEXT("CustomEventQueue", "Custom Event Queue");
 }
 
-FSlateIcon FControlRigEditor::GetEventQueueIcon(const TArray<FName>& InEventQueue) const
+FSlateIcon FControlRigBaseEditor::GetEventQueueIconImpl(const TArray<FName>& InEventQueue) const
 {
 	if(InEventQueue == ConstructionEventQueue)
 	{
@@ -1373,12 +1221,12 @@ FSlateIcon FControlRigEditor::GetEventQueueIcon(const TArray<FName>& InEventQueu
 	return FSlateIcon(FAppStyle::GetAppStyleSetName(), "GraphEditor.Event_16x");
 }
 
-void FControlRigEditor::HandleSetObjectBeingDebugged(UObject* InObject)
+void FControlRigBaseEditor::HandleSetObjectBeingDebuggedImpl(UObject* InObject)
 {
-	FRigVMEditorBase::HandleSetObjectBeingDebugged(InObject);
+	HandleSetObjectBeingDebuggedSuper(InObject);
 	
 	UControlRig* DebuggedControlRig = Cast<UControlRig>(InObject);
-	if(UControlRig* PreviouslyDebuggedControlRig = Cast<UControlRig>(FRigVMEditorBase::GetBlueprintObj()->GetObjectBeingDebugged()))
+	if(UControlRig* PreviouslyDebuggedControlRig = Cast<UControlRig>(GetRigVMBlueprint()->GetObjectBeingDebugged()))
 	{
 		if(!URigVMHost::IsGarbageOrDestroyed(PreviouslyDebuggedControlRig))
 		{
@@ -1391,7 +1239,7 @@ void FControlRigEditor::HandleSetObjectBeingDebugged(UObject* InObject)
 		}
 	}
 
-	if (UControlRigBlueprint* RigBlueprint = Cast<UControlRigBlueprint>(FRigVMEditorBase::GetBlueprintObj()))
+	if (UControlRigBlueprint* RigBlueprint = Cast<UControlRigBlueprint>(GetRigVMBlueprint()))
 	{
 		RigBlueprint->Validator->SetControlRig(DebuggedControlRig);
 	}
@@ -1431,7 +1279,7 @@ void FControlRigEditor::HandleSetObjectBeingDebugged(UObject* InObject)
 			if(bShouldExecute)
 			{
 				DebuggedControlRig->SetBoneInitialTransformsFromSkeletalMeshComponent(EditorSkelComp);
-				if(UControlRigBlueprint* RigBlueprint = Cast<UControlRigBlueprint>(FRigVMEditorBase::GetBlueprintObj()))
+				if(UControlRigBlueprint* RigBlueprint = Cast<UControlRigBlueprint>(GetRigVMBlueprint()))
 				{
 					// copy the initial transforms back to the blueprint
 					// no need to call modify here since this code only modifies the bp if the preview mesh changed
@@ -1447,12 +1295,12 @@ void FControlRigEditor::HandleSetObjectBeingDebugged(UObject* InObject)
 		DebuggedControlRig->OnPostConstruction_AnyThread().RemoveAll(this);
 		DebuggedControlRig->ControlModified().RemoveAll(this);
 
-		DebuggedControlRig->GetHierarchy()->OnModified().AddSP(this, &FControlRigEditor::OnHierarchyModified_AnyThread);
-		DebuggedControlRig->OnPreForwardsSolve_AnyThread().AddSP(this, &FControlRigEditor::OnPreForwardsSolve_AnyThread);
-		DebuggedControlRig->OnPreConstructionForUI_AnyThread().AddSP(this, &FControlRigEditor::OnPreConstructionForUI_AnyThread);
-		DebuggedControlRig->OnPreConstruction_AnyThread().AddSP(this, &FControlRigEditor::OnPreConstruction_AnyThread);
-		DebuggedControlRig->OnPostConstruction_AnyThread().AddSP(this, &FControlRigEditor::OnPostConstruction_AnyThread);
-		DebuggedControlRig->ControlModified().AddSP(this, &FControlRigEditor::HandleOnControlModified);
+		DebuggedControlRig->GetHierarchy()->OnModified().AddSP(SharedRef(), &FControlRigBaseEditor::OnHierarchyModified_AnyThread);
+		DebuggedControlRig->OnPreForwardsSolve_AnyThread().AddSP(SharedRef(), &FControlRigBaseEditor::OnPreForwardsSolve_AnyThread);
+		DebuggedControlRig->OnPreConstructionForUI_AnyThread().AddSP(SharedRef(), &FControlRigBaseEditor::OnPreConstructionForUI_AnyThread);
+		DebuggedControlRig->OnPreConstruction_AnyThread().AddSP(SharedRef(), &FControlRigBaseEditor::OnPreConstruction_AnyThread);
+		DebuggedControlRig->OnPostConstruction_AnyThread().AddSP(SharedRef(), &FControlRigBaseEditor::OnPostConstruction_AnyThread);
+		DebuggedControlRig->ControlModified().AddSP(SharedRef(), &FControlRigBaseEditor::HandleOnControlModified);
 		
 		LastHierarchyHash = INDEX_NONE;
 
@@ -1478,13 +1326,13 @@ void FControlRigEditor::HandleSetObjectBeingDebugged(UObject* InObject)
 	}
 }
 
-void FControlRigEditor::SetDetailViewForRigElements()
+void FControlRigBaseEditor::SetDetailViewForRigElements()
 {
 	URigHierarchy* HierarchyBeingDebugged = GetHierarchyBeingDebugged();
 	SetDetailViewForRigElements(HierarchyBeingDebugged->GetSelectedHierarchyKeys());
 }
 
-void FControlRigEditor::SetDetailViewForRigElements(const TArray<FRigHierarchyKey>& InKeys)
+void FControlRigBaseEditor::SetDetailViewForRigElements(const TArray<FRigHierarchyKey>& InKeys)
 {
 	if(IsDetailsPanelRefreshSuspended())
 	{
@@ -1519,7 +1367,7 @@ void FControlRigEditor::SetDetailViewForRigElements(const TArray<FRigHierarchyKe
 
 	ClearDetailObject();
 
-	UControlRigBlueprint* RigBlueprint = Cast<UControlRigBlueprint>(FRigVMEditorBase::GetBlueprintObj());
+	UControlRigBlueprint* RigBlueprint = Cast<UControlRigBlueprint>(GetRigVMBlueprint());
 	URigHierarchy* HierarchyBeingDebugged = GetHierarchyBeingDebugged();
 	TArray<UObject*> Objects;
 
@@ -1533,8 +1381,8 @@ void FControlRigEditor::SetDetailViewForRigElements(const TArray<FRigHierarchyKe
 				continue;
 			}
 
-			URigVMDetailsViewWrapperObject* WrapperObject = URigVMDetailsViewWrapperObject::MakeInstance(GetDetailWrapperClass(), FRigVMEditorBase::GetBlueprintObj(), Element->GetScriptStruct(), (uint8*)Element, HierarchyBeingDebugged);
-			WrapperObject->GetWrappedPropertyChangedChainEvent().AddSP(this, &FControlRigEditor::OnWrappedPropertyChangedChainEvent);
+			URigVMDetailsViewWrapperObject* WrapperObject = URigVMDetailsViewWrapperObject::MakeInstance(GetDetailWrapperClass(), GetRigVMBlueprint(), Element->GetScriptStruct(), (uint8*)Element, HierarchyBeingDebugged);
+			WrapperObject->GetWrappedPropertyChangedChainEvent().AddSP(SharedRef(), &FControlRigBaseEditor::OnWrappedPropertyChangedChainEvent);
 			WrapperObject->AddToRoot();
 
 			Objects.Add(WrapperObject);
@@ -1547,8 +1395,8 @@ void FControlRigEditor::SetDetailViewForRigElements(const TArray<FRigHierarchyKe
 				continue;
 			}
 
-			URigVMDetailsViewWrapperObject* WrapperObject = URigVMDetailsViewWrapperObject::MakeInstance(GetDetailWrapperClass(), FRigVMEditorBase::GetBlueprintObj(), Component->GetScriptStruct(), (uint8*)Component, HierarchyBeingDebugged);
-			WrapperObject->GetWrappedPropertyChangedChainEvent().AddSP(this, &FControlRigEditor::OnWrappedPropertyChangedChainEvent);
+			URigVMDetailsViewWrapperObject* WrapperObject = URigVMDetailsViewWrapperObject::MakeInstance(GetDetailWrapperClass(), GetRigVMBlueprint(), Component->GetScriptStruct(), (uint8*)Component, HierarchyBeingDebugged);
+			WrapperObject->GetWrappedPropertyChangedChainEvent().AddSP(SharedRef(), &FControlRigBaseEditor::OnWrappedPropertyChangedChainEvent);
 			WrapperObject->AddToRoot();
 
 			Objects.Add(WrapperObject);
@@ -1558,7 +1406,7 @@ void FControlRigEditor::SetDetailViewForRigElements(const TArray<FRigHierarchyKe
 	SetDetailObjects(Objects);
 }
 
-void FControlRigEditor::SetDetailObjects(const TArray<UObject*>& InObjects)
+void FControlRigBaseEditor::SetDetailObjectsImpl(const TArray<UObject*>& InObjects)
 {
 	// if no modules should be selected - we need to deselect all modules
 	if (!InObjects.ContainsByPredicate([](const UObject* InObject) -> bool
@@ -1569,10 +1417,10 @@ void FControlRigEditor::SetDetailObjects(const TArray<UObject*>& InObjects)
 		ModulesSelected.Reset();
 	}
 	
-	IControlRigEditor::SetDetailObjects(InObjects);
+	SetDetailObjectsSuper(InObjects);
 }
 
-void FControlRigEditor::RefreshDetailView()
+void FControlRigBaseEditor::RefreshDetailViewImpl()
 {
 	if(DetailViewShowsAnyRigElement())
 	{
@@ -1585,15 +1433,15 @@ void FControlRigEditor::RefreshDetailView()
 		return;
 	}
 
-	IControlRigEditor::RefreshDetailView();
+	RefreshDetailViewSuper();
 }
 
-bool FControlRigEditor::DetailViewShowsAnyRigElement() const
+bool FControlRigBaseEditor::DetailViewShowsAnyRigElement() const
 {
 	return DetailViewShowsStruct(FRigBaseElement::StaticStruct()) || DetailViewShowsStruct(FRigBaseComponent::StaticStruct());
 }
 
-bool FControlRigEditor::DetailViewShowsRigElement(FRigHierarchyKey InKey) const
+bool FControlRigBaseEditor::DetailViewShowsRigElement(FRigHierarchyKey InKey) const
 {
 	TArray< TWeakObjectPtr<UObject> > SelectedObjects = GetSelectedObjectsFromDetailView();
 	for (TWeakObjectPtr<UObject> SelectedObject : SelectedObjects)
@@ -1625,7 +1473,7 @@ bool FControlRigEditor::DetailViewShowsRigElement(FRigHierarchyKey InKey) const
 	return false;
 }
 
-TArray<FRigHierarchyKey> FControlRigEditor::GetSelectedRigElementsFromDetailView() const
+TArray<FRigHierarchyKey> FControlRigBaseEditor::GetSelectedRigElementsFromDetailView() const
 {
 	TArray<FRigHierarchyKey> Keys;
 	
@@ -1654,21 +1502,31 @@ TArray<FRigHierarchyKey> FControlRigEditor::GetSelectedRigElementsFromDetailView
 	return Keys;
 }
 
-TArray<TWeakObjectPtr<UObject>> FControlRigEditor::GetSelectedObjectsFromDetailView() const
+TArray<TWeakObjectPtr<UObject>> FControlRigBaseEditor::GetSelectedObjectsFromDetailView() const
 {
+	TSharedPtr<class SWidget> Inspector = GetInspector();
 	if (Inspector.IsValid())
 	{
-		return Inspector->GetSelectedObjects();
+#if WITH_RIGVMLEGACYEDITOR
+		if (IsControlRigLegacyEditor())
+		{
+			TSharedPtr<class SKismetInspector> KismetInspector = StaticCastSharedPtr<SKismetInspector>(Inspector);
+			return KismetInspector->GetSelectedObjects();
+		}
+		
+#endif
+		TSharedPtr<class SRigVMDetailsInspector> RigVMInspector = StaticCastSharedPtr<SRigVMDetailsInspector>(Inspector);
+		return RigVMInspector->GetSelectedObjects();
 	}
 	return TArray<TWeakObjectPtr<UObject>>();
 }
 
-void FControlRigEditor::SetDetailViewForRigModules()
+void FControlRigBaseEditor::SetDetailViewForRigModules()
 {
 	SetDetailViewForRigModules(ModulesSelected);
 }
 
-void FControlRigEditor::SetDetailViewForRigModules(const TArray<FName> InModuleNames)
+void FControlRigBaseEditor::SetDetailViewForRigModules(const TArray<FName> InModuleNames)
 {
 	if(IsDetailsPanelRefreshSuspended())
 	{
@@ -1677,7 +1535,7 @@ void FControlRigEditor::SetDetailViewForRigModules(const TArray<FName> InModuleN
 
 	ClearDetailObject();
 
-	UControlRigBlueprint* RigBlueprint = Cast<UControlRigBlueprint>(FRigVMEditorBase::GetBlueprintObj());
+	UControlRigBlueprint* RigBlueprint = Cast<UControlRigBlueprint>(GetRigVMBlueprint());
 	UModularRig* RigBeingDebugged = Cast<UModularRig>(RigBlueprint->GetDebuggedControlRig());
 
 	if (!RigBeingDebugged)
@@ -1711,12 +1569,12 @@ void FControlRigEditor::SetDetailViewForRigModules(const TArray<FName> InModuleN
 	}
 }
 
-bool FControlRigEditor::DetailViewShowsAnyRigModule() const
+bool FControlRigBaseEditor::DetailViewShowsAnyRigModule() const
 {
 	return DetailViewShowsStruct(FRigModuleInstance::StaticStruct());
 }
 
-bool FControlRigEditor::DetailViewShowsRigModule(FName InModuleName) const
+bool FControlRigBaseEditor::DetailViewShowsRigModule(FName InModuleName) const
 {
 	TArray< TWeakObjectPtr<UObject> > SelectedObjects = GetSelectedObjectsFromDetailView();
 	for (TWeakObjectPtr<UObject> SelectedObject : SelectedObjects)
@@ -1741,7 +1599,7 @@ bool FControlRigEditor::DetailViewShowsRigModule(FName InModuleName) const
 	return false;
 }
 
-void FControlRigEditor::Compile()
+void FControlRigBaseEditor::CompileBaseImpl()
 {
 	{
 		DECLARE_SCOPE_HIERARCHICAL_COUNTER_FUNC()
@@ -1784,7 +1642,7 @@ void FControlRigEditor::Compile()
 		}
 
 		{
-			FRigVMEditorBase::Compile();
+			CompileSuper();
 		}
 
 		ControlRigBlueprint->RecompileModularRig();
@@ -1843,9 +1701,9 @@ void FControlRigEditor::Compile()
 	}
 }
 
-void FControlRigEditor::HandleModifiedEvent(ERigVMGraphNotifType InNotifType, URigVMGraph* InGraph, UObject* InSubject)
+void FControlRigBaseEditor::HandleModifiedEventImpl(ERigVMGraphNotifType InNotifType, URigVMGraph* InGraph, UObject* InSubject)
 {
-	IControlRigEditor::HandleModifiedEvent(InNotifType, InGraph, InSubject);
+	HandleModifiedEventSuper(InNotifType, InGraph, InSubject);
 
 	if(InNotifType == ERigVMGraphNotifType::NodeSelected)
 	{
@@ -1896,24 +1754,26 @@ void FControlRigEditor::HandleModifiedEvent(ERigVMGraphNotifType InNotifType, UR
 	}
 }
 
-void FControlRigEditor::OnCreateGraphEditorCommands(TSharedPtr<FUICommandList> GraphEditorCommandsList)
+void FControlRigBaseEditor::OnCreateGraphEditorCommandsImpl(TSharedPtr<FUICommandList> GraphEditorCommandsList)
 {
-	FRigVMEditorBase::OnCreateGraphEditorCommands(GraphEditorCommandsList);
+	OnCreateGraphEditorCommandsSuper(GraphEditorCommandsList);
+
+	TSharedRef<FControlRigBaseEditor> EditorRef = SharedRef();
 
 	GraphEditorCommandsList->MapAction(
 		FControlRigEditorCommands::Get().RequestDirectManipulationPosition,
-		FExecuteAction::CreateSP(this, &FControlRigEditor::HandleRequestDirectManipulationPosition));
+		FExecuteAction::CreateSP(EditorRef, &FControlRigBaseEditor::HandleRequestDirectManipulationPosition));
 	GraphEditorCommandsList->MapAction(
 		FControlRigEditorCommands::Get().RequestDirectManipulationRotation,
-		FExecuteAction::CreateSP(this, &FControlRigEditor::HandleRequestDirectManipulationRotation));
+		FExecuteAction::CreateSP(EditorRef, &FControlRigBaseEditor::HandleRequestDirectManipulationRotation));
 	GraphEditorCommandsList->MapAction(
 		FControlRigEditorCommands::Get().RequestDirectManipulationScale,
-		FExecuteAction::CreateSP(this, &FControlRigEditor::HandleRequestDirectManipulationScale));
+		FExecuteAction::CreateSP(EditorRef, &FControlRigBaseEditor::HandleRequestDirectManipulationScale));
 }
 
-void FControlRigEditor::HandleVMCompiledEvent(UObject* InCompiledObject, URigVM* InVM, FRigVMExtendedExecuteContext& InContext)
+void FControlRigBaseEditor::HandleVMCompiledEventImpl(UObject* InCompiledObject, URigVM* InVM, FRigVMExtendedExecuteContext& InContext)
 {
-	IControlRigEditor::HandleVMCompiledEvent(InCompiledObject, InVM, InContext);
+	HandleVMCompiledEventSuper(InCompiledObject, InVM, InContext);
 
 	if(bRefreshDirectionManipulationTargetsRequired)
 	{
@@ -1930,12 +1790,12 @@ void FControlRigEditor::HandleVMCompiledEvent(UObject* InCompiledObject, URigVM*
 	}
 }
 
-void FControlRigEditor::SaveAsset_Execute()
+void FControlRigBaseEditor::SaveAsset_ExecuteImpl()
 {
-	FRigVMEditorBase::SaveAsset_Execute();
+	SaveAsset_ExecuteSuper();
 
 	// Save the new state of the hierarchy in the default object, so that it has the correct values on load
-	UControlRigBlueprint* RigBlueprint = Cast<UControlRigBlueprint>(FRigVMEditorBase::GetBlueprintObj());
+	UControlRigBlueprint* RigBlueprint = Cast<UControlRigBlueprint>(GetRigVMBlueprint());
 	if(const UControlRig* ControlRig = GetControlRig())
 	{
 		UControlRig* CDO = ControlRig->GetClass()->GetDefaultObject<UControlRig>();
@@ -1948,12 +1808,12 @@ void FControlRigEditor::SaveAsset_Execute()
 	ActionDatabase.RefreshClassActions(UControlRigBlueprint::StaticClass());
 }
 
-void FControlRigEditor::SaveAssetAs_Execute()
+void FControlRigBaseEditor::SaveAssetAs_ExecuteImpl()
 {
-	FRigVMEditorBase::SaveAssetAs_Execute();
+	SaveAssetAs_ExecuteSuper();
 
 	// Save the new state of the hierarchy in the default object, so that it has the correct values on load
-	UControlRigBlueprint* RigBlueprint = Cast<UControlRigBlueprint>(FRigVMEditorBase::GetBlueprintObj());
+	UControlRigBlueprint* RigBlueprint = Cast<UControlRigBlueprint>(GetRigVMBlueprint());
 	if(const UControlRig* ControlRig = GetControlRig())
 	{
 		UControlRig* CDO = ControlRig->GetClass()->GetDefaultObject<UControlRig>();
@@ -1966,42 +1826,42 @@ void FControlRigEditor::SaveAssetAs_Execute()
 	ActionDatabase.RefreshClassActions(UControlRigBlueprint::StaticClass());
 }
 
-bool FControlRigEditor::IsModularRig() const
+bool FControlRigBaseEditor::IsModularRig() const
 {
-	if(UControlRigBlueprint* RigBlueprint = Cast<UControlRigBlueprint>(FRigVMEditorBase::GetBlueprintObj()))
+	if(UControlRigBlueprint* RigBlueprint = Cast<UControlRigBlueprint>(GetRigVMBlueprint()))
 	{
 		return RigBlueprint->IsModularRig();
 	}
 	return false;
 }
 
-bool FControlRigEditor::IsRigModule() const
+bool FControlRigBaseEditor::IsRigModule() const
 {
-	if(UControlRigBlueprint* RigBlueprint = Cast<UControlRigBlueprint>(FRigVMEditorBase::GetBlueprintObj()))
+	if(UControlRigBlueprint* RigBlueprint = Cast<UControlRigBlueprint>(GetRigVMBlueprint()))
 	{
 		return RigBlueprint->IsControlRigModule();
 	}
 	return false;
 }
 
-FName FControlRigEditor::GetToolkitFName() const
+FName FControlRigBaseEditor::GetToolkitFNameImpl() const
 {
 	return FName("ControlRigEditor");
 }
 
-FText FControlRigEditor::GetBaseToolkitName() const
+FText FControlRigBaseEditor::GetBaseToolkitNameImpl() const
 {
 	return LOCTEXT("AppLabel", "Control Rig Editor");
 }
 
-FString FControlRigEditor::GetWorldCentricTabPrefix() const
+FString FControlRigBaseEditor::GetWorldCentricTabPrefixImpl() const
 {
 	return LOCTEXT("WorldCentricTabPrefix", "Control Rig Editor ").ToString();
 }
 
-FReply FControlRigEditor::OnSpawnGraphNodeByShortcut(FInputChord InChord, const FVector2D& InPosition, UEdGraph* InGraph)
+FReply FControlRigBaseEditor::OnSpawnGraphNodeByShortcutImpl(FInputChord InChord, const FVector2D& InPosition, UEdGraph* InGraph)
 {
-	const FReply SuperReply = FRigVMEditorBase::OnSpawnGraphNodeByShortcut(InChord, InPosition, InGraph);
+	const FReply SuperReply = OnSpawnGraphNodeByShortcutSuper(InChord, InPosition, InGraph);
 	if(SuperReply.IsEventHandled())
 	{
 		return SuperReply;
@@ -2044,11 +1904,11 @@ FReply FControlRigEditor::OnSpawnGraphNodeByShortcut(FInputChord InChord, const
 	return FReply::Unhandled();
 }
 
-void FControlRigEditor::PostTransaction(bool bSuccess, const FTransaction* Transaction, bool bIsRedo)
+void FControlRigBaseEditor::PostTransactionImpl(bool bSuccess, const FTransaction* Transaction, bool bIsRedo)
 {
 	EnsureValidRigElementsInDetailPanel();
 
-	if (UControlRigBlueprint* RigBlueprint = Cast<UControlRigBlueprint>(FRigVMEditorBase::GetBlueprintObj()))
+	if (UControlRigBlueprint* RigBlueprint = Cast<UControlRigBlueprint>(GetRigVMBlueprint()))
 	{
 		// Do not compile here. ControlRigBlueprint::PostTransacted decides when it is necessary to compile depending
 		// on the properties that are affected.
@@ -2081,7 +1941,7 @@ void FControlRigEditor::PostTransaction(bool bSuccess, const FTransaction* Trans
 	}
 }
 
-void FControlRigEditor::EnsureValidRigElementsInDetailPanel()
+void FControlRigBaseEditor::EnsureValidRigElementsInDetailPanel()
 {
 	UControlRigBlueprint* ControlRigBP = GetControlRigBlueprint();
 	URigHierarchy* Hierarchy = ControlRigBP->Hierarchy; 
@@ -2109,7 +1969,7 @@ void FControlRigEditor::EnsureValidRigElementsInDetailPanel()
 	}
 }
 
-void FControlRigEditor::OnAnimInitialized()
+void FControlRigBaseEditor::OnAnimInitialized()
 {
 	UControlRigSkeletalMeshComponent* EditorSkelComp = Cast<UControlRigSkeletalMeshComponent>(GetPersonaToolkit()->GetPreviewScene()->GetPreviewMeshComponent());
 	if (EditorSkelComp)
@@ -2128,9 +1988,9 @@ void FControlRigEditor::OnAnimInitialized()
 	}
 }
 
-void FControlRigEditor::HandleVMExecutedEvent(URigVMHost* InHost, const FName& InEventName)
+void FControlRigBaseEditor::HandleVMExecutedEventImpl(URigVMHost* InHost, const FName& InEventName)
 {
-	IControlRigEditor::HandleVMExecutedEvent(InHost, InEventName);
+	HandleVMExecutedEventSuper(InHost, InEventName);
 
 	if (UControlRigBlueprint* ControlRigBP = GetControlRigBlueprint())
 	{
@@ -2238,14 +2098,14 @@ void FControlRigEditor::HandleVMExecutedEvent(URigVMHost* InHost, const FName& I
 	}
 }
 
-void FControlRigEditor::CreateEditorModeManager()
+void FControlRigBaseEditor::CreateEditorModeManagerImpl()
 {
-	EditorModeManager = MakeShareable(FModuleManager::LoadModuleChecked<FPersonaModule>("Persona").CreatePersonaEditorModeManager());
+	SetEditorModeManager(MakeShareable(FModuleManager::LoadModuleChecked<FPersonaModule>("Persona").CreatePersonaEditorModeManager()));
 }
 
-void FControlRigEditor::Tick(float DeltaTime)
+void FControlRigBaseEditor::TickImpl(float DeltaTime)
 {
-	IControlRigEditor::Tick(DeltaTime);
+	TickSuper(DeltaTime);
 
 	bool bDrawHierarchyBones = false;
 
@@ -2300,7 +2160,7 @@ void FControlRigEditor::Tick(float DeltaTime)
 	}
 }
 
-void FControlRigEditor::HandleViewportCreated(const TSharedRef<class IPersonaViewport>& InViewport)
+void FControlRigBaseEditor::HandleViewportCreated(const TSharedRef<class IPersonaViewport>& InViewport)
 {
 	DECLARE_SCOPE_HIERARCHICAL_COUNTER_FUNC()
 
@@ -2309,7 +2169,7 @@ void FControlRigEditor::HandleViewportCreated(const TSharedRef<class IPersonaVie
 	// TODO: this is duplicated code from FAnimBlueprintEditor, would be nice to consolidate. 
 	auto GetCompilationStateText = [this]()
 	{
-		if (UBlueprint* Blueprint = FRigVMEditorBase::GetBlueprintObj())
+		if (UBlueprint* Blueprint = GetRigVMBlueprint())
 		{
 			switch (Blueprint->Status)
 			{
@@ -2358,7 +2218,7 @@ void FControlRigEditor::HandleViewportCreated(const TSharedRef<class IPersonaVie
 
 	auto CompileBlueprint = [this]()
 	{
-		if (UBlueprint* Blueprint = FRigVMEditorBase::GetBlueprintObj())
+		if (UBlueprint* Blueprint = GetRigVMBlueprint())
 		{
 			if (!Blueprint->IsUpToDate())
 			{
@@ -2371,7 +2231,7 @@ void FControlRigEditor::HandleViewportCreated(const TSharedRef<class IPersonaVie
 
 	auto GetErrorSeverity = [this]()
 	{
-		if (UBlueprint* Blueprint = FRigVMEditorBase::GetBlueprintObj())
+		if (UBlueprint* Blueprint = GetRigVMBlueprint())
 		{
 			return (Blueprint->Status == BS_Error) ? EMessageSeverity::Error : EMessageSeverity::Warning;
 		}
@@ -2381,7 +2241,7 @@ void FControlRigEditor::HandleViewportCreated(const TSharedRef<class IPersonaVie
 
 	auto GetIcon = [this]()
 	{
-		if (UBlueprint* Blueprint = FRigVMEditorBase::GetBlueprintObj())
+		if (UBlueprint* Blueprint = GetRigVMBlueprint())
 		{
 			return (Blueprint->Status == BS_Error) ? FEditorFontGlyphs::Exclamation_Triangle : FEditorFontGlyphs::Eye;
 		}
@@ -2416,14 +2276,14 @@ void FControlRigEditor::HandleViewportCreated(const TSharedRef<class IPersonaVie
 	};
 
 	{
-		FPersonaViewportNotificationOptions DirectManipulationNotificationOptions(TAttribute<EVisibility>::CreateSP(this, &FControlRigEditor::GetDirectManipulationVisibility));
+		FPersonaViewportNotificationOptions DirectManipulationNotificationOptions(TAttribute<EVisibility>::CreateRaw(this, &FControlRigBaseEditor::GetDirectManipulationVisibility));
 		DirectManipulationNotificationOptions.OnGetBrushOverride = TAttribute<const FSlateBrush*>(FControlRigEditorStyle::Get().GetBrush("ControlRig.Viewport.Notification.DirectManipulation"));
 
 		InViewport->AddNotification(
 			EMessageSeverity::Info,
 			false,
 			SNew(SHorizontalBox)
-			.Visibility(this, &FControlRigEditor::GetDirectManipulationVisibility)
+			.Visibility(SharedRef(), &FControlRigBaseEditor::GetDirectManipulationVisibility)
 			.ToolTipText(LOCTEXT("DirectManipulation", "Direct Manipulation"))
 			+SHorizontalBox::Slot()
 			.FillWidth(1.0f)
@@ -2457,11 +2317,11 @@ void FControlRigEditor::HandleViewportCreated(const TSharedRef<class IPersonaVie
 								.Text(FText::FromString(*Item))
 							];
 					} )	
-					.OnSelectionChanged(this, &FControlRigEditor::OnDirectManipulationChanged)
+					.OnSelectionChanged(SharedRef(), &FControlRigBaseEditor::OnDirectManipulationChanged)
 					[
 						SNew(STextBlock)
 						.TextStyle(FAppStyle::Get(), "AnimViewport.MessageText")
-						.Text(this, &FControlRigEditor::GetDirectionManipulationText)
+						.Text(SharedRef(), &FControlRigBaseEditor::GetDirectionManipulationText)
 					]
 				]
 			],
@@ -2474,7 +2334,7 @@ void FControlRigEditor::HandleViewportCreated(const TSharedRef<class IPersonaVie
 			EMessageSeverity::Warning,
 			false,
 			SNew(SHorizontalBox)
-			.Visibility(this, &FControlRigEditor::GetConnectorWarningVisibility)
+			.Visibility(SharedRef(), &FControlRigBaseEditor::GetConnectorWarningVisibility)
 			.ToolTipText(LOCTEXT("ConnectorWarningTooltip", "This rig has unresolved connectors."))
 			+SHorizontalBox::Slot()
 			.FillWidth(1.0f)
@@ -2497,7 +2357,7 @@ void FControlRigEditor::HandleViewportCreated(const TSharedRef<class IPersonaVie
 				[
 					SNew(STextBlock)
 					.TextStyle(FAppStyle::Get(), "AnimViewport.MessageText")
-					.Text(this, &FControlRigEditor::GetConnectorWarningText)
+					.Text(SharedRef(), &FControlRigBaseEditor::GetConnectorWarningText)
 				]
 				+SHorizontalBox::Slot()
 				.AutoWidth()
@@ -2507,7 +2367,7 @@ void FControlRigEditor::HandleViewportCreated(const TSharedRef<class IPersonaVie
 					.ForegroundColor(FSlateColor::UseForeground())
 					.ButtonStyle(FAppStyle::Get(), "FlatButton.Primary")
 					.ToolTipText(LOCTEXT("ConnectorWarningNavigateTooltip", "Navigate to the first unresolved connector in the hierarchy"))
-					.OnClicked(this, &FControlRigEditor::OnNavigateToConnectorWarning)
+					.OnClicked(SharedRef(), &FControlRigBaseEditor::OnNavigateToConnectorWarning)
 					[
 						SNew(SHorizontalBox)
 						+SHorizontalBox::Slot()
@@ -2531,22 +2391,22 @@ void FControlRigEditor::HandleViewportCreated(const TSharedRef<class IPersonaVie
 					]
 				]
 			],
-			FPersonaViewportNotificationOptions(TAttribute<EVisibility>::CreateSP(this, &FControlRigEditor::GetConnectorWarningVisibility))
+			FPersonaViewportNotificationOptions(TAttribute<EVisibility>::CreateRaw(this, &FControlRigBaseEditor::GetConnectorWarningVisibility))
 		);
 	}
 
 	if(CVarControlRigShowTestingToolbar.GetValueOnAnyThread())
 	{
 		{
-			FPersonaViewportNotificationOptions ReplayValidationNotificationOptions(TAttribute<EVisibility>::CreateSP<>(this, &FControlRigEditor::GetReplayValidationErrorVisibility));
+			FPersonaViewportNotificationOptions ReplayValidationNotificationOptions(TAttribute<EVisibility>::CreateRaw<>(this, &FControlRigBaseEditor::GetReplayValidationErrorVisibility));
 			ReplayValidationNotificationOptions.OnGetBrushOverride = TAttribute<const FSlateBrush*>(FControlRigEditorStyle::Get().GetBrush("ControlRig.Viewport.Notification.ReplayValidation"));
 
 			InViewport->AddNotification(
 				EMessageSeverity::Info,
 				false,
 				SNew(SHorizontalBox)
-				.Visibility(this, &FControlRigEditor::GetReplayValidationErrorVisibility)
-				.ToolTipText(this, &FControlRigEditor::GetReplayValidationErrorTooltip)
+				.Visibility(SharedRef(), &FControlRigBaseEditor::GetReplayValidationErrorVisibility)
+				.ToolTipText(SharedRef(), &FControlRigBaseEditor::GetReplayValidationErrorTooltip)
 				+SHorizontalBox::Slot()
 				.FillWidth(1.0f)
 				.Padding(4.0f, 4.0f)
@@ -2839,7 +2699,7 @@ void FControlRigEditor::HandleViewportCreated(const TSharedRef<class IPersonaVie
 																.PaddingTop(60)
 																.PaddingBottom(60)
 																.PaddingInterNode(5)
-																.Visibility(this, &FControlRigEditor::GetSchematicOverlayVisibility)
+																.Visibility(SharedRef(), &FControlRigBaseEditor::GetSchematicOverlayVisibility)
 				;
 				InViewport->AddOverlayWidget(SchematicViewport.ToSharedRef());
 			}
@@ -2868,7 +2728,7 @@ void FControlRigEditor::HandleViewportCreated(const TSharedRef<class IPersonaVie
 	{
 		AnimViewportClient->OnSetBoneSize.BindLambda([this](float InBoneSize)
 		{
-			if (UControlRigBlueprint* RigBlueprint = Cast<UControlRigBlueprint>(FRigVMEditorBase::GetBlueprintObj()))
+			if (UControlRigBlueprint* RigBlueprint = Cast<UControlRigBlueprint>(GetRigVMBlueprint()))
 			{
 				RigBlueprint->Modify();
 				RigBlueprint->DebugBoneRadius = InBoneSize;
@@ -2877,7 +2737,7 @@ void FControlRigEditor::HandleViewportCreated(const TSharedRef<class IPersonaVie
 		
 		AnimViewportClient->OnGetBoneSize.BindLambda([this]() -> float
 		{
-			if (UControlRigBlueprint* RigBlueprint = Cast<UControlRigBlueprint>(FRigVMEditorBase::GetBlueprintObj()))
+			if (UControlRigBlueprint* RigBlueprint = Cast<UControlRigBlueprint>(GetRigVMBlueprint()))
 			{
 				return RigBlueprint->DebugBoneRadius;
 			}
@@ -2887,7 +2747,7 @@ void FControlRigEditor::HandleViewportCreated(const TSharedRef<class IPersonaVie
 	}
 }
 
-void FControlRigEditor::HandleToggleControlVisibility()
+void FControlRigBaseEditor::HandleToggleControlVisibility()
 {
 	if (FControlRigEditMode* EditMode = GetEditMode())
 	{
@@ -2895,7 +2755,7 @@ void FControlRigEditor::HandleToggleControlVisibility()
 	}
 }
 
-bool FControlRigEditor::AreControlsVisible() const
+bool FControlRigBaseEditor::AreControlsVisible() const
 {
 	if (FControlRigEditMode* EditMode = GetEditMode())
 	{
@@ -2904,7 +2764,7 @@ bool FControlRigEditor::AreControlsVisible() const
 	return false;
 }
 
-void FControlRigEditor::HandleToggleControlsAsOverlay()
+void FControlRigBaseEditor::HandleToggleControlsAsOverlay()
 {
 	if (FControlRigEditMode* EditMode = GetEditMode())
 	{
@@ -2914,7 +2774,7 @@ void FControlRigEditor::HandleToggleControlsAsOverlay()
 	}
 }
 
-bool FControlRigEditor::AreControlsAsOverlay() const
+bool FControlRigBaseEditor::AreControlsAsOverlay() const
 {
 	if (FControlRigEditMode* EditMode = GetEditMode())
 	{
@@ -2923,7 +2783,7 @@ bool FControlRigEditor::AreControlsAsOverlay() const
 	return false;
 }
 
-void FControlRigEditor::HandleToggleSchematicViewport()
+void FControlRigBaseEditor::HandleToggleSchematicViewport()
 {
 	if(SchematicViewport.IsValid())
 	{
@@ -2932,7 +2792,7 @@ void FControlRigEditor::HandleToggleSchematicViewport()
 	}
 }
 
-bool FControlRigEditor::IsSchematicViewportActive() const
+bool FControlRigBaseEditor::IsSchematicViewportActive() const
 {
 	if (SchematicViewport.IsValid())
 	{
@@ -2941,7 +2801,7 @@ bool FControlRigEditor::IsSchematicViewportActive() const
 	return false;
 }
 
-EVisibility FControlRigEditor::GetSchematicOverlayVisibility() const
+EVisibility FControlRigBaseEditor::GetSchematicOverlayVisibility() const
 {
 	if(bSchematicViewPortIsHidden)
 	{
@@ -2962,7 +2822,7 @@ EVisibility FControlRigEditor::GetSchematicOverlayVisibility() const
 	return EVisibility::SelfHitTestInvisible;
 }
 
-bool FControlRigEditor::GetToolbarDrawAxesOnSelection() const
+bool FControlRigBaseEditor::GetToolbarDrawAxesOnSelection() const
 {
 	if (const UControlRigEditModeSettings* Settings = GetDefault<UControlRigEditModeSettings>())
 	{
@@ -2971,7 +2831,7 @@ bool FControlRigEditor::GetToolbarDrawAxesOnSelection() const
 	return false;
 }
 
-void FControlRigEditor::HandleToggleToolbarDrawAxesOnSelection()
+void FControlRigBaseEditor::HandleToggleToolbarDrawAxesOnSelection()
 {
 	if (UControlRigEditModeSettings* Settings = GetMutableDefault<UControlRigEditModeSettings>())
 	{
@@ -2979,7 +2839,7 @@ void FControlRigEditor::HandleToggleToolbarDrawAxesOnSelection()
 	}
 }
 
-bool FControlRigEditor::IsToolbarDrawNullsEnabled() const
+bool FControlRigBaseEditor::IsToolbarDrawNullsEnabled() const
 {
 	if (const UControlRig* ControlRig = GetControlRig())
 	{
@@ -2991,7 +2851,7 @@ bool FControlRigEditor::IsToolbarDrawNullsEnabled() const
 	return false;
 }
 
-bool FControlRigEditor::GetToolbarDrawNulls() const
+bool FControlRigBaseEditor::GetToolbarDrawNulls() const
 {
 	if (const UControlRigEditModeSettings* Settings = GetDefault<UControlRigEditModeSettings>())
 	{
@@ -3000,7 +2860,7 @@ bool FControlRigEditor::GetToolbarDrawNulls() const
 	return false;
 }
 
-void FControlRigEditor::HandleToggleToolbarDrawNulls()
+void FControlRigBaseEditor::HandleToggleToolbarDrawNulls()
 {
 	if (UControlRigEditModeSettings* Settings = GetMutableDefault<UControlRigEditModeSettings>())
 	{
@@ -3008,7 +2868,7 @@ void FControlRigEditor::HandleToggleToolbarDrawNulls()
 	}
 }
 
-bool FControlRigEditor::IsToolbarDrawSocketsEnabled() const
+bool FControlRigBaseEditor::IsToolbarDrawSocketsEnabled() const
 {
 	if (const UControlRig* ControlRig = GetControlRig())
 	{
@@ -3020,7 +2880,7 @@ bool FControlRigEditor::IsToolbarDrawSocketsEnabled() const
 	return false;
 }
 
-bool FControlRigEditor::GetToolbarDrawSockets() const
+bool FControlRigBaseEditor::GetToolbarDrawSockets() const
 {
 	if (const UControlRigEditModeSettings* Settings = GetDefault<UControlRigEditModeSettings>())
 	{
@@ -3029,7 +2889,7 @@ bool FControlRigEditor::GetToolbarDrawSockets() const
 	return false;
 }
 
-void FControlRigEditor::HandleToggleToolbarDrawSockets()
+void FControlRigBaseEditor::HandleToggleToolbarDrawSockets()
 {
 	if (UControlRigEditModeSettings* Settings = GetMutableDefault<UControlRigEditModeSettings>())
 	{
@@ -3037,12 +2897,12 @@ void FControlRigEditor::HandleToggleToolbarDrawSockets()
 	}
 }
 
-bool FControlRigEditor::IsConstructionModeEnabled() const
+bool FControlRigBaseEditor::IsConstructionModeEnabled() const
 {
 	return GetEventQueue() == ConstructionEventQueue;
 }
 
-bool FControlRigEditor::IsDebuggingExternalControlRig(const UControlRig* InControlRig) const
+bool FControlRigBaseEditor::IsDebuggingExternalControlRig(const UControlRig* InControlRig) const
 {
 	if(InControlRig == nullptr)
 	{
@@ -3054,12 +2914,12 @@ bool FControlRigEditor::IsDebuggingExternalControlRig(const UControlRig* InContr
 	return InControlRig != GetControlRig();
 }
 
-bool FControlRigEditor::ShouldExecuteControlRig(const UControlRig* InControlRig) const
+bool FControlRigBaseEditor::ShouldExecuteControlRig(const UControlRig* InControlRig) const
 {
 	return (!IsDebuggingExternalControlRig(InControlRig)) && bExecutionControlRig;
 }
 
-void FControlRigEditor::HandlePreviewSceneCreated(const TSharedRef<IPersonaPreviewScene>& InPersonaPreviewScene)
+void FControlRigBaseEditor::HandlePreviewSceneCreated(const TSharedRef<IPersonaPreviewScene>& InPersonaPreviewScene)
 {
 	DECLARE_SCOPE_HIERARCHICAL_COUNTER_FUNC()
 
@@ -3074,7 +2934,7 @@ void FControlRigEditor::HandlePreviewSceneCreated(const TSharedRef<IPersonaPrevi
 	if(FWorldContext* WorldContext = GEngine->GetWorldContextFromWorld(InPersonaPreviewScene->GetWorld()))
 	{
 		static constexpr TCHAR Format[] = TEXT("ControlRigEditor (%s)");
-		WorldContext->CustomDescription = FString::Printf(Format, *FRigVMEditorBase::GetBlueprintObj()->GetName());
+		WorldContext->CustomDescription = FString::Printf(Format, *GetRigVMBlueprint()->GetName());
 	}
 
 	// create ground mesh actor
@@ -3125,13 +2985,13 @@ void FControlRigEditor::HandlePreviewSceneCreated(const TSharedRef<IPersonaPrevi
 	InPersonaPreviewScene->UnregisterForUndo();
 }
 
-void FControlRigEditor::UpdateRigVMHost()
+void FControlRigBaseEditor::UpdateRigVMHostImpl()
 {
 	DECLARE_SCOPE_HIERARCHICAL_COUNTER_FUNC()
 
-	IControlRigEditor::UpdateRigVMHost();
+	UpdateRigVMHostSuper();
 
-	UControlRigBlueprint* Blueprint = Cast<UControlRigBlueprint>(FRigVMEditorBase::GetBlueprintObj());
+	UControlRigBlueprint* Blueprint = Cast<UControlRigBlueprint>(GetRigVMBlueprint());
 	if(UClass* Class = Blueprint->GeneratedClass)
 	{
 		UControlRigSkeletalMeshComponent* EditorSkelComp = Cast<UControlRigSkeletalMeshComponent>(GetPersonaToolkit()->GetPreviewScene()->GetPreviewMeshComponent());
@@ -3172,8 +3032,8 @@ void FControlRigEditor::UpdateRigVMHost()
 			ControlRig->OnPreForwardsSolve_AnyThread().RemoveAll(this);
 			ControlRig->ControlModified().RemoveAll(this);
 
-			ControlRig->OnPreForwardsSolve_AnyThread().AddSP(this, &FControlRigEditor::OnPreForwardsSolve_AnyThread);
-			ControlRig->ControlModified().AddSP(this, &FControlRigEditor::HandleOnControlModified);
+			ControlRig->OnPreForwardsSolve_AnyThread().AddSP(SharedRef(), &FControlRigBaseEditor::OnPreForwardsSolve_AnyThread);
+			ControlRig->ControlModified().AddSP(SharedRef(), &FControlRigBaseEditor::HandleOnControlModified);
 		}
 
 		if(IsModularRig() && ControlRig)
@@ -3186,7 +3046,7 @@ void FControlRigEditor::UpdateRigVMHost()
 	}
 }
 
-void FControlRigEditor::UpdateRigVMHost_PreClearOldHost(URigVMHost* InPreviousHost)
+void FControlRigBaseEditor::UpdateRigVMHost_PreClearOldHostImpl(URigVMHost* InPreviousHost)
 {
 	if(ReplayStrongPtr.IsValid())
 	{
@@ -3194,11 +3054,11 @@ void FControlRigEditor::UpdateRigVMHost_PreClearOldHost(URigVMHost* InPreviousHo
 	}
 }
 
-void FControlRigEditor::CacheNameLists()
+void FControlRigBaseEditor::CacheNameListsImpl()
 {
 	DECLARE_SCOPE_HIERARCHICAL_COUNTER_FUNC()
 
-	IControlRigEditor::CacheNameLists();
+	CacheNameListsSuper();
 
 	if (UControlRigBlueprint* ControlRigBP = GetControlRigBlueprint())
 	{
@@ -3224,7 +3084,7 @@ void FControlRigEditor::CacheNameLists()
 	}
 }
 
-FVector2D FControlRigEditor::ComputePersonaProjectedScreenPos(const FVector& InWorldPos, bool bClampToScreenRectangle)
+FVector2D FControlRigBaseEditor::ComputePersonaProjectedScreenPos(const FVector& InWorldPos, bool bClampToScreenRectangle)
 {
 	if (PreviewViewport.IsValid())
 	{
@@ -3270,7 +3130,7 @@ FVector2D FControlRigEditor::ComputePersonaProjectedScreenPos(const FVector& InW
 	return FVector2D::ZeroVector;
 }
 
-void FControlRigEditor::FindReferencesOfItem(const FRigHierarchyKey& InKey)
+void FControlRigBaseEditor::FindReferencesOfItem(const FRigHierarchyKey& InKey)
 {
 	if(InKey.IsElement())
 	{
@@ -3288,7 +3148,7 @@ void FControlRigEditor::FindReferencesOfItem(const FRigHierarchyKey& InKey)
 	}
 }
 
-void FControlRigEditor::HandlePreviewMeshChanged(USkeletalMesh* InOldSkeletalMesh, USkeletalMesh* InNewSkeletalMesh)
+void FControlRigBaseEditor::HandlePreviewMeshChanged(USkeletalMesh* InOldSkeletalMesh, USkeletalMesh* InNewSkeletalMesh)
 {
 	RebindToSkeletalMeshComponent();
 
@@ -3479,7 +3339,7 @@ void FControlRigEditor::HandlePreviewMeshChanged(USkeletalMesh* InOldSkeletalMes
 	}
 }
 
-void FControlRigEditor::RebindToSkeletalMeshComponent()
+void FControlRigBaseEditor::RebindToSkeletalMeshComponent()
 {
 	DECLARE_SCOPE_HIERARCHICAL_COUNTER_FUNC()
 
@@ -3491,19 +3351,19 @@ void FControlRigEditor::RebindToSkeletalMeshComponent()
 	}
 }
 
-void FControlRigEditor::GenerateEventQueueMenuContent(FMenuBuilder& MenuBuilder)
+void FControlRigBaseEditor::GenerateEventQueueMenuContentImpl(FMenuBuilder& MenuBuilder)
 {
 	MenuBuilder.BeginSection(TEXT("Events"));
-    MenuBuilder.AddMenuEntry(FControlRigEditorCommands::Get().ConstructionEvent, TEXT("Setup"), TAttribute<FText>(), TAttribute<FText>(), GetEventQueueIcon(ConstructionEventQueue));
-    MenuBuilder.AddMenuEntry(FControlRigEditorCommands::Get().ForwardsSolveEvent, TEXT("Forwards Solve"), TAttribute<FText>(), TAttribute<FText>(), GetEventQueueIcon(ForwardsSolveEventQueue));
-    MenuBuilder.AddMenuEntry(FControlRigEditorCommands::Get().BackwardsSolveEvent, TEXT("Backwards Solve"), TAttribute<FText>(), TAttribute<FText>(), GetEventQueueIcon(BackwardsSolveEventQueue));
+    MenuBuilder.AddMenuEntry(FControlRigEditorCommands::Get().ConstructionEvent, TEXT("Setup"), TAttribute<FText>(), TAttribute<FText>(), GetEventQueueIconImpl(ConstructionEventQueue));
+    MenuBuilder.AddMenuEntry(FControlRigEditorCommands::Get().ForwardsSolveEvent, TEXT("Forwards Solve"), TAttribute<FText>(), TAttribute<FText>(), GetEventQueueIconImpl(ForwardsSolveEventQueue));
+    MenuBuilder.AddMenuEntry(FControlRigEditorCommands::Get().BackwardsSolveEvent, TEXT("Backwards Solve"), TAttribute<FText>(), TAttribute<FText>(), GetEventQueueIconImpl(BackwardsSolveEventQueue));
     MenuBuilder.EndSection();
 
     MenuBuilder.BeginSection(TEXT("Validation"));
-    MenuBuilder.AddMenuEntry(FControlRigEditorCommands::Get().BackwardsAndForwardsSolveEvent, TEXT("BackwardsAndForwards"), TAttribute<FText>(), TAttribute<FText>(), GetEventQueueIcon(BackwardsAndForwardsSolveEventQueue));
+    MenuBuilder.AddMenuEntry(FControlRigEditorCommands::Get().BackwardsAndForwardsSolveEvent, TEXT("BackwardsAndForwards"), TAttribute<FText>(), TAttribute<FText>(), GetEventQueueIconImpl(BackwardsAndForwardsSolveEventQueue));
     MenuBuilder.EndSection();
 
-    if (const UControlRigBlueprint* RigBlueprint = Cast<UControlRigBlueprint>(FRigVMEditorBase::GetBlueprintObj()))
+    if (const UControlRigBlueprint* RigBlueprint = Cast<UControlRigBlueprint>(GetRigVMBlueprint()))
     {
     	URigVMEdGraphSchema* Schema = CastChecked<URigVMEdGraphSchema>(RigBlueprint->GetRigVMEdGraphSchemaClass()->GetDefaultObject());
     	
@@ -3531,13 +3391,9 @@ void FControlRigEditor::GenerateEventQueueMenuContent(FMenuBuilder& MenuBuilder)
     		MenuBuilder.AddMenuEntry(
     			FText::FromString(EventNameStr),
     			FText::FromString(FString::Printf(TEXT("Runs the user defined %s"), *EventNameStr)),
-    			GetEventQueueIcon({EntryName}),
+    			GetEventQueueIconImpl({EntryName}),
     			FUIAction(
-#if WITH_RIGVMLEGACYEDITOR
-    				FExecuteAction::CreateSP(this, &FRigVMLegacyEditor::SetEventQueue, TArray<FName>({EntryName})),
-#else
-    				FExecuteAction::CreateSP(this, &FRigVMNewEditor::SetEventQueue, TArray<FName>({EntryName})),
-#endif
+    			FExecuteAction::CreateSP(SharedRef(), &FControlRigBaseEditor::SetEventQueueSuper, TArray<FName>({EntryName})),
     				FCanExecuteAction()
     			)
     		);
@@ -3545,7 +3401,7 @@ void FControlRigEditor::GenerateEventQueueMenuContent(FMenuBuilder& MenuBuilder)
     }
 }
 
-void FControlRigEditor::FilterDraggedKeys(TArray<FRigElementKey>& Keys, bool bRemoveNameSpace)
+void FControlRigBaseEditor::FilterDraggedKeys(TArray<FRigElementKey>& Keys, bool bRemoveNameSpace)
 {
 	// if the keys being dragged contain something mapped to a connector - use that instead
 	if(UControlRigBlueprint* ControlRigBlueprint = GetControlRigBlueprint())
@@ -3592,13 +3448,13 @@ void FControlRigEditor::FilterDraggedKeys(TArray<FRigElementKey>& Keys, bool bRe
 	}
 }
 
-FTransform FControlRigEditor::GetRigElementTransform(const FRigElementKey& InElement, bool bLocal, bool bOnDebugInstance) const
+FTransform FControlRigBaseEditor::GetRigElementTransform(const FRigElementKey& InElement, bool bLocal, bool bOnDebugInstance) const
 {
 	DECLARE_SCOPE_HIERARCHICAL_COUNTER_FUNC()
 
 	if(bOnDebugInstance)
 	{
-		UControlRig* DebuggedControlRig = Cast<UControlRig>(FRigVMEditorBase::GetBlueprintObj()->GetObjectBeingDebugged());
+		UControlRig* DebuggedControlRig = Cast<UControlRig>(GetRigVMBlueprint()->GetObjectBeingDebugged());
 		if (DebuggedControlRig == nullptr)
 		{
 			DebuggedControlRig = GetControlRig();
@@ -3621,7 +3477,7 @@ FTransform FControlRigEditor::GetRigElementTransform(const FRigElementKey& InEle
 	return GetHierarchyBeingDebugged()->GetGlobalTransform(InElement);
 }
 
-void FControlRigEditor::SetRigElementTransform(const FRigElementKey& InElement, const FTransform& InTransform, bool bLocal)
+void FControlRigBaseEditor::SetRigElementTransform(const FRigElementKey& InElement, const FTransform& InTransform, bool bLocal)
 {
 	DECLARE_SCOPE_HIERARCHICAL_COUNTER_FUNC()
 
@@ -3706,9 +3562,9 @@ void FControlRigEditor::SetRigElementTransform(const FRigElementKey& InElement,
 	}
 }
 
-void FControlRigEditor::OnFinishedChangingProperties(const FPropertyChangedEvent& PropertyChangedEvent)
+void FControlRigBaseEditor::OnFinishedChangingPropertiesImpl(const FPropertyChangedEvent& PropertyChangedEvent)
 {
-	FRigVMEditorBase::OnFinishedChangingProperties(PropertyChangedEvent);
+	OnFinishedChangingPropertiesSuper(PropertyChangedEvent);
 	
 	DECLARE_SCOPE_HIERARCHICAL_COUNTER_FUNC()
 
@@ -3733,9 +3589,9 @@ void FControlRigEditor::OnFinishedChangingProperties(const FPropertyChangedEvent
 	}
 }
 
-void FControlRigEditor::OnWrappedPropertyChangedChainEvent(URigVMDetailsViewWrapperObject* InWrapperObject, const FString& InPropertyPath, FPropertyChangedChainEvent& InPropertyChangedChainEvent)
+void FControlRigBaseEditor::OnWrappedPropertyChangedChainEventImpl(URigVMDetailsViewWrapperObject* InWrapperObject, const FString& InPropertyPath, FPropertyChangedChainEvent& InPropertyChangedChainEvent)
 {
-	IControlRigEditor::OnWrappedPropertyChangedChainEvent(InWrapperObject, InPropertyPath, InPropertyChangedChainEvent);
+	OnWrappedPropertyChangedChainEventSuper(InWrapperObject, InPropertyPath, InPropertyChangedChainEvent);
 	
 	check(InWrapperObject);
 	check(!GetWrapperObjects().IsEmpty());
@@ -3915,7 +3771,16 @@ void FControlRigEditor::OnWrappedPropertyChangedChainEvent(URigVMDetailsViewWrap
 	}
 }
 
-bool FControlRigEditor::HandleRequestDirectManipulation(ERigControlType InControlType) const
+void FControlRigBaseEditor::OnClose()
+{
+	if (ControlRigEditorClosedDelegate.IsBound())
+	{
+		ControlRigEditorClosedDelegate.Broadcast(this, GetControlRigBlueprint());
+	}
+	OnCloseSuper();
+}
+
+bool FControlRigBaseEditor::HandleRequestDirectManipulation(ERigControlType InControlType) const
 {
 	TArray<FRigDirectManipulationTarget> Targets = GetDirectManipulationTargets();
 	for(const FRigDirectManipulationTarget& Target : Targets)
@@ -3958,7 +3823,7 @@ bool FControlRigEditor::HandleRequestDirectManipulation(ERigControlType InContro
 	return false;
 }
 
-bool FControlRigEditor::SetDirectionManipulationSubject(const URigVMUnitNode* InNode)
+bool FControlRigBaseEditor::SetDirectionManipulationSubject(const URigVMUnitNode* InNode)
 {
 	if(DirectManipulationSubject.Get() == InNode)
 	{
@@ -3975,19 +3840,19 @@ bool FControlRigEditor::SetDirectionManipulationSubject(const URigVMUnitNode* In
 	return true;
 }
 
-bool FControlRigEditor::IsDirectManipulationEnabled() const
+bool FControlRigBaseEditor::IsDirectManipulationEnabled() const
 {
 	return !GetDirectManipulationTargets().IsEmpty();
 }
 
-EVisibility FControlRigEditor::GetDirectManipulationVisibility() const
+EVisibility FControlRigBaseEditor::GetDirectManipulationVisibility() const
 {
 	return IsDirectManipulationEnabled() ? EVisibility::Visible : EVisibility::Hidden;
 }
 
-FText FControlRigEditor::GetDirectionManipulationText() const
+FText FControlRigBaseEditor::GetDirectionManipulationText() const
 {
-	if (UControlRig* DebuggedControlRig = Cast<UControlRig>(FRigVMEditorBase::GetBlueprintObj()->GetObjectBeingDebugged()))
+	if (UControlRig* DebuggedControlRig = Cast<UControlRig>(GetRigVMBlueprint()->GetObjectBeingDebugged()))
 	{
 		if (URigHierarchy* Hierarchy = DebuggedControlRig->GetHierarchy())
 		{
@@ -4006,7 +3871,7 @@ FText FControlRigEditor::GetDirectionManipulationText() const
 	return DefaultText;
 }
 
-void FControlRigEditor::OnDirectManipulationChanged(TSharedPtr<FString> NewValue, ESelectInfo::Type SelectInfo)
+void FControlRigBaseEditor::OnDirectManipulationChanged(TSharedPtr<FString> NewValue, ESelectInfo::Type SelectInfo)
 {
 	if(!NewValue.IsValid())
 	{
@@ -4019,7 +3884,7 @@ void FControlRigEditor::OnDirectManipulationChanged(TSharedPtr<FString> NewValue
 		return;
 	}
 	
-	UControlRigBlueprint* ControlRigBlueprint = CastChecked<UControlRigBlueprint>(FRigVMEditorBase::GetBlueprintObj());
+	UControlRigBlueprint* ControlRigBlueprint = CastChecked<UControlRigBlueprint>(GetRigVMBlueprint());
 	if(ControlRigBlueprint == nullptr)
 	{
 		return;
@@ -4048,11 +3913,11 @@ void FControlRigEditor::OnDirectManipulationChanged(TSharedPtr<FString> NewValue
 	}
 }
 
-const TArray<FRigDirectManipulationTarget> FControlRigEditor::GetDirectManipulationTargets() const
+const TArray<FRigDirectManipulationTarget> FControlRigBaseEditor::GetDirectManipulationTargets() const
 {
 	if(DirectManipulationSubject.IsValid())
 	{
-		if (UControlRig* DebuggedControlRig = Cast<UControlRig>(FRigVMEditorBase::GetBlueprintObj()->GetObjectBeingDebugged()))
+		if (UControlRig* DebuggedControlRig = Cast<UControlRig>(GetRigVMBlueprint()->GetObjectBeingDebugged()))
 		{
 			if(const URigVMUnitNode* Node = DirectManipulationSubject.Get())
 			{
@@ -4079,7 +3944,7 @@ const TArray<FRigDirectManipulationTarget> FControlRigEditor::GetDirectManipulat
 	return EmptyTargets;
 }
 
-const TArray<TSharedPtr<FString>>& FControlRigEditor::GetDirectManipulationTargetTextList() const
+const TArray<TSharedPtr<FString>>& FControlRigBaseEditor::GetDirectManipulationTargetTextList() const
 {
 	if(DirectManipulationTextList.IsEmpty())
 	{
@@ -4092,7 +3957,7 @@ const TArray<TSharedPtr<FString>>& FControlRigEditor::GetDirectManipulationTarge
 	return DirectManipulationTextList;
 }
 
-void FControlRigEditor::RefreshDirectManipulationTextList()
+void FControlRigBaseEditor::RefreshDirectManipulationTextList()
 {
 	DirectManipulationTextList.Reset();
 	(void)GetDirectManipulationTargetTextList();
@@ -4105,7 +3970,7 @@ void FControlRigEditor::RefreshDirectManipulationTextList()
 	}
 }
 
-EVisibility FControlRigEditor::GetConnectorWarningVisibility() const
+EVisibility FControlRigBaseEditor::GetConnectorWarningVisibility() const
 {
 	if(GetConnectorWarningText().IsEmpty())
 	{
@@ -4114,7 +3979,7 @@ EVisibility FControlRigEditor::GetConnectorWarningVisibility() const
 	return EVisibility::Visible;
 }
 
-FText FControlRigEditor::GetConnectorWarningText() const
+FText FControlRigBaseEditor::GetConnectorWarningText() const
 {
 	if (const UControlRigBlueprint* Blueprint = GetControlRigBlueprint())
 	{
@@ -4138,102 +4003,88 @@ FText FControlRigEditor::GetConnectorWarningText() const
 	return FText();
 }
 
-FReply FControlRigEditor::OnNavigateToConnectorWarning() const
+FReply FControlRigBaseEditor::OnNavigateToConnectorWarning() const
 {
 	RequestNavigateToConnectorWarningDelegate.Broadcast();
 	return FReply::Handled();
 }
 
-void FControlRigEditor::BindCommands()
+void FControlRigBaseEditor::BindCommandsImpl()
 {
-	IControlRigEditor::BindCommands();
+	BindCommandsSuper();
+
+	TSharedRef<FControlRigBaseEditor> SharedEditor = SharedRef();
 
 	GetToolkitCommands()->MapAction(
 		FControlRigEditorCommands::Get().ConstructionEvent,
-#if WITH_RIGVMLEGACYEDITOR
-		FExecuteAction::CreateSP(this, &FRigVMLegacyEditor::SetEventQueue, TArray<FName>(ConstructionEventQueue)),
-#else
-		FExecuteAction::CreateSP(this, &FRigVMNewEditor::SetEventQueue, TArray<FName>(ConstructionEventQueue)),
-#endif
+		FExecuteAction::CreateSP(SharedEditor, &FControlRigBaseEditor::SetEventQueueSuper, TArray<FName>(ConstructionEventQueue)),
 		FCanExecuteAction());
 
 	GetToolkitCommands()->MapAction(
 		FControlRigEditorCommands::Get().ForwardsSolveEvent,
-#if WITH_RIGVMLEGACYEDITOR
-		FExecuteAction::CreateSP(this, &FRigVMLegacyEditor::SetEventQueue, TArray<FName>(ForwardsSolveEventQueue)),
-#else
-		FExecuteAction::CreateSP(this, &FRigVMNewEditor::SetEventQueue, TArray<FName>(ForwardsSolveEventQueue)),
-#endif
+		FExecuteAction::CreateSP(SharedEditor, &FControlRigBaseEditor::SetEventQueueSuper, TArray<FName>(ForwardsSolveEventQueue)),
 		FCanExecuteAction());
 
 	GetToolkitCommands()->MapAction(
 		FControlRigEditorCommands::Get().BackwardsSolveEvent,
-#if WITH_RIGVMLEGACYEDITOR
-		FExecuteAction::CreateSP(this, &FRigVMLegacyEditor::SetEventQueue, TArray<FName>(BackwardsSolveEventQueue)),
-#else
-		FExecuteAction::CreateSP(this, &FRigVMNewEditor::SetEventQueue, TArray<FName>(BackwardsSolveEventQueue)),
-#endif
+		FExecuteAction::CreateSP(SharedEditor, &FControlRigBaseEditor::SetEventQueueSuper, TArray<FName>(BackwardsSolveEventQueue)),
 		FCanExecuteAction());
 
 	GetToolkitCommands()->MapAction(
 		FControlRigEditorCommands::Get().BackwardsAndForwardsSolveEvent,
-#if WITH_RIGVMLEGACYEDITOR
-		FExecuteAction::CreateSP(this, &FRigVMLegacyEditor::SetEventQueue, TArray<FName>(BackwardsAndForwardsSolveEventQueue)),
-#else
-		FExecuteAction::CreateSP(this, &FRigVMNewEditor::SetEventQueue, TArray<FName>(BackwardsAndForwardsSolveEventQueue)),
-#endif
+		FExecuteAction::CreateSP(SharedEditor, &FControlRigBaseEditor::SetEventQueueSuper, TArray<FName>(BackwardsAndForwardsSolveEventQueue)),
 		FCanExecuteAction());
 
 	GetToolkitCommands()->MapAction(
 		FControlRigEditorCommands::Get().ToggleControlVisibility,
-		FExecuteAction::CreateSP(this, &FControlRigEditor::HandleToggleControlVisibility),
+		FExecuteAction::CreateSP(SharedEditor, &FControlRigBaseEditor::HandleToggleControlVisibility),
 		FCanExecuteAction(),
-		FIsActionChecked::CreateSP(this, &FControlRigEditor::AreControlsVisible));
+		FIsActionChecked::CreateSP(SharedEditor, &FControlRigBaseEditor::AreControlsVisible));
 
 	GetToolkitCommands()->MapAction(
 		FControlRigEditorCommands::Get().ToggleControlsAsOverlay,
-		FExecuteAction::CreateSP(this, &FControlRigEditor::HandleToggleControlsAsOverlay),
+		FExecuteAction::CreateSP(SharedEditor, &FControlRigBaseEditor::HandleToggleControlsAsOverlay),
 		FCanExecuteAction(),
-		FIsActionChecked::CreateSP(this, &FControlRigEditor::AreControlsAsOverlay));
+		FIsActionChecked::CreateSP(SharedEditor, &FControlRigBaseEditor::AreControlsAsOverlay));
 
 	GetToolkitCommands()->MapAction(
 		FControlRigEditorCommands::Get().ToggleDrawNulls,
-		FExecuteAction::CreateSP(this, &FControlRigEditor::HandleToggleToolbarDrawNulls),
-		FCanExecuteAction::CreateSP(this, &FControlRigEditor::IsToolbarDrawNullsEnabled),
-		FIsActionChecked::CreateSP(this, &FControlRigEditor::GetToolbarDrawNulls));
+		FExecuteAction::CreateSP(SharedEditor, &FControlRigBaseEditor::HandleToggleToolbarDrawNulls),
+		FCanExecuteAction::CreateSP(SharedEditor, &FControlRigBaseEditor::IsToolbarDrawNullsEnabled),
+		FIsActionChecked::CreateSP(SharedEditor, &FControlRigBaseEditor::GetToolbarDrawNulls));
 
 	GetToolkitCommands()->MapAction(
 		FControlRigEditorCommands::Get().ToggleDrawSockets,
-		FExecuteAction::CreateSP(this, &FControlRigEditor::HandleToggleToolbarDrawSockets),
-		FCanExecuteAction::CreateSP(this, &FControlRigEditor::IsToolbarDrawSocketsEnabled),
-		FIsActionChecked::CreateSP(this, &FControlRigEditor::GetToolbarDrawSockets));
+		FExecuteAction::CreateSP(SharedEditor, &FControlRigBaseEditor::HandleToggleToolbarDrawSockets),
+		FCanExecuteAction::CreateSP(SharedEditor, &FControlRigBaseEditor::IsToolbarDrawSocketsEnabled),
+		FIsActionChecked::CreateSP(SharedEditor, &FControlRigBaseEditor::GetToolbarDrawSockets));
 
 	GetToolkitCommands()->MapAction(
 		FControlRigEditorCommands::Get().ToggleDrawAxesOnSelection,
-		FExecuteAction::CreateSP(this, &FControlRigEditor::HandleToggleToolbarDrawAxesOnSelection),
+		FExecuteAction::CreateSP(SharedEditor, &FControlRigBaseEditor::HandleToggleToolbarDrawAxesOnSelection),
 		FCanExecuteAction(),
-		FIsActionChecked::CreateSP(this, &FControlRigEditor::GetToolbarDrawAxesOnSelection));
+		FIsActionChecked::CreateSP(SharedEditor, &FControlRigBaseEditor::GetToolbarDrawAxesOnSelection));
 
 	GetToolkitCommands()->MapAction(
 		FControlRigEditorCommands::Get().ToggleSchematicViewportVisibility,
-		FExecuteAction::CreateSP(this, &FControlRigEditor::HandleToggleSchematicViewport),
+		FExecuteAction::CreateSP(SharedEditor, &FControlRigBaseEditor::HandleToggleSchematicViewport),
 		FCanExecuteAction(),
-		FIsActionChecked::CreateSP(this, &FControlRigEditor::IsSchematicViewportActive));
+		FIsActionChecked::CreateSP(SharedEditor, &FControlRigBaseEditor::IsSchematicViewportActive));
 
 	GetToolkitCommands()->MapAction(
 		FControlRigEditorCommands::Get().SwapModuleWithinAsset,
-		FExecuteAction::CreateSP(this, &FControlRigEditor::SwapModuleWithinAsset),
-		FCanExecuteAction::CreateSP(this, &FControlRigEditor::IsModularRig));
+		FExecuteAction::CreateSP(SharedEditor, &FControlRigBaseEditor::SwapModuleWithinAsset),
+		FCanExecuteAction::CreateSP(SharedEditor, &FControlRigBaseEditor::IsModularRig));
 
 	GetToolkitCommands()->MapAction(
 		FControlRigEditorCommands::Get().SwapModuleAcrossProject,
-		FExecuteAction::CreateSP(this, &FControlRigEditor::SwapModuleAcrossProject),
-		FCanExecuteAction::CreateSP(this, &FControlRigEditor::IsRigModule));
+		FExecuteAction::CreateSP(SharedEditor, &FControlRigBaseEditor::SwapModuleAcrossProject),
+		FCanExecuteAction::CreateSP(SharedEditor, &FControlRigBaseEditor::IsRigModule));
 }
 
-FMenuBuilder FControlRigEditor::GenerateBulkEditMenu()
+FMenuBuilder FControlRigBaseEditor::GenerateBulkEditMenuImpl()
 {
-	FMenuBuilder MenuBuilder = IControlRigEditor::GenerateBulkEditMenu();
+	FMenuBuilder MenuBuilder = GenerateBulkEditMenuSuper();
 	MenuBuilder.BeginSection(TEXT("Asset"), LOCTEXT("Asset", "Asset"));
 	if (UControlRigBlueprint* Blueprint = GetControlRigBlueprint())
 	{
@@ -4250,7 +4101,7 @@ FMenuBuilder FControlRigEditor::GenerateBulkEditMenu()
 	return MenuBuilder;
 }
 
-void FControlRigEditor::OnHierarchyChanged()
+void FControlRigBaseEditor::OnHierarchyChanged()
 {
 	DECLARE_SCOPE_HIERARCHICAL_COUNTER_FUNC()
 
@@ -4290,9 +4141,9 @@ void FControlRigEditor::OnHierarchyChanged()
 }
 
 
-void FControlRigEditor::OnHierarchyModified(ERigHierarchyNotification InNotif, URigHierarchy* InHierarchy, const FRigNotificationSubject& InSubject)
+void FControlRigBaseEditor::OnHierarchyModified(ERigHierarchyNotification InNotif, URigHierarchy* InHierarchy, const FRigNotificationSubject& InSubject)
 {
-	UControlRigBlueprint* RigBlueprint = Cast<UControlRigBlueprint>(FRigVMEditorBase::GetBlueprintObj());
+	UControlRigBlueprint* RigBlueprint = Cast<UControlRigBlueprint>(GetRigVMBlueprint());
 	if(RigBlueprint == nullptr)
 	{
 		return;
@@ -4436,7 +4287,7 @@ void FControlRigEditor::OnHierarchyModified(ERigHierarchyNotification InNotif, U
 	}
 }
 
-void FControlRigEditor::OnHierarchyModified_AnyThread(ERigHierarchyNotification InNotif, URigHierarchy* InHierarchy, const FRigNotificationSubject& InSubject)
+void FControlRigBaseEditor::OnHierarchyModified_AnyThread(ERigHierarchyNotification InNotif, URigHierarchy* InHierarchy, const FRigNotificationSubject& InSubject)
 {
 	if(bIsConstructionEventRunning)
 	{
@@ -4517,7 +4368,7 @@ void FControlRigEditor::OnHierarchyModified_AnyThread(ERigHierarchyNotification
 						else
 						{
 							// only clear the details if we are not looking at a transient control
-							if (UControlRig* DebuggedControlRig = Cast<UControlRig>(FRigVMEditorBase::GetBlueprintObj()->GetObjectBeingDebugged()))
+							if (UControlRig* DebuggedControlRig = Cast<UControlRig>(GetRigVMBlueprint()->GetObjectBeingDebugged()))
 							{
 								if(DebuggedControlRig->RigUnitManipulationInfos.IsEmpty())
 								{
@@ -4617,7 +4468,7 @@ void FControlRigEditor::OnHierarchyModified_AnyThread(ERigHierarchyNotification
 	}
 }
 
-void FControlRigEditor::HandleRigTypeChanged(UControlRigBlueprint* InBlueprint)
+void FControlRigBaseEditor::HandleRigTypeChanged(UControlRigBlueprint* InBlueprint)
 {
 	// todo: fire a notification.
 	// todo: reapply the preview mesh and react to it accordingly.
@@ -4625,7 +4476,7 @@ void FControlRigEditor::HandleRigTypeChanged(UControlRigBlueprint* InBlueprint)
 	Compile();
 }
 
-void FControlRigEditor::HandleModularRigModified(EModularRigNotification InNotification, const FRigModuleReference* InModule)
+void FControlRigBaseEditor::HandleModularRigModified(EModularRigNotification InNotification, const FRigModuleReference* InModule)
 {
 	UControlRigBlueprint* RigBlueprint = GetControlRigBlueprint();
 	if(RigBlueprint == nullptr)
@@ -4676,12 +4527,12 @@ void FControlRigEditor::HandleModularRigModified(EModularRigNotification InNotif
 	}
 }
 
-void FControlRigEditor::HandlePostCompileModularRigs(URigVMBlueprint* InBlueprint)
+void FControlRigBaseEditor::HandlePostCompileModularRigs(URigVMBlueprint* InBlueprint)
 {
 	RefreshDetailView();
 }
 
-void FControlRigEditor::SwapModuleWithinAsset()
+void FControlRigBaseEditor::SwapModuleWithinAsset()
 {
 	const UControlRigBlueprint* Blueprint = GetControlRigBlueprint();
 	const FAssetData Asset = UE::RigVM::Editor::Tools::FindAssetFromAnyPath(GetRigVMBlueprint()->GetPathName(), true);
@@ -4764,7 +4615,7 @@ void FControlRigEditor::SwapModuleWithinAsset()
 	SwapModulesDialog->ShowNormal();
 }
 
-void FControlRigEditor::SwapModuleAcrossProject()
+void FControlRigBaseEditor::SwapModuleAcrossProject()
 {
 	const UControlRigBlueprint* Blueprint = GetControlRigBlueprint();
 	const FAssetData Asset = UE::RigVM::Editor::Tools::FindAssetFromAnyPath(GetRigVMBlueprint()->GetPathName(), true);
@@ -4811,7 +4662,7 @@ void FControlRigEditor::SwapModuleAcrossProject()
 	SwapModulesDialog->ShowNormal();
 }
 
-void FControlRigEditor::SynchronizeViewportBoneSelection()
+void FControlRigBaseEditor::SynchronizeViewportBoneSelection()
 {
 	UControlRigBlueprint* RigBlueprint = GetControlRigBlueprint();
 	if (RigBlueprint == nullptr)
@@ -4836,7 +4687,7 @@ void FControlRigEditor::SynchronizeViewportBoneSelection()
 	}
 }
 
-void FControlRigEditor::UpdateBoneModification(FName BoneName, const FTransform& LocalTransform)
+void FControlRigBaseEditor::UpdateBoneModification(FName BoneName, const FTransform& LocalTransform)
 {
 	if (UControlRig* ControlRig = GetControlRig())
 	{ 
@@ -4852,7 +4703,7 @@ void FControlRigEditor::UpdateBoneModification(FName BoneName, const FTransform&
 		}
 		
 		TMap<FName, FTransform>* TransformOverrideMap = &ControlRig->TransformOverrideForUserCreatedBones;
-		if (UControlRig* DebuggedControlRig = Cast<UControlRig>(FRigVMEditorBase::GetBlueprintObj()->GetObjectBeingDebugged()))
+		if (UControlRig* DebuggedControlRig = Cast<UControlRig>(GetRigVMBlueprint()->GetObjectBeingDebugged()))
 		{
 			TransformOverrideMap = &DebuggedControlRig->TransformOverrideForUserCreatedBones;
 		}
@@ -4864,7 +4715,7 @@ void FControlRigEditor::UpdateBoneModification(FName BoneName, const FTransform&
 	}
 }
 
-void FControlRigEditor::RemoveBoneModification(FName BoneName)
+void FControlRigBaseEditor::RemoveBoneModification(FName BoneName)
 {
 	if (UControlRig* ControlRig = GetControlRig())
 	{
@@ -4874,7 +4725,7 @@ void FControlRigEditor::RemoveBoneModification(FName BoneName)
 		}
 
 		TMap<FName, FTransform>* TransformOverrideMap = &ControlRig->TransformOverrideForUserCreatedBones;
-		if (UControlRig* DebuggedControlRig = Cast<UControlRig>(FRigVMEditorBase::GetBlueprintObj()->GetObjectBeingDebugged()))
+		if (UControlRig* DebuggedControlRig = Cast<UControlRig>(GetRigVMBlueprint()->GetObjectBeingDebugged()))
 		{
 			TransformOverrideMap = &DebuggedControlRig->TransformOverrideForUserCreatedBones;
 		}
@@ -4883,7 +4734,7 @@ void FControlRigEditor::RemoveBoneModification(FName BoneName)
 	}
 }
 
-void FControlRigEditor::ResetAllBoneModification()
+void FControlRigBaseEditor::ResetAllBoneModification()
 {
 	if (UControlRig* ControlRig = GetControlRig())
 	{
@@ -4893,7 +4744,7 @@ void FControlRigEditor::ResetAllBoneModification()
 		}
 
 		TMap<FName, FTransform>* TransformOverrideMap = &ControlRig->TransformOverrideForUserCreatedBones;
-		if (UControlRig* DebuggedControlRig = Cast<UControlRig>(FRigVMEditorBase::GetBlueprintObj()->GetObjectBeingDebugged()))
+		if (UControlRig* DebuggedControlRig = Cast<UControlRig>(GetRigVMBlueprint()->GetObjectBeingDebugged()))
 		{
 			TransformOverrideMap = &DebuggedControlRig->TransformOverrideForUserCreatedBones;
 		}
@@ -4902,17 +4753,17 @@ void FControlRigEditor::ResetAllBoneModification()
 	}
 }
 
-FControlRigEditorEditMode* FControlRigEditor::GetEditMode() const
+FControlRigEditorEditMode* FControlRigBaseEditor::GetEditMode() const
 {
 	if(IsModularRig())
 	{
-		return static_cast<FModularRigEditorEditMode*>(GetEditorModeManager().GetActiveMode(GetEditorModeName()));
+		return static_cast<FModularRigEditorEditMode*>(GetEditorModeManagerImpl().GetActiveMode(GetEditorModeNameImpl()));
 	}
-	return static_cast<FControlRigEditorEditMode*>(GetEditorModeManager().GetActiveMode(GetEditorModeName()));
+	return static_cast<FControlRigEditorEditMode*>(GetEditorModeManagerImpl().GetActiveMode(GetEditorModeNameImpl()));
 }
 
 
-void FControlRigEditor::OnCurveContainerChanged()
+void FControlRigBaseEditor::OnCurveContainerChanged()
 {
 	DECLARE_SCOPE_HIERARCHICAL_COUNTER_FUNC()
 
@@ -4939,7 +4790,7 @@ void FControlRigEditor::OnCurveContainerChanged()
 	NotificationPtr->SetCompletionState(SNotificationItem::CS_Success);
 }
 
-void FControlRigEditor::CreateRigHierarchyToGraphDragAndDropMenu() const
+void FControlRigBaseEditor::CreateRigHierarchyToGraphDragAndDropMenu() const
 {
 	const FName MenuName = RigHierarchyToGraphDragAndDropMenuName;
 	UToolMenus* ToolMenus = UToolMenus::Get();
@@ -4957,7 +4808,7 @@ void FControlRigEditor::CreateRigHierarchyToGraphDragAndDropMenu() const
 			{
 				UControlRigContextMenuContext* MainContext = InMenu->FindContext<UControlRigContextMenuContext>();
 				
-				if (FControlRigEditor* ControlRigEditor = MainContext->GetControlRigEditor())
+				if (IControlRigBaseEditor* ControlRigEditor = MainContext->GetControlRigEditor())
 				{
 					const FControlRigRigHierarchyToGraphDragAndDropContext& DragDropContext = MainContext->GetRigHierarchyToGraphDragAndDropContext();
 
@@ -4995,7 +4846,7 @@ void FControlRigEditor::CreateRigHierarchyToGraphDragAndDropMenu() const
 							GetterTooltip,
 							FSlateIcon(),
 							FUIAction(
-								FExecuteAction::CreateSP(ControlRigEditor, &FControlRigEditor::HandleMakeComponentGetterSetter, true, DraggedComponents, Graph, NodePosition),
+								FExecuteAction::CreateSP(ControlRigEditor->SharedRef(), &FControlRigBaseEditor::HandleMakeComponentGetterSetter, true, DraggedComponents, Graph, NodePosition),
 								FCanExecuteAction()
 							)
 						);
@@ -5010,7 +4861,7 @@ void FControlRigEditor::CreateRigHierarchyToGraphDragAndDropMenu() const
 							SetterTooltip,
 							FSlateIcon(),
 							FUIAction(
-								FExecuteAction::CreateSP(ControlRigEditor, &FControlRigEditor::HandleMakeComponentGetterSetter, false, DraggedComponents, Graph, NodePosition),
+								FExecuteAction::CreateSP(ControlRigEditor->SharedRef(), &FControlRigBaseEditor::HandleMakeComponentGetterSetter, false, DraggedComponents, Graph, NodePosition),
 								FCanExecuteAction()
 							)
 						);
@@ -5123,7 +4974,7 @@ void FControlRigEditor::CreateRigHierarchyToGraphDragAndDropMenu() const
 							GetterTooltip,
 							FSlateIcon(),
 							FUIAction(
-								FExecuteAction::CreateSP(ControlRigEditor, &FControlRigEditor::HandleMakeElementGetterSetter, ERigElementGetterSetterType_Transform, true, DraggedElements, Graph, NodePosition),
+								FExecuteAction::CreateSP(ControlRigEditor->SharedRef(), &FControlRigBaseEditor::HandleMakeElementGetterSetter, ERigElementGetterSetterType_Transform, true, DraggedElements, Graph, NodePosition),
 								FCanExecuteAction()
 							)
 						);
@@ -5139,7 +4990,7 @@ void FControlRigEditor::CreateRigHierarchyToGraphDragAndDropMenu() const
 							SetterTooltip,
 							FSlateIcon(),
 							FUIAction(
-								FExecuteAction::CreateSP(ControlRigEditor, &FControlRigEditor::HandleMakeElementGetterSetter, ERigElementGetterSetterType_Transform, false, DraggedElements, Graph, NodePosition),
+								FExecuteAction::CreateSP(ControlRigEditor->SharedRef(), &FControlRigBaseEditor::HandleMakeElementGetterSetter, ERigElementGetterSetterType_Transform, false, DraggedElements, Graph, NodePosition),
 								FCanExecuteAction()
 							)
 						);
@@ -5163,7 +5014,7 @@ void FControlRigEditor::CreateRigHierarchyToGraphDragAndDropMenu() const
 								LOCTEXT("SetRotation_ToolTip","Setter for Rotation"),
 								FSlateIcon(),
 								FUIAction(
-									FExecuteAction::CreateSP(ControlRigEditor, &FControlRigEditor::HandleMakeElementGetterSetter, ERigElementGetterSetterType_Rotation, false, DraggedElements, Graph, NodePosition),
+									FExecuteAction::CreateSP(ControlRigEditor->SharedRef(), &FControlRigBaseEditor::HandleMakeElementGetterSetter, ERigElementGetterSetterType_Rotation, false, DraggedElements, Graph, NodePosition),
 									FCanExecuteAction()
 								)
 							);
@@ -5178,7 +5029,7 @@ void FControlRigEditor::CreateRigHierarchyToGraphDragAndDropMenu() const
 								LOCTEXT("SetTranslation_ToolTip","Setter for Translation"),
 								FSlateIcon(),
 								FUIAction(
-									FExecuteAction::CreateSP(ControlRigEditor, &FControlRigEditor::HandleMakeElementGetterSetter, ERigElementGetterSetterType_Translation, false, DraggedElements, Graph, NodePosition),
+									FExecuteAction::CreateSP(ControlRigEditor->SharedRef(), &FControlRigBaseEditor::HandleMakeElementGetterSetter, ERigElementGetterSetterType_Translation, false, DraggedElements, Graph, NodePosition),
 									FCanExecuteAction()
 								)
 							);
@@ -5193,7 +5044,7 @@ void FControlRigEditor::CreateRigHierarchyToGraphDragAndDropMenu() const
 								LOCTEXT("AddOffset_ToolTip","Setter for Offset"),
 								FSlateIcon(),
 								FUIAction(
-									FExecuteAction::CreateSP(ControlRigEditor, &FControlRigEditor::HandleMakeElementGetterSetter, ERigElementGetterSetterType_Offset, false, DraggedElements, Graph, NodePosition),
+									FExecuteAction::CreateSP(ControlRigEditor->SharedRef(), &FControlRigBaseEditor::HandleMakeElementGetterSetter, ERigElementGetterSetterType_Offset, false, DraggedElements, Graph, NodePosition),
 									FCanExecuteAction()
 								)
 							);
@@ -5212,7 +5063,7 @@ void FControlRigEditor::CreateRigHierarchyToGraphDragAndDropMenu() const
 								LOCTEXT("GetRelativeTransform_ToolTip", "Getter for Relative Transform"),
 								FSlateIcon(),
 								FUIAction(
-									FExecuteAction::CreateSP(ControlRigEditor, &FControlRigEditor::HandleMakeElementGetterSetter, ERigElementGetterSetterType_Relative, true, DraggedElements, Graph, NodePosition),
+									FExecuteAction::CreateSP(ControlRigEditor->SharedRef(), &FControlRigBaseEditor::HandleMakeElementGetterSetter, ERigElementGetterSetterType_Relative, true, DraggedElements, Graph, NodePosition),
 									FCanExecuteAction()
 								)
 							);
@@ -5227,7 +5078,7 @@ void FControlRigEditor::CreateRigHierarchyToGraphDragAndDropMenu() const
 								LOCTEXT("SetRelativeTransform_ToolTip", "Setter for Relative Transform"),
 								FSlateIcon(),
 								FUIAction(
-									FExecuteAction::CreateSP(ControlRigEditor, &FControlRigEditor::HandleMakeElementGetterSetter, ERigElementGetterSetterType_Relative, false, DraggedElements, Graph, NodePosition),
+									FExecuteAction::CreateSP(ControlRigEditor->SharedRef(), &FControlRigBaseEditor::HandleMakeElementGetterSetter, ERigElementGetterSetterType_Relative, false, DraggedElements, Graph, NodePosition),
 									FCanExecuteAction()
 								)
 							);
@@ -5293,7 +5144,7 @@ void FControlRigEditor::CreateRigHierarchyToGraphDragAndDropMenu() const
 	}
 }
 
-void FControlRigEditor::OnGraphNodeDropToPerform(TSharedPtr<FDragDropOperation> InDragDropOp, UEdGraph* InGraph, const FVector2D& InNodePosition, const FVector2D& InScreenPosition)
+void FControlRigBaseEditor::OnGraphNodeDropToPerformImpl(TSharedPtr<FDragDropOperation> InDragDropOp, UEdGraph* InGraph, const FVector2D& InNodePosition, const FVector2D& InScreenPosition)
 {
 	DECLARE_SCOPE_HIERARCHICAL_COUNTER_FUNC()
 
@@ -5301,7 +5152,7 @@ void FControlRigEditor::OnGraphNodeDropToPerform(TSharedPtr<FDragDropOperation>
 	{
 		TSharedPtr<FRigElementHierarchyDragDropOp> RigHierarchyOp = StaticCastSharedPtr<FRigElementHierarchyDragDropOp>(InDragDropOp);
 
-		if (RigHierarchyOp->GetElements().Num() > 0 && FocusedGraphEdPtr.IsValid())
+		if (RigHierarchyOp->GetElements().Num() > 0 && GetFocusedGraphEd().IsValid())
 		{
 			const FName MenuName = RigHierarchyToGraphDragAndDropMenuName;
 			
@@ -5313,12 +5164,12 @@ void FControlRigEditor::OnGraphNodeDropToPerform(TSharedPtr<FDragDropOperation>
 					InGraph,
 					InNodePosition
 				);
-			MenuContext->Init(SharedThis(this), MenuSpecificContext);
+			MenuContext->Init(StaticCastSharedRef<FControlRigBaseEditor>(SharedControlRigEditorRef()), MenuSpecificContext);
 			
 			UToolMenus* ToolMenus = UToolMenus::Get();
 			TSharedRef<SWidget>	MenuWidget = ToolMenus->GenerateWidget(MenuName, FToolMenuContext(MenuContext));
 			
-			TSharedRef<SWidget> GraphEditorPanel = FocusedGraphEdPtr.Pin().ToSharedRef();
+			TSharedRef<SWidget> GraphEditorPanel = GetFocusedGraphEd().Pin().ToSharedRef();
 
 			// Show menu to choose getter vs setter
 			FSlateApplication::Get().PushMenu(
@@ -5333,10 +5184,10 @@ void FControlRigEditor::OnGraphNodeDropToPerform(TSharedPtr<FDragDropOperation>
 		return;
 	}
 
-	IControlRigEditor::OnGraphNodeDropToPerform(InDragDropOp, InGraph, InNodePosition, InScreenPosition);
+	OnGraphNodeDropToPerformSuper(InDragDropOp, InGraph, InNodePosition, InScreenPosition);
 }
 
-void FControlRigEditor::HandleMakeElementGetterSetter(ERigElementGetterSetterType Type, bool bIsGetter, TArray<FRigElementKey> Keys, UEdGraph* Graph, FVector2D NodePosition)
+void FControlRigBaseEditor::HandleMakeElementGetterSetter(ERigElementGetterSetterType Type, bool bIsGetter, TArray<FRigElementKey> Keys, UEdGraph* Graph, FVector2D NodePosition)
 {
 	DECLARE_SCOPE_HIERARCHICAL_COUNTER_FUNC()
 
@@ -5847,7 +5698,7 @@ void FControlRigEditor::HandleMakeElementGetterSetter(ERigElementGetterSetterTyp
 	}
 }
 
-void FControlRigEditor::HandleMakeComponentGetterSetter(bool bIsGetter, TArray<FRigComponentKey> Keys, UEdGraph* Graph, FVector2D NodePosition)
+void FControlRigBaseEditor::HandleMakeComponentGetterSetter(bool bIsGetter, TArray<FRigComponentKey> Keys, UEdGraph* Graph, FVector2D NodePosition)
 {
 	DECLARE_SCOPE_HIERARCHICAL_COUNTER_FUNC()
 
@@ -5950,15 +5801,15 @@ void FControlRigEditor::HandleMakeComponentGetterSetter(bool bIsGetter, TArray<F
 	}
 }
 
-void FControlRigEditor::HandleOnControlModified(UControlRig* Subject, FRigControlElement* ControlElement, const FRigControlModifiedContext& Context)
+void FControlRigBaseEditor::HandleOnControlModified(UControlRig* Subject, FRigControlElement* ControlElement, const FRigControlModifiedContext& Context)
 {
-	UControlRig* DebuggedControlRig = Cast<UControlRig>(FRigVMEditorBase::GetBlueprintObj()->GetObjectBeingDebugged());
+	UControlRig* DebuggedControlRig = Cast<UControlRig>(GetRigVMBlueprint()->GetObjectBeingDebugged());
 	if (Subject != DebuggedControlRig)
 	{
 		return;
 	}
 
-	UControlRigBlueprint* Blueprint = Cast<UControlRigBlueprint>(FRigVMEditorBase::GetBlueprintObj());
+	UControlRigBlueprint* Blueprint = Cast<UControlRigBlueprint>(GetRigVMBlueprint());
 	if (Blueprint == nullptr)
 	{
 		return;
@@ -6106,13 +5957,13 @@ void FControlRigEditor::HandleOnControlModified(UControlRig* Subject, FRigContro
 	}
 }
 
-void FControlRigEditor::HandleRefreshEditorFromBlueprint(URigVMBlueprint* InBlueprint)
+void FControlRigBaseEditor::HandleRefreshEditorFromBlueprintImpl(URigVMBlueprint* InBlueprint)
 {
 	OnHierarchyChanged();
-	IControlRigEditor::HandleRefreshEditorFromBlueprint(InBlueprint);
+	HandleRefreshEditorFromBlueprintSuper(InBlueprint);
 }
 
-UToolMenu* FControlRigEditor::HandleOnGetViewportContextMenuDelegate()
+UToolMenu* FControlRigBaseEditor::HandleOnGetViewportContextMenuDelegate()
 {
 	if (OnGetViewportContextMenuDelegate.IsBound())
 	{
@@ -6121,7 +5972,7 @@ UToolMenu* FControlRigEditor::HandleOnGetViewportContextMenuDelegate()
 	return nullptr;
 }
 
-TSharedPtr<FUICommandList> FControlRigEditor::HandleOnViewportContextMenuCommandsDelegate()
+TSharedPtr<FUICommandList> FControlRigBaseEditor::HandleOnViewportContextMenuCommandsDelegate()
 {
 	if (OnViewportContextMenuCommandsDelegate.IsBound())
 	{
@@ -6130,7 +5981,7 @@ TSharedPtr<FUICommandList> FControlRigEditor::HandleOnViewportContextMenuCommand
 	return TSharedPtr<FUICommandList>();
 }
 
-void FControlRigEditor::OnPreForwardsSolve_AnyThread(UControlRig* InRig, const FName& InEventName)
+void FControlRigBaseEditor::OnPreForwardsSolve_AnyThread(UControlRig* InRig, const FName& InEventName)
 {
 	// if we are debugging a PIE instance, we need to remember the input pose on the
 	// rig so we can perform multiple evaluations. this is to avoid double transforms / double forward solve results.
@@ -6150,7 +6001,7 @@ void FControlRigEditor::OnPreForwardsSolve_AnyThread(UControlRig* InRig, const F
 	}
 }
 
-void FControlRigEditor::OnPreConstructionForUI_AnyThread(UControlRig* InRig, const FName& InEventName)
+void FControlRigBaseEditor::OnPreConstructionForUI_AnyThread(UControlRig* InRig, const FName& InEventName)
 {
 	bIsConstructionEventRunning = true;
 
@@ -6176,7 +6027,7 @@ void FControlRigEditor::OnPreConstructionForUI_AnyThread(UControlRig* InRig, con
 	}
 }
 
-void FControlRigEditor::OnPreConstruction_AnyThread(UControlRig* InRig, const FName& InEventName)
+void FControlRigBaseEditor::OnPreConstruction_AnyThread(UControlRig* InRig, const FName& InEventName)
 {
 	if(UControlRigBlueprint* RigBlueprint = GetControlRigBlueprint())
 	{
@@ -6202,7 +6053,7 @@ void FControlRigEditor::OnPreConstruction_AnyThread(UControlRig* InRig, const FN
 	}
 }
 
-void FControlRigEditor::OnPostConstruction_AnyThread(UControlRig* InRig, const FName& InEventName)
+void FControlRigBaseEditor::OnPostConstruction_AnyThread(UControlRig* InRig, const FName& InEventName)
 {
 	bIsConstructionEventRunning = false;
 	const bool bShouldExecute = ShouldExecuteControlRig(InRig);
@@ -6285,25 +6136,26 @@ void FControlRigEditor::OnPostConstruction_AnyThread(UControlRig* InRig, const F
 	}
 }
 
-void FControlRigEditor::SetupTimelineDelegates(FAnimationScrubPanelDelegates& InOutDelegates)
+void FControlRigBaseEditor::SetupTimelineDelegates(FAnimationScrubPanelDelegates& InOutDelegates)
 {
-	InOutDelegates.IsRecordingActiveDelegate.BindSP(this, &FControlRigEditor::HandleReplayIsRecordingActive);
-	InOutDelegates.GetRecordingVisibilityDelegate.BindSP(this, &FControlRigEditor::HandleGetReplayRecordButtonVisibility);
-	InOutDelegates.StartRecordingDelegate.BindSP(this, &FControlRigEditor::HandleReplayStartRecording);
-	InOutDelegates.StopRecordingDelegate.BindSP(this, &FControlRigEditor::HandleReplayStopRecording);
-	InOutDelegates.GetPlaybackModeDelegate.BindSP(this, &FControlRigEditor::HandleReplayGetPlaybackMode);
-	InOutDelegates.SetPlaybackModeDelegate.BindSP(this, &FControlRigEditor::HandleReplaySetPlaybackMode);
-	InOutDelegates.GetPlaybackTimeDelegate.BindSP(this, &FControlRigEditor::HandleReplayGetPlaybackTime);
-	InOutDelegates.SetPlaybackTimeDelegate.BindSP(this, &FControlRigEditor::HandleReplaySetPlaybackTime);
-	InOutDelegates.StepForwardDelegate.BindSP(this, &FControlRigEditor::HandleReplayStepForward);
-	InOutDelegates.StepBackwardDelegate.BindSP(this, &FControlRigEditor::HandleReplayStepBackward);
-	InOutDelegates.GetIsLoopingDelegate.BindSP(this, &FControlRigEditor::HandleReplayGetIsLooping);
-	InOutDelegates.SetIsLoopingDelegate.BindSP(this, &FControlRigEditor::HandleReplaySetIsLooping);
-	InOutDelegates.GetPlaybackTimeRangeDelegate.BindSP(this, &FControlRigEditor::HandleReplayGetPlaybackTimeRange);
-	InOutDelegates.GetNumberOfKeysDelegate.BindSP(this, &FControlRigEditor::HandleReplayGetNumberOfKeys);
+	TSharedRef<FControlRigBaseEditor> SharedEditor = SharedRef();
+	InOutDelegates.IsRecordingActiveDelegate.BindSP(SharedEditor, &FControlRigBaseEditor::HandleReplayIsRecordingActive);
+	InOutDelegates.GetRecordingVisibilityDelegate.BindSP(SharedEditor, &FControlRigBaseEditor::HandleGetReplayRecordButtonVisibility);
+	InOutDelegates.StartRecordingDelegate.BindSP(SharedEditor, &FControlRigBaseEditor::HandleReplayStartRecording);
+	InOutDelegates.StopRecordingDelegate.BindSP(SharedEditor, &FControlRigBaseEditor::HandleReplayStopRecording);
+	InOutDelegates.GetPlaybackModeDelegate.BindSP(SharedEditor, &FControlRigBaseEditor::HandleReplayGetPlaybackMode);
+	InOutDelegates.SetPlaybackModeDelegate.BindSP(SharedEditor, &FControlRigBaseEditor::HandleReplaySetPlaybackMode);
+	InOutDelegates.GetPlaybackTimeDelegate.BindSP(SharedEditor, &FControlRigBaseEditor::HandleReplayGetPlaybackTime);
+	InOutDelegates.SetPlaybackTimeDelegate.BindSP(SharedEditor, &FControlRigBaseEditor::HandleReplaySetPlaybackTime);
+	InOutDelegates.StepForwardDelegate.BindSP(SharedEditor, &FControlRigBaseEditor::HandleReplayStepForward);
+	InOutDelegates.StepBackwardDelegate.BindSP(SharedEditor, &FControlRigBaseEditor::HandleReplayStepBackward);
+	InOutDelegates.GetIsLoopingDelegate.BindSP(SharedEditor, &FControlRigBaseEditor::HandleReplayGetIsLooping);
+	InOutDelegates.SetIsLoopingDelegate.BindSP(SharedEditor, &FControlRigBaseEditor::HandleReplaySetIsLooping);
+	InOutDelegates.GetPlaybackTimeRangeDelegate.BindSP(SharedEditor, &FControlRigBaseEditor::HandleReplayGetPlaybackTimeRange);
+	InOutDelegates.GetNumberOfKeysDelegate.BindSP(SharedEditor, &FControlRigBaseEditor::HandleReplayGetNumberOfKeys);
 }
 
-bool FControlRigEditor::ShowReplayOnTimeline() const
+bool FControlRigBaseEditor::ShowReplayOnTimeline() const
 {
 	if(ReplayStrongPtr.IsValid())
 	{
@@ -6315,7 +6167,7 @@ bool FControlRigEditor::ShowReplayOnTimeline() const
 	return false;
 }
 
-TOptional<bool> FControlRigEditor::HandleReplayIsRecordingActive() const
+TOptional<bool> FControlRigBaseEditor::HandleReplayIsRecordingActive() const
 {
 	if(!ShowReplayOnTimeline())
 	{
@@ -6324,7 +6176,7 @@ TOptional<bool> FControlRigEditor::HandleReplayIsRecordingActive() const
 	return ReplayStrongPtr->IsRecording();
 }
 
-TOptional<EVisibility> FControlRigEditor::HandleGetReplayRecordButtonVisibility() const
+TOptional<EVisibility> FControlRigBaseEditor::HandleGetReplayRecordButtonVisibility() const
 {
 	if(ReplayStrongPtr.IsValid())
 	{
@@ -6333,19 +6185,19 @@ TOptional<EVisibility> FControlRigEditor::HandleGetReplayRecordButtonVisibility(
 	return TOptional<EVisibility>();
 }
 
-bool FControlRigEditor::HandleReplayStartRecording()
+bool FControlRigBaseEditor::HandleReplayStartRecording()
 {
 	// for now this button won't be supported from here.
 	return false;
 }
 
-bool FControlRigEditor::HandleReplayStopRecording()
+bool FControlRigBaseEditor::HandleReplayStopRecording()
 {
 	// for now this button won't be supported from here.
 	return false;
 }
 
-TOptional<int32> FControlRigEditor::HandleReplayGetPlaybackMode() const
+TOptional<int32> FControlRigBaseEditor::HandleReplayGetPlaybackMode() const
 {
 	if(!ShowReplayOnTimeline())
 	{
@@ -6354,7 +6206,7 @@ TOptional<int32> FControlRigEditor::HandleReplayGetPlaybackMode() const
 	return (ReplayStrongPtr->IsReplaying() && !ReplayStrongPtr->IsPaused()) ? (int32)EPlaybackMode::PlayingForward : (int32)EPlaybackMode::Stopped; 
 }
 
-bool FControlRigEditor::HandleReplaySetPlaybackMode(int32 InPlaybackMode)
+bool FControlRigBaseEditor::HandleReplaySetPlaybackMode(int32 InPlaybackMode)
 {
 	if(!ShowReplayOnTimeline())
 	{
@@ -6371,7 +6223,7 @@ bool FControlRigEditor::HandleReplaySetPlaybackMode(int32 InPlaybackMode)
 	return true;
 }
 
-TOptional<float> FControlRigEditor::HandleReplayGetPlaybackTime() const
+TOptional<float> FControlRigBaseEditor::HandleReplayGetPlaybackTime() const
 {
 	if(!ShowReplayOnTimeline())
 	{
@@ -6380,7 +6232,7 @@ TOptional<float> FControlRigEditor::HandleReplayGetPlaybackTime() const
 	return GetControlRig()->GetAbsoluteTime();
 }
 
-bool FControlRigEditor::HandleReplaySetPlaybackTime(float InTime, bool bStopPlayback)
+bool FControlRigBaseEditor::HandleReplaySetPlaybackTime(float InTime, bool bStopPlayback)
 {
 	if(!ShowReplayOnTimeline())
 	{
@@ -6399,7 +6251,7 @@ bool FControlRigEditor::HandleReplaySetPlaybackTime(float InTime, bool bStopPlay
 	return true;
 }
 
-bool FControlRigEditor::HandleReplayStepForward()
+bool FControlRigBaseEditor::HandleReplayStepForward()
 {
 	if(!ShowReplayOnTimeline())
 	{
@@ -6416,7 +6268,7 @@ bool FControlRigEditor::HandleReplayStepForward()
 	return true;
 }
 
-bool FControlRigEditor::HandleReplayStepBackward()
+bool FControlRigBaseEditor::HandleReplayStepBackward()
 {
 	if(!ShowReplayOnTimeline())
 	{
@@ -6433,7 +6285,7 @@ bool FControlRigEditor::HandleReplayStepBackward()
 	return true;
 }
 
-TOptional<bool> FControlRigEditor::HandleReplayGetIsLooping() const
+TOptional<bool> FControlRigBaseEditor::HandleReplayGetIsLooping() const
 {
 	if(!ShowReplayOnTimeline())
 	{
@@ -6443,13 +6295,13 @@ TOptional<bool> FControlRigEditor::HandleReplayGetIsLooping() const
 	return true;
 }
 
-bool FControlRigEditor::HandleReplaySetIsLooping(bool bIsLooping)
+bool FControlRigBaseEditor::HandleReplaySetIsLooping(bool bIsLooping)
 {
 	// we are always set to loop for now.
 	return true;
 }
 
-TOptional<FVector2f> FControlRigEditor::HandleReplayGetPlaybackTimeRange() const
+TOptional<FVector2f> FControlRigBaseEditor::HandleReplayGetPlaybackTimeRange() const
 {
 	if(!ShowReplayOnTimeline())
 	{
@@ -6458,7 +6310,7 @@ TOptional<FVector2f> FControlRigEditor::HandleReplayGetPlaybackTimeRange() const
 	return FVector2f(ReplayStrongPtr->GetTimeRange());
 }
 
-TOptional<uint32> FControlRigEditor::HandleReplayGetNumberOfKeys() const
+TOptional<uint32> FControlRigBaseEditor::HandleReplayGetNumberOfKeys() const
 {
 	if(!ShowReplayOnTimeline())
 	{
@@ -6467,7 +6319,7 @@ TOptional<uint32> FControlRigEditor::HandleReplayGetNumberOfKeys() const
 	return static_cast<uint32>(ReplayStrongPtr->InputTracks.GetNumTimes());
 }
 
-EVisibility FControlRigEditor::GetReplayValidationErrorVisibility() const
+EVisibility FControlRigBaseEditor::GetReplayValidationErrorVisibility() const
 {
 	if(ReplayStrongPtr.IsValid())
 	{
@@ -6481,7 +6333,7 @@ EVisibility FControlRigEditor::GetReplayValidationErrorVisibility() const
 	return EVisibility::Hidden;
 }
 
-FText FControlRigEditor::GetReplayValidationErrorTooltip() const
+FText FControlRigBaseEditor::GetReplayValidationErrorTooltip() const
 {
 	if(ReplayStrongPtr.IsValid())
 	{
diff --git a/Engine/Plugins/Animation/ControlRig/Source/ControlRigEditor/Private/Editor/ControlRigEditor.h b/Engine/Plugins/Animation/ControlRig/Source/ControlRigEditor/Private/Editor/ControlRigEditor.h
--- a/Engine/Plugins/Animation/ControlRig/Source/ControlRigEditor/Private/Editor/ControlRigEditor.h
+++ b/Engine/Plugins/Animation/ControlRig/Source/ControlRigEditor/Private/Editor/ControlRigEditor.h
@@ -36,7 +36,7 @@ class SBorder;
 class USkeletalMesh;
 class FStructOnScope;
 class UToolMenu;
-class FControlRigEditor;
+class FControlRigBaseEditor;
 
 struct FControlRigEditorModes
 {
@@ -60,108 +60,222 @@ private:
 	FControlRigEditorModes() {}
 };
 
-class FControlRigEditor : public IControlRigEditor
+class IControlRigBaseEditor
 {
 public:
+	virtual TSharedRef<IRigVMEditor> SharedRigVMEditorRef() = 0;
+	virtual TSharedRef<IControlRigBaseEditor> SharedControlRigEditorRef() = 0;
+	virtual TSharedRef<FControlRigBaseEditor> SharedRef() = 0;	
+	virtual TSharedRef<const IRigVMEditor> SharedRigVMEditorRef() const = 0;
+	
+	virtual UControlRigBlueprint* GetControlRigBlueprint() const = 0;
+	virtual TSharedPtr<FAssetEditorToolkit> GetHostingApp() = 0;
+	virtual UControlRig* GetControlRig() const = 0;
+
+	virtual TSharedRef<IPersonaToolkit> GetPersonaToolkit() const = 0;
+
+	virtual void Compile() = 0;
+
+	DECLARE_MULTICAST_DELEGATE_TwoParams(FControlRigEditorClosed, IControlRigBaseEditor*, UControlRigBlueprint*);
+	virtual FControlRigEditorClosed& OnEditorClosed() = 0;
+	virtual FPersonaViewportKeyDownDelegate& GetKeyDownDelegate() = 0;
+	virtual FOnGetContextMenu& OnGetViewportContextMenu() = 0;
+	virtual FNewMenuCommandsDelegate& OnViewportContextMenuCommands() = 0;
+	virtual FSimpleMulticastDelegate& OnRequestNavigateToConnectorWarning() = 0;
+	virtual FControlRigEditorEditMode* GetEditMode() const = 0;
+	virtual int32 GetEventQueueComboValue() const = 0;
+	virtual void OnHierarchyChanged() = 0;
+
+	virtual void SetDetailViewForRigElements() = 0;
+	virtual void SetDetailViewForRigElements(const TArray<FRigHierarchyKey>& InKeys) = 0;
+	virtual void ClearDetailObject(bool bChangeUISelectionState = true) = 0;
+	virtual void FindReferencesOfItem(const FRigHierarchyKey& InKey) = 0;
+	
+	virtual int32 GetRigHierarchyTabCount() const = 0;
+	virtual void IncreaseRigHierarchyTabCount() = 0;
+	virtual void DecreaseRigHierarchyTabCount() = 0;
+
+	virtual int32 GetModularRigHierarchyTabCount() const = 0;
+	virtual void IncreaseModularRigHierarchyTabCount() = 0;
+	virtual void DecreaseModularRigHierarchyTabCount() = 0;
+	
+	virtual bool& GetSuspendDetailsPanelRefreshFlag() = 0;
+	virtual EControlRigReplayPlaybackMode GetReplayPlaybackMode() const = 0;
+	virtual TArray<FName> GetSelectedModules() const = 0;
 
-	virtual void InitRigVMEditor(const EToolkitMode::Type Mode, const TSharedPtr< class IToolkitHost >& InitToolkitHost, class URigVMBlueprint* InRigVMBlueprint) override;
+	virtual void RefreshDetailView() = 0;
+	virtual FVector2D ComputePersonaProjectedScreenPos(const FVector& InWorldPos, bool bClampToScreenRectangle = false) = 0;
+	virtual void SetDetailViewForRigModules(const TArray<FName> InModuleNames) = 0;
 
-	virtual const FName GetEditorAppName() const override;
-	virtual const FName GetEditorModeName() const override;
-	virtual TSharedPtr<FApplicationMode> CreateEditorMode() override;
-	virtual const FSlateBrush* GetDefaultTabIcon() const override;
+	virtual UAnimPreviewInstance* GetPreviewInstance() const = 0;
+	virtual void SetPreviewInstance(UAnimPreviewInstance* InPreviewInstance) = 0;
+	
+	virtual void RemoveBoneModification(FName BoneName) = 0;
+	virtual URigHierarchy* GetHierarchyBeingDebugged() const = 0;
 
+	virtual void FilterDraggedKeys(TArray<FRigElementKey>& Keys, bool bRemoveNameSpace) = 0;
+	virtual URigVMController* GetFocusedController() const = 0;
+	
+protected:
+
+	virtual void InitRigVMEditorSuper(const EToolkitMode::Type Mode, const TSharedPtr< class IToolkitHost >& InitToolkitHost, class URigVMBlueprint* InRigVMBlueprint) = 0;
+
+	virtual bool IsControlRigLegacyEditor() const = 0;
+	virtual URigVMBlueprint* GetRigVMBlueprint() const = 0;
+	virtual URigVMHost* GetRigVMHost() const = 0;
+	
+	
+	virtual TSharedRef<FUICommandList> GetToolkitCommands() = 0;
+	virtual FPreviewScene* GetPreviewScene() = 0;
+	virtual bool IsDetailsPanelRefreshSuspended() const = 0;
+	virtual TArray<TWeakObjectPtr<UObject>> GetSelectedObjects() const = 0;
+	virtual UClass* GetDetailWrapperClass() const = 0;
+	virtual void OnWrappedPropertyChangedChainEvent(URigVMDetailsViewWrapperObject* InWrapperObject, const FString& InPropertyPath, FPropertyChangedChainEvent& InPropertyChangedChainEvent) = 0;
+	virtual void SetDetailObjects(const TArray<UObject*>& InObjects) = 0;
+	virtual void SetDetailObjectsSuper(const TArray<UObject*>& InObjects) = 0;
+	virtual bool DetailViewShowsStruct(UScriptStruct* InStruct) const = 0;
+	virtual TSharedPtr<class SWidget> GetInspector() const = 0;
+	virtual TArray<FName> GetEventQueue() const = 0;
+	virtual void SummonSearchUI(bool bSetFindWithinBlueprint, FString NewSearchTerms = FString(), bool bSelectFirstResult = false) = 0;
+	virtual const TArray< UObject* >* GetObjectsCurrentlyBeingEdited() const = 0;
+	virtual void UpdateRigVMHost() = 0;
+	virtual void RefreshDetailViewSuper() = 0;
+	virtual void CacheNameLists() = 0;
+	virtual FEditorModeTools& GetEditorModeManagerImpl() const = 0;
+	virtual const FName GetEditorModeNameImpl() const = 0;
+	virtual URigVMGraph* GetFocusedModel() const = 0;
+	virtual UObject* GetOuterForHostSuper() const = 0;
+	virtual void CompileSuper() = 0;
+	virtual void HandleModifiedEventSuper(ERigVMGraphNotifType InNotifType, URigVMGraph* InGraph, UObject* InSubject) = 0;
+	virtual void OnCreateGraphEditorCommandsSuper(TSharedPtr<FUICommandList> GraphEditorCommandsList) = 0;
+	virtual void HandleVMCompiledEventSuper(UObject* InCompiledObject, URigVM* InVM, FRigVMExtendedExecuteContext& InContext) = 0;
+	virtual FReply OnViewportDropSuper(const FGeometry& MyGeometry, const FDragDropEvent& DragDropEvent) = 0;
+	virtual void FillToolbarSuper(FToolBarBuilder& ToolbarBuilder, bool bEndSection = true) = 0;
+	virtual TArray<FName> GetLastEventQueue() const = 0;
+	virtual int32 GetEventQueueComboValueSuper() const = 0;
+	virtual void HandleSetObjectBeingDebuggedSuper(UObject* InObject) = 0;
+	virtual void SetEventQueue(TArray<FName> InEventQueue, bool bCompile) = 0;
+	virtual void SetEventQueueSuper(TArray<FName> InEventQueue, bool bCompile) = 0;
+	virtual void SetEventQueueSuper(TArray<FName> InEventQueue) = 0;
+	virtual void SaveAsset_ExecuteSuper() = 0;
+	virtual void SaveAssetAs_ExecuteSuper() = 0;
+	virtual FReply OnSpawnGraphNodeByShortcutSuper(FInputChord InChord, const FVector2D& InPosition, UEdGraph* InGraph) = 0;
+	virtual void HandleVMExecutedEventSuper(URigVMHost* InHost, const FName& InEventName) = 0;
+	virtual void SetEditorModeManager(TSharedPtr<FEditorModeTools> InManager) = 0;
+	virtual void TickSuper(float DeltaTime) = 0;
+	virtual void UpdateRigVMHostSuper() = 0;
+	virtual void CacheNameListsSuper() = 0;
+	virtual void OnFinishedChangingPropertiesSuper(const FPropertyChangedEvent& PropertyChangedEvent) = 0;
+	virtual void OnWrappedPropertyChangedChainEventSuper(URigVMDetailsViewWrapperObject* InWrapperObject, const FString& InPropertyPath, FPropertyChangedChainEvent& InPropertyChangedChainEvent) = 0;
+	virtual const TArray<TStrongObjectPtr<URigVMDetailsViewWrapperObject>>& GetWrapperObjects() const = 0;
+	virtual void BindCommandsSuper() = 0;
+	virtual FMenuBuilder GenerateBulkEditMenuSuper() = 0;
+	virtual TWeakPtr<class SGraphEditor> GetFocusedGraphEd() const = 0;
+	virtual void OnGraphNodeDropToPerformSuper(TSharedPtr<FDragDropOperation> InDragDropOp, UEdGraph* InGraph, const FVector2D& InNodePosition, const FVector2D& InScreenPosition) = 0;
+	virtual void HandleRefreshEditorFromBlueprintSuper(URigVMBlueprint* InBlueprint) = 0;
+	virtual void OnGraphNodeDropToPerform(TSharedPtr<FDragDropOperation> InDragDropOp, UEdGraph* InGraph, const FVector2D& InNodePosition, const FVector2D& InScreenPosition) = 0;
+	
+	virtual void OnCloseSuper() = 0;
+
+};
+
+class FControlRigBaseEditor : public IControlRigBaseEditor
+{
 public:
-	FControlRigEditor();
-	virtual ~FControlRigEditor();
-
-	virtual UObject* GetOuterForHost() const override;
-
-	// FRigVMEditorBase interface
-	virtual UClass* GetDetailWrapperClass() const;
-	virtual void Compile() override;
-	virtual void HandleModifiedEvent(ERigVMGraphNotifType InNotifType, URigVMGraph* InGraph, UObject* InSubject) override;
-	virtual void OnCreateGraphEditorCommands(TSharedPtr<FUICommandList> GraphEditorCommandsList) override;
-	UE_DEPRECATED(5.4, "Please use HandleVMCompiledEvent with ExtendedExecuteContext parameter.")
-	virtual void HandleVMCompiledEvent(UObject* InCompiledObject, URigVM* InVM) override {}
-	virtual void HandleVMCompiledEvent(UObject* InCompiledObject, URigVM* InVM, FRigVMExtendedExecuteContext& InContext) override;
-	virtual bool ShouldOpenGraphByDefault() const override { return !IsModularRig(); }
-	virtual FReply OnViewportDrop(const FGeometry& MyGeometry, const FDragDropEvent& DragDropEvent) override;
+	static FControlRigBaseEditor* GetFromAssetEditorInstance(IAssetEditorInstance* Instance);
+	virtual TSharedRef<FControlRigBaseEditor> SharedRef() override { return StaticCastSharedRef<FControlRigBaseEditor>(SharedControlRigEditorRef());}
+
+protected:
+	void InitRigVMEditorImpl(const EToolkitMode::Type Mode, const TSharedPtr< class IToolkitHost >& InitToolkitHost, class URigVMBlueprint* InRigVMBlueprint);
+
+	virtual const FName GetEditorAppNameImpl() const;
+	virtual const FName GetEditorModeNameImpl() const;
+	virtual const FSlateBrush* GetDefaultTabIconImpl() const;
+
+
+	FControlRigBaseEditor();
+	virtual ~FControlRigBaseEditor() {}
+
+	virtual UObject* GetOuterForHostImpl() const;
+
+	virtual UClass* GetDetailWrapperClassImpl() const;
+	virtual void CompileBaseImpl();
+	virtual void HandleModifiedEventImpl(ERigVMGraphNotifType InNotifType, URigVMGraph* InGraph, UObject* InSubject);
+	virtual void OnCreateGraphEditorCommandsImpl(TSharedPtr<FUICommandList> GraphEditorCommandsList);
+	virtual void HandleVMCompiledEventImpl(UObject* InCompiledObject, URigVM* InVM, FRigVMExtendedExecuteContext& InContext);
+	virtual bool ShouldOpenGraphByDefaultImpl() const { return !IsModularRig(); }
+	virtual FReply OnViewportDropImpl(const FGeometry& MyGeometry, const FDragDropEvent& DragDropEvent);
 
 	// allows the editor to fill an empty graph
-	virtual void CreateEmptyGraphContent(URigVMController* InController) override;
+	virtual void CreateEmptyGraphContentImpl(URigVMController* InController);
 
-	int32 GetRigHierarchyTabCount() const { return RigHierarchyTabCount; }
-	int32 GetModularRigHierarchyTabCount() const { return ModularRigHierarchyTabCount; }
+	virtual int32 GetRigHierarchyTabCount() const override { return RigHierarchyTabCount; }
+	virtual void IncreaseRigHierarchyTabCount() override { RigHierarchyTabCount++; }
+	virtual void DecreaseRigHierarchyTabCount() override { RigHierarchyTabCount--; }
+	virtual int32 GetModularRigHierarchyTabCount() const override { return ModularRigHierarchyTabCount; }
+	virtual void IncreaseModularRigHierarchyTabCount() override { ModularRigHierarchyTabCount++; }
+	virtual void DecreaseModularRigHierarchyTabCount() override { ModularRigHierarchyTabCount--; }
 
 	bool IsModularRig() const;
 	bool IsRigModule() const;
 
-public:
+
 	
 	// IToolkit Interface
-	virtual FName GetToolkitFName() const override;
-	virtual FText GetBaseToolkitName() const override;
-	virtual FString GetWorldCentricTabPrefix() const override;
-	virtual FString GetDocumentationLink() const override
+	virtual FName GetToolkitFNameImpl() const;
+	virtual FText GetBaseToolkitNameImpl() const;
+	virtual FString GetWorldCentricTabPrefixImpl() const;
+	virtual FString GetDocumentationLinkImpl() const 
 	{
 		return FString(TEXT("Engine/Animation/ControlRig"));
 	}
 
 	// BlueprintEditor interface
-	virtual FReply OnSpawnGraphNodeByShortcut(FInputChord InChord, const FVector2D& InPosition, UEdGraph* InGraph) override;
-#if WITH_RIGVMLEGACYEDITOR
-	virtual bool IsSectionVisible(NodeSectionID::Type InSectionID) const override;
-	virtual bool NewDocument_IsVisibleForType(FBlueprintEditor::ECreatedDocumentType GraphType) const override;
-#else
-	virtual bool IsSectionVisible(RigVMNodeSectionID::Type InSectionID) const override;
-	virtual bool NewDocument_IsVisibleForType(FRigVMEditorBase::ECreatedDocumentType GraphType) const override;
-#endif
+	virtual FReply OnSpawnGraphNodeByShortcutImpl(FInputChord InChord, const FVector2D& InPosition, UEdGraph* InGraph);
 
-	virtual void PostTransaction(bool bSuccess, const FTransaction* Transaction, bool bIsRedo) override;
+	virtual void PostTransactionImpl(bool bSuccess, const FTransaction* Transaction, bool bIsRedo);
 
 	void EnsureValidRigElementsInDetailPanel();
 
 	//  FTickableEditorObject Interface
-	virtual void Tick(float DeltaTime) override;
+	virtual void TickImpl(float DeltaTime);
 
 	// Gets the Control Rig Blueprint being edited/viewed
-	UControlRigBlueprint* GetControlRigBlueprint() const;
+	virtual UControlRigBlueprint* GetControlRigBlueprint() const override;
 
 	// returns the hierarchy being debugged
-	UControlRig* GetControlRig() const;
+	virtual UControlRig* GetControlRig() const override;
 
 	// returns the hierarchy being debugged
-	URigHierarchy* GetHierarchyBeingDebugged() const;
+	virtual URigHierarchy* GetHierarchyBeingDebugged() const override;
 
-	void SetDetailViewForRigElements();
-	void SetDetailViewForRigElements(const TArray<FRigHierarchyKey>& InKeys);
+	virtual void SetDetailViewForRigElements() override;
+	virtual void SetDetailViewForRigElements(const TArray<FRigHierarchyKey>& InKeys) override;
 	bool DetailViewShowsAnyRigElement() const;
 	bool DetailViewShowsRigElement(FRigHierarchyKey InKey) const;
 	TArray<FRigHierarchyKey> GetSelectedRigElementsFromDetailView() const;
 	TArray< TWeakObjectPtr<UObject> > GetSelectedObjectsFromDetailView() const;
 
 	void SetDetailViewForRigModules();
-	void SetDetailViewForRigModules(const TArray<FName> InModuleNames);
+	virtual void SetDetailViewForRigModules(const TArray<FName> InModuleNames) override;
 	bool DetailViewShowsAnyRigModule() const;
 	bool DetailViewShowsRigModule(FName InModuleName) const;
+
 	TArray<FName> ModulesSelected;
+	virtual TArray<FName> GetSelectedModules() const override { return ModulesSelected; }
 
-	virtual void SetDetailObjects(const TArray<UObject*>& InObjects) override;
-	virtual void RefreshDetailView() override;
+	virtual void SetDetailObjectsImpl(const TArray<UObject*>& InObjects);
+	virtual void RefreshDetailViewImpl();
 
 	void CreatePersonaToolKitIfRequired();
 
 public:
-
 	/** Get the persona toolkit */
-	TSharedRef<IPersonaToolkit> GetPersonaToolkit() const { return PersonaToolkit.ToSharedRef(); }
+	virtual TSharedRef<IPersonaToolkit> GetPersonaToolkit() const override { return PersonaToolkit.ToSharedRef(); }
+protected:
 
 	/** Get the edit mode */
-	FControlRigEditorEditMode* GetEditMode() const;
-
-	// this changes everytime you compile, so don't cache it expecting it will last. 
-	UControlRig* GetInstanceRig() const { return GetControlRig();  }
+	virtual FControlRigEditorEditMode* GetEditMode() const override;
 
 	void OnCurveContainerChanged();
 
@@ -177,12 +291,12 @@ public:
 
 	const FName RigHierarchyToGraphDragAndDropMenuName = TEXT("ControlRigEditor.RigHierarchyToGraphDragAndDropMenu");
 	void CreateRigHierarchyToGraphDragAndDropMenu() const;
-	virtual void OnGraphNodeDropToPerform(TSharedPtr<FDragDropOperation> InDragDropOp, UEdGraph* InGraph, const FVector2D& InNodePosition, const FVector2D& InScreenPosition) override;
+	virtual void OnGraphNodeDropToPerformImpl(TSharedPtr<FDragDropOperation> InDragDropOp, UEdGraph* InGraph, const FVector2D& InNodePosition, const FVector2D& InScreenPosition);
 
-	FPersonaViewportKeyDownDelegate& GetKeyDownDelegate() { return OnKeyDownDelegate; }
+	virtual FPersonaViewportKeyDownDelegate& GetKeyDownDelegate() override { return OnKeyDownDelegate; }
 
-	FOnGetContextMenu& OnGetViewportContextMenu() { return OnGetViewportContextMenuDelegate; }
-	FNewMenuCommandsDelegate& OnViewportContextMenuCommands() { return OnViewportContextMenuCommandsDelegate; }
+	virtual FOnGetContextMenu& OnGetViewportContextMenu() override { return OnGetViewportContextMenuDelegate; }
+	virtual FNewMenuCommandsDelegate& OnViewportContextMenuCommands() override { return OnViewportContextMenuCommandsDelegate; }
 
 	// DirectManipulation functionality
 	void HandleRequestDirectManipulationPosition() const { (void)HandleRequestDirectManipulation(ERigControlType::Position); }
@@ -203,52 +317,54 @@ public:
 	EVisibility GetConnectorWarningVisibility() const;
 	FText GetConnectorWarningText() const;
 	FReply OnNavigateToConnectorWarning() const;
-	FSimpleMulticastDelegate& OnRequestNavigateToConnectorWarning() { return RequestNavigateToConnectorWarningDelegate; }
+	virtual FSimpleMulticastDelegate& OnRequestNavigateToConnectorWarning() override { return RequestNavigateToConnectorWarningDelegate; }
 
-	FVector2D ComputePersonaProjectedScreenPos(const FVector& InWorldPos, bool bClampToScreenRectangle = false);
+	virtual FVector2D ComputePersonaProjectedScreenPos(const FVector& InWorldPos, bool bClampToScreenRectangle = false) override;
 
-	void FindReferencesOfItem(const FRigHierarchyKey& InKey);
+	virtual void FindReferencesOfItem(const FRigHierarchyKey& InKey) override;
 	
-protected:
 
-	virtual void BindCommands() override;
-	virtual FMenuBuilder GenerateBulkEditMenu() override;
 
-	void OnHierarchyChanged();
+	virtual void BindCommandsImpl();
+	virtual FMenuBuilder GenerateBulkEditMenuImpl();
+
+	virtual void OnHierarchyChanged() override;
 
 	void SynchronizeViewportBoneSelection();
 
-	virtual void SaveAsset_Execute() override;
-	virtual void SaveAssetAs_Execute() override;
+	virtual void SaveAsset_ExecuteImpl();
+	virtual void SaveAssetAs_ExecuteImpl();
 
 	// update the cached modification value
 	void UpdateBoneModification(FName BoneName, const FTransform& Transform);
 
 	// remove a single bone modification across all instance
-	void RemoveBoneModification(FName BoneName);
+	virtual void RemoveBoneModification(FName BoneName) override;
 
 	// reset all bone modification across all instance
 	void ResetAllBoneModification();
 
-	virtual void HandleVMExecutedEvent(URigVMHost* InHost, const FName& InEventName) override;
+	virtual void HandleVMExecutedEventImpl(URigVMHost* InHost, const FName& InEventName);
 
 	// FBaseToolKit overrides
-	virtual void CreateEditorModeManager() override;
+	virtual void CreateEditorModeManagerImpl();
 
-private:
 	/** Fill the toolbar with content */
-	virtual void FillToolbar(FToolBarBuilder& ToolbarBuilder, bool bEndSection = true) override;
+	virtual void FillToolbarImpl(FToolBarBuilder& ToolbarBuilder, bool bEndSection = true);
+
+	virtual TArray<FName> GetDefaultEventQueueImpl() const;
+	virtual void SetEventQueueImpl(TArray<FName> InEventQueue, bool bCompile);
+	virtual int32 GetEventQueueComboValueImpl() const;
+	virtual FText GetEventQueueLabelImpl() const;
+	virtual FSlateIcon GetEventQueueIconImpl(const TArray<FName>& InEventQueue) const;
+	virtual void HandleSetObjectBeingDebuggedImpl(UObject* InObject);
 
-	virtual TArray<FName> GetDefaultEventQueue() const override;
-	virtual void SetEventQueue(TArray<FName> InEventQueue, bool bCompile) override;
-	virtual int32 GetEventQueueComboValue() const override;
-	virtual FText GetEventQueueLabel() const override;
-	virtual FSlateIcon GetEventQueueIcon(const TArray<FName>& InEventQueue) const override;
-	virtual void HandleSetObjectBeingDebugged(UObject* InObject) override;
 
 	/** Handle preview scene setup */
 	void HandlePreviewSceneCreated(const TSharedRef<IPersonaPreviewScene>& InPersonaPreviewScene);
+public:
 	void HandleViewportCreated(const TSharedRef<class IPersonaViewport>& InViewport);
+protected:
 	void HandleToggleControlVisibility();
 	bool AreControlsVisible() const;
 	void HandleToggleControlsAsOverlay();
@@ -269,16 +385,16 @@ private:
 	void HandlePreviewMeshChanged(USkeletalMesh* InOldSkeletalMesh, USkeletalMesh* InNewSkeletalMesh);
 
 	/** Push a newly compiled/opened control rig to the edit mode */
-	virtual void UpdateRigVMHost() override;
-	virtual void UpdateRigVMHost_PreClearOldHost(URigVMHost* InPreviousHost) override;
+	virtual void UpdateRigVMHostImpl();
+	virtual void UpdateRigVMHost_PreClearOldHostImpl(URigVMHost* InPreviousHost);
 
 	/** Update the name lists for use in name combo boxes */
-	virtual void CacheNameLists() override;
+	virtual void CacheNameListsImpl();
 
 	/** Rebind our anim instance to the preview's skeletal mesh component */
 	void RebindToSkeletalMeshComponent();
 
-	virtual void GenerateEventQueueMenuContent(FMenuBuilder& MenuBuilder) override;
+	virtual void GenerateEventQueueMenuContentImpl(FMenuBuilder& MenuBuilder);
 
 	enum ERigElementGetterSetterType
 	{
@@ -291,13 +407,13 @@ private:
 		ERigElementGetterSetterType_Name
 	};
 
-	void FilterDraggedKeys(TArray<FRigElementKey>& Keys, bool bRemoveNameSpace);
+	virtual void FilterDraggedKeys(TArray<FRigElementKey>& Keys, bool bRemoveNameSpace) override;
 	void HandleMakeElementGetterSetter(ERigElementGetterSetterType Type, bool bIsGetter, TArray<FRigElementKey> Keys, UEdGraph* Graph, FVector2D NodePosition);
 	void HandleMakeComponentGetterSetter(bool bIsGetter, TArray<FRigComponentKey> Keys, UEdGraph* Graph, FVector2D NodePosition);
 
 	void HandleOnControlModified(UControlRig* Subject, FRigControlElement* ControlElement, const FRigControlModifiedContext& Context);
 
-	virtual void HandleRefreshEditorFromBlueprint(URigVMBlueprint* InBlueprint) override;
+	virtual void HandleRefreshEditorFromBlueprintImpl(URigVMBlueprint* InBlueprint);
 
 	FText GetReplayAssetName() const;
 	FText GetReplayAssetTooltip() const;
@@ -308,10 +424,9 @@ private:
 	bool RecordReplay(double InRecordingDuration);
 	void ToggleReplay();
 
-public:
-	EControlRigReplayPlaybackMode GetReplayPlaybackMode() const;
 
-protected:
+	virtual EControlRigReplayPlaybackMode GetReplayPlaybackMode() const override;
+
 
 	/** Persona toolkit used to support skeletal mesh preview */
 	TSharedPtr<IPersonaToolkit> PersonaToolkit;
@@ -324,6 +439,8 @@ protected:
 
 	/** preview animation instance */
 	UAnimPreviewInstance* PreviewInstance;
+	virtual UAnimPreviewInstance* GetPreviewInstance() const override { return PreviewInstance; }
+	virtual void SetPreviewInstance(UAnimPreviewInstance* InPreviewInstance) override { PreviewInstance = InPreviewInstance; }
 
 	/** Model for the schematic views */
 	FControlRigSchematicModel SchematicModel;
@@ -342,8 +459,8 @@ protected:
 	void SetRigElementTransform(const FRigElementKey& InElement, const FTransform& InTransform, bool bLocal);
 	
 	/** delegate for changing property */
-	virtual void OnFinishedChangingProperties(const FPropertyChangedEvent& PropertyChangedEvent) override;
-	virtual void OnWrappedPropertyChangedChainEvent(URigVMDetailsViewWrapperObject* InWrapperObject, const FString& InPropertyPath, FPropertyChangedChainEvent& InPropertyChangedChainEvent) override;
+	virtual void OnFinishedChangingPropertiesImpl(const FPropertyChangedEvent& PropertyChangedEvent);
+	virtual void OnWrappedPropertyChangedChainEventImpl(URigVMDetailsViewWrapperObject* InWrapperObject, const FString& InPropertyPath, FPropertyChangedChainEvent& InPropertyChangedChainEvent);
 
 	URigVMController* ActiveController;
 
@@ -382,7 +499,9 @@ protected:
 	bool bSchematicViewPortIsHidden;
 	FTimerHandle RecordReplayTimerHandle;
 
+public:
 	void SetupTimelineDelegates(FAnimationScrubPanelDelegates& InOutDelegates);
+protected:
 	bool ShowReplayOnTimeline() const;
 	TOptional<bool> HandleReplayIsRecordingActive() const;
 	TOptional<EVisibility> HandleGetReplayRecordButtonVisibility() const;
@@ -402,6 +521,11 @@ protected:
 	EVisibility GetReplayValidationErrorVisibility() const;
 	FText GetReplayValidationErrorTooltip() const;
 
+	virtual FControlRigEditorClosed& OnEditorClosed() override { return ControlRigEditorClosedDelegate; }
+	FControlRigEditorClosed ControlRigEditorClosedDelegate;
+
+	void OnClose();
+
 	static const TArray<FName> ForwardsSolveEventQueue;
 	static const TArray<FName> BackwardsSolveEventQueue;
 	static const TArray<FName> ConstructionEventQueue;
diff --git a/Engine/Plugins/Animation/ControlRig/Source/ControlRigEditor/Private/Editor/ControlRigEditorMode.cpp b/Engine/Plugins/Animation/ControlRig/Source/ControlRigEditor/Private/Editor/ControlRigEditorMode.cpp
--- a/Engine/Plugins/Animation/ControlRig/Source/ControlRigEditor/Private/Editor/ControlRigEditorMode.cpp
+++ b/Engine/Plugins/Animation/ControlRig/Source/ControlRigEditor/Private/Editor/ControlRigEditorMode.cpp
@@ -17,15 +17,13 @@
 #include "Editor/RigValidationTabSummoner.h"
 #include "Editor/RigAnimAttributeTabSummoner.h"
 #include "ToolMenus.h"
+#include "Editor/RigVMDetailsInspectorTabSummoner.h"
 
-FControlRigEditorMode::FControlRigEditorMode(const TSharedRef<FControlRigEditor>& InControlRigEditor, bool bCreateDefaultLayout)
 #if WITH_RIGVMLEGACYEDITOR
-	: FBlueprintEditorApplicationMode(InControlRigEditor, FControlRigEditorModes::ControlRigEditorMode, FControlRigEditorModes::GetLocalizedMode, false, false)
-#else
-	: FRigVMNewEditorMode(InControlRigEditor)
-#endif
+FControlRigLegacyEditorMode::FControlRigLegacyEditorMode(const TSharedRef<FControlRigLegacyEditor>& InControlRigEditor, bool bCreateDefaultLayout)
+	: FBlueprintEditorApplicationMode(InControlRigEditor.ToSharedPtr(), FControlRigEditorModes::ControlRigEditorMode, FControlRigEditorModes::GetLocalizedMode, false, false)
 {
-	ControlRigBlueprintPtr = CastChecked<UControlRigBlueprint>(InControlRigEditor->FRigVMEditorBase::GetBlueprintObj());
+	ControlRigBlueprintPtr = CastChecked<UControlRigBlueprint>(InControlRigEditor->GetRigVMBlueprint());
 
 	TabFactories.RegisterFactory(MakeShared<FRigHierarchyTabSummoner>(InControlRigEditor));
 	TabFactories.RegisterFactory(MakeShared<FRigVMExecutionStackTabSummoner>(InControlRigEditor));
@@ -37,20 +35,19 @@ FControlRigEditorMode::FControlRigEditorMode(const TSharedRef<FControlRigEditor>
 	FPersonaModule& PersonaModule = FModuleManager::LoadModuleChecked<FPersonaModule>("Persona");
 
  	FPersonaViewportArgs ViewportArgs(InControlRigEditor->GetPersonaToolkit()->GetPreviewScene());
-#if WITH_RIGVMLEGACYEDITOR
  	ViewportArgs.BlueprintEditor = InControlRigEditor;
-#endif
  	ViewportArgs.bShowStats = false;
 	ViewportArgs.bShowPlaySpeedMenu = false;
 	ViewportArgs.bShowTimeline = true;
 	ViewportArgs.bShowTurnTable = false;
 	ViewportArgs.bAlwaysShowTransformToolbar = true;
-	ViewportArgs.OnViewportCreated = FOnViewportCreated::CreateSP(InControlRigEditor, &FControlRigEditor::HandleViewportCreated);
+	ViewportArgs.OnViewportCreated = FOnViewportCreated::CreateSP(InControlRigEditor, &FControlRigBaseEditor::HandleViewportCreated);
 
 	InControlRigEditor->SetupTimelineDelegates(ViewportArgs.TimelineDelegates);
- 
- 	TabFactories.RegisterFactory(PersonaModule.CreatePersonaViewportTabFactory(InControlRigEditor, ViewportArgs));
-	TabFactories.RegisterFactory(PersonaModule.CreateAdvancedPreviewSceneTabFactory(InControlRigEditor, InControlRigEditor->GetPersonaToolkit()->GetPreviewScene()));
+
+	TSharedRef<FWorkflowCentricApplication> EditorSharedRef = StaticCastSharedRef<FWorkflowCentricApplication>(InControlRigEditor->GetHostingApp().ToSharedRef());
+ 	TabFactories.RegisterFactory(PersonaModule.CreatePersonaViewportTabFactory(EditorSharedRef, ViewportArgs));
+	TabFactories.RegisterFactory(PersonaModule.CreateAdvancedPreviewSceneTabFactory(EditorSharedRef, InControlRigEditor->GetPersonaToolkit()->GetPreviewScene()));
 
 	if(bCreateDefaultLayout)
 	{
@@ -86,10 +83,7 @@ FControlRigEditorMode::FControlRigEditorMode(const TSharedRef<FControlRigEditor>
 							->AddTab(FRigHierarchyTabSummoner::TabID, ETabState::OpenedTab)
 							->AddTab(FRigVMExecutionStackTabSummoner::TabID, ETabState::OpenedTab)
 							->AddTab(FRigCurveContainerTabSummoner::TabID, ETabState::OpenedTab)
-#if WITH_RIGVMLEGACYEDITOR
 							->AddTab(FBlueprintEditorTabs::MyBlueprintID, ETabState::OpenedTab)
-#endif
-							->AddTab(FRigVMEditorGraphExplorerTabSummoner::TabID(), ETabState::ClosedTab)
 						)
 					)
 					->Split
@@ -110,10 +104,8 @@ FControlRigEditorMode::FControlRigEditorMode(const TSharedRef<FControlRigEditor>
 							// Middle bottom - compiler results & find
 							FTabManager::NewStack()
 							->SetSizeCoefficient(0.2f)
-#if WITH_RIGVMLEGACYEDITOR
 							->AddTab(FBlueprintEditorTabs::CompilerResultsID, ETabState::ClosedTab)
 							->AddTab(FBlueprintEditorTabs::FindResultsID, ETabState::ClosedTab)
-#endif
 						)
 					)
 					->Split
@@ -128,50 +120,39 @@ FControlRigEditorMode::FControlRigEditorMode(const TSharedRef<FControlRigEditor>
 							FTabManager::NewStack()
 							->SetHideTabWell(false)
 							->SetSizeCoefficient(1.f)
-#if WITH_RIGVMLEGACYEDITOR
 							->AddTab(FBlueprintEditorTabs::DetailsID, ETabState::OpenedTab)
-#endif
 							->AddTab(FPersonaTabs::AdvancedPreviewSceneSettingsID, ETabState::OpenedTab)
 							->AddTab(FRigAnimAttributeTabSummoner::TabID, ETabState::OpenedTab)
-#if WITH_RIGVMLEGACYEDITOR
 							->SetForegroundTab(FBlueprintEditorTabs::DetailsID)
-#endif
 						)
 					)
 				)
 			);
 	}
 
-#if WITH_RIGVMLEGACYEDITOR
 	if (UToolMenu* Toolbar = InControlRigEditor->RegisterModeToolbarIfUnregistered(GetModeName()))
 	{
 		InControlRigEditor->GetToolbarBuilder()->AddCompileToolbar(Toolbar);
 		InControlRigEditor->GetToolbarBuilder()->AddScriptingToolbar(Toolbar);
 		InControlRigEditor->GetToolbarBuilder()->AddBlueprintGlobalOptionsToolbar(Toolbar);
 	}
-#endif
 }
 
-void FControlRigEditorMode::RegisterTabFactories(TSharedPtr<FTabManager> InTabManager)
+void FControlRigLegacyEditorMode::RegisterTabFactories(TSharedPtr<FTabManager> InTabManager)
 {
-#if WITH_RIGVMLEGACYEDITOR
 	TSharedPtr<FBlueprintEditor> BP = MyBlueprintEditor.Pin();
-#else
-	TSharedPtr<FRigVMNewEditor> BP = Editor.Pin();
-#endif
 	
 	BP->RegisterToolbarTab(InTabManager.ToSharedRef());
 
 	// Mode-specific setup
-#if WITH_RIGVMLEGACYEDITOR
 	BP->PushTabFactories(CoreTabFactories);
 	BP->PushTabFactories(BlueprintEditorTabFactories);
-#endif
 	BP->PushTabFactories(TabFactories);
 }
 
-FModularRigEditorMode::FModularRigEditorMode(const TSharedRef<FControlRigEditor>& InControlRigEditor)
-	: FControlRigEditorMode(InControlRigEditor, false)
+
+FModularRigLegacyEditorMode::FModularRigLegacyEditorMode(const TSharedRef<FControlRigLegacyEditor>& InControlRigEditor)
+	: FControlRigLegacyEditorMode(InControlRigEditor, false)
 {
 
 	TabFactories.RegisterFactory(MakeShared<FModularRigModelTabSummoner>(InControlRigEditor));
@@ -199,9 +180,7 @@ FModularRigEditorMode::FModularRigEditorMode(const TSharedRef<FControlRigEditor>
 						->SetHideTabWell(false)
 						->SetSizeCoefficient(0.5f)
 						->AddTab(FRigModuleAssetBrowserTabSummoner::TabID, ETabState::OpenedTab)
-#if WITH_RIGVMLEGACYEDITOR
 						->SetForegroundTab(FBlueprintEditorTabs::DetailsID)
-#endif
 					)
 					->Split
 					(
@@ -211,9 +190,7 @@ FModularRigEditorMode::FModularRigEditorMode(const TSharedRef<FControlRigEditor>
 						->SetSizeCoefficient(0.5f)
 						->SetForegroundTab(FModularRigModelTabSummoner::TabID)
 						->AddTab(FModularRigModelTabSummoner::TabID, ETabState::OpenedTab)
-#if WITH_RIGVMLEGACYEDITOR
 						->AddTab(FBlueprintEditorTabs::MyBlueprintID, ETabState::OpenedTab)
-#endif
 					)
 				)
 				->Split
@@ -230,10 +207,8 @@ FModularRigEditorMode::FModularRigEditorMode(const TSharedRef<FControlRigEditor>
 					FTabManager::NewStack()
 					->SetSizeCoefficient(0.2f)
 					->SetHideTabWell(false)
-#if WITH_RIGVMLEGACYEDITOR
 					->SetForegroundTab(FBlueprintEditorTabs::DetailsID)
 					->AddTab(FBlueprintEditorTabs::DetailsID, ETabState::OpenedTab)
-#endif
 					->AddTab(FPersonaTabs::AdvancedPreviewSceneSettingsID, ETabState::OpenedTab)
 					->AddTab(FRigHierarchyTabSummoner::TabID, ETabState::OpenedTab)
 				)
@@ -241,24 +216,18 @@ FModularRigEditorMode::FModularRigEditorMode(const TSharedRef<FControlRigEditor>
 		);
 }
 
-void FModularRigEditorMode::RegisterTabFactories(TSharedPtr<FTabManager> InTabManager)
+void FModularRigLegacyEditorMode::RegisterTabFactories(TSharedPtr<FTabManager> InTabManager)
 {
-#if WITH_RIGVMLEGACYEDITOR
 	TSharedPtr<FBlueprintEditor> BP = MyBlueprintEditor.Pin();
-#else
-	TSharedPtr<FRigVMNewEditor> BP = Editor.Pin();
-#endif
 	
 	BP->RegisterToolbarTab(InTabManager.ToSharedRef());
 
 	static const TArray<FName> DisallowedTabs = {
-#if WITH_RIGVMLEGACYEDITOR
 		FBlueprintEditorTabs::PaletteID,
 		FBlueprintEditorTabs::ReplaceNodeReferencesID,
 		FBlueprintEditorTabs::CompilerResultsID,
 		FBlueprintEditorTabs::FindResultsID,
 		FBlueprintEditorTabs::BookmarksID,
-#endif
 		FRigVMExecutionStackTabSummoner::TabID
 	};
 
@@ -281,3 +250,208 @@ void FModularRigEditorMode::RegisterTabFactories(TSharedPtr<FTabManager> InTabMa
 #endif
 	PushTabFactories(TabFactories);
 }
+#endif
+
+FControlRigEditorMode::FControlRigEditorMode(const TSharedRef<FControlRigEditor>& InControlRigEditor, bool bCreateDefaultLayout)
+	: FRigVMNewEditorMode(InControlRigEditor)
+{
+	ControlRigBlueprintPtr = CastChecked<UControlRigBlueprint>(InControlRigEditor->GetRigVMBlueprint());
+
+	TabFactories.RegisterFactory(MakeShared<FRigHierarchyTabSummoner>(InControlRigEditor));
+	TabFactories.RegisterFactory(MakeShared<FRigCurveContainerTabSummoner>(InControlRigEditor));
+	TabFactories.RegisterFactory(MakeShared<FRigValidationTabSummoner>(InControlRigEditor));
+	TabFactories.RegisterFactory(MakeShared<FRigAnimAttributeTabSummoner>(InControlRigEditor));
+
+	FPersonaModule& PersonaModule = FModuleManager::LoadModuleChecked<FPersonaModule>("Persona");
+
+ 	FPersonaViewportArgs ViewportArgs(InControlRigEditor->GetPersonaToolkit()->GetPreviewScene());
+ 	//ViewportArgs.BlueprintEditor = InControlRigEditor;
+ 	ViewportArgs.bShowStats = false;
+	ViewportArgs.bShowPlaySpeedMenu = false;
+	ViewportArgs.bShowTimeline = true;
+	ViewportArgs.bShowTurnTable = false;
+	ViewportArgs.bAlwaysShowTransformToolbar = true;
+	ViewportArgs.OnViewportCreated = FOnViewportCreated::CreateSP(InControlRigEditor, &FControlRigBaseEditor::HandleViewportCreated);
+
+	InControlRigEditor->SetupTimelineDelegates(ViewportArgs.TimelineDelegates);
+
+	TSharedRef<FWorkflowCentricApplication> EditorSharedRef = StaticCastSharedRef<FWorkflowCentricApplication>(InControlRigEditor->GetHostingApp().ToSharedRef());
+ 	TabFactories.RegisterFactory(PersonaModule.CreatePersonaViewportTabFactory(EditorSharedRef, ViewportArgs));
+	TabFactories.RegisterFactory(PersonaModule.CreateAdvancedPreviewSceneTabFactory(EditorSharedRef, InControlRigEditor->GetPersonaToolkit()->GetPreviewScene()));
+
+	if(bCreateDefaultLayout)
+	{
+		TabLayout = FTabManager::NewLayout("Standalone_ControlRigEditMode_Layout_v1.8")
+			->AddArea
+			(
+				// Main application area
+				FTabManager::NewPrimaryArea()
+				->SetOrientation(Orient_Vertical)
+				->Split
+				(
+					FTabManager::NewSplitter()
+					->SetOrientation(Orient_Horizontal)
+					->Split
+					(
+						FTabManager::NewSplitter()
+						->SetOrientation(Orient_Vertical)
+						->SetSizeCoefficient(0.2f)
+						->Split
+						(
+							//	Left top - viewport
+							FTabManager::NewStack()
+							->SetSizeCoefficient(0.5f)
+							->SetHideTabWell(true)
+							->AddTab(FPersonaTabs::PreviewViewportID, ETabState::OpenedTab)
+						
+						)
+						->Split
+						(
+							//	Left bottom - rig/hierarchy
+							FTabManager::NewStack()
+							->SetSizeCoefficient(0.5f)
+							->AddTab(FRigHierarchyTabSummoner::TabID, ETabState::OpenedTab)
+							->AddTab(FRigVMExecutionStackTabSummoner::TabID, ETabState::OpenedTab)
+							->AddTab(FRigCurveContainerTabSummoner::TabID, ETabState::OpenedTab)
+							->AddTab(FRigVMEditorGraphExplorerTabSummoner::TabID(), ETabState::OpenedTab)
+						)
+					)
+					->Split
+					(
+						// Middle 
+						FTabManager::NewSplitter()
+						->SetOrientation(Orient_Vertical)
+						->SetSizeCoefficient(0.6f)
+						->Split
+						(
+							// Middle top - document edit area
+							FTabManager::NewStack()
+							->SetSizeCoefficient(0.8f)
+							->AddTab("Document", ETabState::ClosedTab)
+						)
+						->Split
+						(
+							// Middle bottom - compiler results & find
+							FTabManager::NewStack()
+							->SetSizeCoefficient(0.2f)
+						)
+					)
+					->Split
+					(
+						// Right side
+						FTabManager::NewSplitter()
+						->SetOrientation(Orient_Vertical)
+						->SetSizeCoefficient(0.2f)
+						->Split
+						(
+							// Right top
+							FTabManager::NewStack()
+							->SetHideTabWell(false)
+							->SetSizeCoefficient(1.f)
+							->AddTab(FRigVMDetailsInspectorTabSummoner::TabID(), ETabState::OpenedTab)
+							->AddTab(FPersonaTabs::AdvancedPreviewSceneSettingsID, ETabState::OpenedTab)
+							->AddTab(FRigAnimAttributeTabSummoner::TabID, ETabState::OpenedTab)
+							->SetForegroundTab(FRigVMDetailsInspectorTabSummoner::TabID())
+						)
+					)
+				)
+			);
+	}
+}
+
+void FControlRigEditorMode::RegisterTabFactories(TSharedPtr<FTabManager> InTabManager)
+{
+	FRigVMNewEditorMode::RegisterTabFactories(InTabManager);
+}
+
+
+FModularRigEditorMode::FModularRigEditorMode(const TSharedRef<FControlRigEditor>& InControlRigEditor)
+	: FControlRigEditorMode(InControlRigEditor, false)
+{
+
+	TabFactories.RegisterFactory(MakeShared<FModularRigModelTabSummoner>(InControlRigEditor));
+	TabFactories.RegisterFactory(MakeShared<FRigModuleAssetBrowserTabSummoner>(InControlRigEditor));
+	
+	TabLayout = FTabManager::NewLayout("Standalone_ModularRigEditMode_Layout_v1.3")
+		->AddArea
+		(
+			// Main application area
+			FTabManager::NewPrimaryArea()
+			->SetOrientation(Orient_Vertical)
+			->Split
+			(
+				FTabManager::NewSplitter()
+				->SetOrientation(Orient_Horizontal)
+				->Split
+				(
+					FTabManager::NewSplitter()
+					->SetOrientation(Orient_Vertical)
+					->SetSizeCoefficient(0.2f)
+					->Split
+					(
+						// Left top - Module Library
+						FTabManager::NewStack()
+						->SetHideTabWell(false)
+						->SetSizeCoefficient(0.5f)
+						->AddTab(FRigModuleAssetBrowserTabSummoner::TabID, ETabState::OpenedTab)
+						->SetForegroundTab(FRigVMDetailsInspectorTabSummoner::TabID())
+					)
+					->Split
+					(
+						//	Left bottom - rig/hierarchy/modules
+						FTabManager::NewStack()
+						->SetHideTabWell(false)
+						->SetSizeCoefficient(0.5f)
+						->SetForegroundTab(FModularRigModelTabSummoner::TabID)
+						->AddTab(FModularRigModelTabSummoner::TabID, ETabState::OpenedTab)
+						->AddTab(FRigVMEditorGraphExplorerTabSummoner::TabID(), ETabState::OpenedTab)
+					)
+				)
+				->Split
+				(
+					//	Center - viewport
+					FTabManager::NewStack()
+					->SetSizeCoefficient(0.6f)
+					->SetHideTabWell(true)
+					->AddTab(FPersonaTabs::PreviewViewportID, ETabState::OpenedTab)
+				)
+				->Split
+				(
+					// Right side
+					FTabManager::NewStack()
+					->SetSizeCoefficient(0.2f)
+					->SetHideTabWell(false)
+					->SetForegroundTab(FRigVMDetailsInspectorTabSummoner::TabID())
+					->AddTab(FRigVMDetailsInspectorTabSummoner::TabID(), ETabState::OpenedTab)
+					->AddTab(FPersonaTabs::AdvancedPreviewSceneSettingsID, ETabState::OpenedTab)
+					->AddTab(FRigHierarchyTabSummoner::TabID, ETabState::OpenedTab)
+				)
+			)
+		);
+}
+
+void FModularRigEditorMode::RegisterTabFactories(TSharedPtr<FTabManager> InTabManager)
+{
+	TSharedPtr<FRigVMNewEditor> BP = Editor.Pin();
+	
+	BP->RegisterToolbarTab(InTabManager.ToSharedRef());
+
+	static const TArray<FName> DisallowedTabs = {
+		FRigVMExecutionStackTabSummoner::TabID
+	};
+
+	auto PushTabFactories = [&](FWorkflowAllowedTabSet& Tabs)
+	{
+		for (auto FactoryIt = Tabs.CreateIterator(); FactoryIt; ++FactoryIt)
+		{
+			if (DisallowedTabs.Contains(FactoryIt->Key))
+			{
+				continue;
+			}
+			FactoryIt.Value()->RegisterTabSpawner(InTabManager.ToSharedRef(), BP->GetCurrentModePtr().Get());
+		}
+	};
+
+	// Mode-specific setup
+	PushTabFactories(TabFactories);
+}
\ No newline at end of file
diff --git a/Engine/Plugins/Animation/ControlRig/Source/ControlRigEditor/Private/Editor/ControlRigEditorMode.h b/Engine/Plugins/Animation/ControlRig/Source/ControlRigEditor/Private/Editor/ControlRigEditorMode.h
--- a/Engine/Plugins/Animation/ControlRig/Source/ControlRigEditor/Private/Editor/ControlRigEditorMode.h
+++ b/Engine/Plugins/Animation/ControlRig/Source/ControlRigEditor/Private/Editor/ControlRigEditorMode.h
@@ -8,20 +8,18 @@
 #include "ControlRigBlueprint.h"
 #if WITH_RIGVMLEGACYEDITOR
 #include "BlueprintEditorModes.h"
-#else
-#include "Editor/RigVMNewEditorMode.h" 
+#include "Editor/ControlRigLegacyEditor.h"
 #endif
+#include "Editor/RigVMNewEditorMode.h" 
+#include "Editor/ControlRigNewEditor.h"
 
 class UControlRigBlueprint;
 
 #if WITH_RIGVMLEGACYEDITOR
-class FControlRigEditorMode : public FBlueprintEditorApplicationMode
-#else
-class FControlRigEditorMode : public FRigVMNewEditorMode
-#endif
+class FControlRigLegacyEditorMode : public FBlueprintEditorApplicationMode
 {
 public:
-	FControlRigEditorMode(const TSharedRef<FControlRigEditor>& InControlRigEditor, bool bCreateDefaultLayout = true);
+	FControlRigLegacyEditorMode(const TSharedRef<FControlRigLegacyEditor>& InControlRigEditor, bool bCreateDefaultLayout = true);
 
 	// FApplicationMode interface
 	virtual void RegisterTabFactories(TSharedPtr<FTabManager> InTabManager) override;
@@ -34,6 +32,31 @@ private:
 	TWeakObjectPtr<UControlRigBlueprint> ControlRigBlueprintPtr;
 };
 
+class FModularRigLegacyEditorMode : public FControlRigLegacyEditorMode
+{
+public:
+	FModularRigLegacyEditorMode(const TSharedRef<FControlRigLegacyEditor>& InControlRigEditor);
+
+	// FApplicationMode interface
+	virtual void RegisterTabFactories(TSharedPtr<FTabManager> InTabManager) override;
+
+	// for now just don't open up the previous edited documents
+	virtual void PostActivateMode() override {}
+};
+#endif
+
+class FControlRigEditorMode : public FRigVMNewEditorMode
+{
+public:
+	FControlRigEditorMode(const TSharedRef<FControlRigEditor>& InControlRigEditor, bool bCreateDefaultLayout = true);
+
+	// FApplicationMode interface
+	virtual void RegisterTabFactories(TSharedPtr<FTabManager> InTabManager) override;
+
+private:
+	TWeakObjectPtr<UControlRigBlueprint> ControlRigBlueprintPtr;
+};
+
 class FModularRigEditorMode : public FControlRigEditorMode
 {
 public:
diff --git a/Engine/Plugins/Animation/ControlRig/Source/ControlRigEditor/Private/Editor/ControlRigLegacyEditor.cpp b/Engine/Plugins/Animation/ControlRig/Source/ControlRigEditor/Private/Editor/ControlRigLegacyEditor.cpp
new file mode 100644
--- /dev/null
+++ b/Engine/Plugins/Animation/ControlRig/Source/ControlRigEditor/Private/Editor/ControlRigLegacyEditor.cpp
@@ -0,0 +1,127 @@
+// Copyright Epic Games, Inc. All Rights Reserved.
+
+#include "Editor/ControlRigLegacyEditor.h"
+
+#if WITH_RIGVMLEGACYEDITOR
+#include "ControlRigEditorMode.h"
+#include "ControlRigEditorStyle.h"
+#include "IPersonaPreviewScene.h"
+#include "IPersonaToolkit.h"
+
+
+FControlRigLegacyEditor::FControlRigLegacyEditor()
+: IControlRigLegacyEditor()
+, FControlRigBaseEditor()
+{
+	LastEventQueue = ConstructionEventQueue;
+}
+
+FControlRigLegacyEditor::~FControlRigLegacyEditor()
+{
+	UControlRigBlueprint* RigBlueprint = GetControlRigBlueprint();
+
+	if (RigBlueprint)
+	{
+		UControlRigBlueprint::sCurrentlyOpenedRigBlueprints.Remove(RigBlueprint);
+
+		RigBlueprint->OnHierarchyModified().RemoveAll(this);
+		if (FControlRigEditMode* EditMode = GetEditMode())
+		{
+			RigBlueprint->OnHierarchyModified().RemoveAll(EditMode);
+			EditMode->OnEditorClosed();
+		}
+
+		RigBlueprint->OnRigTypeChanged().RemoveAll(this);
+		if (RigBlueprint->IsModularRig())
+		{
+			RigBlueprint->GetModularRigController()->OnModified().RemoveAll(this);
+			RigBlueprint->OnModularRigCompiled().RemoveAll(this);
+
+			RigBlueprint->OnSetObjectBeingDebugged().RemoveAll(&SchematicModel);
+			RigBlueprint->OnHierarchyModified().RemoveAll(&SchematicModel);
+			RigBlueprint->GetModularRigController()->OnModified().RemoveAll(&SchematicModel);
+		}
+	}
+
+	if (PersonaToolkit.IsValid())
+	{
+		constexpr bool bSetPreviewMeshInAsset = false;
+		PersonaToolkit->SetPreviewMesh(nullptr, bSetPreviewMeshInAsset);
+	}
+}
+
+TSharedPtr<FApplicationMode> FControlRigLegacyEditor::CreateEditorMode()
+{
+	CreatePersonaToolKitIfRequired();
+
+	if(IsModularRig())
+	{
+		return MakeShareable(new FModularRigLegacyEditorMode(SharedThis(this)));
+	}
+	return MakeShareable(new FControlRigLegacyEditorMode(SharedThis(this)));
+}
+
+bool FControlRigLegacyEditor::IsSectionVisible(NodeSectionID::Type InSectionID) const
+{
+	if(!IControlRigLegacyEditor::IsSectionVisible(InSectionID))
+	{
+		return false;
+	}
+
+	if(const UControlRigBlueprint* RigBlueprint = GetControlRigBlueprint())
+	{
+		if(IsModularRig())
+		{
+			switch (InSectionID)
+			{
+				case NodeSectionID::GRAPH:
+				{
+					return RigBlueprint->SupportsEventGraphs();
+				}
+				case NodeSectionID::FUNCTION:
+				{
+					return RigBlueprint->SupportsFunctions();
+				}
+				default:
+				{
+					break;
+				}
+			}
+		}
+	}
+	return true;
+}
+
+bool FControlRigLegacyEditor::NewDocument_IsVisibleForType(FBlueprintEditor::ECreatedDocumentType GraphType) const
+{
+	if(!IControlRigLegacyEditor::NewDocument_IsVisibleForType(GraphType))
+	{
+		return false;
+	}
+
+	if(const UControlRigBlueprint* RigBlueprint = GetControlRigBlueprint())
+	{
+		if(IsModularRig())
+		{
+			switch(GraphType)
+			{
+				case FBlueprintEditor::CGT_NewEventGraph:
+				{
+					return RigBlueprint->SupportsEventGraphs();
+				}
+				case FBlueprintEditor::CGT_NewFunctionGraph:
+				{
+					return RigBlueprint->SupportsFunctions();
+				}
+				default:
+				{
+					break;
+				}
+			}
+		}
+	}
+
+	return true;
+}
+
+#endif
diff --git a/Engine/Plugins/Animation/ControlRig/Source/ControlRigEditor/Private/Editor/ControlRigLegacyEditor.h b/Engine/Plugins/Animation/ControlRig/Source/ControlRigEditor/Private/Editor/ControlRigLegacyEditor.h
new file mode 100644
--- /dev/null
+++ b/Engine/Plugins/Animation/ControlRig/Source/ControlRigEditor/Private/Editor/ControlRigLegacyEditor.h
@@ -0,0 +1,178 @@
+// Copyright Epic Games, Inc. All Rights Reserved.
+
+#pragma once
+
+#if WITH_RIGVMLEGACYEDITOR
+#include "ControlRigEditor.h"
+
+class FControlRigLegacyEditor : public IControlRigLegacyEditor, public FControlRigBaseEditor
+{
+public:
+
+	
+	virtual void InitRigVMEditor(const EToolkitMode::Type Mode, const TSharedPtr< class IToolkitHost >& InitToolkitHost, class URigVMBlueprint* InRigVMBlueprint) override { return FControlRigBaseEditor::InitRigVMEditorImpl(Mode, InitToolkitHost, InRigVMBlueprint); }
+	virtual void InitRigVMEditorSuper(const EToolkitMode::Type Mode, const TSharedPtr< class IToolkitHost >& InitToolkitHost, class URigVMBlueprint* InRigVMBlueprint) override { return IControlRigLegacyEditor::InitRigVMEditor(Mode, InitToolkitHost, InRigVMBlueprint); }
+
+	virtual const FName GetEditorAppName() const override { return FControlRigBaseEditor::GetEditorAppNameImpl(); }
+	virtual const FName GetEditorModeName() const override { return FControlRigBaseEditor::GetEditorModeNameImpl(); }
+	virtual TSharedPtr<FApplicationMode> CreateEditorMode() override;
+	virtual const FSlateBrush* GetDefaultTabIcon() const override { return FControlRigBaseEditor::GetDefaultTabIconImpl(); }
+
+public:
+	FControlRigLegacyEditor();
+	virtual ~FControlRigLegacyEditor();
+
+	// FControlRigBaseEditor
+	virtual TSharedPtr<FAssetEditorToolkit> GetHostingApp() override { return IControlRigLegacyEditor::GetHostingApp(); }
+	virtual TSharedRef<IControlRigBaseEditor> SharedControlRigEditorRef() override { return StaticCastSharedRef<IControlRigBaseEditor>(SharedThis(this)); }
+	virtual TSharedRef<IRigVMEditor> SharedRigVMEditorRef() override { return StaticCastSharedRef<IRigVMEditor>(SharedThis(this)); }
+	virtual TSharedRef<const IRigVMEditor> SharedRigVMEditorRef() const override { return StaticCastSharedRef<const IRigVMEditor>(SharedThis(this)); }
+	
+	virtual bool IsControlRigLegacyEditor() const { return true; };
+	virtual URigVMBlueprint* GetRigVMBlueprint() const { return IControlRigLegacyEditor::GetRigVMBlueprint(); }
+	virtual URigVMHost* GetRigVMHost() const { return FRigVMEditorBase::GetRigVMHost(); }
+	virtual TSharedRef<FUICommandList> GetToolkitCommands() { return IControlRigLegacyEditor::GetToolkitCommands(); }
+	virtual FPreviewScene* GetPreviewScene() { return IControlRigLegacyEditor::GetPreviewScene(); }
+	virtual bool IsDetailsPanelRefreshSuspended() const { return IControlRigLegacyEditor::IsDetailsPanelRefreshSuspended(); }
+	virtual TArray<TWeakObjectPtr<UObject>> GetSelectedObjects() const { return IControlRigLegacyEditor::GetSelectedObjects(); }
+	virtual void ClearDetailObject(bool bChangeUISelectionState = true) { IControlRigLegacyEditor::ClearDetailObject(bChangeUISelectionState); }
+	virtual bool DetailViewShowsStruct(UScriptStruct* InStruct) const { return IControlRigLegacyEditor::DetailViewShowsStruct(InStruct); }
+	virtual TSharedPtr<class SWidget> GetInspector() const { return Inspector; }
+	virtual TArray<FName> GetEventQueue() const { return IControlRigLegacyEditor::GetEventQueue(); }
+	virtual void SummonSearchUI(bool bSetFindWithinBlueprint, FString NewSearchTerms = FString(), bool bSelectFirstResult = false) { IControlRigLegacyEditor::SummonSearchUI(bSetFindWithinBlueprint, NewSearchTerms, bSelectFirstResult); }
+	virtual const TArray< UObject* >* GetObjectsCurrentlyBeingEdited() const { return IControlRigLegacyEditor::GetObjectsCurrentlyBeingEdited(); }
+	virtual FEditorModeTools& GetEditorModeManagerImpl() const { return GetEditorModeManager(); }
+	virtual const FName GetEditorModeNameImpl() const { return GetEditorModeName(); }
+	virtual URigVMController* GetFocusedController() const { return IControlRigLegacyEditor::GetFocusedController(); }
+	virtual URigVMGraph* GetFocusedModel() const { return IControlRigLegacyEditor::GetFocusedModel(); }
+	virtual TArray<FName> GetLastEventQueue() const override { return LastEventQueue; }
+
+
+	// FRigVMEditorBase interface
+	virtual UObject* GetOuterForHost() const override { return GetOuterForHostImpl();}
+	virtual UObject* GetOuterForHostSuper() const override { return IControlRigLegacyEditor::GetOuterForHost(); }
+	
+	virtual UClass* GetDetailWrapperClass() const { return GetDetailWrapperClassImpl(); }
+	virtual void Compile() override { CompileBaseImpl(); }
+	virtual void CompileSuper() override { IControlRigLegacyEditor::Compile(); }
+
+	virtual void HandleModifiedEvent(ERigVMGraphNotifType InNotifType, URigVMGraph* InGraph, UObject* InSubject) override { FControlRigBaseEditor::HandleModifiedEventImpl(InNotifType, InGraph, InSubject); }
+	virtual void HandleModifiedEventSuper(ERigVMGraphNotifType InNotifType, URigVMGraph* InGraph, UObject* InSubject) override { IControlRigLegacyEditor::HandleModifiedEvent(InNotifType, InGraph, InSubject); }
+
+	virtual void OnCreateGraphEditorCommands(TSharedPtr<FUICommandList> GraphEditorCommandsList) override { FControlRigBaseEditor::OnCreateGraphEditorCommandsImpl(GraphEditorCommandsList); }
+	virtual void OnCreateGraphEditorCommandsSuper(TSharedPtr<FUICommandList> GraphEditorCommandsList) override { IControlRigLegacyEditor::OnCreateGraphEditorCommands(GraphEditorCommandsList); }
+	
+	virtual void HandleVMCompiledEvent(UObject* InCompiledObject, URigVM* InVM, FRigVMExtendedExecuteContext& InContext) override { FControlRigBaseEditor::HandleVMCompiledEventImpl(InCompiledObject, InVM, InContext); }
+	virtual void HandleVMCompiledEventSuper(UObject* InCompiledObject, URigVM* InVM, FRigVMExtendedExecuteContext& InContext) override { IControlRigLegacyEditor::HandleVMCompiledEvent(InCompiledObject, InVM, InContext); }
+	
+	virtual bool ShouldOpenGraphByDefault() const override { return FControlRigBaseEditor::ShouldOpenGraphByDefaultImpl(); }
+	virtual FReply OnViewportDrop(const FGeometry& MyGeometry, const FDragDropEvent& DragDropEvent) override { return FControlRigBaseEditor::OnViewportDropImpl(MyGeometry, DragDropEvent); }
+	virtual FReply OnViewportDropSuper(const FGeometry& MyGeometry, const FDragDropEvent& DragDropEvent) override { return IControlRigLegacyEditor::OnViewportDrop(MyGeometry, DragDropEvent); }
+
+	// allows the editor to fill an empty graph
+	virtual void CreateEmptyGraphContent(URigVMController* InController) override { FControlRigBaseEditor::CreateEmptyGraphContentImpl(InController); }
+
+public:
+	
+	// IToolkit Interface
+	virtual FName GetToolkitFName() const override { return FControlRigBaseEditor::GetToolkitFNameImpl(); }
+	virtual FText GetBaseToolkitName() const override { return FControlRigBaseEditor::GetBaseToolkitNameImpl(); }
+	virtual FString GetWorldCentricTabPrefix() const override { return FControlRigBaseEditor::GetWorldCentricTabPrefixImpl(); }
+	virtual FString GetDocumentationLink() const override { return FControlRigBaseEditor::GetDocumentationLinkImpl(); }
+
+	// BlueprintEditor interface
+	virtual FReply OnSpawnGraphNodeByShortcut(FInputChord InChord, const FVector2D& InPosition, UEdGraph* InGraph) override { return FControlRigBaseEditor::OnSpawnGraphNodeByShortcutImpl(InChord, InPosition, InGraph); }
+	virtual FReply OnSpawnGraphNodeByShortcutSuper(FInputChord InChord, const FVector2D& InPosition, UEdGraph* InGraph) override { return IControlRigLegacyEditor::OnSpawnGraphNodeByShortcut(InChord, InPosition, InGraph); }
+#if WITH_RIGVMLEGACYEDITOR
+	virtual bool IsSectionVisible(NodeSectionID::Type InSectionID) const override;
+	virtual bool NewDocument_IsVisibleForType(FBlueprintEditor::ECreatedDocumentType GraphType) const override;
+#else
+	virtual bool IsSectionVisible(RigVMNodeSectionID::Type InSectionID) const override;
+	virtual bool NewDocument_IsVisibleForType(FRigVMEditorBase::ECreatedDocumentType GraphType) const override;
+#endif
+
+	virtual void PostTransaction(bool bSuccess, const FTransaction* Transaction, bool bIsRedo) override { return FControlRigBaseEditor::PostTransactionImpl(bSuccess, Transaction, bIsRedo); }
+
+	//  FTickableEditorObject Interface
+	virtual void Tick(float DeltaTime) override { return FControlRigBaseEditor::TickImpl(DeltaTime); }
+	virtual void TickSuper(float DeltaTime) override { return IControlRigLegacyEditor::Tick(DeltaTime); }
+
+	virtual void SetDetailObjects(const TArray<UObject*>& InObjects) override { return FControlRigBaseEditor::SetDetailObjectsImpl(InObjects); }
+	virtual void SetDetailObjectsSuper(const TArray<UObject*>& InObjects) override { return IControlRigLegacyEditor::SetDetailObjects(InObjects); }
+	virtual void RefreshDetailView() override { return FControlRigBaseEditor::RefreshDetailViewImpl(); }
+	virtual void RefreshDetailViewSuper() override { return IControlRigLegacyEditor::RefreshDetailView(); }
+
+public:
+
+	
+	virtual void OnGraphNodeDropToPerform(TSharedPtr<FDragDropOperation> InDragDropOp, UEdGraph* InGraph, const FVector2D& InNodePosition, const FVector2D& InScreenPosition) override  { return FControlRigBaseEditor::OnGraphNodeDropToPerformImpl(InDragDropOp, InGraph, InNodePosition, InScreenPosition); }
+	virtual void OnGraphNodeDropToPerformSuper(TSharedPtr<FDragDropOperation> InDragDropOp, UEdGraph* InGraph, const FVector2D& InNodePosition, const FVector2D& InScreenPosition) override  { return IControlRigLegacyEditor::OnGraphNodeDropToPerform(InDragDropOp, InGraph, InNodePosition, InScreenPosition); }
+
+	
+protected:
+
+	virtual void BindCommands() override { return FControlRigBaseEditor::BindCommandsImpl(); }
+	virtual void BindCommandsSuper() override { return IControlRigLegacyEditor::BindCommands(); }
+	virtual FMenuBuilder GenerateBulkEditMenu() override { return FControlRigBaseEditor::GenerateBulkEditMenuImpl(); }
+	virtual FMenuBuilder GenerateBulkEditMenuSuper() override { return IControlRigLegacyEditor::GenerateBulkEditMenu(); }
+
+	virtual void SaveAsset_Execute() override { return FControlRigBaseEditor::SaveAsset_ExecuteImpl(); }
+	virtual void SaveAsset_ExecuteSuper() override { return IControlRigLegacyEditor::SaveAsset_Execute(); }
+	virtual void SaveAssetAs_Execute() override { return FControlRigBaseEditor::SaveAssetAs_ExecuteImpl(); }
+	virtual void SaveAssetAs_ExecuteSuper() override { return IControlRigLegacyEditor::SaveAssetAs_Execute(); }
+
+	virtual void HandleVMExecutedEvent(URigVMHost* InHost, const FName& InEventName) override { return FControlRigBaseEditor::HandleVMExecutedEventImpl(InHost, InEventName); }
+	virtual void HandleVMExecutedEventSuper(URigVMHost* InHost, const FName& InEventName) override { return IControlRigLegacyEditor::HandleVMExecutedEvent(InHost, InEventName); }
+
+	// FBaseToolKit overrides
+	virtual void CreateEditorModeManager() override { return FControlRigBaseEditor::CreateEditorModeManagerImpl(); }
+
+private:
+	/** Fill the toolbar with content */
+	virtual void FillToolbar(FToolBarBuilder& ToolbarBuilder, bool bEndSection = true) override { return FControlRigBaseEditor::FillToolbarImpl(ToolbarBuilder, bEndSection); }
+	virtual void FillToolbarSuper(FToolBarBuilder& ToolbarBuilder, bool bEndSection = true) override { return IControlRigLegacyEditor::FillToolbar(ToolbarBuilder, bEndSection); }
+
+	virtual TArray<FName> GetDefaultEventQueue() const override { return FControlRigBaseEditor::GetDefaultEventQueueImpl(); }
+	virtual void SetEventQueue(TArray<FName> InEventQueue, bool bCompile) override { return FControlRigBaseEditor::SetEventQueueImpl(InEventQueue, bCompile); }
+	virtual void SetEventQueueSuper(TArray<FName> InEventQueue, bool bCompile) override { return IControlRigLegacyEditor::SetEventQueue(InEventQueue, bCompile); }
+	virtual void SetEventQueueSuper(TArray<FName> InEventQueue) override { return IControlRigLegacyEditor::SetEventQueue(InEventQueue); }
+	virtual int32 GetEventQueueComboValue() const override { return FControlRigBaseEditor::GetEventQueueComboValueImpl(); }
+	virtual int32 GetEventQueueComboValueSuper() const override { return IControlRigLegacyEditor::GetEventQueueComboValue(); }
+	virtual FText GetEventQueueLabel() const override { return FControlRigBaseEditor::GetEventQueueLabelImpl(); }
+	virtual FSlateIcon GetEventQueueIcon(const TArray<FName>& InEventQueue) const override { return FControlRigBaseEditor::GetEventQueueIconImpl(InEventQueue); }
+	virtual void HandleSetObjectBeingDebugged(UObject* InObject) override { return FControlRigBaseEditor::HandleSetObjectBeingDebuggedImpl(InObject); }
+	virtual void HandleSetObjectBeingDebuggedSuper(UObject* InObject) override { return IControlRigLegacyEditor::HandleSetObjectBeingDebugged(InObject); }
+
+	/** Push a newly compiled/opened control rig to the edit mode */
+	virtual void UpdateRigVMHost() override { return FControlRigBaseEditor::UpdateRigVMHostImpl(); }
+	virtual void UpdateRigVMHostSuper() override { return IControlRigLegacyEditor::UpdateRigVMHost(); }
+	virtual void UpdateRigVMHost_PreClearOldHost(URigVMHost* InPreviousHost) override { return FControlRigBaseEditor::UpdateRigVMHost_PreClearOldHostImpl(InPreviousHost); }
+
+	/** Update the name lists for use in name combo boxes */
+	virtual void CacheNameLists() override { return FControlRigBaseEditor::CacheNameListsImpl(); }
+	virtual void CacheNameListsSuper() override { return IControlRigLegacyEditor::CacheNameLists(); }
+
+	virtual void GenerateEventQueueMenuContent(FMenuBuilder& MenuBuilder) override { return FControlRigBaseEditor::GenerateEventQueueMenuContentImpl(MenuBuilder); }
+
+	virtual void HandleRefreshEditorFromBlueprint(URigVMBlueprint* InBlueprint) override { return FControlRigBaseEditor::HandleRefreshEditorFromBlueprintImpl(InBlueprint); }
+	virtual void HandleRefreshEditorFromBlueprintSuper(URigVMBlueprint* InBlueprint) override { return IControlRigLegacyEditor::HandleRefreshEditorFromBlueprint(InBlueprint); }
+
+	
+	/** delegate for changing property */
+	virtual void OnFinishedChangingProperties(const FPropertyChangedEvent& PropertyChangedEvent) override { return FControlRigBaseEditor::OnFinishedChangingPropertiesImpl(PropertyChangedEvent); }
+	virtual void OnFinishedChangingPropertiesSuper(const FPropertyChangedEvent& PropertyChangedEvent) override { return IControlRigLegacyEditor::OnFinishedChangingProperties(PropertyChangedEvent); }
+	virtual void OnWrappedPropertyChangedChainEvent(URigVMDetailsViewWrapperObject* InWrapperObject, const FString& InPropertyPath, FPropertyChangedChainEvent& InPropertyChangedChainEvent) override { return FControlRigBaseEditor::OnWrappedPropertyChangedChainEventImpl(InWrapperObject, InPropertyPath, InPropertyChangedChainEvent); }
+	virtual void OnWrappedPropertyChangedChainEventSuper(URigVMDetailsViewWrapperObject* InWrapperObject, const FString& InPropertyPath, FPropertyChangedChainEvent& InPropertyChangedChainEvent) override { return IControlRigLegacyEditor::OnWrappedPropertyChangedChainEvent(InWrapperObject, InPropertyPath, InPropertyChangedChainEvent); }
+
+	virtual void SetEditorModeManager(TSharedPtr<FEditorModeTools> InManager) override { EditorModeManager = InManager; }
+
+	virtual const TArray<TStrongObjectPtr<URigVMDetailsViewWrapperObject>>& GetWrapperObjects() const override { return IControlRigLegacyEditor::GetWrapperObjects(); }
+	virtual bool& GetSuspendDetailsPanelRefreshFlag() override { return IControlRigLegacyEditor::GetSuspendDetailsPanelRefreshFlag(); }
+
+	virtual TWeakPtr<class SGraphEditor> GetFocusedGraphEd() const override { return FocusedGraphEdPtr;}
+
+	virtual void OnClose() override { FControlRigBaseEditor::OnClose(); }
+	virtual void OnCloseSuper() override { IControlRigLegacyEditor::OnClose(); }
+};
+
+#endif
\ No newline at end of file
diff --git a/Engine/Plugins/Animation/ControlRig/Source/ControlRigEditor/Private/Editor/ControlRigNewEditor.cpp b/Engine/Plugins/Animation/ControlRig/Source/ControlRigEditor/Private/Editor/ControlRigNewEditor.cpp
new file mode 100644
--- /dev/null
+++ b/Engine/Plugins/Animation/ControlRig/Source/ControlRigEditor/Private/Editor/ControlRigNewEditor.cpp
@@ -0,0 +1,130 @@
+// Copyright Epic Games, Inc. All Rights Reserved.
+
+#include "Editor/ControlRigNewEditor.h"
+
+#include "ControlRigEditorMode.h"
+#include "ControlRigEditorStyle.h"
+#include "IPersonaPreviewScene.h"
+#include "IPersonaToolkit.h"
+
+
+FControlRigEditor::FControlRigEditor()
+: IControlRigNewEditor()
+, FControlRigBaseEditor()
+{
+	LastEventQueue = ConstructionEventQueue;
+}
+
+FControlRigEditor::~FControlRigEditor()
+{
+	UControlRigBlueprint* RigBlueprint = GetControlRigBlueprint();
+
+	if (RigBlueprint)
+	{
+		UControlRigBlueprint::sCurrentlyOpenedRigBlueprints.Remove(RigBlueprint);
+
+		RigBlueprint->OnHierarchyModified().RemoveAll(this);
+		if (FControlRigEditMode* EditMode = GetEditMode())
+		{
+			RigBlueprint->OnHierarchyModified().RemoveAll(EditMode);
+			EditMode->OnEditorClosed();
+		}
+
+		RigBlueprint->OnRigTypeChanged().RemoveAll(this);
+		if (RigBlueprint->IsModularRig())
+		{
+			RigBlueprint->GetModularRigController()->OnModified().RemoveAll(this);
+			RigBlueprint->OnModularRigCompiled().RemoveAll(this);
+
+			RigBlueprint->OnSetObjectBeingDebugged().RemoveAll(&SchematicModel);
+			RigBlueprint->OnHierarchyModified().RemoveAll(&SchematicModel);
+			RigBlueprint->GetModularRigController()->OnModified().RemoveAll(&SchematicModel);
+		}
+	}
+
+	if (PersonaToolkit.IsValid())
+	{
+		constexpr bool bSetPreviewMeshInAsset = false;
+		PersonaToolkit->SetPreviewMesh(nullptr, bSetPreviewMeshInAsset);
+	}
+}
+
+TSharedPtr<FApplicationMode> FControlRigEditor::CreateEditorMode()
+{
+	CreatePersonaToolKitIfRequired();
+
+	if(IsModularRig())
+	{
+		return MakeShareable(new FModularRigEditorMode(SharedThis(this)));
+	}
+	return MakeShareable(new FControlRigEditorMode(SharedThis(this)));
+}
+
+bool FControlRigEditor::IsSectionVisible(RigVMNodeSectionID::Type InSectionID) const
+{
+	if(!IControlRigNewEditor::IsSectionVisible(InSectionID))
+	{
+		return false;
+	}
+
+	if(const UControlRigBlueprint* RigBlueprint = GetControlRigBlueprint())
+	{
+		if(IsModularRig())
+		{
+			switch (InSectionID)
+			{
+				case RigVMNodeSectionID::GRAPH:
+				{
+					return RigBlueprint->SupportsEventGraphs();
+				}
+				case RigVMNodeSectionID::FUNCTION:
+				{
+					return RigBlueprint->SupportsFunctions();
+				}
+				default:
+				{
+					break;
+				}
+			}
+		}
+	}
+	return true;
+}
+
+bool FControlRigEditor::NewDocument_IsVisibleForType(FRigVMEditorBase::ECreatedDocumentType GraphType) const
+{
+	if(!IControlRigNewEditor::NewDocument_IsVisibleForType(GraphType))
+	{
+		return false;
+	}
+
+	if(const UControlRigBlueprint* RigBlueprint = GetControlRigBlueprint())
+	{
+		if(IsModularRig())
+		{
+			switch(GraphType)
+			{
+				case FRigVMEditorBase::CGT_NewEventGraph:
+				{
+					return RigBlueprint->SupportsEventGraphs();
+				}
+				case FRigVMEditorBase::CGT_NewFunctionGraph:
+				{
+					return RigBlueprint->SupportsFunctions();
+				}
+				default:
+				{
+					break;
+				}
+			}
+		}
+	}
+
+	return true;
+}
+
+void FControlRigEditor::Tick(float DeltaTime)
+{
+	PreviewScene.UpdateCaptureContents();
+	FControlRigBaseEditor::TickImpl(DeltaTime);
+}
diff --git a/Engine/Plugins/Animation/ControlRig/Source/ControlRigEditor/Private/Editor/ControlRigNewEditor.h b/Engine/Plugins/Animation/ControlRig/Source/ControlRigEditor/Private/Editor/ControlRigNewEditor.h
new file mode 100644
--- /dev/null
+++ b/Engine/Plugins/Animation/ControlRig/Source/ControlRigEditor/Private/Editor/ControlRigNewEditor.h
@@ -0,0 +1,178 @@
+// Copyright Epic Games, Inc. All Rights Reserved.
+
+#pragma once
+
+#include "ControlRigEditor.h"
+#include "Editor/SRigVMDetailsInspector.h"
+#include "PreviewScene.h"
+
+class FControlRigEditor : public IControlRigNewEditor, public FControlRigBaseEditor
+{
+public:
+
+	virtual void InitRigVMEditor(const EToolkitMode::Type Mode, const TSharedPtr< class IToolkitHost >& InitToolkitHost, class URigVMBlueprint* InRigVMBlueprint) override { return FControlRigBaseEditor::InitRigVMEditorImpl(Mode, InitToolkitHost, InRigVMBlueprint); }
+	virtual void InitRigVMEditorSuper(const EToolkitMode::Type Mode, const TSharedPtr<class IToolkitHost>& InitToolkitHost, class URigVMBlueprint* InRigVMBlueprint) override { return IControlRigNewEditor::InitRigVMEditor(Mode, InitToolkitHost, InRigVMBlueprint); }
+
+	virtual const FName GetEditorAppName() const override { return FControlRigBaseEditor::GetEditorAppNameImpl(); }
+	virtual const FName GetEditorModeName() const override { return FControlRigBaseEditor::GetEditorModeNameImpl(); }
+	virtual TSharedPtr<FApplicationMode> CreateEditorMode() override;
+	virtual const FSlateBrush* GetDefaultTabIcon() const override { return FControlRigBaseEditor::GetDefaultTabIconImpl(); }
+	
+public:
+	FControlRigEditor();
+	virtual ~FControlRigEditor();
+
+	// FControlRigBaseEditor
+	virtual bool IsControlRigLegacyEditor() const override { return false; }
+	virtual TSharedPtr<FAssetEditorToolkit> GetHostingApp() override { return IControlRigNewEditor::GetHostingApp(); }
+	virtual TSharedRef<IControlRigBaseEditor> SharedControlRigEditorRef() override { return StaticCastSharedRef<IControlRigBaseEditor>(SharedThis(this)); }
+	virtual TSharedRef<IRigVMEditor> SharedRigVMEditorRef() override { return StaticCastSharedRef<IRigVMEditor>(SharedThis(this)); }
+	virtual TSharedRef<const IRigVMEditor> SharedRigVMEditorRef() const override { return StaticCastSharedRef<const IRigVMEditor>(SharedThis(this)); }
+	
+	virtual bool IsControlRigNewEditor() const { return true; };
+	virtual URigVMBlueprint* GetRigVMBlueprint() const { return IControlRigNewEditor::GetRigVMBlueprint(); }
+	virtual URigVMHost* GetRigVMHost() const { return FRigVMEditorBase::GetRigVMHost(); }
+	virtual TSharedRef<FUICommandList> GetToolkitCommands() { return IControlRigNewEditor::GetToolkitCommands(); }
+	virtual FPreviewScene* GetPreviewScene() { return &PreviewScene; }
+	virtual bool IsDetailsPanelRefreshSuspended() const { return IControlRigNewEditor::IsDetailsPanelRefreshSuspended(); }
+	virtual TArray<TWeakObjectPtr<UObject>> GetSelectedObjects() const { return IControlRigNewEditor::GetSelectedObjects(); }
+	virtual void ClearDetailObject(bool bChangeUISelectionState = true) { IControlRigNewEditor::ClearDetailObject(bChangeUISelectionState); }
+	virtual bool DetailViewShowsStruct(UScriptStruct* InStruct) const { return IControlRigNewEditor::DetailViewShowsStruct(InStruct); }
+	virtual TSharedPtr<class SWidget> GetInspector() const { return Inspector; }
+	virtual TArray<FName> GetEventQueue() const { return IControlRigNewEditor::GetEventQueue(); }
+	virtual void SummonSearchUI(bool bSetFindWithinBlueprint, FString NewSearchTerms = FString(), bool bSelectFirstResult = false) { IControlRigNewEditor::SummonSearchUI(bSetFindWithinBlueprint, NewSearchTerms, bSelectFirstResult); }
+	virtual const TArray< UObject* >* GetObjectsCurrentlyBeingEdited() const { return IControlRigNewEditor::GetObjectsCurrentlyBeingEdited(); }
+	virtual FEditorModeTools& GetEditorModeManagerImpl() const { return GetEditorModeManager(); }
+	virtual const FName GetEditorModeNameImpl() const { return GetEditorModeName(); }
+	virtual URigVMController* GetFocusedController() const { return IControlRigNewEditor::GetFocusedController(); }
+	virtual URigVMGraph* GetFocusedModel() const { return IControlRigNewEditor::GetFocusedModel(); }
+	virtual TArray<FName> GetLastEventQueue() const override { return LastEventQueue; }
+
+
+	// FRigVMEditorBase interface
+	virtual UObject* GetOuterForHost() const override { return GetOuterForHostImpl();}
+	virtual UObject* GetOuterForHostSuper() const override { return IControlRigNewEditor::GetOuterForHost(); }
+	
+	virtual UClass* GetDetailWrapperClass() const { return GetDetailWrapperClassImpl(); }
+	virtual void Compile() override { CompileBaseImpl(); }
+	virtual void CompileSuper() override { IControlRigNewEditor::Compile(); }
+
+	virtual void HandleModifiedEvent(ERigVMGraphNotifType InNotifType, URigVMGraph* InGraph, UObject* InSubject) override { FControlRigBaseEditor::HandleModifiedEventImpl(InNotifType, InGraph, InSubject); }
+	virtual void HandleModifiedEventSuper(ERigVMGraphNotifType InNotifType, URigVMGraph* InGraph, UObject* InSubject) override { IControlRigNewEditor::HandleModifiedEvent(InNotifType, InGraph, InSubject); }
+
+	virtual void OnCreateGraphEditorCommands(TSharedPtr<FUICommandList> GraphEditorCommandsList) override { FControlRigBaseEditor::OnCreateGraphEditorCommandsImpl(GraphEditorCommandsList); }
+	virtual void OnCreateGraphEditorCommandsSuper(TSharedPtr<FUICommandList> GraphEditorCommandsList) override { IControlRigNewEditor::OnCreateGraphEditorCommands(GraphEditorCommandsList); }
+	
+	virtual void HandleVMCompiledEvent(UObject* InCompiledObject, URigVM* InVM, FRigVMExtendedExecuteContext& InContext) override { FControlRigBaseEditor::HandleVMCompiledEventImpl(InCompiledObject, InVM, InContext); }
+	virtual void HandleVMCompiledEventSuper(UObject* InCompiledObject, URigVM* InVM, FRigVMExtendedExecuteContext& InContext) override { IControlRigNewEditor::HandleVMCompiledEvent(InCompiledObject, InVM, InContext); }
+	
+	virtual bool ShouldOpenGraphByDefault() const override { return FControlRigBaseEditor::ShouldOpenGraphByDefaultImpl(); }
+	virtual FReply OnViewportDrop(const FGeometry& MyGeometry, const FDragDropEvent& DragDropEvent) override { return FControlRigBaseEditor::OnViewportDropImpl(MyGeometry, DragDropEvent); }
+	virtual FReply OnViewportDropSuper(const FGeometry& MyGeometry, const FDragDropEvent& DragDropEvent) override { return IControlRigNewEditor::OnViewportDrop(MyGeometry, DragDropEvent); }
+
+	// allows the editor to fill an empty graph
+	virtual void CreateEmptyGraphContent(URigVMController* InController) override { FControlRigBaseEditor::CreateEmptyGraphContentImpl(InController); }
+
+public:
+	
+	// IToolkit Interface
+	virtual FName GetToolkitFName() const override { return FControlRigBaseEditor::GetToolkitFNameImpl(); }
+	virtual FText GetBaseToolkitName() const override { return FControlRigBaseEditor::GetBaseToolkitNameImpl(); }
+	virtual FString GetWorldCentricTabPrefix() const override { return FControlRigBaseEditor::GetWorldCentricTabPrefixImpl(); }
+	virtual FString GetDocumentationLink() const override { return FControlRigBaseEditor::GetDocumentationLinkImpl(); }
+
+	// BlueprintEditor interface
+	virtual FReply OnSpawnGraphNodeByShortcut(FInputChord InChord, const FVector2D& InPosition, UEdGraph* InGraph) override { return FControlRigBaseEditor::OnSpawnGraphNodeByShortcutImpl(InChord, InPosition, InGraph); }
+	virtual FReply OnSpawnGraphNodeByShortcutSuper(FInputChord InChord, const FVector2D& InPosition, UEdGraph* InGraph) override { return IControlRigNewEditor::OnSpawnGraphNodeByShortcut(InChord, InPosition, InGraph); }
+
+	virtual bool IsSectionVisible(RigVMNodeSectionID::Type InSectionID) const override;
+	virtual bool NewDocument_IsVisibleForType(FRigVMEditorBase::ECreatedDocumentType GraphType) const override;
+
+
+	virtual void PostTransaction(bool bSuccess, const FTransaction* Transaction, bool bIsRedo) override { return FControlRigBaseEditor::PostTransactionImpl(bSuccess, Transaction, bIsRedo); }
+
+	//  FTickableEditorObject Interface
+	virtual void Tick(float DeltaTime) override;
+	virtual void TickSuper(float DeltaTime) override { return IControlRigNewEditor::Tick(DeltaTime); }
+
+	virtual void SetDetailObjects(const TArray<UObject*>& InObjects) override { return FControlRigBaseEditor::SetDetailObjectsImpl(InObjects); }
+	virtual void SetDetailObjectsSuper(const TArray<UObject*>& InObjects) override { return IControlRigNewEditor::SetDetailObjects(InObjects); }
+	virtual void RefreshDetailView() override { return FControlRigBaseEditor::RefreshDetailViewImpl(); }
+	virtual void RefreshDetailViewSuper() override { return IControlRigNewEditor::RefreshDetailView(); }
+
+public:
+
+	
+	virtual void OnGraphNodeDropToPerform(TSharedPtr<FDragDropOperation> InDragDropOp, UEdGraph* InGraph, const FVector2D& InNodePosition, const FVector2D& InScreenPosition) override  { return FControlRigBaseEditor::OnGraphNodeDropToPerformImpl(InDragDropOp, InGraph, InNodePosition, InScreenPosition); }
+	virtual void OnGraphNodeDropToPerformSuper(TSharedPtr<FDragDropOperation> InDragDropOp, UEdGraph* InGraph, const FVector2D& InNodePosition, const FVector2D& InScreenPosition) override  { return IControlRigNewEditor::OnGraphNodeDropToPerform(InDragDropOp, InGraph, InNodePosition, InScreenPosition); }
+
+	
+protected:
+
+	virtual void BindCommands() override { return FControlRigBaseEditor::BindCommandsImpl(); }
+	virtual void BindCommandsSuper() override { return IControlRigNewEditor::BindCommands(); }
+	virtual FMenuBuilder GenerateBulkEditMenu() override { return FControlRigBaseEditor::GenerateBulkEditMenuImpl(); }
+	virtual FMenuBuilder GenerateBulkEditMenuSuper() override { return IControlRigNewEditor::GenerateBulkEditMenu(); }
+
+	virtual void SaveAsset_Execute() override { return FControlRigBaseEditor::SaveAsset_ExecuteImpl(); }
+	virtual void SaveAsset_ExecuteSuper() override { return FRigVMEditorBase::SaveAsset_Execute(); }
+	virtual void SaveAssetAs_Execute() override { return FControlRigBaseEditor::SaveAssetAs_ExecuteImpl(); }
+	virtual void SaveAssetAs_ExecuteSuper() override { return FRigVMEditorBase::SaveAssetAs_Execute(); }
+
+	virtual void HandleVMExecutedEvent(URigVMHost* InHost, const FName& InEventName) override { return FControlRigBaseEditor::HandleVMExecutedEventImpl(InHost, InEventName); }
+	virtual void HandleVMExecutedEventSuper(URigVMHost* InHost, const FName& InEventName) override { return IControlRigNewEditor::HandleVMExecutedEvent(InHost, InEventName); }
+
+	// FBaseToolKit overrides
+	virtual void CreateEditorModeManager() override { return FControlRigBaseEditor::CreateEditorModeManagerImpl(); }
+
+private:
+	/** Fill the toolbar with content */
+	virtual void FillToolbar(FToolBarBuilder& ToolbarBuilder, bool bEndSection = true) override { return FControlRigBaseEditor::FillToolbarImpl(ToolbarBuilder, bEndSection); }
+	virtual void FillToolbarSuper(FToolBarBuilder& ToolbarBuilder, bool bEndSection = true) override { return IControlRigNewEditor::FillToolbar(ToolbarBuilder, bEndSection); }
+
+	virtual TArray<FName> GetDefaultEventQueue() const override { return FControlRigBaseEditor::GetDefaultEventQueueImpl(); }
+	virtual void SetEventQueue(TArray<FName> InEventQueue, bool bCompile) override { return FControlRigBaseEditor::SetEventQueueImpl(InEventQueue, bCompile); }
+	virtual void SetEventQueueSuper(TArray<FName> InEventQueue, bool bCompile) override { return IControlRigNewEditor::SetEventQueue(InEventQueue, bCompile); }
+	virtual void SetEventQueueSuper(TArray<FName> InEventQueue) override { return IControlRigNewEditor::SetEventQueue(InEventQueue); }
+	virtual int32 GetEventQueueComboValue() const override { return FControlRigBaseEditor::GetEventQueueComboValueImpl(); }
+	virtual int32 GetEventQueueComboValueSuper() const override { return IControlRigNewEditor::GetEventQueueComboValue(); }
+	virtual FText GetEventQueueLabel() const override { return FControlRigBaseEditor::GetEventQueueLabelImpl(); }
+	virtual FSlateIcon GetEventQueueIcon(const TArray<FName>& InEventQueue) const override { return FControlRigBaseEditor::GetEventQueueIconImpl(InEventQueue); }
+	virtual void HandleSetObjectBeingDebugged(UObject* InObject) override { return FControlRigBaseEditor::HandleSetObjectBeingDebuggedImpl(InObject); }
+	virtual void HandleSetObjectBeingDebuggedSuper(UObject* InObject) override { return IControlRigNewEditor::HandleSetObjectBeingDebugged(InObject); }
+
+	/** Push a newly compiled/opened control rig to the edit mode */
+	virtual void UpdateRigVMHost() override { return FControlRigBaseEditor::UpdateRigVMHostImpl(); }
+	virtual void UpdateRigVMHostSuper() override { return IControlRigNewEditor::UpdateRigVMHost(); }
+	virtual void UpdateRigVMHost_PreClearOldHost(URigVMHost* InPreviousHost) override { return FControlRigBaseEditor::UpdateRigVMHost_PreClearOldHostImpl(InPreviousHost); }
+
+	/** Update the name lists for use in name combo boxes */
+	virtual void CacheNameLists() override { return FControlRigBaseEditor::CacheNameListsImpl(); }
+	virtual void CacheNameListsSuper() override { return IControlRigNewEditor::CacheNameLists(); }
+
+	virtual void GenerateEventQueueMenuContent(FMenuBuilder& MenuBuilder) override { return FControlRigBaseEditor::GenerateEventQueueMenuContentImpl(MenuBuilder); }
+
+	virtual void HandleRefreshEditorFromBlueprint(URigVMBlueprint* InBlueprint) override { return FControlRigBaseEditor::HandleRefreshEditorFromBlueprintImpl(InBlueprint); }
+	virtual void HandleRefreshEditorFromBlueprintSuper(URigVMBlueprint* InBlueprint) override { return IControlRigNewEditor::HandleRefreshEditorFromBlueprint(InBlueprint); }
+
+	
+	/** delegate for changing property */
+	virtual void OnFinishedChangingProperties(const FPropertyChangedEvent& PropertyChangedEvent) override { return FControlRigBaseEditor::OnFinishedChangingPropertiesImpl(PropertyChangedEvent); }
+	virtual void OnFinishedChangingPropertiesSuper(const FPropertyChangedEvent& PropertyChangedEvent) override { return IControlRigNewEditor::OnFinishedChangingProperties(PropertyChangedEvent); }
+	virtual void OnWrappedPropertyChangedChainEvent(URigVMDetailsViewWrapperObject* InWrapperObject, const FString& InPropertyPath, FPropertyChangedChainEvent& InPropertyChangedChainEvent) override { return FControlRigBaseEditor::OnWrappedPropertyChangedChainEventImpl(InWrapperObject, InPropertyPath, InPropertyChangedChainEvent); }
+	virtual void OnWrappedPropertyChangedChainEventSuper(URigVMDetailsViewWrapperObject* InWrapperObject, const FString& InPropertyPath, FPropertyChangedChainEvent& InPropertyChangedChainEvent) override { return IControlRigNewEditor::OnWrappedPropertyChangedChainEvent(InWrapperObject, InPropertyPath, InPropertyChangedChainEvent); }
+
+	virtual void SetEditorModeManager(TSharedPtr<FEditorModeTools> InManager) override { EditorModeManager = InManager; }
+
+	virtual const TArray<TStrongObjectPtr<URigVMDetailsViewWrapperObject>>& GetWrapperObjects() const override { return IControlRigNewEditor::GetWrapperObjects(); }
+	virtual bool& GetSuspendDetailsPanelRefreshFlag() override { return IControlRigNewEditor::GetSuspendDetailsPanelRefreshFlag(); }
+
+	virtual TWeakPtr<class SGraphEditor> GetFocusedGraphEd() const override { return FocusedGraphEdPtr;}
+
+	virtual void OnClose() override { FControlRigBaseEditor::OnClose(); }
+	virtual void OnCloseSuper() override { IControlRigNewEditor::OnClose(); }
+
+private:
+	/** Blueprint preview scene */
+	FPreviewScene PreviewScene;
+};
diff --git a/Engine/Plugins/Animation/ControlRig/Source/ControlRigEditor/Private/Editor/ControlRigSchematicModel.cpp b/Engine/Plugins/Animation/ControlRig/Source/ControlRigEditor/Private/Editor/ControlRigSchematicModel.cpp
--- a/Engine/Plugins/Animation/ControlRig/Source/ControlRigEditor/Private/Editor/ControlRigSchematicModel.cpp
+++ b/Engine/Plugins/Animation/ControlRig/Source/ControlRigEditor/Private/Editor/ControlRigSchematicModel.cpp
@@ -114,7 +114,7 @@ FControlRigSchematicModel::~FControlRigSchematicModel()
 	ControlRigBlueprint.Reset();
 }
 
-void FControlRigSchematicModel::SetEditor(const TSharedRef<FControlRigEditor>& InEditor)
+void FControlRigSchematicModel::SetEditor(const TSharedRef<IControlRigBaseEditor>& InEditor)
 {
 	ControlRigEditor = InEditor;
 	ControlRigBlueprint = ControlRigEditor.Pin()->GetControlRigBlueprint();
@@ -354,7 +354,7 @@ void FControlRigSchematicModel::UpdateElementKeyLinks()
 		return;
 	}
 	
-	const TSharedPtr<FControlRigEditor> Editor = ControlRigEditor.Pin();
+	const TSharedPtr<IControlRigBaseEditor> Editor = ControlRigEditor.Pin();
 	if(!Editor.IsValid())
 	{
 		return;
@@ -663,7 +663,7 @@ FVector2d FControlRigSchematicModel::GetPositionForNode(const FSchematicGraphNod
 
 			if(Key.IsValid())
 			{
-				if(const TSharedPtr<FControlRigEditor> Editor = ControlRigEditor.Pin())
+				if(const TSharedPtr<IControlRigBaseEditor> Editor = ControlRigEditor.Pin())
 				{
 					if(Editor.IsValid())
 					{
diff --git a/Engine/Plugins/Animation/ControlRig/Source/ControlRigEditor/Private/Editor/ControlRigSchematicModel.h b/Engine/Plugins/Animation/ControlRig/Source/ControlRigEditor/Private/Editor/ControlRigSchematicModel.h
--- a/Engine/Plugins/Animation/ControlRig/Source/ControlRigEditor/Private/Editor/ControlRigSchematicModel.h
+++ b/Engine/Plugins/Animation/ControlRig/Source/ControlRigEditor/Private/Editor/ControlRigSchematicModel.h
@@ -8,7 +8,7 @@
 #include "ControlRigBlueprint.h"
 #include "Input/DragAndDrop.h"
 
-class FControlRigEditor;
+class IControlRigBaseEditor;
 class UControlRig;
 class URigHierarchy;
 struct FRigBaseElement;
@@ -75,7 +75,7 @@ public:
 
 	virtual ~FControlRigSchematicModel() override;
 	
-	void SetEditor(const TSharedRef<FControlRigEditor>& InEditor);
+	void SetEditor(const TSharedRef<IControlRigBaseEditor>& InEditor);
 	virtual void ApplyToPanel(SSchematicGraphPanel* InPanel) override;
 
 	virtual void Reset() override;
@@ -132,7 +132,7 @@ private:
 	void OnHideCandidatesForConnector();
 	void OnHierarchyModified(ERigHierarchyNotification InNotif, URigHierarchy* InHierarchy, const FRigNotificationSubject& InSubject);
 
-	TWeakPtr<FControlRigEditor> ControlRigEditor;
+	TWeakPtr<IControlRigBaseEditor> ControlRigEditor;
 	TWeakObjectPtr<UControlRigBlueprint> ControlRigBlueprint;
 	TWeakObjectPtr<UControlRig> ControlRigBeingDebuggedPtr;
 
@@ -142,6 +142,6 @@ private:
 	mutable TMap<FSoftObjectPath, FSlateBrush> ModuleIcons;
 	bool bUpdateElementKeyLinks = true;
 
-	friend class FControlRigEditor;
+	friend class FControlRigBaseEditor;
 	friend class FControlRigSchematicRigElementKeyNode;
 };
\ No newline at end of file
diff --git a/Engine/Plugins/Animation/ControlRig/Source/ControlRigEditor/Private/Editor/ModularRigModelTabSummoner.cpp b/Engine/Plugins/Animation/ControlRig/Source/ControlRigEditor/Private/Editor/ModularRigModelTabSummoner.cpp
--- a/Engine/Plugins/Animation/ControlRig/Source/ControlRigEditor/Private/Editor/ModularRigModelTabSummoner.cpp
+++ b/Engine/Plugins/Animation/ControlRig/Source/ControlRigEditor/Private/Editor/ModularRigModelTabSummoner.cpp
@@ -10,8 +10,8 @@
 
 const FName FModularRigModelTabSummoner::TabID(TEXT("ModularRigModel"));
 
-FModularRigModelTabSummoner::FModularRigModelTabSummoner(const TSharedRef<FControlRigEditor>& InControlRigEditor)
-	: FWorkflowTabFactory(TabID, InControlRigEditor)
+FModularRigModelTabSummoner::FModularRigModelTabSummoner(const TSharedRef<IControlRigBaseEditor>& InControlRigEditor)
+	: FWorkflowTabFactory(TabID, InControlRigEditor->GetHostingApp())
 	, ControlRigEditor(InControlRigEditor)
 {
 	TabLabel = LOCTEXT("ModularRigHierarchyTabLabel", "Module Hierarchy");
@@ -36,7 +36,7 @@ FTabSpawnerEntry& FModularRigModelTabSummoner::RegisterTabSpawner(TSharedRef<FTa
 
 TSharedRef<SWidget> FModularRigModelTabSummoner::CreateTabBody(const FWorkflowTabSpawnInfo& Info) const
 {
-	ControlRigEditor.Pin()->ModularRigHierarchyTabCount++;
+	ControlRigEditor.Pin()->IncreaseModularRigHierarchyTabCount();
 	return SNew(SModularRigModel, ControlRigEditor.Pin().ToSharedRef());
 }
 
@@ -52,7 +52,7 @@ TSharedRef<SDockTab> FModularRigModelTabSummoner::SpawnTab(const FWorkflowTabSpa
 			if(SWidget* Content = &SharedDocTab->GetContent().Get())
 			{
 				SModularRigModel* RigHierarchy = (SModularRigModel*)Content;
-				if(FControlRigEditor* ControlRigEditorForTab = RigHierarchy->GetControlRigEditor())
+				if(IControlRigBaseEditor* ControlRigEditorForTab = RigHierarchy->GetControlRigEditor())
 				{
 					HierarchyTabCount = ControlRigEditorForTab->GetModularRigHierarchyTabCount();
 				}
@@ -69,9 +69,9 @@ TSharedRef<SDockTab> FModularRigModelTabSummoner::SpawnTab(const FWorkflowTabSpa
 		if(SWidget* Content = &DockTab->GetContent().Get())
 		{
 			SModularRigModel* RigHierarchy = (SModularRigModel*)Content;
-			if(FControlRigEditor* ControlRigEditorForTab = RigHierarchy->GetControlRigEditor())
+			if(IControlRigBaseEditor* ControlRigEditorForTab = RigHierarchy->GetControlRigEditor())
 			{
-				ControlRigEditorForTab->ModularRigHierarchyTabCount--;
+				ControlRigEditorForTab->DecreaseModularRigHierarchyTabCount();
 			}
 		}
 	}));
diff --git a/Engine/Plugins/Animation/ControlRig/Source/ControlRigEditor/Private/Editor/ModularRigModelTabSummoner.h b/Engine/Plugins/Animation/ControlRig/Source/ControlRigEditor/Private/Editor/ModularRigModelTabSummoner.h
--- a/Engine/Plugins/Animation/ControlRig/Source/ControlRigEditor/Private/Editor/ModularRigModelTabSummoner.h
+++ b/Engine/Plugins/Animation/ControlRig/Source/ControlRigEditor/Private/Editor/ModularRigModelTabSummoner.h
@@ -5,7 +5,7 @@
 #include "CoreMinimal.h"
 #include "WorkflowOrientedApp/WorkflowTabFactory.h"
 
-class FControlRigEditor;
+class IControlRigBaseEditor;
 
 struct FModularRigModelTabSummoner : public FWorkflowTabFactory
 {
@@ -13,12 +13,12 @@ public:
 	static const FName TabID;
 	
 public:
-	FModularRigModelTabSummoner(const TSharedRef<FControlRigEditor>& InControlRigEditor);
+	FModularRigModelTabSummoner(const TSharedRef<IControlRigBaseEditor>& InControlRigEditor);
 	
 	virtual FTabSpawnerEntry& RegisterTabSpawner(TSharedRef<FTabManager> TabManager, const FApplicationMode* CurrentApplicationMode) const;
 	virtual TSharedRef<SWidget> CreateTabBody(const FWorkflowTabSpawnInfo& Info) const override;
 	virtual TSharedRef<SDockTab> SpawnTab(const FWorkflowTabSpawnInfo& Info) const override;
 
 protected:
-	TWeakPtr<FControlRigEditor> ControlRigEditor;
+	TWeakPtr<IControlRigBaseEditor> ControlRigEditor;
 };
diff --git a/Engine/Plugins/Animation/ControlRig/Source/ControlRigEditor/Private/Editor/RigAnimAttributeTabSummoner.cpp b/Engine/Plugins/Animation/ControlRig/Source/ControlRigEditor/Private/Editor/RigAnimAttributeTabSummoner.cpp
--- a/Engine/Plugins/Animation/ControlRig/Source/ControlRigEditor/Private/Editor/RigAnimAttributeTabSummoner.cpp
+++ b/Engine/Plugins/Animation/ControlRig/Source/ControlRigEditor/Private/Editor/RigAnimAttributeTabSummoner.cpp
@@ -9,8 +9,8 @@
 
 const FName FRigAnimAttributeTabSummoner::TabID(TEXT("RigAnimAttribute"));
 
-FRigAnimAttributeTabSummoner::FRigAnimAttributeTabSummoner(const TSharedRef<FControlRigEditor>& InControlRigEditor)
-	: FWorkflowTabFactory(TabID, InControlRigEditor)
+FRigAnimAttributeTabSummoner::FRigAnimAttributeTabSummoner(const TSharedRef<IControlRigBaseEditor>& InControlRigEditor)
+	: FWorkflowTabFactory(TabID, InControlRigEditor->GetHostingApp())
 	, ControlRigEditor(InControlRigEditor)
 {
 	TabLabel = LOCTEXT("RigAnimAttributeTabLabel", "Animation Attributes");
diff --git a/Engine/Plugins/Animation/ControlRig/Source/ControlRigEditor/Private/Editor/RigAnimAttributeTabSummoner.h b/Engine/Plugins/Animation/ControlRig/Source/ControlRigEditor/Private/Editor/RigAnimAttributeTabSummoner.h
--- a/Engine/Plugins/Animation/ControlRig/Source/ControlRigEditor/Private/Editor/RigAnimAttributeTabSummoner.h
+++ b/Engine/Plugins/Animation/ControlRig/Source/ControlRigEditor/Private/Editor/RigAnimAttributeTabSummoner.h
@@ -5,7 +5,7 @@
 #include "CoreMinimal.h"
 #include "WorkflowOrientedApp/WorkflowTabFactory.h"
 
-class FControlRigEditor;
+class IControlRigBaseEditor;
 
 struct FRigAnimAttributeTabSummoner : public FWorkflowTabFactory
 {
@@ -13,10 +13,10 @@ public:
 	static const FName TabID;
 	
 public:
-	FRigAnimAttributeTabSummoner(const TSharedRef<FControlRigEditor>& InControlRigEditor);
+	FRigAnimAttributeTabSummoner(const TSharedRef<IControlRigBaseEditor>& InControlRigEditor);
 	
 	virtual TSharedRef<SWidget> CreateTabBody(const FWorkflowTabSpawnInfo& Info) const override;
 	
 protected:
-	TWeakPtr<FControlRigEditor> ControlRigEditor;
+	TWeakPtr<IControlRigBaseEditor> ControlRigEditor;
 };
diff --git a/Engine/Plugins/Animation/ControlRig/Source/ControlRigEditor/Private/Editor/RigCurveContainerTabSummoner.cpp b/Engine/Plugins/Animation/ControlRig/Source/ControlRigEditor/Private/Editor/RigCurveContainerTabSummoner.cpp
--- a/Engine/Plugins/Animation/ControlRig/Source/ControlRigEditor/Private/Editor/RigCurveContainerTabSummoner.cpp
+++ b/Engine/Plugins/Animation/ControlRig/Source/ControlRigEditor/Private/Editor/RigCurveContainerTabSummoner.cpp
@@ -9,8 +9,8 @@
 
 const FName FRigCurveContainerTabSummoner::TabID(TEXT("RigCurveContainer"));
 
-FRigCurveContainerTabSummoner::FRigCurveContainerTabSummoner(const TSharedRef<FControlRigEditor>& InControlRigEditor)
-	: FWorkflowTabFactory(TabID, InControlRigEditor)
+FRigCurveContainerTabSummoner::FRigCurveContainerTabSummoner(const TSharedRef<IControlRigBaseEditor>& InControlRigEditor)
+	: FWorkflowTabFactory(TabID, InControlRigEditor->GetHostingApp())
 	, ControlRigEditor(InControlRigEditor)
 {
 	TabLabel = LOCTEXT("RigCurveContainerTabLabel", "Curve Container");
diff --git a/Engine/Plugins/Animation/ControlRig/Source/ControlRigEditor/Private/Editor/RigCurveContainerTabSummoner.h b/Engine/Plugins/Animation/ControlRig/Source/ControlRigEditor/Private/Editor/RigCurveContainerTabSummoner.h
--- a/Engine/Plugins/Animation/ControlRig/Source/ControlRigEditor/Private/Editor/RigCurveContainerTabSummoner.h
+++ b/Engine/Plugins/Animation/ControlRig/Source/ControlRigEditor/Private/Editor/RigCurveContainerTabSummoner.h
@@ -5,7 +5,7 @@
 #include "CoreMinimal.h"
 #include "WorkflowOrientedApp/WorkflowTabFactory.h"
 
-class FControlRigEditor;
+class IControlRigBaseEditor;
 
 struct FRigCurveContainerTabSummoner : public FWorkflowTabFactory
 {
@@ -13,10 +13,10 @@ public:
 	static const FName TabID;
 	
 public:
-	FRigCurveContainerTabSummoner(const TSharedRef<FControlRigEditor>& InControlRigEditor);
+	FRigCurveContainerTabSummoner(const TSharedRef<IControlRigBaseEditor>& InControlRigEditor);
 	
 	virtual TSharedRef<SWidget> CreateTabBody(const FWorkflowTabSpawnInfo& Info) const override;
 	
 protected:
-	TWeakPtr<FControlRigEditor> ControlRigEditor;
+	TWeakPtr<IControlRigBaseEditor> ControlRigEditor;
 };
diff --git a/Engine/Plugins/Animation/ControlRig/Source/ControlRigEditor/Private/Editor/RigHierarchyTabSummoner.cpp b/Engine/Plugins/Animation/ControlRig/Source/ControlRigEditor/Private/Editor/RigHierarchyTabSummoner.cpp
--- a/Engine/Plugins/Animation/ControlRig/Source/ControlRigEditor/Private/Editor/RigHierarchyTabSummoner.cpp
+++ b/Engine/Plugins/Animation/ControlRig/Source/ControlRigEditor/Private/Editor/RigHierarchyTabSummoner.cpp
@@ -11,8 +11,8 @@
 
 const FName FRigHierarchyTabSummoner::TabID(TEXT("RigHierarchy"));
 
-FRigHierarchyTabSummoner::FRigHierarchyTabSummoner(const TSharedRef<FControlRigEditor>& InControlRigEditor)
-	: FWorkflowTabFactory(TabID, InControlRigEditor)
+FRigHierarchyTabSummoner::FRigHierarchyTabSummoner(const TSharedRef<IControlRigBaseEditor>& InControlRigEditor)
+	: FWorkflowTabFactory(TabID, InControlRigEditor->GetHostingApp())
 	, ControlRigEditor(InControlRigEditor)
 {
 	TabLabel = LOCTEXT("RigHierarchyTabLabel", "Rig Hierarchy");
@@ -37,7 +37,7 @@ FTabSpawnerEntry& FRigHierarchyTabSummoner::RegisterTabSpawner(TSharedRef<FTabMa
 
 TSharedRef<SWidget> FRigHierarchyTabSummoner::CreateTabBody(const FWorkflowTabSpawnInfo& Info) const
 {
-	ControlRigEditor.Pin()->RigHierarchyTabCount++;
+	ControlRigEditor.Pin()->IncreaseRigHierarchyTabCount();
 	return SNew(SRigHierarchy, ControlRigEditor.Pin().ToSharedRef());
 }
 
@@ -53,7 +53,7 @@ TSharedRef<SDockTab> FRigHierarchyTabSummoner::SpawnTab(const FWorkflowTabSpawnI
 			if(SWidget* Content = &SharedDocTab->GetContent().Get())
 			{
 				SRigHierarchy* RigHierarchy = (SRigHierarchy*)Content;
-				if(FControlRigEditor* ControlRigEditorForTab = RigHierarchy->GetControlRigEditor())
+				if(IControlRigBaseEditor* ControlRigEditorForTab = RigHierarchy->GetControlRigEditor())
 				{
 					HierarchyTabCount = ControlRigEditorForTab->GetRigHierarchyTabCount();
 				}
@@ -70,9 +70,9 @@ TSharedRef<SDockTab> FRigHierarchyTabSummoner::SpawnTab(const FWorkflowTabSpawnI
 		if(SWidget* Content = &DockTab->GetContent().Get())
 		{
 			SRigHierarchy* RigHierarchy = (SRigHierarchy*)Content;
-			if(FControlRigEditor* ControlRigEditorForTab = RigHierarchy->GetControlRigEditor())
+			if(IControlRigBaseEditor* ControlRigEditorForTab = RigHierarchy->GetControlRigEditor())
 			{
-				ControlRigEditorForTab->RigHierarchyTabCount--;
+				ControlRigEditorForTab->DecreaseRigHierarchyTabCount();
 			}
 		}
 	}));
diff --git a/Engine/Plugins/Animation/ControlRig/Source/ControlRigEditor/Private/Editor/RigHierarchyTabSummoner.h b/Engine/Plugins/Animation/ControlRig/Source/ControlRigEditor/Private/Editor/RigHierarchyTabSummoner.h
--- a/Engine/Plugins/Animation/ControlRig/Source/ControlRigEditor/Private/Editor/RigHierarchyTabSummoner.h
+++ b/Engine/Plugins/Animation/ControlRig/Source/ControlRigEditor/Private/Editor/RigHierarchyTabSummoner.h
@@ -5,7 +5,7 @@
 #include "CoreMinimal.h"
 #include "WorkflowOrientedApp/WorkflowTabFactory.h"
 
-class FControlRigEditor;
+class IControlRigBaseEditor;
 
 struct FRigHierarchyTabSummoner : public FWorkflowTabFactory
 {
@@ -13,12 +13,12 @@ public:
 	static const FName TabID;
 	
 public:
-	FRigHierarchyTabSummoner(const TSharedRef<FControlRigEditor>& InControlRigEditor);
+	FRigHierarchyTabSummoner(const TSharedRef<IControlRigBaseEditor>& InControlRigEditor);
 	
 	virtual FTabSpawnerEntry& RegisterTabSpawner(TSharedRef<FTabManager> TabManager, const FApplicationMode* CurrentApplicationMode) const;
 	virtual TSharedRef<SWidget> CreateTabBody(const FWorkflowTabSpawnInfo& Info) const override;
 	virtual TSharedRef<SDockTab> SpawnTab(const FWorkflowTabSpawnInfo& Info) const override;
 
 protected:
-	TWeakPtr<FControlRigEditor> ControlRigEditor;
+	TWeakPtr<IControlRigBaseEditor> ControlRigEditor;
 };
diff --git a/Engine/Plugins/Animation/ControlRig/Source/ControlRigEditor/Private/Editor/RigInfluenceMapTabSummoner.cpp b/Engine/Plugins/Animation/ControlRig/Source/ControlRigEditor/Private/Editor/RigInfluenceMapTabSummoner.cpp
--- a/Engine/Plugins/Animation/ControlRig/Source/ControlRigEditor/Private/Editor/RigInfluenceMapTabSummoner.cpp
+++ b/Engine/Plugins/Animation/ControlRig/Source/ControlRigEditor/Private/Editor/RigInfluenceMapTabSummoner.cpp
@@ -16,8 +16,8 @@
 
 const FName FRigInfluenceMapTabSummoner::TabID(TEXT("RigInfluenceMap"));
 
-FRigInfluenceMapTabSummoner::FRigInfluenceMapTabSummoner(const TSharedRef<FControlRigEditor>& InControlRigEditor)
-	: FWorkflowTabFactory(TabID, InControlRigEditor)
+FRigInfluenceMapTabSummoner::FRigInfluenceMapTabSummoner(const TSharedRef<IControlRigBaseEditor>& InControlRigEditor)
+	: FWorkflowTabFactory(TabID, InControlRigEditor->GetHostingApp())
 	, ControlRigEditor(InControlRigEditor)
 {
 	TabLabel = LOCTEXT("RigInfluenceMapTabLabel", "Rig Influence Map");
@@ -37,9 +37,9 @@ TSharedRef<SWidget> FRigInfluenceMapTabSummoner::CreateTabBody(const FWorkflowTa
 	TSharedRef<SRigVMDetailsInspector> KismetInspector = SNew(SRigVMDetailsInspector);
 #endif
 
-	if (TSharedPtr<FControlRigEditor> Editor = ControlRigEditor.Pin())
+	if (TSharedPtr<IControlRigBaseEditor> Editor = ControlRigEditor.Pin())
 	{
-		if (UControlRigBlueprint* RigBlueprint = Cast<UControlRigBlueprint>(Editor->FRigVMEditorBase::GetBlueprintObj()))
+		if (UControlRigBlueprint* RigBlueprint = Cast<UControlRigBlueprint>(Editor->GetControlRigBlueprint()))
 		{
 			TSharedPtr<FStructOnScope> StructToDisplay = MakeShareable(new FStructOnScope(
 				FRigInfluenceMapPerEvent::StaticStruct(), (uint8*)&RigBlueprint->Influences));
diff --git a/Engine/Plugins/Animation/ControlRig/Source/ControlRigEditor/Private/Editor/RigInfluenceMapTabSummoner.h b/Engine/Plugins/Animation/ControlRig/Source/ControlRigEditor/Private/Editor/RigInfluenceMapTabSummoner.h
--- a/Engine/Plugins/Animation/ControlRig/Source/ControlRigEditor/Private/Editor/RigInfluenceMapTabSummoner.h
+++ b/Engine/Plugins/Animation/ControlRig/Source/ControlRigEditor/Private/Editor/RigInfluenceMapTabSummoner.h
@@ -5,7 +5,7 @@
 #include "CoreMinimal.h"
 #include "WorkflowOrientedApp/WorkflowTabFactory.h"
 
-class FControlRigEditor;
+class IControlRigBaseEditor;
 
 struct FRigInfluenceMapTabSummoner : public FWorkflowTabFactory
 {
@@ -13,10 +13,10 @@ public:
 	static const FName TabID;
 	
 public:
-	FRigInfluenceMapTabSummoner(const TSharedRef<FControlRigEditor>& InControlRigEditor);
+	FRigInfluenceMapTabSummoner(const TSharedRef<IControlRigBaseEditor>& InControlRigEditor);
 	
 	virtual TSharedRef<SWidget> CreateTabBody(const FWorkflowTabSpawnInfo& Info) const override;
 	
 protected:
-	TWeakPtr<FControlRigEditor> ControlRigEditor;
+	TWeakPtr<IControlRigBaseEditor> ControlRigEditor;
 };
diff --git a/Engine/Plugins/Animation/ControlRig/Source/ControlRigEditor/Private/Editor/RigModuleAssetBrowserTabSummoner.cpp b/Engine/Plugins/Animation/ControlRig/Source/ControlRigEditor/Private/Editor/RigModuleAssetBrowserTabSummoner.cpp
--- a/Engine/Plugins/Animation/ControlRig/Source/ControlRigEditor/Private/Editor/RigModuleAssetBrowserTabSummoner.cpp
+++ b/Engine/Plugins/Animation/ControlRig/Source/ControlRigEditor/Private/Editor/RigModuleAssetBrowserTabSummoner.cpp
@@ -10,8 +10,8 @@
 
 const FName FRigModuleAssetBrowserTabSummoner::TabID(TEXT("RigModuleAssetBrowser"));
 
-FRigModuleAssetBrowserTabSummoner::FRigModuleAssetBrowserTabSummoner(const TSharedRef<FControlRigEditor>& InEditor)
-	: FWorkflowTabFactory(TabID, InEditor),
+FRigModuleAssetBrowserTabSummoner::FRigModuleAssetBrowserTabSummoner(const TSharedRef<IControlRigBaseEditor>& InEditor)
+	: FWorkflowTabFactory(TabID, InEditor->GetHostingApp()),
 	ControlRigEditor(InEditor)
 {
 	bIsSingleton = true; // only allow a single instance of this tab
diff --git a/Engine/Plugins/Animation/ControlRig/Source/ControlRigEditor/Private/Editor/RigModuleAssetBrowserTabSummoner.h b/Engine/Plugins/Animation/ControlRig/Source/ControlRigEditor/Private/Editor/RigModuleAssetBrowserTabSummoner.h
--- a/Engine/Plugins/Animation/ControlRig/Source/ControlRigEditor/Private/Editor/RigModuleAssetBrowserTabSummoner.h
+++ b/Engine/Plugins/Animation/ControlRig/Source/ControlRigEditor/Private/Editor/RigModuleAssetBrowserTabSummoner.h
@@ -5,18 +5,18 @@
 #include "CoreMinimal.h"
 #include "WorkflowOrientedApp/WorkflowTabFactory.h"
 
-class FControlRigEditor;
+class IControlRigBaseEditor;
 
 struct FRigModuleAssetBrowserTabSummoner : public FWorkflowTabFactory
 {
 public:
 	static const FName TabID;
 	
-	FRigModuleAssetBrowserTabSummoner(const TSharedRef<FControlRigEditor>& InControlRigEditor);
+	FRigModuleAssetBrowserTabSummoner(const TSharedRef<IControlRigBaseEditor>& InControlRigEditor);
 	virtual TSharedRef<SWidget> CreateTabBody(const FWorkflowTabSpawnInfo& Info) const override;
 	virtual TSharedPtr<SToolTip> CreateTabToolTipWidget(const FWorkflowTabSpawnInfo& Info) const override;
 	virtual TSharedRef<SDockTab> SpawnTab(const FWorkflowTabSpawnInfo& Info) const override;
 
 protected:
-	TWeakPtr<FControlRigEditor> ControlRigEditor;
+	TWeakPtr<IControlRigBaseEditor> ControlRigEditor;
 };
diff --git a/Engine/Plugins/Animation/ControlRig/Source/ControlRigEditor/Private/Editor/RigValidationTabSummoner.cpp b/Engine/Plugins/Animation/ControlRig/Source/ControlRigEditor/Private/Editor/RigValidationTabSummoner.cpp
--- a/Engine/Plugins/Animation/ControlRig/Source/ControlRigEditor/Private/Editor/RigValidationTabSummoner.cpp
+++ b/Engine/Plugins/Animation/ControlRig/Source/ControlRigEditor/Private/Editor/RigValidationTabSummoner.cpp
@@ -11,8 +11,8 @@
 
 const FName FRigValidationTabSummoner::TabID(TEXT("RigValidation"));
 
-FRigValidationTabSummoner::FRigValidationTabSummoner(const TSharedRef<FControlRigEditor>& InControlRigEditor)
-	: FWorkflowTabFactory(TabID, InControlRigEditor)
+FRigValidationTabSummoner::FRigValidationTabSummoner(const TSharedRef<IControlRigBaseEditor>& InControlRigEditor)
+	: FWorkflowTabFactory(TabID, InControlRigEditor->GetHostingApp())
 	, WeakControlRigEditor(InControlRigEditor)
 {
 	TabLabel = LOCTEXT("RigValidationTabLabel", "Rig Validation");
@@ -27,7 +27,10 @@ FRigValidationTabSummoner::FRigValidationTabSummoner(const TSharedRef<FControlRi
 TSharedRef<SWidget> FRigValidationTabSummoner::CreateTabBody(const FWorkflowTabSpawnInfo& Info) const
 {
 	ensure(WeakControlRigEditor.IsValid());
-	if (TSharedPtr<FRigVMEditorBase> Editor = WeakControlRigEditor.Pin())
+	TSharedRef<IRigVMEditor> EditorRef = WeakControlRigEditor.Pin()->SharedRigVMEditorRef();
+	TSharedPtr<FRigVMEditorBase> EditorPtr = StaticCastSharedPtr<FRigVMEditorBase>(EditorRef.ToSharedPtr());
+	
+	if (TSharedPtr<FRigVMEditorBase> Editor = EditorPtr)
 	{
 		UControlRigBlueprint* RigBlueprint = Cast<UControlRigBlueprint>(Editor->FRigVMEditorBase::GetBlueprintObj());
 		check(RigBlueprint);
diff --git a/Engine/Plugins/Animation/ControlRig/Source/ControlRigEditor/Private/Editor/RigValidationTabSummoner.h b/Engine/Plugins/Animation/ControlRig/Source/ControlRigEditor/Private/Editor/RigValidationTabSummoner.h
--- a/Engine/Plugins/Animation/ControlRig/Source/ControlRigEditor/Private/Editor/RigValidationTabSummoner.h
+++ b/Engine/Plugins/Animation/ControlRig/Source/ControlRigEditor/Private/Editor/RigValidationTabSummoner.h
@@ -5,7 +5,7 @@
 #include "CoreMinimal.h"
 #include "WorkflowOrientedApp/WorkflowTabFactory.h"
 
-class FControlRigEditor;
+class IControlRigBaseEditor;
 
 struct FRigValidationTabSummoner : public FWorkflowTabFactory
 {
@@ -13,10 +13,10 @@ public:
 	static const FName TabID;
 	
 public:
-	FRigValidationTabSummoner(const TSharedRef<FControlRigEditor>& InControlRigEditor);
+	FRigValidationTabSummoner(const TSharedRef<IControlRigBaseEditor>& InControlRigEditor);
 	
 	virtual TSharedRef<SWidget> CreateTabBody(const FWorkflowTabSpawnInfo& Info) const override;
 	
 protected:
-	TWeakPtr<FControlRigEditor> WeakControlRigEditor;
+	TWeakPtr<IControlRigBaseEditor> WeakControlRigEditor;
 };
diff --git a/Engine/Plugins/Animation/ControlRig/Source/ControlRigEditor/Private/Editor/SControlRigAnimAttributeView.cpp b/Engine/Plugins/Animation/ControlRig/Source/ControlRigEditor/Private/Editor/SControlRigAnimAttributeView.cpp
--- a/Engine/Plugins/Animation/ControlRig/Source/ControlRigEditor/Private/Editor/SControlRigAnimAttributeView.cpp
+++ b/Engine/Plugins/Animation/ControlRig/Source/ControlRigEditor/Private/Editor/SControlRigAnimAttributeView.cpp
@@ -21,7 +21,7 @@ SControlRigAnimAttributeView::~SControlRigAnimAttributeView()
 	ControlRigBlueprint->OnSetObjectBeingDebugged().RemoveAll(this);
 }
 
-void SControlRigAnimAttributeView::Construct(const FArguments& InArgs, TSharedRef<FControlRigEditor> InControlRigEditor)
+void SControlRigAnimAttributeView::Construct(const FArguments& InArgs, TSharedRef<IControlRigBaseEditor> InControlRigEditor)
 {
 	SAssignNew(AttributeView, SAnimAttributeView);
 	
diff --git a/Engine/Plugins/Animation/ControlRig/Source/ControlRigEditor/Private/Editor/SControlRigAnimAttributeView.h b/Engine/Plugins/Animation/ControlRig/Source/ControlRigEditor/Private/Editor/SControlRigAnimAttributeView.h
--- a/Engine/Plugins/Animation/ControlRig/Source/ControlRigEditor/Private/Editor/SControlRigAnimAttributeView.h
+++ b/Engine/Plugins/Animation/ControlRig/Source/ControlRigEditor/Private/Editor/SControlRigAnimAttributeView.h
@@ -5,7 +5,7 @@
 #include "CoreMinimal.h"
 #include "Widgets/SCompoundWidget.h"
 
-class FControlRigEditor;
+class IControlRigBaseEditor;
 class UControlRigBlueprint;
 class UControlRig;
 
@@ -21,7 +21,7 @@ class SControlRigAnimAttributeView : public SCompoundWidget
 
 	~SControlRigAnimAttributeView();
 	
-	void Construct( const FArguments& InArgs, TSharedRef<FControlRigEditor> InControlRigEditor);
+	void Construct( const FArguments& InArgs, TSharedRef<IControlRigBaseEditor> InControlRigEditor);
 
 	void HandleSetObjectBeingDebugged(UObject* InObject);
 	
@@ -30,7 +30,7 @@ class SControlRigAnimAttributeView : public SCompoundWidget
 
 	void HandleControlRigPostForwardSolve(UControlRig* InControlRig, const FName& InEventName) const;
 	
-	TWeakPtr<FControlRigEditor> ControlRigEditor;
+	TWeakPtr<IControlRigBaseEditor> ControlRigEditor;
 	TWeakObjectPtr<UControlRigBlueprint> ControlRigBlueprint;
 	TWeakObjectPtr<UControlRig> ControlRigBeingDebuggedPtr;
 	TSharedPtr<SAnimAttributeView> AttributeView;
diff --git a/Engine/Plugins/Animation/ControlRig/Source/ControlRigEditor/Private/Editor/SModularRigModel.cpp b/Engine/Plugins/Animation/ControlRig/Source/ControlRigEditor/Private/Editor/SModularRigModel.cpp
--- a/Engine/Plugins/Animation/ControlRig/Source/ControlRigEditor/Private/Editor/SModularRigModel.cpp
+++ b/Engine/Plugins/Animation/ControlRig/Source/ControlRigEditor/Private/Editor/SModularRigModel.cpp
@@ -71,11 +71,11 @@ const FName SModularRigModel::ContextMenuName = TEXT("ControlRigEditor.ModularRi
 
 SModularRigModel::~SModularRigModel()
 {
-	const FControlRigEditor* Editor = ControlRigEditor.IsValid() ? ControlRigEditor.Pin().Get() : nullptr;
+	IControlRigBaseEditor* Editor = ControlRigEditor.IsValid() ? ControlRigEditor.Pin().Get() : nullptr;
 	OnEditorClose(Editor, ControlRigBlueprint.Get());
 }
 
-void SModularRigModel::Construct(const FArguments& InArgs, TSharedRef<FControlRigEditor> InControlRigEditor)
+void SModularRigModel::Construct(const FArguments& InArgs, TSharedRef<IControlRigBaseEditor> InControlRigEditor)
 {
 	ControlRigEditor = InControlRigEditor;
 
@@ -242,22 +242,23 @@ void SModularRigModel::Construct(const FArguments& InArgs, TSharedRef<FControlRi
 	}
 }
 
-void SModularRigModel::OnEditorClose(const FRigVMEditorBase* InEditor, URigVMBlueprint* InBlueprint)
+void SModularRigModel::OnEditorClose(IControlRigBaseEditor* InEditor, UControlRigBlueprint* InBlueprint)
 {
 	if (InEditor)
 	{
-		FControlRigEditor* Editor = (FControlRigEditor*)InEditor;  
+		IControlRigBaseEditor* Editor = (IControlRigBaseEditor*)InEditor;  
 		Editor->OnGetViewportContextMenu().Unbind();
 		Editor->OnViewportContextMenuCommands().Unbind();
+		Editor->OnEditorClosed().RemoveAll(this);
 	}
 
-	if (UControlRigBlueprint* BP = Cast<UControlRigBlueprint>(InBlueprint))
+	if (InBlueprint)
 	{
 		InBlueprint->OnRefreshEditor().RemoveAll(this);
 		InBlueprint->OnSetObjectBeingDebugged().RemoveAll(this);
-		BP->OnModularRigPreCompiled().RemoveAll(this);
-		BP->OnModularRigCompiled().RemoveAll(this);
-		if(UModularRigController* ModularRigController = BP->GetModularRigController())
+		InBlueprint->OnModularRigPreCompiled().RemoveAll(this);
+		InBlueprint->OnModularRigCompiled().RemoveAll(this);
+		if(UModularRigController* ModularRigController = InBlueprint->GetModularRigController())
 		{
 			ModularRigController->OnModified().RemoveAll(this);
 		}
@@ -406,7 +407,7 @@ void SModularRigModel::HandlePostCompileModularRigs(URigVMBlueprint* InBlueprint
 	if (ControlRigEditor.IsValid())
 	{
 		TArray<TSharedPtr<FModularRigTreeElement>> SelectedElements;
-		Algo::Transform(ControlRigEditor.Pin()->ModulesSelected, SelectedElements, [this](const FName& ModuleName)
+		Algo::Transform(ControlRigEditor.Pin()->GetSelectedModules(), SelectedElements, [this](const FName& ModuleName)
 		{
 			return TreeView->FindElement(ModuleName.ToString());
 		});
diff --git a/Engine/Plugins/Animation/ControlRig/Source/ControlRigEditor/Private/Editor/SModularRigModel.h b/Engine/Plugins/Animation/ControlRig/Source/ControlRigEditor/Private/Editor/SModularRigModel.h
--- a/Engine/Plugins/Animation/ControlRig/Source/ControlRigEditor/Private/Editor/SModularRigModel.h
+++ b/Engine/Plugins/Animation/ControlRig/Source/ControlRigEditor/Private/Editor/SModularRigModel.h
@@ -11,7 +11,7 @@
 #include "ControlRigDragOps.h"
 
 class SModularRigModel;
-class FControlRigEditor;
+class IControlRigBaseEditor;
 class SSearchBox;
 class FUICommandList;
 class URigVMBlueprint;
@@ -30,9 +30,9 @@ public:
 
 	~SModularRigModel();
 
-	void Construct(const FArguments& InArgs, TSharedRef<FControlRigEditor> InControlRigEditor);
+	void Construct(const FArguments& InArgs, TSharedRef<IControlRigBaseEditor> InControlRigEditor);
 
-	FControlRigEditor* GetControlRigEditor() const
+	IControlRigBaseEditor* GetControlRigEditor() const
 	{
 		if(ControlRigEditor.IsValid())
 		{
@@ -43,7 +43,7 @@ public:
 
 private:
 
-	void OnEditorClose(const FRigVMEditorBase* InEditor, URigVMBlueprint* InBlueprint);
+	void OnEditorClose(IControlRigBaseEditor* InEditor, UControlRigBlueprint* InBlueprint);
 
 	/** Bind commands that this widget handles */
 	void BindCommands();
@@ -124,7 +124,7 @@ private:
 	TSharedPtr<FUICommandList> GetContextMenuCommands() const;
 	
 	/** Our owning control rig editor */
-	TWeakPtr<FControlRigEditor> ControlRigEditor;
+	TWeakPtr<IControlRigBaseEditor> ControlRigEditor;
 
 	/** Tree view widget */
 	TSharedPtr<SModularRigTreeView> TreeView;
diff --git a/Engine/Plugins/Animation/ControlRig/Source/ControlRigEditor/Private/Editor/SRigCurveContainer.cpp b/Engine/Plugins/Animation/ControlRig/Source/ControlRigEditor/Private/Editor/SRigCurveContainer.cpp
--- a/Engine/Plugins/Animation/ControlRig/Source/ControlRigEditor/Private/Editor/SRigCurveContainer.cpp
+++ b/Engine/Plugins/Animation/ControlRig/Source/ControlRigEditor/Private/Editor/SRigCurveContainer.cpp
@@ -138,7 +138,7 @@ float SRigCurveListRow::GetValue() const
 //////////////////////////////////////////////////////////////////////////
 // SRigCurveContainer
 
-void SRigCurveContainer::Construct(const FArguments& InArgs, TSharedRef<FControlRigEditor> InControlRigEditor)
+void SRigCurveContainer::Construct(const FArguments& InArgs, TSharedRef<IControlRigBaseEditor> InControlRigEditor)
 {
 	ControlRigEditor = InControlRigEditor;
 	ControlRigBlueprint = InControlRigEditor.Get().GetControlRigBlueprint();
@@ -211,7 +211,7 @@ void SRigCurveContainer::Construct(const FArguments& InArgs, TSharedRef<FControl
 
 SRigCurveContainer::~SRigCurveContainer()
 {
-	const FControlRigEditor* Editor = ControlRigEditor.IsValid() ? ControlRigEditor.Pin().Get() : nullptr;
+	IControlRigBaseEditor* Editor = ControlRigEditor.IsValid() ? ControlRigEditor.Pin().Get() : nullptr;
 	OnEditorClose(Editor, ControlRigBlueprint.Get());
 }
 
@@ -310,12 +310,17 @@ TSharedPtr<SWidget> SRigCurveContainer::OnGetContextMenuContent() const
 	return MenuBuilder.MakeWidget();
 }
 
-void SRigCurveContainer::OnEditorClose(const FRigVMEditorBase* InEditor, URigVMBlueprint* InBlueprint)
+void SRigCurveContainer::OnEditorClose(IControlRigBaseEditor* InEditor, UControlRigBlueprint* InBlueprint)
 {
-	if (UControlRigBlueprint* BP = Cast<UControlRigBlueprint>(InBlueprint))
+	if (InBlueprint)
 	{
-		BP->Hierarchy->OnModified().RemoveAll(this);
-		BP->OnRefreshEditor().RemoveAll(this);
+		InBlueprint->Hierarchy->OnModified().RemoveAll(this);
+		InBlueprint->OnRefreshEditor().RemoveAll(this);
+	}
+
+	if (InEditor)
+	{
+		InEditor->OnEditorClosed().RemoveAll(this);
 	}
 
 	ControlRigBlueprint.Reset();
@@ -740,7 +745,7 @@ URigHierarchy* SRigCurveContainer::GetInstanceHierarchy() const
 {
 	if (ControlRigEditor.IsValid())
 	{
-		UControlRig* ControlRig = ControlRigEditor.Pin()->GetInstanceRig();
+		UControlRig* ControlRig = ControlRigEditor.Pin()->GetControlRig();
 		if (ControlRig)
 		{
 			return ControlRig->GetHierarchy();
diff --git a/Engine/Plugins/Animation/ControlRig/Source/ControlRigEditor/Private/Editor/SRigCurveContainer.h b/Engine/Plugins/Animation/ControlRig/Source/ControlRigEditor/Private/Editor/SRigCurveContainer.h
--- a/Engine/Plugins/Animation/ControlRig/Source/ControlRigEditor/Private/Editor/SRigCurveContainer.h
+++ b/Engine/Plugins/Animation/ControlRig/Source/ControlRigEditor/Private/Editor/SRigCurveContainer.h
@@ -15,7 +15,7 @@
 // FDisplayedRigCurveInfo
 
 class SInlineEditableTextBlock;
-class FControlRigEditor;
+class IControlRigBaseEditor;
 struct FAssetData;
 class URigVMBlueprint;
 class FRigVMEditorBase;
@@ -136,7 +136,7 @@ public:
 	* @param InArgs - Arguments passed from Slate
 	*
 	*/
-	void Construct( const FArguments& InArgs, TSharedRef<FControlRigEditor> InControlRigEditor);
+	void Construct( const FArguments& InArgs, TSharedRef<IControlRigBaseEditor> InControlRigEditor);
 
 	/**
 	* Destructor - resets the animation curve
@@ -213,7 +213,7 @@ private:
 	/** Handler for context menus */
 	TSharedPtr<SWidget> OnGetContextMenuContent() const;
 
-	void OnEditorClose(const FRigVMEditorBase* InEditor, URigVMBlueprint* InBlueprint);
+	void OnEditorClose(IControlRigBaseEditor* InEditor, UControlRigBlueprint* InBlueprint);
 
 	/**
 	* Clears and rebuilds the table, according to an optional search string
@@ -261,7 +261,7 @@ private:
 
 	/** Control Rig Blueprint */
 	TWeakObjectPtr<UControlRigBlueprint> ControlRigBlueprint;
-	TWeakPtr<FControlRigEditor> ControlRigEditor;
+	TWeakPtr<IControlRigBaseEditor> ControlRigEditor;
 
 	/** Current text typed into NameFilterBox */
 	FText FilterText;
diff --git a/Engine/Plugins/Animation/ControlRig/Source/ControlRigEditor/Private/Editor/SRigHierarchy.cpp b/Engine/Plugins/Animation/ControlRig/Source/ControlRigEditor/Private/Editor/SRigHierarchy.cpp
--- a/Engine/Plugins/Animation/ControlRig/Source/ControlRigEditor/Private/Editor/SRigHierarchy.cpp
+++ b/Engine/Plugins/Animation/ControlRig/Source/ControlRigEditor/Private/Editor/SRigHierarchy.cpp
@@ -72,11 +72,11 @@ const FName SRigHierarchy::DragDropMenuName = TEXT("ControlRigEditor.RigHierarch
 
 SRigHierarchy::~SRigHierarchy()
 {
-	const FControlRigEditor* Editor = ControlRigEditor.IsValid() ? ControlRigEditor.Pin().Get() : nullptr;
+	IControlRigBaseEditor* Editor = ControlRigEditor.IsValid() ? ControlRigEditor.Pin().Get() : nullptr;
 	OnEditorClose(Editor, ControlRigBlueprint.Get());
 }
 
-void SRigHierarchy::Construct(const FArguments& InArgs, TSharedRef<FControlRigEditor> InControlRigEditor)
+void SRigHierarchy::Construct(const FArguments& InArgs, TSharedRef<IControlRigBaseEditor> InControlRigEditor)
 {
 	ControlRigEditor = InControlRigEditor;
 
@@ -256,14 +256,14 @@ void SRigHierarchy::Construct(const FArguments& InArgs, TSharedRef<FControlRigEd
 	}));
 }
 
-void SRigHierarchy::OnEditorClose(const FRigVMEditorBase* InEditor, URigVMBlueprint* InBlueprint)
+void SRigHierarchy::OnEditorClose(IControlRigBaseEditor* InEditor, UControlRigBlueprint* InBlueprint)
 {
 	if (InEditor)
 	{
-		FControlRigEditor* Editor = (FControlRigEditor*)InEditor;  
-		Editor->GetKeyDownDelegate().Unbind();
-		Editor->OnGetViewportContextMenu().Unbind();
-		Editor->OnViewportContextMenuCommands().Unbind();
+		InEditor->GetKeyDownDelegate().Unbind();
+		InEditor->OnGetViewportContextMenu().Unbind();
+		InEditor->OnViewportContextMenuCommands().Unbind();
+		InEditor->OnEditorClosed().RemoveAll(this);
 	}
 
 	if (UControlRigBlueprint* BP = Cast<UControlRigBlueprint>(InBlueprint))
@@ -577,9 +577,12 @@ FReply SRigHierarchy::OnImportSkeletonClicked()
 
 FText SRigHierarchy::GetImportHierarchyText() const
 {
-	if(ControlRigBlueprint->IsControlRigModule())
+	if (TStrongObjectPtr<UControlRigBlueprint> Blueprint = ControlRigBlueprint.Pin())
 	{
-		return LOCTEXT("SetPreviewMesh", "Set Preview Mesh");
+		if(Blueprint->IsControlRigModule())
+		{
+			return LOCTEXT("SetPreviewMesh", "Set Preview Mesh");
+		}
 	}
 	return LOCTEXT("ImportHierarchy", "Import Hierarchy");
 }
@@ -1816,7 +1819,7 @@ void SRigHierarchy::UpdateMesh(USkeletalMesh* InMesh, const bool bImport) const
 		return;
 	}
 
-	const TSharedPtr<FControlRigEditor> EditorSharedPtr = ControlRigEditor.Pin();
+	const TSharedPtr<IControlRigBaseEditor> EditorSharedPtr = ControlRigEditor.Pin();
 	EditorSharedPtr->GetPersonaToolkit()->SetPreviewMesh(InMesh, true);
 
 	UControlRigSkeletalMeshComponent* Component = Cast<UControlRigSkeletalMeshComponent>(EditorSharedPtr->GetPersonaToolkit()->GetPreviewMeshComponent());
@@ -1828,11 +1831,11 @@ void SRigHierarchy::UpdateMesh(USkeletalMesh* InMesh, const bool bImport) const
 	UAnimInstance* AnimInstance = Component->GetAnimInstance();
 	if (UControlRigLayerInstance* ControlRigLayerInstance = Cast<UControlRigLayerInstance>(AnimInstance))
 	{
-		EditorSharedPtr->PreviewInstance = Cast<UAnimPreviewInstance>(ControlRigLayerInstance->GetSourceAnimInstance());
+		EditorSharedPtr->SetPreviewInstance(Cast<UAnimPreviewInstance>(ControlRigLayerInstance->GetSourceAnimInstance()));
 	}
 	else
 	{
-		EditorSharedPtr->PreviewInstance = Cast<UAnimPreviewInstance>(AnimInstance);
+		EditorSharedPtr->SetPreviewInstance(Cast<UAnimPreviewInstance>(AnimInstance));
 	}
 
 	EditorSharedPtr->Compile();
@@ -1851,7 +1854,7 @@ void SRigHierarchy::RefreshHierarchy(const FAssetData& InAssetData, bool bOnlyRe
 		return;
 	}
 
-	FControlRigEditor* StrongEditor = ControlRigEditor.Pin().Get();
+	IControlRigBaseEditor* StrongEditor = ControlRigEditor.Pin().Get();
 	StrongEditor->ClearDetailObject();
 
 	URigHierarchy* Hierarchy = GetDefaultHierarchy();
@@ -1983,7 +1986,7 @@ void SRigHierarchy::ImportHierarchy(const FAssetData& InAssetData)
 		return;
 	}
 	
-	const TSharedPtr<FControlRigEditor> EditorSharedPtr = ControlRigEditor.Pin();
+	const TSharedPtr<IControlRigBaseEditor> EditorSharedPtr = ControlRigEditor.Pin();
 	if (URigHierarchy* Hierarchy = GetDefaultHierarchy())
 	{
 		// filter out meshes that don't contain a skeleton
@@ -2960,12 +2963,12 @@ FReply SRigHierarchy::OnDragDetected(const FGeometry& MyGeometry, const FPointer
 	TArray<FRigHierarchyKey> DraggedElements = GetSelectedKeys();
 	if (MouseEvent.IsMouseButtonDown(EKeys::LeftMouseButton) && DraggedElements.Num() > 0)
 	{
-		if (ControlRigEditor.IsValid())
+		if (TSharedPtr<IControlRigBaseEditor> EditorPtr = ControlRigEditor.Pin())
 		{
 			UpdateConnectorMatchesOnDrag(DraggedElements);
 			
 			TSharedRef<FRigElementHierarchyDragDropOp> DragDropOp = FRigElementHierarchyDragDropOp::New(MoveTemp(DraggedElements));
-			DragDropOp->OnPerformDropToGraph.BindSP(ControlRigEditor.Pin().Get(), &FControlRigEditor::OnGraphNodeDropToPerform);
+			DragDropOp->OnPerformDropToGraph.BindSP(EditorPtr.ToSharedRef(), &FControlRigBaseEditor::OnGraphNodeDropToPerform);
 			return FReply::Handled().BeginDragDrop(DragDropOp);
 		}
 	}
@@ -3736,9 +3739,9 @@ void SRigHierarchy::HandleSetInitialTransformFromCurrentTransform()
 						SelectedKey.Type == ERigElementType::Connector)
 					{
 						FTransform InitialTransform = LocalTransform;
-						if (ControlRigEditor.Pin()->PreviewInstance)
+						if (ControlRigEditor.Pin()->GetPreviewInstance())
 						{
-							if (FAnimNode_ModifyBone* ModifyBone = ControlRigEditor.Pin()->PreviewInstance->FindModifiedBone(SelectedKey.Name))
+							if (FAnimNode_ModifyBone* ModifyBone = ControlRigEditor.Pin()->GetPreviewInstance()->FindModifiedBone(SelectedKey.Name))
 							{
 								InitialTransform.SetTranslation(ModifyBone->Translation);
 								InitialTransform.SetRotation(FQuat(ModifyBone->Rotation));
diff --git a/Engine/Plugins/Animation/ControlRig/Source/ControlRigEditor/Private/Editor/SRigHierarchy.h b/Engine/Plugins/Animation/ControlRig/Source/ControlRigEditor/Private/Editor/SRigHierarchy.h
--- a/Engine/Plugins/Animation/ControlRig/Source/ControlRigEditor/Private/Editor/SRigHierarchy.h
+++ b/Engine/Plugins/Animation/ControlRig/Source/ControlRigEditor/Private/Editor/SRigHierarchy.h
@@ -8,13 +8,12 @@
 #include "Editor/SRigHierarchyTreeView.h"
 #include "Units/RigUnitContext.h"
 #include "ControlRigBlueprint.h"
-#include "Editor/RigVMNewEditor.h"
 #include "ControlRigSchematicModel.h"
 #include "ControlRigDragOps.h"
 #include "SRigHierarchy.generated.h"
 
 class SRigHierarchy;
-class FControlRigEditor;
+class IControlRigBaseEditor;
 class SSearchBox;
 class FUICommandList;
 class URigVMBlueprint;
@@ -46,9 +45,9 @@ public:
 
 	~SRigHierarchy();
 
-	void Construct(const FArguments& InArgs, TSharedRef<FControlRigEditor> InControlRigEditor);
+	void Construct(const FArguments& InArgs, TSharedRef<IControlRigBaseEditor> InControlRigEditor);
 
-	FControlRigEditor* GetControlRigEditor() const
+	IControlRigBaseEditor* GetControlRigEditor() const
 	{
 		if(ControlRigEditor.IsValid())
 		{
@@ -59,7 +58,7 @@ public:
 
 private:
 
-	void OnEditorClose(const FRigVMEditorBase* InEditor, URigVMBlueprint* InBlueprint);
+	void OnEditorClose(IControlRigBaseEditor* InEditor, UControlRigBlueprint* InBlueprint);
 
 	/** Bind commands that this widget handles */
 	void BindCommands();
@@ -146,7 +145,7 @@ private:
 	UToolMenu* GetDragDropMenu(const TArray<FRigHierarchyKey>& DraggedKeys, FRigElementKey TargetKey);
 
 	/** Our owning control rig editor */
-	TWeakPtr<FControlRigEditor> ControlRigEditor;
+	TWeakPtr<IControlRigBaseEditor> ControlRigEditor;
 
 	FRigTreeDisplaySettings DisplaySettings;
 	const FRigTreeDisplaySettings& GetDisplaySettings() const { return DisplaySettings; }
diff --git a/Engine/Plugins/Animation/ControlRig/Source/ControlRigEditor/Private/Editor/SRigModuleAssetBrowser.cpp b/Engine/Plugins/Animation/ControlRig/Source/ControlRigEditor/Private/Editor/SRigModuleAssetBrowser.cpp
--- a/Engine/Plugins/Animation/ControlRig/Source/ControlRigEditor/Private/Editor/SRigModuleAssetBrowser.cpp
+++ b/Engine/Plugins/Animation/ControlRig/Source/ControlRigEditor/Private/Editor/SRigModuleAssetBrowser.cpp
@@ -16,7 +16,7 @@
 
 void SRigModuleAssetBrowser::Construct(
 	const FArguments& InArgs,
-	TSharedRef<FControlRigEditor> InEditor)
+	TSharedRef<IControlRigBaseEditor> InEditor)
 {
 	ControlRigEditor = InEditor;
 
diff --git a/Engine/Plugins/Animation/ControlRig/Source/ControlRigEditor/Private/Editor/SRigModuleAssetBrowser.h b/Engine/Plugins/Animation/ControlRig/Source/ControlRigEditor/Private/Editor/SRigModuleAssetBrowser.h
--- a/Engine/Plugins/Animation/ControlRig/Source/ControlRigEditor/Private/Editor/SRigModuleAssetBrowser.h
+++ b/Engine/Plugins/Animation/ControlRig/Source/ControlRigEditor/Private/Editor/SRigModuleAssetBrowser.h
@@ -9,7 +9,7 @@
 #include "Widgets/SWindow.h"
 
 
-class FControlRigEditor;
+class IControlRigBaseEditor;
 
 class SRigModuleAssetBrowser : public SBox
 {
@@ -17,7 +17,7 @@ public:
 	SLATE_BEGIN_ARGS(SRigModuleAssetBrowser) {}
 	SLATE_END_ARGS()
 
-	void Construct(const FArguments& InArgs, TSharedRef<FControlRigEditor> InEditor);
+	void Construct(const FArguments& InArgs, TSharedRef<IControlRigBaseEditor> InEditor);
 
 	void RefreshView();
 	
@@ -33,11 +33,11 @@ private:
 	FGetCurrentSelectionDelegate GetCurrentSelectionDelegate;
 		
 	/** editor controller */
-	TWeakPtr<FControlRigEditor> ControlRigEditor;
+	TWeakPtr<IControlRigBaseEditor> ControlRigEditor;
 
 	/** the animation asset browser */
 	TSharedPtr<SBox> AssetBrowserBox;
 
 
-	friend FControlRigEditor;
+	friend class FControlRigBaseEditor;
 };
diff --git a/Engine/Plugins/Animation/ControlRig/Source/ControlRigEditor/Public/EditMode/ControlRigEditMode.h b/Engine/Plugins/Animation/ControlRig/Source/ControlRigEditor/Public/EditMode/ControlRigEditMode.h
--- a/Engine/Plugins/Animation/ControlRig/Source/ControlRigEditor/Public/EditMode/ControlRigEditMode.h
+++ b/Engine/Plugins/Animation/ControlRig/Source/ControlRigEditor/Public/EditMode/ControlRigEditMode.h
@@ -691,6 +691,10 @@ private:
 	mutable UE::AnimationEditMode::FControlRigKeyframer Keyframer;
 
 	friend class FControlRigEditorModule;
+	friend class FControlRigBaseEditor;
+#if WITH_RIGVMLEGACYEDITOR
+	friend class FControlRigLegacyEditor;
+#endif
 	friend class FControlRigEditor;
 	friend class FControlRigEditModeGenericDetails;
 	friend class UControlRigEditModeDelegateHelper;
diff --git a/Engine/Plugins/Animation/ControlRig/Source/ControlRigEditor/Public/IControlRigEditor.h b/Engine/Plugins/Animation/ControlRig/Source/ControlRigEditor/Public/IControlRigEditor.h
--- a/Engine/Plugins/Animation/ControlRig/Source/ControlRigEditor/Public/IControlRigEditor.h
+++ b/Engine/Plugins/Animation/ControlRig/Source/ControlRigEditor/Public/IControlRigEditor.h
@@ -6,20 +6,24 @@
 
 #if WITH_RIGVMLEGACYEDITOR
 #include "Editor/RigVMLegacyEditor.h"
-#else
-#include "Editor/RigVMNewEditor.h"
 #endif
+#include "Editor/RigVMNewEditor.h"
 
 class UControlRigBlueprint;
 class FBlueprintActionDatabaseRegistrar;
 
+class IControlRigEditor
+{
+public:
+	virtual URigVMHost* GetRigVMHost() const = 0;
+};
+
 #if WITH_RIGVMLEGACYEDITOR
-class IControlRigEditor : public FRigVMLegacyEditor
+class IControlRigLegacyEditor : public FRigVMLegacyEditor, public IControlRigEditor
 {
 };
+#endif
 
-#else
-class IControlRigEditor : public FRigVMNewEditor
+class IControlRigNewEditor : public FRigVMNewEditor, public IControlRigEditor
 {
 };
-#endif
\ No newline at end of file
diff --git a/Engine/Plugins/Runtime/RigVM/Source/RigVMEditor/Private/Editor/Kismet/RigVMBlueprintCompilationManager.cpp b/Engine/Plugins/Runtime/RigVM/Source/RigVMEditor/Private/Editor/Kismet/RigVMBlueprintCompilationManager.cpp
--- a/Engine/Plugins/Runtime/RigVM/Source/RigVMEditor/Private/Editor/Kismet/RigVMBlueprintCompilationManager.cpp
+++ b/Engine/Plugins/Runtime/RigVM/Source/RigVMEditor/Private/Editor/Kismet/RigVMBlueprintCompilationManager.cpp
@@ -3081,7 +3081,9 @@ UClass* FRigVMBlueprintCompilationManagerImpl::FastGenerateSkeletonClass(UBluepr
 			CompilerContext.ValidateVariableNames();
 		}
 
-		CompilerContext.CleanAndSanitizeClass(Ret, MutableView(Ret->GetDefaultObject(false)));
+PRAGMA_DISABLE_DEPRECATION_WARNINGS
+		CompilerContext.CleanAndSanitizeClass(Ret, MutableView(Ret->ClassDefaultObject));
+PRAGMA_ENABLE_DEPRECATION_WARNINGS
 	}
 
 	Ret->ClassGeneratedBy = BP;
diff --git a/Engine/Plugins/Runtime/RigVM/Source/RigVMEditor/Private/RigVMEditorModule.cpp b/Engine/Plugins/Runtime/RigVM/Source/RigVMEditor/Private/RigVMEditorModule.cpp
--- a/Engine/Plugins/Runtime/RigVM/Source/RigVMEditor/Private/RigVMEditorModule.cpp
+++ b/Engine/Plugins/Runtime/RigVM/Source/RigVMEditor/Private/RigVMEditorModule.cpp
@@ -53,6 +53,9 @@ DEFINE_LOG_CATEGORY(LogRigVMEditor);
 
 IMPLEMENT_MODULE(FRigVMEditorModule, RigVMEditor)
 
+TAutoConsoleVariable<bool> CVarRigVMUseNewEditor(TEXT("RigVM.UseNewEditor"), 0, TEXT("Use editor which does not inherit from BlueprintEditor."));
+TAutoConsoleVariable<bool> CVarRigVMUseDualEditor(TEXT("RigVM.UseDualEditor"), 0, TEXT("Use both legacy and new editor for A/B comparison."));
+
 FRigVMEditorModule& FRigVMEditorModule::Get()
 {
 	return FModuleManager::LoadModuleChecked< FRigVMEditorModule >(TEXT("RigVMEditor"));
diff --git a/Engine/Plugins/Runtime/RigVM/Source/RigVMEditor/Public/Editor/RigVMEditor.h b/Engine/Plugins/Runtime/RigVM/Source/RigVMEditor/Public/Editor/RigVMEditor.h
--- a/Engine/Plugins/Runtime/RigVM/Source/RigVMEditor/Public/Editor/RigVMEditor.h
+++ b/Engine/Plugins/Runtime/RigVM/Source/RigVMEditor/Public/Editor/RigVMEditor.h
@@ -101,6 +101,7 @@ class RIGVMEDITOR_API IRigVMEditor
 public:
 	
 	virtual TSharedRef<IRigVMEditor> SharedRef() = 0;
+	virtual TSharedRef<const IRigVMEditor> SharedRef() const = 0;
 
 	virtual TSharedPtr<FAssetEditorToolkit> GetHostingApp() = 0;
 	virtual const TSharedPtr<FAssetEditorToolkit> GetHostingApp() const = 0;
diff --git a/Engine/Plugins/Runtime/RigVM/Source/RigVMEditor/Public/Editor/RigVMLegacyEditor.h b/Engine/Plugins/Runtime/RigVM/Source/RigVMEditor/Public/Editor/RigVMLegacyEditor.h
--- a/Engine/Plugins/Runtime/RigVM/Source/RigVMEditor/Public/Editor/RigVMLegacyEditor.h
+++ b/Engine/Plugins/Runtime/RigVM/Source/RigVMEditor/Public/Editor/RigVMLegacyEditor.h
@@ -25,6 +25,7 @@ public:
 	FRigVMLegacyEditor();
 
 	virtual TSharedRef<IRigVMEditor> SharedRef() override { return StaticCastSharedRef<IRigVMEditor>(SharedThis(this)); }
+	virtual TSharedRef<const IRigVMEditor> SharedRef() const override { return StaticCastSharedRef<const IRigVMEditor>(SharedThis(this)); }
 
 protected:
 
diff --git a/Engine/Plugins/Runtime/RigVM/Source/RigVMEditor/Public/Editor/RigVMNewEditor.h b/Engine/Plugins/Runtime/RigVM/Source/RigVMEditor/Public/Editor/RigVMNewEditor.h
--- a/Engine/Plugins/Runtime/RigVM/Source/RigVMEditor/Public/Editor/RigVMNewEditor.h
+++ b/Engine/Plugins/Runtime/RigVM/Source/RigVMEditor/Public/Editor/RigVMNewEditor.h
@@ -22,6 +22,7 @@ public:
 	virtual void OnClose() override;
 
 	virtual TSharedRef<IRigVMEditor> SharedRef() override { return StaticCastSharedRef<IRigVMEditor>(SharedThis(this)); }
+	virtual TSharedRef<const IRigVMEditor> SharedRef() const override { return StaticCastSharedRef<const IRigVMEditor>(SharedThis(this)); }
 
 	virtual TSharedPtr<FAssetEditorToolkit> GetHostingApp() override { return AsShared(); }
 	virtual const TSharedPtr<FAssetEditorToolkit> GetHostingApp() const override { return ConstCastSharedRef<FAssetEditorToolkit>(AsShared()); }
diff --git a/Engine/Plugins/Runtime/RigVM/Source/RigVMEditor/Public/RigVMEditorModule.h b/Engine/Plugins/Runtime/RigVM/Source/RigVMEditor/Public/RigVMEditorModule.h
--- a/Engine/Plugins/Runtime/RigVM/Source/RigVMEditor/Public/RigVMEditorModule.h
+++ b/Engine/Plugins/Runtime/RigVM/Source/RigVMEditor/Public/RigVMEditorModule.h
@@ -27,6 +27,9 @@ class FBlueprintActionDatabaseRegistrar;
 
 DECLARE_LOG_CATEGORY_EXTERN(LogRigVMEditor, Log, All);
 
+extern RIGVMEDITOR_API TAutoConsoleVariable<bool> CVarRigVMUseNewEditor;
+extern RIGVMEDITOR_API TAutoConsoleVariable<bool> CVarRigVMUseDualEditor;
+
 /** Describes the reason for Refreshing the editor */
 namespace ERefreshRigVMEditorReason
 {
