ControlRig: delegate the construction and propagation of dependencies between elements

- IDependenciesProvider provides an interface for constructing and passing elements dependencies to functions that need to perform dependency tests.
- in the context of animation, the URigVM dependency map was built even if it was never used, resulting in poor playback performances. URigHierarchy now only requires these dependencies to be built when necessary to avoid this.
- it also enables other objects to build dependencies, such as modular rig.

#rnx

diff --git a/Engine/Plugins/Animation/ControlRig/Source/ControlRig/Private/Constraints/ControlRigTransformableHandle.cpp b/Engine/Plugins/Animation/ControlRig/Source/ControlRig/Private/Constraints/ControlRigTransformableHandle.cpp
--- a/Engine/Plugins/Animation/ControlRig/Source/ControlRig/Private/Constraints/ControlRigTransformableHandle.cpp
+++ b/Engine/Plugins/Animation/ControlRig/Source/ControlRig/Private/Constraints/ControlRigTransformableHandle.cpp
@@ -315,11 +315,9 @@ bool UTransformableControlHandle::HasDirectDependencyWith(const UTransformableHa
 	{
 		FRigControlElement* Parent = AllControls[IndexOfPossibleParent];
 
-		FString FailureReason;
-
 #if WITH_EDITOR
-		const URigHierarchy::TElementDependencyMap DependencyMap = Hierarchy->GetDependenciesForVM(ControlRig->GetVM());
-		const bool bIsParentedTo = Hierarchy->IsParentedTo(ControlElement, Parent, DependencyMap);
+		const FRigVMDependenciesProvider DependencyProvider(Hierarchy, ControlRig->GetVM());
+		const bool bIsParentedTo = Hierarchy->IsParentedTo(ControlElement, Parent, DependencyProvider);
 #else
 		const bool bIsParentedTo = Hierarchy->IsParentedTo(ControlElement, Parent);
 #endif
diff --git a/Engine/Plugins/Animation/ControlRig/Source/ControlRig/Private/ControlRig.cpp b/Engine/Plugins/Animation/ControlRig/Source/ControlRig/Private/ControlRig.cpp
--- a/Engine/Plugins/Animation/ControlRig/Source/ControlRig/Private/ControlRig.cpp
+++ b/Engine/Plugins/Animation/ControlRig/Source/ControlRig/Private/ControlRig.cpp
@@ -433,8 +433,8 @@ void UControlRig::Evaluate_AnyThread()
 				if (FRigBaseElement* NewParent = Hierarchy->Find(Pair.Value.NewParent))
 				{
 #if WITH_EDITOR
-					URigHierarchy::TElementDependencyMap DependencyMap = Hierarchy->GetDependenciesForVM(GetVM());
-					Hierarchy->SwitchToParent(Element, NewParent, Pair.Value.bInitial, Pair.Value.bAffectChildren, DependencyMap, nullptr);
+					FRigVMDependenciesProvider DependencyProvider(Hierarchy, GetVM());
+					Hierarchy->SwitchToParent(Element, NewParent, Pair.Value.bInitial, Pair.Value.bAffectChildren, DependencyProvider, nullptr);
 #else
 					Hierarchy->SwitchToParent(Element, NewParent, Pair.Value.bInitial, Pair.Value.bAffectChildren);
 #endif
@@ -2237,8 +2237,8 @@ void UControlRig::SwitchToParent(const FRigElementKey& InElementKey, const FRigE
 	else
 	{
 #if WITH_EDITOR
-		URigHierarchy::TElementDependencyMap Dependencies = DynamicHierarchy->GetDependenciesForVM(GetVM());
-		DynamicHierarchy->SwitchToParent(InElementKey, InNewParentKey, bInitial, bAffectChildren, Dependencies, nullptr);
+		FRigVMDependenciesProvider DependenciesProvider(DynamicHierarchy, GetVM());
+		DynamicHierarchy->SwitchToParent(InElementKey, InNewParentKey, bInitial, bAffectChildren, DependenciesProvider, nullptr);
 #else
 		DynamicHierarchy->SwitchToParent(InElementKey, InNewParentKey, bInitial, bAffectChildren);
 #endif
diff --git a/Engine/Plugins/Animation/ControlRig/Source/ControlRig/Private/Rigs/RigConnectionRules.cpp b/Engine/Plugins/Animation/ControlRig/Source/ControlRig/Private/Rigs/RigConnectionRules.cpp
--- a/Engine/Plugins/Animation/ControlRig/Source/ControlRig/Private/Rigs/RigConnectionRules.cpp
+++ b/Engine/Plugins/Animation/ControlRig/Source/ControlRig/Private/Rigs/RigConnectionRules.cpp
@@ -333,9 +333,9 @@ FRigElementResolveResult FRigChildOfPrimaryConnectionRule::Resolve(const FRigBas
 		}
 	}
 
-	static const URigHierarchy::TElementDependencyMap EmptyDependencyMap;
+	static const FNoDependenciesProvider NoDependencyProvider;
 	if(!InRuleInput.GetHierarchy()->IsParentedTo(
-		const_cast<FRigBaseElement*>(InTarget), const_cast<FRigTransformElement*>(PrimaryTarget), EmptyDependencyMap))
+		const_cast<FRigBaseElement*>(InTarget), const_cast<FRigTransformElement*>(PrimaryTarget), NoDependencyProvider))
 	{
 		Result.SetInvalidTarget(FText::FromString(FString::Printf(IsNotAChildOfFormat, *InTarget->GetKey().ToString())));
 	}
diff --git a/Engine/Plugins/Animation/ControlRig/Source/ControlRig/Private/Rigs/RigDependency.cpp b/Engine/Plugins/Animation/ControlRig/Source/ControlRig/Private/Rigs/RigDependency.cpp
new file mode 100644
--- /dev/null
+++ b/Engine/Plugins/Animation/ControlRig/Source/ControlRig/Private/Rigs/RigDependency.cpp
@@ -0,0 +1,52 @@
+// Copyright Epic Games, Inc. All Rights Reserved.
+
+#include "Rigs/RigDependency.h"
+
+#include "ControlRig.h"
+#include "Rigs/RigHierarchy.h"
+
+const TMap<int32, TArray<int32>>& FNoDependenciesProvider::GetDependencies() const
+{
+	static const TMap<int32, TArray<int32>> Empty;
+	return Empty;
+}
+
+FRigVMDependenciesProvider::FRigVMDependenciesProvider(const URigHierarchy* InHierarchy, URigVM* InVM, FName InEventName)
+	: WeakHierarchy(InHierarchy)
+	, WeakRigVM(InVM)
+	, EventName(InEventName)
+{}
+	
+const TMap<int32, TArray<int32>>& FRigVMDependenciesProvider::GetDependencies() const
+{
+	// NOTE: read URigHierarchy::GetDependenciesForVM about using this function with modular rigs.
+
+#if WITH_EDITOR
+	const URigHierarchy* Hierarchy = WeakHierarchy.Get();
+	if (ensure(Hierarchy))
+	{
+		uint32 NewHash = 0; 
+		auto LazyGetTopologyHash = [Hierarchy, &NewHash]()
+		{
+			if (NewHash == 0)
+			{
+				// NOTE: GetTopologyVersion() is used to track simple topology changes (GetTopologyHash() is probably more complete but also slower)
+				NewHash = Hierarchy->GetTopologyVersion();
+			}
+			return NewHash;
+		};
+		
+		if (CachedDependencies.IsEmpty() || TopologyHash != LazyGetTopologyHash())
+		{
+			URigVM* RigVM = WeakRigVM.Get();
+			if (ensure(RigVM))
+			{
+				CachedDependencies = Hierarchy->GetDependenciesForVM(RigVM, EventName);
+				TopologyHash = NewHash;
+			}
+		}
+	}
+#endif
+
+	return CachedDependencies;
+}
\ No newline at end of file
diff --git a/Engine/Plugins/Animation/ControlRig/Source/ControlRig/Private/Rigs/RigHierarchy.cpp b/Engine/Plugins/Animation/ControlRig/Source/ControlRig/Private/Rigs/RigHierarchy.cpp
--- a/Engine/Plugins/Animation/ControlRig/Source/ControlRig/Private/Rigs/RigHierarchy.cpp
+++ b/Engine/Plugins/Animation/ControlRig/Source/ControlRig/Private/Rigs/RigHierarchy.cpp
@@ -3723,7 +3723,7 @@ bool URigHierarchy::SetParentWeightArray(FRigBaseElement* InChild,  const TArray
 	return false;
 }
 
-bool URigHierarchy::CanSwitchToParent(FRigElementKey InChild, FRigElementKey InParent, const TElementDependencyMap& InDependencyMap, FString* OutFailureReason)
+bool URigHierarchy::CanSwitchToParent(FRigElementKey InChild, FRigElementKey InParent, const IDependenciesProvider& InDependencyProvider, FString* OutFailureReason)
 {
 	InParent = PreprocessParentElementKeyForSpaceSwitching(InChild, InParent);
 
@@ -3778,7 +3778,7 @@ bool URigHierarchy::CanSwitchToParent(FRigElementKey InChild, FRigElementKey InP
 		}
 	}
 
-	if(IsParentedTo(Parent, Child, InDependencyMap))
+	if(IsParentedTo(Parent, Child, InDependencyProvider))
 	{
 		if(OutFailureReason)
 		{
@@ -3790,14 +3790,14 @@ bool URigHierarchy::CanSwitchToParent(FRigElementKey InChild, FRigElementKey InP
 	return true;
 }
 
-bool URigHierarchy::SwitchToParent(FRigElementKey InChild, FRigElementKey InParent, bool bInitial, bool bAffectChildren, const TElementDependencyMap& InDependencyMap, FString* OutFailureReason)
+bool URigHierarchy::SwitchToParent(FRigElementKey InChild, FRigElementKey InParent, bool bInitial, bool bAffectChildren, const IDependenciesProvider& InDependencyProvider, FString* OutFailureReason)
 {
 	InParent = PreprocessParentElementKeyForSpaceSwitching(InChild, InParent);
-	return SwitchToParent(Find(InChild), Find(InParent), bInitial, bAffectChildren, InDependencyMap, OutFailureReason);
+	return SwitchToParent(Find(InChild), Find(InParent), bInitial, bAffectChildren, InDependencyProvider, OutFailureReason);
 }
 
 bool URigHierarchy::SwitchToParent(FRigBaseElement* InChild, FRigBaseElement* InParent, bool bInitial,
-	bool bAffectChildren, const TElementDependencyMap& InDependencyMap, FString* OutFailureReason)
+	bool bAffectChildren, const IDependenciesProvider& InDependencyProvider, FString* OutFailureReason)
 {
 	FRigHierarchyEnableControllerBracket EnableController(this, true);
 
@@ -3815,28 +3815,27 @@ bool URigHierarchy::SwitchToParent(FRigBaseElement* InChild, FRigBaseElement* In
 	}
 
 	// rely on the VM's dependency map if there's currently an available context.
-	const TElementDependencyMap* DependencyMapPtr = &InDependencyMap;
+	const IDependenciesProvider* DependencyProviderPtr = &InDependencyProvider;
 
 #if WITH_EDITOR
 	// Keep this in function scope, since we might be taking a pointer to it.
-	TElementDependencyMap DependencyMapFromVM;
+	FRigVMDependenciesProvider DependencyProviderFromVM;
 	{
 		UE::TScopeLock Lock(ExecuteContextLock);
-		if(ExecuteContext != nullptr && DependencyMapPtr->IsEmpty())
+		if(ExecuteContext != nullptr && DependencyProviderPtr->GetDependencies().IsEmpty())
 		{
 			if(ExecuteContext->VM)
 			{
-				DependencyMapFromVM = GetDependenciesForVM(ExecuteContext->VM);
-				DependencyMapPtr = &DependencyMapFromVM;
+				DependencyProviderFromVM = FRigVMDependenciesProvider(this, ExecuteContext->VM);
+				DependencyProviderPtr = &DependencyProviderFromVM;
 			}
 		}
 	}
-	
 #endif
 	
 	if(InChild && InParent)
 	{
-		if(!CanSwitchToParent(InChild->GetKey(), InParent->GetKey(), *DependencyMapPtr, OutFailureReason))
+		if(!CanSwitchToParent(InChild->GetKey(), InParent->GetKey(), *DependencyProviderPtr, OutFailureReason))
 		{
 			return false;
 		}
@@ -6141,17 +6140,17 @@ FRigHierarchyKey URigHierarchy::GetPreviousHierarchyParent(const FRigHierarchyKe
 	return FRigHierarchyKey();
 }
 
-bool URigHierarchy::IsParentedTo(FRigBaseElement* InChild, FRigBaseElement* InParent, const TElementDependencyMap& InDependencyMap) const
+bool URigHierarchy::IsParentedTo(FRigBaseElement* InChild, FRigBaseElement* InParent, const IDependenciesProvider& InDependencyProvider) const
 {
 	ElementDependencyVisited.Reset();
-	if(!InDependencyMap.IsEmpty())
+	if(!InDependencyProvider.GetDependencies().IsEmpty())
 	{
 		ElementDependencyVisited.SetNumZeroed(Elements.Num());
 	}
-	return IsDependentOn(InChild, InParent, InDependencyMap, true);
+	return IsDependentOn(InChild, InParent, InDependencyProvider, true);
 }
 
-bool URigHierarchy::IsDependentOn(FRigBaseElement* InDependent, FRigBaseElement* InDependency, const TElementDependencyMap& InDependencyMap, bool bIsOnActualTopology) const
+bool URigHierarchy::IsDependentOn(FRigBaseElement* InDependent, FRigBaseElement* InDependency, const IDependenciesProvider& InDependencyProvider, bool bIsOnActualTopology) const
 {
 	if((InDependent == nullptr) || (InDependency == nullptr))
 	{
@@ -6208,7 +6207,7 @@ bool URigHierarchy::IsDependentOn(FRigBaseElement* InDependent, FRigBaseElement*
 	// collect all possible parents of the dependent
 	if(const FRigSingleParentElement* SingleParentElement = Cast<FRigSingleParentElement>(InDependent))
 	{
-		if(IsDependentOn(SingleParentElement->ParentElement, InDependency, InDependencyMap, true))
+		if(IsDependentOn(SingleParentElement->ParentElement, InDependency, InDependencyProvider, true))
 		{
 			// we'll only update the caches if we are following edges on the actual topology
 			if(bIsOnActualTopology)
@@ -6222,7 +6221,7 @@ bool URigHierarchy::IsDependentOn(FRigBaseElement* InDependent, FRigBaseElement*
 	{
 		for(const FRigElementParentConstraint& ParentConstraint : MultiParentElement->ParentConstraints)
 		{
-			if(IsDependentOn(ParentConstraint.ParentElement, InDependency, InDependencyMap, true))
+			if(IsDependentOn(ParentConstraint.ParentElement, InDependency, InDependencyProvider, true))
 			{
 				// we'll only update the caches if we are following edges on the actual topology
 				if(bIsOnActualTopology)
@@ -6235,13 +6234,14 @@ bool URigHierarchy::IsDependentOn(FRigBaseElement* InDependent, FRigBaseElement*
 	}
 
 	// check the optional dependency map
-	if(const TArray<int32>* DependentIndicesPtr = InDependencyMap.Find(InDependent->GetIndex()))
+	const TElementDependencyMap& DependencyMap = InDependencyProvider.GetDependencies();
+	if(const TArray<int32>* DependentIndicesPtr = DependencyMap.Find(InDependent->GetIndex()))
 	{
 		const TArray<int32>& DependentIndices = *DependentIndicesPtr;
 		for(const int32 DependentIndex : DependentIndices)
 		{
 			ensure(Elements.IsValidIndex(DependentIndex));
-			if(IsDependentOn(Elements[DependentIndex], InDependency, InDependencyMap, false))
+			if(IsDependentOn(Elements[DependentIndex], InDependency, InDependencyProvider, false))
 			{
 				// we'll only update the caches if we are following edges on the actual topology
 				if(bIsOnActualTopology)
@@ -9741,4 +9741,3 @@ FArchive& FRigHierarchyMemoryReader::operator<<(FText& Value)
 	Value = FText::FromString(ValueString);
 	return *this;
 }
-;
diff --git a/Engine/Plugins/Animation/ControlRig/Source/ControlRig/Private/Units/Execution/RigUnit_DynamicHierarchy.cpp b/Engine/Plugins/Animation/ControlRig/Source/ControlRig/Private/Units/Execution/RigUnit_DynamicHierarchy.cpp
--- a/Engine/Plugins/Animation/ControlRig/Source/ControlRig/Private/Units/Execution/RigUnit_DynamicHierarchy.cpp
+++ b/Engine/Plugins/Animation/ControlRig/Source/ControlRig/Private/Units/Execution/RigUnit_DynamicHierarchy.cpp
@@ -215,8 +215,8 @@ FRigUnit_SwitchParent_Execute()
 		default:
 		{
 			FString FailureReason;
-			static const URigHierarchy::TElementDependencyMap EmptyDependencyMap;
-			if(!ExecuteContext.Hierarchy->SwitchToParent(ChildElement, ParentElement, false, true, EmptyDependencyMap, &FailureReason))
+			static const FNoDependenciesProvider NoDependencyProvider;
+			if(!ExecuteContext.Hierarchy->SwitchToParent(ChildElement, ParentElement, false, true, NoDependencyProvider, &FailureReason))
 			{
 				if(!FailureReason.IsEmpty())
 				{
@@ -228,7 +228,7 @@ FRigUnit_SwitchParent_Execute()
 			// during construction event also change the initial weights
 			if(ExecuteContext.IsRunningConstructionEvent())
 			{
-				if(!ExecuteContext.Hierarchy->SwitchToParent(ChildElement, ParentElement, true, true, EmptyDependencyMap, &FailureReason))
+				if(!ExecuteContext.Hierarchy->SwitchToParent(ChildElement, ParentElement, true, true, NoDependencyProvider, &FailureReason))
 				{
 					if(!FailureReason.IsEmpty())
 					{
diff --git a/Engine/Plugins/Animation/ControlRig/Source/ControlRig/Public/Rigs/RigDependency.h b/Engine/Plugins/Animation/ControlRig/Source/ControlRig/Public/Rigs/RigDependency.h
new file mode 100644
--- /dev/null
+++ b/Engine/Plugins/Animation/ControlRig/Source/ControlRig/Public/Rigs/RigDependency.h
@@ -0,0 +1,56 @@
+// Copyright Epic Games, Inc. All Rights Reserved.
+
+#pragma once
+
+#include "Containers/Array.h"
+#include "Containers/Map.h"
+#include "UObject/NameTypes.h"
+#include "UObject/WeakObjectPtr.h"
+#include "UObject/WeakObjectPtrTemplates.h"
+
+class URigHierarchy;
+class URigVM;
+
+/**
+ * IDependenciesProvider provides an interface for constructing and passing elements dependencies to functions that need to perform dependency tests.
+ * This is currently mainly used to ensure that certain rig elements are not dependent on each other (parent switching, constraints ordering, etc.) 
+ */
+
+struct IDependenciesProvider
+{
+	virtual ~IDependenciesProvider() = default;
+	virtual const TMap<int32, TArray<int32>>& GetDependencies() const = 0;
+};
+
+/**
+ * FNoDependenciesProvider is the default dependency provider (no dependency at all) 
+ */
+
+struct CONTROLRIG_API FNoDependenciesProvider : public IDependenciesProvider
+{
+	virtual ~FNoDependenciesProvider() override = default;
+	virtual const TMap<int32, TArray<int32>>& GetDependencies() const override;
+};
+
+/**
+ * FRigVMDependenciesProvider builds a dependency map based on the provided RigVM's instructions
+ */
+
+struct CONTROLRIG_API FRigVMDependenciesProvider : public IDependenciesProvider
+{
+	FRigVMDependenciesProvider() = default;
+	FRigVMDependenciesProvider(const URigHierarchy* InHierarchy, URigVM* InVM, FName InEventName = NAME_None);
+	virtual ~FRigVMDependenciesProvider() override = default;
+	
+	virtual const TMap<int32, TArray<int32>>& GetDependencies() const override;
+	
+private:
+		
+	TWeakObjectPtr<const URigHierarchy> WeakHierarchy = nullptr;
+	TWeakObjectPtr<URigVM> WeakRigVM = nullptr;
+	FName EventName = NAME_None;
+	
+	mutable TMap<int32, TArray<int32>> CachedDependencies;
+	mutable uint32 TopologyHash = 0;
+};
+
diff --git a/Engine/Plugins/Animation/ControlRig/Source/ControlRig/Public/Rigs/RigHierarchy.h b/Engine/Plugins/Animation/ControlRig/Source/ControlRig/Public/Rigs/RigHierarchy.h
--- a/Engine/Plugins/Animation/ControlRig/Source/ControlRig/Public/Rigs/RigHierarchy.h
+++ b/Engine/Plugins/Animation/ControlRig/Source/ControlRig/Public/Rigs/RigHierarchy.h
@@ -16,6 +16,7 @@
 #if WITH_EDITOR
 #include "RigVMPythonUtils.h"
 #endif
+#include "RigDependency.h"
 #include "Containers/Queue.h"
 #include "Misc/TransactionallySafeCriticalSection.h"
 #include "Serialization/MemoryReader.h"
@@ -3665,11 +3666,15 @@ public:
 	* Determines if the element can be switched to a provided parent
 	* @param InChild The key of the multi parented element
 	* @param InParent The key of the parent to look up the weight for
-    * @param InDependencyMap An additional map of dependencies to respect
+    * @param InDependencyProvider An additional object providing dependencies to respect
 	* @param OutFailureReason An optional pointer to retrieve the reason for failure
 	* @return Returns true if changing the weight was successful
 	*/
-	bool CanSwitchToParent(FRigElementKey InChild, FRigElementKey InParent, const TElementDependencyMap& InDependencyMap = TElementDependencyMap(), FString* OutFailureReason = nullptr);
+	bool CanSwitchToParent(
+		FRigElementKey InChild,
+		FRigElementKey InParent,
+		const IDependenciesProvider& InDependencyProvider = FNoDependenciesProvider(),
+		FString* OutFailureReason = nullptr);
 
 	/**
 	 * Switches a multi parent element to a single parent.
@@ -3684,9 +3689,15 @@ public:
 	UFUNCTION(BlueprintCallable, Category = URigHierarchy)
 	bool SwitchToParent(FRigElementKey InChild, FRigElementKey InParent, bool bInitial = false, bool bAffectChildren = true)
 	{
-		return SwitchToParent(InChild, InParent, bInitial, bAffectChildren, TElementDependencyMap(), nullptr);
+		return SwitchToParent(InChild, InParent, bInitial, bAffectChildren, FNoDependenciesProvider(), nullptr);
 	}
-	bool SwitchToParent(FRigElementKey InChild, FRigElementKey InParent, bool bInitial, bool bAffectChildren, const TElementDependencyMap& InDependencyMap, FString* OutFailureReason);
+
+	bool SwitchToParent(FRigElementKey InChild,
+		FRigElementKey InParent,
+		bool bInitial,
+		bool bAffectChildren,
+		const IDependenciesProvider& InDependencyProvider,
+		FString* OutFailureReason);
 
 	/**
 	 * Switches a multi parent element to a single parent.
@@ -3696,11 +3707,16 @@ public:
 	 * @param InParent The parent to look up the weight for
 	 * @param bInitial If true the initial weights will be used
 	 * @param bAffectChildren If set to false children will not move (maintain global).
-     * @param InDependencyMap An additional map of dependencies to respect
+     * @param InDependencyProvider An additional object providing dependencies to respect
 	 * @param OutFailureReason An optional pointer to retrieve the reason for failure
 	 * @return Returns true if changing the weight was successful
 	 */
-	bool SwitchToParent(FRigBaseElement* InChild, FRigBaseElement* InParent, bool bInitial = false, bool bAffectChildren = true, const TElementDependencyMap& InDependencyMap = TElementDependencyMap(), FString* OutFailureReason = nullptr);
+	bool SwitchToParent(FRigBaseElement* InChild,
+		FRigBaseElement* InParent,
+		bool bInitial = false,
+		bool bAffectChildren = true,
+		const IDependenciesProvider& InDependencyProvider = FNoDependenciesProvider(),
+		FString* OutFailureReason = nullptr);
 
 	/**
 	 * Switches a multi parent element to a single parent.
@@ -3782,14 +3798,14 @@ public:
 	 * Returns true if an element is parented to another element
 	 * @param InChildIndex The index of the child element to check for a parent
 	 * @param InParentIndex The index of the parent element to check for
-	 * @param InDependencyMap An additional map of dependencies to respect
+	 * @param InDependencyProvider An additional object providing dependencies to respect
 	 * @return True if the given parent and given child have a parent-child relationship
 	 */
-    bool IsParentedTo(int32 InChildIndex, int32 InParentIndex, const TElementDependencyMap& InDependencyMap = TElementDependencyMap()) const
+    bool IsParentedTo(int32 InChildIndex, int32 InParentIndex, const IDependenciesProvider& InDependencyProvider = FNoDependenciesProvider()) const
 	{
 		if(Elements.IsValidIndex(InChildIndex) && Elements.IsValidIndex(InParentIndex))
 		{
-			return IsParentedTo(Elements[InChildIndex], Elements[InParentIndex], InDependencyMap);
+			return IsParentedTo(Elements[InChildIndex], Elements[InParentIndex], InDependencyProvider);
 		}
 		return false;
 	}
@@ -4503,21 +4519,26 @@ public:
 	 * Returns true if an element is parented to another element
 	 * @param InChild The child element to check for a parent
 	 * @param InParent The parent element to check for
-	 * @param InDependencyMap An additional map of dependencies to respect
+	 * @param InDependencyProvider An additional object providing dependencies to respect
 	 * @return True if the given parent and given child have a parent-child relationship
 	 */
-	bool IsParentedTo(FRigBaseElement* InChild, FRigBaseElement* InParent, const TElementDependencyMap& InDependencyMap = TElementDependencyMap()) const;
+	bool IsParentedTo(FRigBaseElement* InChild,
+		FRigBaseElement* InParent,
+		const IDependenciesProvider& InDependencyProvider = FNoDependenciesProvider()) const;
 
 private:
 	/**
 	 * Returns true if an element is affected to another element
 	 * @param InDependent The dependent element to check for a dependency
 	 * @param InDependency The dependency element to check for
-	 * @param InDependencyMap An additional map of dependencies to respect
+	 * @param InDependencyProvider An additional object providing dependencies to respect
 	 * @param bIsOnActualTopology Indicates that the passed dependent and dependency are expected to be on the current topology (if false they are provided with the dependency map)
 	 * @return True if the given dependent is affected by the given dependency 
 	 */
-	bool IsDependentOn(FRigBaseElement* InDependent, FRigBaseElement* InDependency, const TElementDependencyMap& InDependencyMap = TElementDependencyMap(), bool bIsOnActualTopology = true) const;
+	bool IsDependentOn(FRigBaseElement* InDependent,
+		FRigBaseElement* InDependency,
+		const IDependenciesProvider& InDependencyProvider = FNoDependenciesProvider(),
+		bool bIsOnActualTopology = true) const;
 
 public:
 	/**
diff --git a/Engine/Plugins/Animation/ControlRig/Source/ControlRigEditor/Private/EditMode/ControlRigEditMode.cpp b/Engine/Plugins/Animation/ControlRig/Source/ControlRigEditor/Private/EditMode/ControlRigEditMode.cpp
--- a/Engine/Plugins/Animation/ControlRig/Source/ControlRigEditor/Private/EditMode/ControlRigEditMode.cpp
+++ b/Engine/Plugins/Animation/ControlRig/Source/ControlRigEditor/Private/EditMode/ControlRigEditMode.cpp
@@ -3302,10 +3302,11 @@ void FControlRigEditMode::OpenSpacePickerWidget()
 			}
 			else
 			{
+				constexpr bool bInitial = false, bAffectChildren = true;
 				const FTransform Transform = InHierarchy->GetGlobalTransform(InControlKey);
-				URigHierarchy::TElementDependencyMap Dependencies = InHierarchy->GetDependenciesForVM(RuntimeRig->GetVM());
+				FRigVMDependenciesProvider DependenciesProvider(InHierarchy, RuntimeRig->GetVM());
 				FString OutFailureReason;
-				if (InHierarchy->SwitchToParent(InControlKey, InSpaceKey, false, true, Dependencies, &OutFailureReason))
+				if (InHierarchy->SwitchToParent(InControlKey, InSpaceKey, bInitial, bAffectChildren, DependenciesProvider, &OutFailureReason))
 				{
 					InHierarchy->SetGlobalTransform(InControlKey, Transform);
 				}
@@ -6461,7 +6462,7 @@ FInteractionDependencyCache& FControlRigEditMode::GetInteractionDependencies(UCo
 		NewDependency.WeakHierarchy = Hierarchy;
 
 		// NOTE: this is not enough for modular rigs since there are several VMs
-		const URigHierarchy::TElementDependencyMap HierarchyDependencyMap = Hierarchy->GetDependenciesForVM(InControlRig->GetVM());
+		const FRigVMDependenciesProvider DependencyProvider(Hierarchy, InControlRig->GetVM());
 		
 		for (int32 Index = 0; Index < SelectedControls.Num(); ++Index)
 		{
@@ -6469,12 +6470,12 @@ FInteractionDependencyCache& FControlRigEditMode::GetInteractionDependencies(UCo
 			for (int32 NextIndex = Index+1; NextIndex < SelectedControls.Num(); ++NextIndex)
 			{
 				FRigControlElement* OtherControl = SelectedControls[NextIndex];
-				if (Hierarchy->IsParentedTo(Control, OtherControl, HierarchyDependencyMap))
+				if (Hierarchy->IsParentedTo(Control, OtherControl, DependencyProvider))
 				{
 					NewDependency.Parents.Add(OtherControl->GetKey());
 					NewDependency.Children.Add(Control->GetKey());
 				}
-				else if (Hierarchy->IsParentedTo(OtherControl, Control, HierarchyDependencyMap))
+				else if (Hierarchy->IsParentedTo(OtherControl, Control, DependencyProvider))
 				{
 					NewDependency.Parents.Add(Control->GetKey());
 					NewDependency.Children.Add(OtherControl->GetKey());
diff --git a/Engine/Plugins/Animation/ControlRig/Source/ControlRigEditor/Private/EditMode/SControlRigEditModeTools.cpp b/Engine/Plugins/Animation/ControlRig/Source/ControlRigEditor/Private/EditMode/SControlRigEditModeTools.cpp
--- a/Engine/Plugins/Animation/ControlRig/Source/ControlRigEditor/Private/EditMode/SControlRigEditModeTools.cpp
+++ b/Engine/Plugins/Animation/ControlRig/Source/ControlRigEditor/Private/EditMode/SControlRigEditModeTools.cpp
@@ -924,8 +924,8 @@ void SControlRigEditModeTools::HandleActiveSpaceChanged(URigHierarchy* InHierarc
 			if (ControlRig.IsValid() && ControlRig->GetHierarchy() == InHierarchy)
 			{
 				FString FailureReason;
-				URigHierarchy::TElementDependencyMap DependencyMap = InHierarchy->GetDependenciesForVM(ControlRig->GetVM());
-				if(!InHierarchy->CanSwitchToParent(InControlKey, InSpaceKey, DependencyMap, &FailureReason))
+				FRigVMDependenciesProvider DependencyProvider(InHierarchy, ControlRig->GetVM());
+				if(!InHierarchy->CanSwitchToParent(InControlKey, InSpaceKey, DependencyProvider, &FailureReason))
 				{
 					// notification
 					FNotificationInfo Info(FText::FromString(FailureReason));
diff --git a/Engine/Plugins/Animation/ControlRig/Source/ControlRigEditor/Private/EditMode/SControlRigSpacePicker.cpp b/Engine/Plugins/Animation/ControlRig/Source/ControlRigEditor/Private/EditMode/SControlRigSpacePicker.cpp
--- a/Engine/Plugins/Animation/ControlRig/Source/ControlRigEditor/Private/EditMode/SControlRigSpacePicker.cpp
+++ b/Engine/Plugins/Animation/ControlRig/Source/ControlRigEditor/Private/EditMode/SControlRigSpacePicker.cpp
@@ -161,8 +161,8 @@ void SControlRigSpacePicker::HandleActiveSpaceChanged(URigHierarchy* InHierarchy
 		if (UControlRig* ControlRig = GetControlRig())
 		{
 			FString FailureReason;
-			URigHierarchy::TElementDependencyMap DependencyMap = InHierarchy->GetDependenciesForVM(ControlRig->GetVM());
-			if (!InHierarchy->CanSwitchToParent(InControlKey, InSpaceKey, DependencyMap, &FailureReason))
+			FRigVMDependenciesProvider DependencyProvider(InHierarchy, ControlRig->GetVM());
+			if (!InHierarchy->CanSwitchToParent(InControlKey, InSpaceKey, DependencyProvider, &FailureReason))
 			{
 				// notification
 				FNotificationInfo Info(FText::FromString(FailureReason));
@@ -332,12 +332,8 @@ FReply SControlRigSpacePicker::OnBakeControlsToNewSpaceButtonClicked()
 	FSpaceChannelAndSection SpaceChannelAndSection = FControlRigSpaceChannelHelpers::FindSpaceChannelAndSectionForControl(ControlRig, ControlKeys[0].Name, Sequencer, true /*bCreateIfNeeded*/);
 	if (SpaceChannelAndSection.SpaceChannel != nullptr)
 	{
-		const FFrameRate TickResolution = Sequencer->GetFocusedTickResolution();
-		const FFrameTime FrameTime = Sequencer->GetLocalTime().ConvertTo(TickResolution);
-		FFrameNumber CurrentTime = FrameTime.GetFrame();
-		FMovieSceneControlRigSpaceBaseKey Value;
 		using namespace UE::MovieScene;
-		URigHierarchy* RigHierarchy = SpacePickerWidget->GetHierarchy();
+
 		Settings.TargetSpace = URigHierarchy::GetDefaultParentKey();
 
 		TRange<FFrameNumber> Range = Sequencer->GetFocusedMovieSceneSequence()->GetMovieScene()->GetPlaybackRange();
@@ -352,7 +348,7 @@ FReply SControlRigSpacePicker::OnBakeControlsToNewSpaceButtonClicked()
 			.Controls(ControlKeys) // use the cached controls here since the selection is not recovered until next tick.
 			.Sequencer(Sequencer)
 			.GetControlCustomization(this, &SControlRigSpacePicker::HandleGetControlElementCustomization)
-			.OnBake_Lambda([Sequencer, ControlRig, TickResolution](URigHierarchy* InHierarchy, TArray<FRigElementKey> InControls, FRigSpacePickerBakeSettings InSettings)
+			.OnBake_Lambda([Sequencer, ControlRig](URigHierarchy* InHierarchy, TArray<FRigElementKey> InControls, FRigSpacePickerBakeSettings InSettings)
 		{
 	
 			FScopedTransaction Transaction(LOCTEXT("BakeControlToSpace", "Bake Control In Space"));
diff --git a/Engine/Plugins/Animation/ControlRig/Source/ControlRigEditor/Private/Editor/SRigSpacePickerWidget.cpp b/Engine/Plugins/Animation/ControlRig/Source/ControlRigEditor/Private/Editor/SRigSpacePickerWidget.cpp
--- a/Engine/Plugins/Animation/ControlRig/Source/ControlRigEditor/Private/Editor/SRigSpacePickerWidget.cpp
+++ b/Engine/Plugins/Animation/ControlRig/Source/ControlRigEditor/Private/Editor/SRigSpacePickerWidget.cpp
@@ -752,22 +752,22 @@ FReply SRigSpacePickerWidget::HandleAddElementClicked()
 					URigHierarchy* StrongHierarchy = Hierarchy.Get();
 					for(const FRigElementKey& ControlKey : ControlKeys)
 					{
-						URigHierarchy::TElementDependencyMap DependencyMap;
+						FRigVMDependenciesProvider DependencyProvider;
 						FString FailureReason;
 
 						if(UControlRig* ControlRig = StrongHierarchy->GetTypedOuter<UControlRig>())
 						{
-							DependencyMap = StrongHierarchy->GetDependenciesForVM(ControlRig->GetVM()); 
+							DependencyProvider = FRigVMDependenciesProvider(StrongHierarchy, ControlRig->GetVM()); 
 						}
 						else if(UControlRigBlueprint* RigBlueprint = StrongHierarchy->GetTypedOuter<UControlRigBlueprint>())
 						{
 							if(UControlRig* CDO = Cast<UControlRig>(RigBlueprint->GetRigVMBlueprintGeneratedClass()->GetDefaultObject()))
 							{
-								DependencyMap = StrongHierarchy->GetDependenciesForVM(CDO->GetVM()); 
+								DependencyProvider = FRigVMDependenciesProvider(StrongHierarchy, CDO->GetVM()); 
 							}
 						}
 						
-						if(!StrongHierarchy->CanSwitchToParent(ControlKey, Key, DependencyMap, &FailureReason))
+						if(!StrongHierarchy->CanSwitchToParent(ControlKey, Key, DependencyProvider, &FailureReason))
 						{
 							// notification
 							FNotificationInfo Info(FText::FromString(FailureReason));
