Control Rig: Fix crash in ECS keying

#rnx

diff --git a/Engine/Plugins/Animation/ControlRig/Source/ControlRigEditor/Private/Sequencer/ControlRigParameterTrackEditor.cpp b/Engine/Plugins/Animation/ControlRig/Source/ControlRigEditor/Private/Sequencer/ControlRigParameterTrackEditor.cpp
--- a/Engine/Plugins/Animation/ControlRig/Source/ControlRigEditor/Private/Sequencer/ControlRigParameterTrackEditor.cpp
+++ b/Engine/Plugins/Animation/ControlRig/Source/ControlRigEditor/Private/Sequencer/ControlRigParameterTrackEditor.cpp
@@ -4933,7 +4933,7 @@ void FControlRigParameterTrackEditor::ProcessKeyOperation(UObject* ObjectToKey,
 				FMovieSceneIntegerChannel* Int = static_cast<FMovieSceneIntegerChannel*>(Channel);
 				Int->GetData().UpdateOrAddKey(KeyTime, *DesiredValue.Cast<int32>());
 			}
-			else 
+			else if (ChannelHandle.GetChannelTypeName() == FMovieSceneFloatChannel::StaticStruct()->GetFName())
 			{
 				const int32 ControlChannelIndex = ChannelMetaData.GetChannelIndex();
 
