Anim In Engine: When a control is added to a control rig, we now mask it out automatically if the section has a mask, othewise it would get added to anim layers that are already set up. Not a common use case but on shows with rigs getting worked on during production it is happening often

diff --git a/Engine/Plugins/Animation/ControlRig/Source/ControlRig/Private/Sequencer/MovieSceneControlRigParameterSection.cpp b/Engine/Plugins/Animation/ControlRig/Source/ControlRig/Private/Sequencer/MovieSceneControlRigParameterSection.cpp
--- a/Engine/Plugins/Animation/ControlRig/Source/ControlRig/Private/Sequencer/MovieSceneControlRigParameterSection.cpp
+++ b/Engine/Plugins/Animation/ControlRig/Source/ControlRig/Private/Sequencer/MovieSceneControlRigParameterSection.cpp
@@ -1235,6 +1235,14 @@ FName UMovieSceneControlRigParameterSection::FindControlNameFromSpaceChannel(con
 	return NAME_None;
 }
 
+void UMovieSceneControlRigParameterSection::MaskOutIfThereAreMaskedControls(const FName& InControlName)
+{
+	if (ControlNameMask.IsEmpty() == false)
+	{
+		ControlNameMask.Add(InControlName);
+	}
+}
+
 void UMovieSceneControlRigParameterSection::AddScalarParameter(FName InParameterName, TOptional<float> DefaultValue, bool bReconstructChannel)
 {
 	FMovieSceneFloatChannel* ExistingChannel = nullptr;
@@ -1250,10 +1258,14 @@ void UMovieSceneControlRigParameterSection::AddScalarParameter(FName InParameter
 		{
 			ExistingChannel->SetDefault(0.0f);
 		}
+
+		MaskOutIfThereAreMaskedControls(InParameterName);
+		
 		if (bReconstructChannel)
 		{
 			ReconstructChannelProxy();
 		}
+
 	}
 }
 
@@ -1273,6 +1285,9 @@ void UMovieSceneControlRigParameterSection::AddBoolParameter(FName InParameterNa
 		{
 			ExistingChannel->SetDefault(false);
 		}
+
+		MaskOutIfThereAreMaskedControls(InParameterName);
+
 		if (bReconstructChannel)
 		{
 			ReconstructChannelProxy();
@@ -1294,6 +1309,9 @@ void UMovieSceneControlRigParameterSection::AddEnumParameter(FName InParameterNa
 		{
 			ExistingChannel->SetDefault(false);
 		}
+
+		MaskOutIfThereAreMaskedControls(InParameterName);
+
 		ExistingChannel->SetEnum(Enum);
 		if (bReconstructChannel)
 		{
@@ -1317,6 +1335,9 @@ void UMovieSceneControlRigParameterSection::AddIntegerParameter(FName InParamete
 		{
 			ExistingChannel->SetDefault(false);
 		}
+
+		MaskOutIfThereAreMaskedControls(InParameterName);
+
 		if (bReconstructChannel)
 		{
 			ReconstructChannelProxy();
@@ -1342,6 +1363,9 @@ void UMovieSceneControlRigParameterSection::AddVector2DParameter(FName InParamet
 			ExistingCurves->XCurve.SetDefault(0.0f);
 			ExistingCurves->YCurve.SetDefault(0.0f);
 		}
+
+		MaskOutIfThereAreMaskedControls(InParameterName);
+
 		if (bReconstructChannel)
 		{
 			ReconstructChannelProxy();
@@ -1370,6 +1394,9 @@ void UMovieSceneControlRigParameterSection::AddVectorParameter(FName InParameter
 			ExistingCurves->YCurve.SetDefault(0.0f);
 			ExistingCurves->ZCurve.SetDefault(0.0f);
 		}
+
+		MaskOutIfThereAreMaskedControls(InParameterName);
+
 		if (bReconstructChannel)
 		{
 			ReconstructChannelProxy();
@@ -1399,6 +1426,9 @@ void UMovieSceneControlRigParameterSection::AddColorParameter(FName InParameterN
 			ExistingCurves->BlueCurve.SetDefault(0.0f);
 			ExistingCurves->AlphaCurve.SetDefault(0.0f);
 		}
+
+		MaskOutIfThereAreMaskedControls(InParameterName);
+
 		if (bReconstructChannel)
 		{
 			ReconstructChannelProxy();
@@ -1447,6 +1477,9 @@ void UMovieSceneControlRigParameterSection::AddTransformParameter(FName InParame
 			ExistingCurves->Scale[1].SetDefault(0.0f);
 			ExistingCurves->Scale[2].SetDefault(0.0f);
 		}
+
+		MaskOutIfThereAreMaskedControls(InParameterName);
+
 		if(bReconstructChannel)
 		{
 			ReconstructChannelProxy();
diff --git a/Engine/Plugins/Animation/ControlRig/Source/ControlRig/Public/Sequencer/MovieSceneControlRigParameterSection.h b/Engine/Plugins/Animation/ControlRig/Source/ControlRig/Public/Sequencer/MovieSceneControlRigParameterSection.h
--- a/Engine/Plugins/Animation/ControlRig/Source/ControlRig/Public/Sequencer/MovieSceneControlRigParameterSection.h
+++ b/Engine/Plugins/Animation/ControlRig/Source/ControlRig/Public/Sequencer/MovieSceneControlRigParameterSection.h
@@ -631,6 +631,7 @@ public:
 protected:
 
 	void ConvertMaskArrayToNameSet();
+	void MaskOutIfThereAreMaskedControls(const FName& InControlName);
 
 	//~ UMovieSceneSection interface
 	virtual void Serialize(FArchive& Ar) override;
