Control Rig Edit Mode: update the keyframer state when sequencer's playback status changes while manipulating

#rnx

diff --git a/Engine/Plugins/Animation/ControlRig/Source/ControlRigEditor/Private/EditMode/ControlRigEditMode.cpp b/Engine/Plugins/Animation/ControlRig/Source/ControlRigEditor/Private/EditMode/ControlRigEditMode.cpp
--- a/Engine/Plugins/Animation/ControlRig/Source/ControlRigEditor/Private/EditMode/ControlRigEditMode.cpp
+++ b/Engine/Plugins/Animation/ControlRig/Source/ControlRigEditor/Private/EditMode/ControlRigEditMode.cpp
@@ -3992,7 +3992,16 @@ void FControlRigEditMode::UnsetSequencerDelegates() const
 void FControlRigEditMode::UpdateSequencerStatus()
 {
 	const TSharedPtr<ISequencer> Sequencer = WeakSequencer.IsValid() ? WeakSequencer.Pin() : nullptr;
-	bSequencerPlaying = Sequencer.IsValid() && Sequencer->GetPlaybackStatus() == EMovieScenePlayerStatus::Playing;
+    const bool bIsSequencerPlaying = Sequencer.IsValid() && Sequencer->GetPlaybackStatus() == EMovieScenePlayerStatus::Playing;
+
+    if (bIsSequencerPlaying != bSequencerPlaying)
+    {
+    	bSequencerPlaying = bIsSequencerPlaying;
+
+    	// update keyframer state
+    	const bool bDeferAutokeyOnMouseRelease = !bSequencerPlaying && IsInLevelEditor();
+    	Keyframer.Enable(bDeferAutokeyOnMouseRelease);
+    }
 }
 
 bool FControlRigEditMode::CheckMovieSceneSig()
