ControlRig :Added an option to exclude virtual bones to Import Skeleton rig unit

diff --git a/Engine/Plugins/Animation/ControlRig/Source/ControlRig/Private/Units/Execution/RigUnit_DynamicHierarchy.cpp b/Engine/Plugins/Animation/ControlRig/Source/ControlRig/Private/Units/Execution/RigUnit_DynamicHierarchy.cpp
--- a/Engine/Plugins/Animation/ControlRig/Source/ControlRig/Private/Units/Execution/RigUnit_DynamicHierarchy.cpp
+++ b/Engine/Plugins/Animation/ControlRig/Source/ControlRig/Private/Units/Execution/RigUnit_DynamicHierarchy.cpp
@@ -343,7 +343,37 @@ FRigUnit_HierarchyImportFromSkeleton_Execute()
 			if(SkelMeshComponent->GetSkeletalMeshAsset())
 			{
 				const FReferenceSkeleton& ReferenceSkeleton = SkelMeshComponent->GetSkeletalMeshAsset()->GetRefSkeleton();
-				Items = Controller->ImportBones(ReferenceSkeleton, NameSpace, false, false, false, false);
+				
+				if (bIncludeVirtualBones)
+				{
+					Items = Controller->ImportBones(ReferenceSkeleton, NameSpace, false, false, false, false);
+				}
+				else
+				{
+					const TArray<FMeshBoneInfo>& BoneInfos = ReferenceSkeleton.GetRefBoneInfo();
+					const TArray<FTransform>& BoneTransforms = ReferenceSkeleton.GetRefBonePose();
+
+					TArray<FMeshBoneInfo> FilteredBoneInfos;
+					TArray<FTransform> FilteredBoneTransforms;
+
+					const int32 NumBones = BoneInfos.Num();
+					for (int32 Index = 0; Index < NumBones; ++Index)
+					{
+						const FMeshBoneInfo& MeshBoneInfo = BoneInfos[Index];
+						const FTransform& BoneTransform = BoneTransforms[Index];
+
+						if (VirtualBoneNameHelpers::CheckVirtualBonePrefix(MeshBoneInfo.Name.ToString()))
+						{
+							continue;
+						}
+						FilteredBoneInfos.Add(MeshBoneInfo);
+						FilteredBoneTransforms.Add(BoneTransform);
+					}
+
+					Items = Controller->ImportBones(FilteredBoneInfos, FilteredBoneTransforms, NameSpace, false, false, false, false);
+				}
+
+
 				if(bIncludeCurves)
 				{
 					if(USkeleton* Skeleton = SkelMeshComponent->GetSkeletalMeshAsset()->GetSkeleton())
diff --git a/Engine/Plugins/Animation/ControlRig/Source/ControlRig/Public/Units/Execution/RigUnit_DynamicHierarchy.h b/Engine/Plugins/Animation/ControlRig/Source/ControlRig/Public/Units/Execution/RigUnit_DynamicHierarchy.h
--- a/Engine/Plugins/Animation/ControlRig/Source/ControlRig/Public/Units/Execution/RigUnit_DynamicHierarchy.h
+++ b/Engine/Plugins/Animation/ControlRig/Source/ControlRig/Public/Units/Execution/RigUnit_DynamicHierarchy.h
@@ -335,6 +335,7 @@ struct FRigUnit_HierarchyImportFromSkeleton : public FRigUnit_DynamicHierarchyBa
 	{
 		NameSpace = NAME_None;
 		bIncludeCurves = bIncludeMeshSockets = false;
+		bIncludeVirtualBones = true;
 	}
 
 	RIGVM_METHOD()
@@ -349,6 +350,9 @@ struct FRigUnit_HierarchyImportFromSkeleton : public FRigUnit_DynamicHierarchyBa
 	UPROPERTY(meta = (Input))
 	bool bIncludeMeshSockets;
 
+	UPROPERTY(meta = (Input))
+	bool bIncludeVirtualBones;
+
 	UPROPERTY(meta = (Output))
 	TArray<FRigElementKey> Items;
 };
