Motion Trail Fixes:  Support multiple rigs with same name and crashes with unbound rigs

diff --git a/Engine/Plugins/Animation/ControlRig/Source/ControlRigEditor/Private/Tools/EvaluateSequencerTools.cpp b/Engine/Plugins/Animation/ControlRig/Source/ControlRigEditor/Private/Tools/EvaluateSequencerTools.cpp
--- a/Engine/Plugins/Animation/ControlRig/Source/ControlRigEditor/Private/Tools/EvaluateSequencerTools.cpp
+++ b/Engine/Plugins/Animation/ControlRig/Source/ControlRigEditor/Private/Tools/EvaluateSequencerTools.cpp
@@ -113,13 +113,14 @@ namespace AIE
 
 				for (USkeletalMeshComponent* MeshComp : MeshComps)
 				{
-					MeshComp->TickAnimation(0.03f, false);
+					MeshComp->TickAnimation(0.0f, false);
 					MeshComp->RefreshBoneTransforms();
 					MeshComp->RefreshFollowerComponents();
 					MeshComp->UpdateComponentToWorld();
 					MeshComp->FinalizeBoneTransform();
-					MeshComp->MarkRenderTransformDirty();
-					MeshComp->MarkRenderDynamicDataDirty();
+					//WIP here for flashing mannies
+					//MeshComp->MarkRenderTransformDirty();
+					//MeshComp->MarkRenderDynamicDataDirty();
 				}
 
 				Parent = Parent->GetAttachParentActor();
@@ -209,6 +210,7 @@ namespace AIE
 
 
 		TSet<UControlRig*> ControlsToRevalaute;
+		
 		for (int32 Index : Indices)
 		{
 			const FFrameNumber FrameNumber = FrameTimeByIndex.CalculateFrame(Index);
@@ -266,11 +268,14 @@ namespace AIE
 				if (UControlRig* ControlRig = CRTrack->GetControlRig())
 				{
 					ControlRig->Evaluate_AnyThread();
-					ControlRig->EvaluateSkeletalMeshComponent(0.03);
+					if (ControlRig->GetObjectBinding())
+					{
+						ControlRig->EvaluateSkeletalMeshComponent(0.0);
+					}
 				}
 			}
 		}
-
+		
 		return true;
 	}
 
diff --git a/Engine/Plugins/Editor/SequencerAnimTools/Source/SequencerAnimTools/Private/Sequencer/MovieSceneTransformTrail.cpp b/Engine/Plugins/Editor/SequencerAnimTools/Source/SequencerAnimTools/Private/Sequencer/MovieSceneTransformTrail.cpp
--- a/Engine/Plugins/Editor/SequencerAnimTools/Source/SequencerAnimTools/Private/Sequencer/MovieSceneTransformTrail.cpp
+++ b/Engine/Plugins/Editor/SequencerAnimTools/Source/SequencerAnimTools/Private/Sequencer/MovieSceneTransformTrail.cpp
@@ -372,7 +372,11 @@ FTrailCurrentStatus FMovieSceneTransformTrail::UpdateTrail(const FNewSceneContex
 	
 	CachedHierarchyGuid = InSceneContext.YourNode;
 	
-	const bool bTrackUnchanged = (InSceneContext.bCheckForChange || TransformDependencies.IsEmpty() ) ? !TrailOrSpaceHasChanged() : true;
+	bool bTrackUnchanged = true;
+	if (InSceneContext.bCheckForChange || TransformDependencies.IsEmpty())
+	{
+		bTrackUnchanged = !TrailOrSpaceHasChanged();
+	}
 
 	UMovieSceneSection* Section = GetSection();
 	TSharedPtr<ISequencer> Sequencer = WeakSequencer.Pin();
diff --git a/Engine/Plugins/Editor/SequencerAnimTools/Source/SequencerAnimTools/Private/Sequencer/SequencerTrailHierarchy.cpp b/Engine/Plugins/Editor/SequencerAnimTools/Source/SequencerAnimTools/Private/Sequencer/SequencerTrailHierarchy.cpp
--- a/Engine/Plugins/Editor/SequencerAnimTools/Source/SequencerAnimTools/Private/Sequencer/SequencerTrailHierarchy.cpp
+++ b/Engine/Plugins/Editor/SequencerAnimTools/Source/SequencerAnimTools/Private/Sequencer/SequencerTrailHierarchy.cpp
@@ -800,8 +800,11 @@ FGuid FSequencerTrailHierarchy::AddControlRigTrail(USkeletalMeshComponent* Compo
 				{
 					if (UMovieSceneControlRigParameterTrack* CRTrack = Cast<UMovieSceneControlRigParameterTrack>(Track))
 					{
-						BindingID = Binding.GetObjectGuid();
-						break;
+						if (CRTrack->GetControlRig() == ControlRig)
+						{
+							BindingID = Binding.GetObjectGuid();
+							break;
+						}
 					}
 				}
 			}
