Tweening Utils: Refactor construction of tweening models:
- FKeyBlendInfo was replaced by TweenRange<EBlendFunction> template: its specializations handle extracting only the parameters required by the blend function.
- The members (BeforeBlendRange, etc.) of FKeyBlendInfo were converted to getters in ContiguousKeyMapping.h.
- FCurveTweenModel was converted to templated TCurveTweenModel, which accepts EBlendFunction. Code only compiles if TweenRange is specialized for the blend function.
- The building of tween functions was extracted to dedicated classes: FControlRigTweenModels and FCurveEditorTweenModels.

diff --git a/Engine/Plugins/Animation/ControlRig/Source/ControlRigEditor/ControlRigEditor.Build.cs b/Engine/Plugins/Animation/ControlRig/Source/ControlRigEditor/ControlRigEditor.Build.cs
--- a/Engine/Plugins/Animation/ControlRig/Source/ControlRigEditor/ControlRigEditor.Build.cs
+++ b/Engine/Plugins/Animation/ControlRig/Source/ControlRigEditor/ControlRigEditor.Build.cs
@@ -85,7 +85,8 @@ namespace UnrealBuildTool.Rules
                     "Constraints",
                     "AnimationEditMode",
 					"SequencerWidgets",
-					"DeveloperSettings"
+					"DeveloperSettings",
+					"TweeningUtils"
 				}
             );
 
diff --git a/Engine/Plugins/Animation/ControlRig/Source/ControlRigEditor/Private/ControlRigSequencerEditorLibrary.cpp b/Engine/Plugins/Animation/ControlRig/Source/ControlRigEditor/Private/ControlRigSequencerEditorLibrary.cpp
--- a/Engine/Plugins/Animation/ControlRig/Source/ControlRigEditor/Private/ControlRigSequencerEditorLibrary.cpp
+++ b/Engine/Plugins/Animation/ControlRig/Source/ControlRigEditor/Private/ControlRigSequencerEditorLibrary.cpp
@@ -410,7 +410,7 @@ bool UControlRigSequencerEditorLibrary::TweenControlRig(ULevelSequence* LevelSeq
 	{
 		LevelSequence->GetMovieScene()->Modify();
 		
-		FControlRigTweenModel BlendTool(WeakSequencer, EBlendFunction::ControlsToTween);
+		TControlRigTweenModel<EBlendFunction::ControlsToTween> BlendTool(WeakSequencer); 
 		BlendTool.StartBlendOperationWithControlRigs({ ControlRig });
 		BlendTool.BlendValues(TweenValue);
 		BlendTool.StopBlendOperation();
@@ -421,20 +421,33 @@ bool UControlRigSequencerEditorLibrary::TweenControlRig(ULevelSequence* LevelSeq
 
 namespace UE::ControlRigEditor::SequencerLibraryDetail
 {
-	static TweeningUtilsEditor::EBlendFunction ConvertBlendFunction(EAnimToolBlendOperation Operation)
+static bool BlendSingle(
+	const TWeakPtr<ISequencer>& InWeakSequencer,
+	const TWeakPtr<FControlRigEditMode> InWeakMode,
+	EAnimToolBlendOperation InOperation,
+	double InBlendValue
+	)
+{
+	using namespace TweeningUtilsEditor;
+	
+	switch (InOperation)
 	{
-		switch (Operation)
-		{
-		case EAnimToolBlendOperation::Tween: return TweeningUtilsEditor::EBlendFunction::ControlsToTween;
-		case EAnimToolBlendOperation::BlendToNeighbor: return TweeningUtilsEditor::EBlendFunction::BlendNeighbor;
-		case EAnimToolBlendOperation::PushPull: return TweeningUtilsEditor::EBlendFunction::PushPull;
-		case EAnimToolBlendOperation::BlendRelative: return TweeningUtilsEditor::EBlendFunction::BlendRelative;
-		case EAnimToolBlendOperation::BlendToEase: return TweeningUtilsEditor::EBlendFunction::BlendEase;
-		case EAnimToolBlendOperation::SmoothRough: return TweeningUtilsEditor::EBlendFunction::SmoothRough;
-		default: checkNoEntry(); return TweeningUtilsEditor::EBlendFunction::ControlsToTween;;
-		}
+	case EAnimToolBlendOperation::Tween:
+		return TControlRigTweenModel<EBlendFunction::ControlsToTween>(InWeakSequencer, InWeakMode).BlendSingle(InBlendValue);
+	case EAnimToolBlendOperation::BlendToNeighbor:
+		return TControlRigTweenModel<EBlendFunction::BlendNeighbor>(InWeakSequencer, InWeakMode).BlendSingle(InBlendValue);
+	case EAnimToolBlendOperation::PushPull:
+		return TControlRigTweenModel<EBlendFunction::PushPull>(InWeakSequencer, InWeakMode).BlendSingle(InBlendValue);
+	case EAnimToolBlendOperation::BlendRelative:
+		return TControlRigTweenModel<EBlendFunction::BlendRelative>(InWeakSequencer, InWeakMode).BlendSingle(InBlendValue);
+	case EAnimToolBlendOperation::BlendToEase:
+		return TControlRigTweenModel<EBlendFunction::BlendEase>(InWeakSequencer, InWeakMode).BlendSingle(InBlendValue);
+	case EAnimToolBlendOperation::SmoothRough:
+		return TControlRigTweenModel<EBlendFunction::SmoothRough>(InWeakSequencer, InWeakMode).BlendSingle(InBlendValue);
+	default: ensure(false); return false;
 	}
 }
+}
 
 bool UControlRigSequencerEditorLibrary::BlendValuesOnSelected(ULevelSequence* LevelSequence, EAnimToolBlendOperation BlendOperation, float BlendValue)
 {
@@ -450,9 +463,7 @@ bool UControlRigSequencerEditorLibrary::BlendValuesOnSelected(ULevelSequence* Le
 			TWeakPtr<FControlRigEditMode> WeakMode = StaticCastSharedRef<FControlRigEditMode, FEdMode>(EditMode->AsShared()).ToWeakPtr();
 			LevelSequence->GetMovieScene()->Modify();
 			
-			const EBlendFunction BlendFunction = SequencerLibraryDetail::ConvertBlendFunction(BlendOperation);
-			FControlRigTweenModel BlendTool(WeakSequencer, WeakMode, BlendFunction);
-			return BlendTool.BlendSingle(BlendValue);
+			return SequencerLibraryDetail::BlendSingle(WeakSequencer, WeakMode, BlendOperation, BlendValue);
 		}
 	}
 	return false;
diff --git a/Engine/Plugins/Animation/ControlRig/Source/ControlRigEditor/Private/EditMode/Tween/ControlRigTweenModel.cpp b/Engine/Plugins/Animation/ControlRig/Source/ControlRigEditor/Private/EditMode/Tween/ControlRigTweenModel.cpp
deleted file mode 100644
diff --git a/Engine/Plugins/Animation/ControlRig/Source/ControlRigEditor/Private/EditMode/Tween/ControlRigTweenModel.h b/Engine/Plugins/Animation/ControlRig/Source/ControlRigEditor/Private/EditMode/Tween/ControlRigTweenModel.h
--- a/Engine/Plugins/Animation/ControlRig/Source/ControlRigEditor/Private/EditMode/Tween/ControlRigTweenModel.h
+++ b/Engine/Plugins/Animation/ControlRig/Source/ControlRigEditor/Private/EditMode/Tween/ControlRigTweenModel.h
@@ -2,15 +2,19 @@
 
 #pragma once
 
+#include "Channels/MovieSceneFloatChannel.h"
+#include "Channels/MovieSceneDoubleChannel.h"
 #include "ControlRigObjectSelection.h"
-#include "Math/ContiguousKeyMapping.h"
+#include "ISequencer.h"
 #include "Math/Abstraction/KeyBlendingAbstraction.h"
+#include "Math/Abstraction/TweenRangeTemplates.h"
+#include "Math/ContiguousKeyMapping.h"
+#include "Math/CurveBlending.h"
 #include "Math/Models/TweenModel.h"
 #include "Misc/Attribute.h"
+#include "Misc/SequencerUtils.h"
 #include "Templates/SharedPointer.h"
 
-class ISequencer;
-
 namespace UE::ControlRigEditor
 {
 /**
@@ -21,20 +25,17 @@ namespace UE::ControlRigEditor
  *	2. Otherwise, if a control rig element is selected, insert a key at the current scrub position into the curves corresponding to the control rig
  *	and then interpolate those keys.
  */
-class FControlRigTweenModel : public TweeningUtilsEditor::FTweenModel
+template<TweeningUtilsEditor::EBlendFunction>
+class TControlRigTweenModel : public TweeningUtilsEditor::FTweenModel
 {
 public:
-	
-	explicit FControlRigTweenModel(
-		TAttribute<TWeakPtr<ISequencer>> InSequencerAttr, TWeakPtr<FControlRigEditMode> InControlRigMode,
-		TweeningUtilsEditor::EBlendFunction InBlendFunction
-		)
+
+	explicit TControlRigTweenModel(TAttribute<TWeakPtr<ISequencer>> InSequencerAttr, TWeakPtr<FControlRigEditMode> InControlRigMode)
 		: SequencerAttr(MoveTemp(InSequencerAttr))
 		, ControlRigMode(MoveTemp(InControlRigMode))
-		, BlendFunction(InBlendFunction)
 	{}
-	explicit FControlRigTweenModel(TAttribute<TWeakPtr<ISequencer>> InSequencerAttr, TweeningUtilsEditor::EBlendFunction InBlendFunction)
-		: FControlRigTweenModel(MoveTemp(InSequencerAttr), nullptr, InBlendFunction)
+	explicit TControlRigTweenModel(TAttribute<TWeakPtr<ISequencer>> InSequencerAttr)
+		: TControlRigTweenModel(MoveTemp(InSequencerAttr), nullptr)
 	{}
 
 	/** Does a blend and returns whether any values were actually blended. */
@@ -56,8 +57,6 @@ private:
 	const TAttribute<TWeakPtr<ISequencer>> SequencerAttr;
 	/** The control rig mode that owns this tweener. Can be nullptr, in which case only the curve editor key selection is blended. */
 	const TWeakPtr<FControlRigEditMode> ControlRigMode;
-	/** The blend function that does the blending. */
-	const TweeningUtilsEditor::EBlendFunction BlendFunction;
 
 	/** The key selection to blend. Created in StartBlendOperation and used for the entirety of the blend operation. */
 	TweeningUtilsEditor::FContiguousKeyMapping KeySelection;
@@ -66,9 +65,149 @@ private:
 	 * Created in StartBlendOperation and used for the entirety of the blend operation.
 	 */
 	FControlRigObjectSelection ObjectSelection; // TODO: Extract this to UE::ControlRigEditor namespace and deprecate the global one.
-	
+
 	/** @return Whether any blended occured. */
 	bool DoBlend(float InNormalizedValue);
+	bool BlendObjectSelection(const TSharedPtr<ISequencer>& SequencerPin, float InScaledBlendValue);
 };
 }
 
+namespace UE::ControlRigEditor
+{
+template <TweeningUtilsEditor::EBlendFunction BlendFunction>
+bool TControlRigTweenModel<BlendFunction>::BlendSingle(float InNormalizedValue)
+{
+	StartBlendOperation();
+	const bool bBlended = DoBlend(InNormalizedValue);
+	StopBlendOperation(); // Technically unneeded but it's expected by the API so we do it in case base version changes in the future
+	return bBlended;
+}
+
+template <TweeningUtilsEditor::EBlendFunction BlendFunction>
+bool TControlRigTweenModel<BlendFunction>::BlendSingleWithControlRigs(float InNormalizedValue, const TArray<UControlRig*>& InControlRigs)
+{
+	StartBlendOperationWithControlRigs(InControlRigs);
+	const bool bBlended = DoBlend(InNormalizedValue);
+	StopBlendOperation(); // Technically unneeded but it's expected by the API so we do it in case base version changes in the future
+	return bBlended;
+}
+
+template <TweeningUtilsEditor::EBlendFunction BlendFunction>
+void TControlRigTweenModel<BlendFunction>::StartBlendOperationWithControlRigs(const TArray<UControlRig*>& InControlRigs)
+{
+	const TWeakPtr<ISequencer> Sequencer = SequencerAttr.Get();
+	if (const TSharedPtr<ISequencer> SequencerPin = Sequencer.Pin())
+	{
+		ObjectSelection.Setup(InControlRigs, Sequencer);
+	}
+}
+
+template <TweeningUtilsEditor::EBlendFunction BlendFunction>
+void TControlRigTweenModel<BlendFunction>::StartBlendOperation()
+{
+	KeySelection.KeyMap.Empty(KeySelection.KeyMap.Num());
+	ObjectSelection.ChannelsArray.Empty(ObjectSelection.ChannelsArray.Num());
+	
+	const TWeakPtr<ISequencer> Sequencer = SequencerAttr.Get();
+	const TSharedPtr<ISequencer> SequencerPin = Sequencer.Pin();
+	const TSharedPtr<FCurveEditor> CurveEditor = GetCurveEditorFromSequencer(SequencerPin);
+	if (!CurveEditor)
+	{
+		return;
+	}
+
+	KeySelection = TweeningUtilsEditor::FContiguousKeyMapping(*CurveEditor);
+	if (KeySelection.KeyMap.IsEmpty())
+	{
+		ObjectSelection.Setup(Sequencer, ControlRigMode);
+	}
+}
+
+template <TweeningUtilsEditor::EBlendFunction BlendFunction>
+bool TControlRigTweenModel<BlendFunction>::DoBlend(float InNormalizedValue)
+{
+	
+	const TWeakPtr<ISequencer> Sequencer = SequencerAttr.Get();
+	const TSharedPtr<ISequencer> SequencerPin = Sequencer.Pin();
+	const TSharedPtr<FCurveEditor> CurveEditor = GetCurveEditorFromSequencer(SequencerPin);
+	if (!CurveEditor)
+	{
+		return false;
+	}
+
+	const float ScaledBlendValue = ScaleBlendValue(InNormalizedValue); 
+	if (!KeySelection.KeyMap.IsEmpty())
+	{
+		using namespace UE::TweeningUtilsEditor;
+		BlendCurves_BySingleKey(*CurveEditor, KeySelection,
+			[this, ScaledBlendValue](
+				const FCurveModelID&,
+				const FContiguousKeyMapping::FContiguousKeysArray& AllBlendedKeys, const FContiguousKeys& CurrentBlendRange,
+				int32 InCurrentKeyIndex
+			)
+		{
+			return TweeningUtilsEditor::TweenRange<BlendFunction>(ScaledBlendValue, AllBlendedKeys, CurrentBlendRange, InCurrentKeyIndex);
+		});
+	}
+
+	if (!ObjectSelection.ChannelsArray.IsEmpty())
+	{
+		BlendObjectSelection(SequencerPin, ScaledBlendValue);
+	}
+
+	return false;
+}
+
+template <TweeningUtilsEditor::EBlendFunction BlendFunction>
+bool TControlRigTweenModel<BlendFunction>::BlendObjectSelection(const TSharedPtr<ISequencer>& SequencerPin, float InScaledBlendValue)
+{
+	using namespace UE::TweeningUtilsEditor;
+	
+	const FFrameTime FrameTime = SequencerPin->GetLocalTime().Time;
+	const FFrameRate TickResolution = SequencerPin->GetFocusedTickResolution();
+	
+	bool bDidBlend = false;
+	for (const FControlRigObjectSelection::FObjectChannels& ObjectChannels : ObjectSelection.ChannelsArray)
+	{
+		if (ObjectChannels.Section)
+		{
+			ObjectChannels.Section->Modify();
+		}
+		for (int32 Index = 0; Index < ObjectChannels.KeyBounds.Num(); ++Index)
+		{
+			if (ObjectChannels.KeyBounds[Index].bValid)
+			{
+				const double PreviousValue = ObjectChannels.KeyBounds[Index].PreviousValue;
+				const double PreviousTime = TickResolution.AsSeconds(FFrameTime(ObjectChannels.KeyBounds[Index].PreviousFrame));
+				const double NextValue = ObjectChannels.KeyBounds[Index].NextValue;
+				const double NextTime = TickResolution.AsSeconds(FFrameTime(ObjectChannels.KeyBounds[Index].NextFrame));
+				const double CurrentValue = ObjectChannels.KeyBounds[Index].CurrentValue;
+				const double CurrentTime = TickResolution.AsSeconds(FFrameTime(ObjectChannels.KeyBounds[Index].CurrentFrame));
+
+				const FVector2d BeforeBlendRange { PreviousTime, PreviousValue };
+				const FVector2d AfterBlendRange { NextTime, NextValue };
+				const FVector2d Current { CurrentTime, CurrentValue };
+				
+				FBlendRangesData AllKeys({ BeforeBlendRange, Current, AfterBlendRange});
+				AllKeys.AddBlendRange({ 1 });
+				const double NewValue = TweenRange<BlendFunction>(InScaledBlendValue, AllKeys, AllKeys.KeysArray[0], 0);
+				
+				using namespace UE::MovieScene;
+				if (ObjectChannels.KeyBounds[Index].FloatChannel)
+				{
+					const EMovieSceneKeyInterpolation KeyInterpolation = GetInterpolationMode(ObjectChannels.KeyBounds[Index].FloatChannel, FrameTime.GetFrame(), SequencerPin->GetKeyInterpolation());
+					AddKeyToChannel(ObjectChannels.KeyBounds[Index].FloatChannel, FrameTime.GetFrame(), (float)NewValue, KeyInterpolation);
+				}
+				else if (ObjectChannels.KeyBounds[Index].DoubleChannel)
+				{
+					const EMovieSceneKeyInterpolation KeyInterpolation = GetInterpolationMode(ObjectChannels.KeyBounds[Index].DoubleChannel, FrameTime.GetFrame(), SequencerPin->GetKeyInterpolation());
+					AddKeyToChannel(ObjectChannels.KeyBounds[Index].DoubleChannel, FrameTime.GetFrame(), NewValue, KeyInterpolation);
+				}
+				bDidBlend = true;
+			}
+		}
+	}
+	return bDidBlend;
+}
+}
+
diff --git a/Engine/Plugins/Animation/ControlRig/Source/ControlRigEditor/Private/EditMode/Tween/ControlRigTweenModels.cpp b/Engine/Plugins/Animation/ControlRig/Source/ControlRigEditor/Private/EditMode/Tween/ControlRigTweenModels.cpp
new file mode 100644
--- /dev/null
+++ b/Engine/Plugins/Animation/ControlRig/Source/ControlRigEditor/Private/EditMode/Tween/ControlRigTweenModels.cpp
@@ -0,0 +1,50 @@
+﻿// Copyright Epic Games, Inc. All Rights Reserved.
+
+#include "ControlRigTweenModels.h"
+
+#include "ControlRigTweenModel.h"
+#include "Math/Abstraction/TweenRangeTemplates.h"
+#include "Math/Models/EditorTweenModel.h"
+#include "Misc/SequencerUtils.h"
+
+namespace UE::ControlRigEditor
+{
+namespace TweenDetail
+{
+static TAttribute<TWeakPtr<FCurveEditor>> MakeCurveEditorAttribute(TAttribute<TWeakPtr<ISequencer>> InSequencerAttr)
+{
+	return TAttribute<TWeakPtr<FCurveEditor>>::CreateLambda([SequencerAttr = InSequencerAttr]
+	{
+		const TSharedPtr<ISequencer> SequencerPin = SequencerAttr.Get().Pin();
+		return GetCurveEditorFromSequencer(SequencerPin);
+	});
+}
+
+static TArray<TweeningUtilsEditor::FTweenModelUIEntry> MakeTweenModels(
+	const TAttribute<TWeakPtr<ISequencer>>& InSequencerAttr, const TSharedRef<FControlRigEditMode>& InOwningEditMode
+	)
+{
+	using namespace TweeningUtilsEditor;
+
+	const TAttribute<TWeakPtr<FCurveEditor>> CurveEditorAttr = MakeCurveEditorAttribute(InSequencerAttr);
+	TArray<FTweenModelUIEntry> Result;
+	ForEachCurveTweenable([&InSequencerAttr, &InOwningEditMode, &CurveEditorAttr, &Result]<EBlendFunction BlendFunction>
+	{
+		using FModel = TEditorTweenModel<TControlRigTweenModel<BlendFunction>>;
+		TUniquePtr<FModel> TweenModel = MakeUnique<FModel>(
+			CurveEditorAttr,										// TTangentFlatteningTweenProxy arg
+			InSequencerAttr, InOwningEditMode						// FControlRigTweenModel args
+			);
+		Result.Emplace(MoveTemp(TweenModel), FTweenModelDisplayInfo(BlendFunction));
+	});
+
+	return Result;
+}
+}
+
+FControlRigTweenModels::FControlRigTweenModels(
+	const TAttribute<TWeakPtr<ISequencer>>& InSequencerAttr, const TSharedRef<FControlRigEditMode>& InOwningEditMode
+	)
+	: FTweenModelArray(TweenDetail::MakeTweenModels(InSequencerAttr, InOwningEditMode))
+{}
+}
diff --git a/Engine/Plugins/Animation/ControlRig/Source/ControlRigEditor/Private/EditMode/Tween/ControlRigTweenModels.h b/Engine/Plugins/Animation/ControlRig/Source/ControlRigEditor/Private/EditMode/Tween/ControlRigTweenModels.h
new file mode 100644
--- /dev/null
+++ b/Engine/Plugins/Animation/ControlRig/Source/ControlRigEditor/Private/EditMode/Tween/ControlRigTweenModels.h
@@ -0,0 +1,22 @@
+﻿// Copyright Epic Games, Inc. All Rights Reserved.
+
+#pragma once
+
+#include "Math/Abstraction/TweenModelArray.h"
+#include "Templates/SharedPointerFwd.h"
+
+class ISequencer;
+class FControlRigEditMode;
+class FCurveEditor;
+template<typename ObjectType>class TAttribute;
+
+namespace UE::ControlRigEditor
+{
+/** Knows of the tween models used in control rig. */
+class FControlRigTweenModels : public TweeningUtilsEditor::FTweenModelArray
+{
+public:
+
+	explicit FControlRigTweenModels(const TAttribute<TWeakPtr<ISequencer>>& InSequencerAttr, const TSharedRef<FControlRigEditMode>& InOwningEditMode);
+};
+}
\ No newline at end of file
diff --git a/Engine/Plugins/Animation/ControlRig/Source/ControlRigEditor/Private/EditMode/Tween/TweenLogic.cpp b/Engine/Plugins/Animation/ControlRig/Source/ControlRigEditor/Private/EditMode/Tween/TweenLogic.cpp
--- a/Engine/Plugins/Animation/ControlRig/Source/ControlRigEditor/Private/EditMode/Tween/TweenLogic.cpp
+++ b/Engine/Plugins/Animation/ControlRig/Source/ControlRigEditor/Private/EditMode/Tween/TweenLogic.cpp
@@ -3,6 +3,7 @@
 #include "TweenLogic.h"
 
 #include "ControlRigTweenModel.h"
+#include "ControlRigTweenModels.h"
 #include "EditMode/ControlRigEditMode.h"
 #include "Framework/Commands/UICommandList.h"
 #include "Math/Abstraction/TweenModelArray.h"
@@ -12,34 +13,9 @@
 
 namespace UE::ControlRigEditor
 {
-namespace TweenLogicDetail
-{
-static TUniquePtr<FControlRigTweenModel> MakeTweenModel(
-	TAttribute<TWeakPtr<ISequencer>> InSequencerAttr, const TSharedRef<FControlRigEditMode>& InOwningEditMode,
-	TweeningUtilsEditor::EBlendFunction BlendFunction
-	)
-{
-	using FModel = TweeningUtilsEditor::TEditorTweenModel<FControlRigTweenModel>;
-	TAttribute<TWeakPtr<FCurveEditor>> CurveEditorAttr = TAttribute<TWeakPtr<FCurveEditor>>::CreateLambda([SequencerAttr = InSequencerAttr]
-	{
-		const TSharedPtr<ISequencer> SequencerPin = SequencerAttr.Get().Pin();
-		return GetCurveEditorFromSequencer(SequencerPin);
-	});
-	
-	return MakeUnique<FModel>(
-		MoveTemp(CurveEditorAttr),										// TTangentFlatteningTweenProxy arg
-		MoveTemp(InSequencerAttr), InOwningEditMode, BlendFunction
-		);
-}
-}
-	
-FTweenLogic::FTweenLogic(TAttribute<TWeakPtr<ISequencer>> InSequencerAttr, const TSharedRef<FControlRigEditMode>& InOwningEditMode)
+FTweenLogic::FTweenLogic(const TAttribute<TWeakPtr<ISequencer>>& InSequencerAttr, const TSharedRef<FControlRigEditMode>& InOwningEditMode)
 	: CommandList(InOwningEditMode->GetCommandBindings().ToSharedRef())
-	, TweenModels(MakeShared<TweeningUtilsEditor::FTweenModelArray>(
-		[&InSequencerAttr, &InOwningEditMode](TweeningUtilsEditor::EBlendFunction BlendFunction)
-		{
-			return TweenLogicDetail::MakeTweenModel(InSequencerAttr, InOwningEditMode, BlendFunction);
-		}))
+	, TweenModels(MakeShared<FControlRigTweenModels>(InSequencerAttr, InOwningEditMode))
 	, Controllers(CommandList, TweenModels, TEXT("ControlRigTween"))
 {}
 
diff --git a/Engine/Plugins/Animation/ControlRig/Source/ControlRigEditor/Private/EditMode/Tween/TweenLogic.h b/Engine/Plugins/Animation/ControlRig/Source/ControlRigEditor/Private/EditMode/Tween/TweenLogic.h
--- a/Engine/Plugins/Animation/ControlRig/Source/ControlRigEditor/Private/EditMode/Tween/TweenLogic.h
+++ b/Engine/Plugins/Animation/ControlRig/Source/ControlRigEditor/Private/EditMode/Tween/TweenLogic.h
@@ -12,14 +12,14 @@ class ISequencer;
 
 namespace UE::ControlRigEditor
 {
-class FControlRigTweenModel;
+class FControlRigTweenModels;
 
 /** Creates the tweening model specific to ControlRig and provides an interface for interacting with it. */
 class FTweenLogic
 {
 public:
 
-	explicit FTweenLogic(TAttribute<TWeakPtr<ISequencer>> InSequencerAttr, const TSharedRef<FControlRigEditMode>& InOwningEditMode);
+	explicit FTweenLogic(const TAttribute<TWeakPtr<ISequencer>>& InSequencerAttr, const TSharedRef<FControlRigEditMode>& InOwningEditMode);
 
 	/** Constructs the content for the viewport tweening widget. */
 	TSharedRef<SWidget> MakeWidget() const;
@@ -30,7 +30,7 @@ private:
 	const TSharedRef<FUICommandList> CommandList;
 
 	/** Holds the used tweening functions and info about how to display them in UI. */
-	const TSharedRef<TweeningUtilsEditor::ITweenModelContainer> TweenModels;
+	const TSharedRef<FControlRigTweenModels> TweenModels;
 
 	/** Common functionality that should be shared consistent with other editor modules, such as Curve Editor. */
 	const TweeningUtilsEditor::FTweenControllers Controllers;
diff --git a/Engine/Plugins/Animation/TweeningUtils/Source/TweeningUtils/Private/Math/KeyBlendingFunctions.cpp b/Engine/Plugins/Animation/TweeningUtils/Source/TweeningUtils/Private/Math/KeyBlendingFunctions.cpp
--- a/Engine/Plugins/Animation/TweeningUtils/Source/TweeningUtils/Private/Math/KeyBlendingFunctions.cpp
+++ b/Engine/Plugins/Animation/TweeningUtils/Source/TweeningUtils/Private/Math/KeyBlendingFunctions.cpp
@@ -2,64 +2,67 @@
 
 #include "Math/KeyBlendingFunctions.h"
 
-#include "Math/KeyBlendInfo.h"
 #include "Math/Vector2D.h"
 
 namespace UE::TweeningUtils
 {
-double Blend_ControlsToTween(const FKeyBlendInfo& InBlendInfo)
+double Blend_ControlsToTween(double InBlendValue, const FVector2d& InBeforeBlendRange, const FVector2d& InAfterBlendRange)
 {
 	// Classic tween will move all to same location, not based on current time at all just blend and values
-	const double NormalizedBlendValue = (InBlendInfo.BlendValue + 1.0f) * 0.5f;
-	const double Value = InBlendInfo.BeforeBlendRange.Y + (InBlendInfo.AfterBlendRange.Y - InBlendInfo.BeforeBlendRange.Y) * (NormalizedBlendValue);
+	const double NormalizedBlendValue = (InBlendValue + 1.0f) * 0.5f;
+	const double Value = InBeforeBlendRange.Y + (InAfterBlendRange.Y - InBeforeBlendRange.Y) * (NormalizedBlendValue);
 	return Value;
 }
 
-double Blend_PushPull(const FKeyBlendInfo& InBlendInfo)
+double Blend_PushPull(double InBlendValue, const FVector2d& InBeforeBlendRange, const FVector2d& InCurrent, const FVector2d& InAfterBlendRange)
 {
-	if (FMath::IsNearlyEqual(InBlendInfo.AfterBlendRange.X, InBlendInfo.BeforeBlendRange.X))
+	if (FMath::IsNearlyEqual(InAfterBlendRange.X, InBeforeBlendRange.X))
 	{
-		return InBlendInfo.Current.Y;
+		return InCurrent.Y;
 	}
-	const double T = (InBlendInfo.Current.X - InBlendInfo.BeforeBlendRange.X) / (InBlendInfo.AfterBlendRange.X - InBlendInfo.BeforeBlendRange.X);
-	const double ValueAtT = InBlendInfo.BeforeBlendRange.Y + T * (InBlendInfo.AfterBlendRange.Y - InBlendInfo.BeforeBlendRange.Y);
+	const double T = (InCurrent.X - InBeforeBlendRange.X) / (InAfterBlendRange.X - InBeforeBlendRange.X);
+	const double ValueAtT = InBeforeBlendRange.Y + T * (InAfterBlendRange.Y - InBeforeBlendRange.Y);
 	double NewValue;
-	if (InBlendInfo.BlendValue < 0.0)
+	if (InBlendValue < 0.0)
 	{
-		NewValue = InBlendInfo.Current.Y + (-1.0 * InBlendInfo.BlendValue) * (ValueAtT - InBlendInfo.Current.Y);
+		NewValue = InCurrent.Y + (-1.0 * InBlendValue) * (ValueAtT - InCurrent.Y);
 	}
 	else
 	{
-		const double AmplifyValueAtT = InBlendInfo.Current.Y + (InBlendInfo.Current.Y - ValueAtT);
-		NewValue = InBlendInfo.Current.Y + InBlendInfo.BlendValue * (AmplifyValueAtT - InBlendInfo.Current.Y);
+		const double AmplifyValueAtT = InCurrent.Y + (InCurrent.Y - ValueAtT);
+		NewValue = InCurrent.Y + InBlendValue * (AmplifyValueAtT - InCurrent.Y);
 	}
 	return NewValue;
 }
 
-double Blend_Neighbor(const FKeyBlendInfo& InBlendInfo)
+double Blend_Neighbor(double InBlendValue, const FVector2d& InBeforeBlendRange, const FVector2d& InCurrent, const FVector2d& InAfterBlendRange)
 {
-	double NewValue = InBlendInfo.Current.Y;
-	if (InBlendInfo.BlendValue < 0.0)
+	double NewValue = InCurrent.Y;
+	if (InBlendValue < 0.0)
 	{
-		NewValue = InBlendInfo.Current.Y + (-1.0 * InBlendInfo.BlendValue) * (InBlendInfo.BeforeBlendRange.Y - InBlendInfo.Current.Y);
+		NewValue = InCurrent.Y + (-1.0 * InBlendValue) * (InBeforeBlendRange.Y - InCurrent.Y);
 	}
 	else
 	{
-		NewValue = InBlendInfo.Current.Y + InBlendInfo.BlendValue * (InBlendInfo.AfterBlendRange.Y - InBlendInfo.Current.Y);
+		NewValue = InCurrent.Y + InBlendValue * (InAfterBlendRange.Y - InCurrent.Y);
 	}
 	return NewValue;
 }
-
-double Blend_Relative(const FKeyBlendInfo& InBlendInfo)
+	
+double Blend_Relative(double InBlendValue, 
+	const FVector2d& InBeforeBlendRange, const FVector2d& InFirstBlended,
+	const FVector2d& InCurrent,
+	const FVector2d& InLastBlended, const FVector2d& InAfterBlendRange
+)
 {
-	double NewValue = InBlendInfo.Current.Y;
-	if (InBlendInfo.BlendValue < 0.0)
+	double NewValue = InCurrent.Y;
+	if (InBlendValue < 0.0)
 	{
-		NewValue = InBlendInfo.Current.Y + (-1.0 * InBlendInfo.BlendValue) * (InBlendInfo.BeforeBlendRange.Y - InBlendInfo.FirstBlended.Y);
+		NewValue = InCurrent.Y + (-1.0 * InBlendValue) * (InBeforeBlendRange.Y - InFirstBlended.Y);
 	}
 	else
 	{
-		NewValue = InBlendInfo.Current.Y + InBlendInfo.BlendValue * (InBlendInfo.AfterBlendRange.Y - InBlendInfo.LastBlended.Y);
+		NewValue = InCurrent.Y + InBlendValue * (InAfterBlendRange.Y - InLastBlended.Y);
 	}
 	return NewValue;
 }
@@ -82,29 +85,32 @@ static float SCurve(float X, float Slope, float Width, float Height, float XShif
 }
 }
 
-double Blend_Ease(const FKeyBlendInfo& InBlendInfo)
+double Blend_Ease(
+	double InBlendValue, 
+	const FVector2d& InBeforeBlendRange, const FVector2d& InCurrent, const FVector2d& InAfterBlendRange
+	)
 {
-	if (FMath::IsNearlyEqual(InBlendInfo.AfterBlendRange.X, InBlendInfo.BeforeBlendRange.X))
+	if (FMath::IsNearlyEqual(InAfterBlendRange.X, InBeforeBlendRange.X))
 	{
-		return InBlendInfo.Current.Y;
+		return InCurrent.Y;
 	}
-	const double Source = InBlendInfo.Current.Y;
-	const double FullTimeDiff = InBlendInfo.AfterBlendRange.X - InBlendInfo.BeforeBlendRange.X;
-	const double AbsValue = FMath::Abs(InBlendInfo.BlendValue);
-	const double X = InBlendInfo.Current.X - InBlendInfo.BeforeBlendRange.X;
+	const double Source = InCurrent.Y;
+	const double FullTimeDiff = InAfterBlendRange.X - InBeforeBlendRange.X;
+	const double AbsValue = FMath::Abs(InBlendValue);
+	const double X = InCurrent.X - InBeforeBlendRange.X;
 	const double Ratio = X / FullTimeDiff;
 	double Shift, Delta, Base;
-	if (InBlendInfo.BlendValue > 0)
+	if (InBlendValue > 0)
 	{
 		Shift = -1.0;
-		Delta = InBlendInfo.AfterBlendRange.Y - Source;
+		Delta = InAfterBlendRange.Y - Source;
 		Base = Source;
 	}
 	else
 	{
 		Shift = 0.0;
-		Delta = Source - InBlendInfo.BeforeBlendRange.Y;
-		Base = InBlendInfo.BeforeBlendRange.Y;
+		Delta = Source - InBeforeBlendRange.Y;
+		Base = InBeforeBlendRange.Y;
 	}
 	const double Slope = 5.0 * AbsValue;
 	const double EaseY = Private::BlendToEase::SCurve(Ratio, Slope, 2.0, 2.0, Shift, Shift);
@@ -113,19 +119,19 @@ double Blend_Ease(const FKeyBlendInfo& InBlendInfo)
 	return NewValue;
 }
 
-double Blend_SmoothRough(const FKeyBlendInfo& InBlendInfo)
+double Blend_SmoothRough(double InBlendValue, const FVector2d& InBeforeCurrent, const FVector2d& InCurrent, const FVector2d& InAfterCurrent)
 {
-	const double PrevVal = InBlendInfo.BeforeCurrent.Y;
-	const double CurVal = InBlendInfo.Current.Y;
-	const double NextVal = InBlendInfo.AfterCurrent.Y;
+	const double PrevVal = InBeforeCurrent.Y;
+	const double CurVal = InCurrent.Y;
+	const double NextVal = InAfterCurrent.Y;
 	const double NewValue = (PrevVal * 0.25 + CurVal * 0.5 + NextVal * 0.25);
-	if (InBlendInfo.BlendValue < 0.0)
+	if (InBlendValue < 0.0)
 	{
-		return InBlendInfo.Current.Y + (-1.0 * InBlendInfo.BlendValue) * (NewValue - InBlendInfo.Current.Y);
+		return InCurrent.Y + (-1.0 * InBlendValue) * (NewValue - InCurrent.Y);
 	}
 	else
 	{
-		return InBlendInfo.Current.Y + InBlendInfo.BlendValue * (InBlendInfo.Current.Y - NewValue);
+		return InCurrent.Y + InBlendValue * (InCurrent.Y - NewValue);
 	}
 }
 }
diff --git a/Engine/Plugins/Animation/TweeningUtils/Source/TweeningUtils/Public/Math/KeyBlendInfo.h b/Engine/Plugins/Animation/TweeningUtils/Source/TweeningUtils/Public/Math/KeyBlendInfo.h
deleted file mode 100644
diff --git a/Engine/Plugins/Animation/TweeningUtils/Source/TweeningUtils/Public/Math/KeyBlendingFunctions.h b/Engine/Plugins/Animation/TweeningUtils/Source/TweeningUtils/Public/Math/KeyBlendingFunctions.h
--- a/Engine/Plugins/Animation/TweeningUtils/Source/TweeningUtils/Public/Math/KeyBlendingFunctions.h
+++ b/Engine/Plugins/Animation/TweeningUtils/Source/TweeningUtils/Public/Math/KeyBlendingFunctions.h
@@ -2,19 +2,79 @@
 
 #pragma once
 
-#include "KeyBlendInfo.h"
+#include "Math/MathFwd.h"
 
+/**
+ * The blending functions are generally used to iterate through a contiguous set of keys.
+ * 
+ * Visualized it looks like .xxxxxx., where '.'are keys that are not blended and 'x' keys that are blended.
+ * There can be no gap between the keys, so '.xx.xxx.' is not a blendable range: it would need to be split up into two ranges: '.xx.' and '.xxx.'.
+ * @see ContiguousKeyMapping.h
+ * 
+ * Generally the blending functions share the following arguments:
+ *	- BlendValue		- The user specified value based on which to blend. Range is in [-1.0, 1.0]. 
+ *	- BeforeBlendRange	- The key right before the range of blended keys. Same as FirstBlended if there is no such key.
+ *	- FirstBlended		- The first key in the range of blended keys.
+ *	- BeforeCurrent		- The key before the currently blended one. If Current is the first key, this is BeforeBlendRange.
+ *	- Current			- The key currently being blended.
+ *	- AfterCurrent		- The key after the current blended one. 
+ *	- LastBlended		- The last key in the range of blended keys. If current is the last key, this is AfterBlendRange.
+ *	- AfterBlendRange	- The key right after the range of blended keys. Same as LastBlended if there is no such key.
+ * The blend functions return the new Y value that the current key should have.
+ */
 namespace UE::TweeningUtils
 {
-TWEENINGUTILS_API double Blend_ControlsToTween(const FKeyBlendInfo& InBlendInfo);
+/**
+ * Blend values of -1 or 1 move all keys to be linearly interpolated between the heights of BeforeBlendRange or AfterBlendRange, respectively.
+ * @return The new Y value that the current key should have.
+ */
+TWEENINGUTILS_API double Blend_ControlsToTween(double InBlendValue, const FVector2d& InBeforeBlendRange, const FVector2d& InAfterBlendRange);
+	
+/**
+ * Blend values of -1 or 1 flatten or exaggerates valleys and hills on the function, respectively.  
+ * @return The new Y value that the current key should have.
+ */
+TWEENINGUTILS_API double Blend_PushPull(double InBlendValue,
+	const FVector2d& InBeforeBlendRange, const FVector2d& InCurrent, const FVector2d& InAfterBlendRange
+	);
+	
+/**
+ * Blend values of -1 or 1 gradually interpolate keys to the height of the BeforeBlendRange or AfterBlendRange, respectively, using linear interpolation.
+ *
+ * Similar to Blend_Ease, which uses a S-curve instead.
+ * Keys are moved up and down using linear interpolation - which comparatively moves much more suddenly.
+ * 
+ * @return The new Y value that the current key should have.
+ */
+TWEENINGUTILS_API double Blend_Neighbor(double InBlendValue,
+	const FVector2d& InBeforeBlendRange, const FVector2d& InCurrent, const FVector2d& InAfterBlendRange
+	);
+	
+/**
+ * Blend values of -1 or 1 uniformly shift all keys down or up such that the left-most or right-most keys match up with the
+ * key before or after the blend range, respectively.
+ * 
+ * @return The new Y value that the current key should have.
+ */
+TWEENINGUTILS_API double Blend_Relative(double InBlendValue, 
+	const FVector2d& InBeforeBlendRange, const FVector2d& InFirstBlended,
+	const FVector2d& InCurrent,
+	const FVector2d& InLastBlended, const FVector2d& InAfterBlendRange
+	);
+	
+/**
+ * Blend values of -1 or 1 gradually interpolate keys to the height of the BeforeBlendRange or AfterBlendRange, respectively, using a smooth S curve.
+ * @return The new Y value that the current key should have.
+ */
+TWEENINGUTILS_API double Blend_Ease(double InBlendValue, 
+	const FVector2d& InBeforeBlendRange, const FVector2d& InCurrent, const FVector2d& InAfterBlendRange
+	);
 
-TWEENINGUTILS_API double Blend_PushPull(const FKeyBlendInfo& InBlendInfo);
-
-TWEENINGUTILS_API double Blend_Neighbor(const FKeyBlendInfo& InBlendInfo);
-
-TWEENINGUTILS_API double Blend_Relative(const FKeyBlendInfo& InBlendInfo);
-
-TWEENINGUTILS_API double Blend_Ease(const FKeyBlendInfo& InBlendInfo);
-
-TWEENINGUTILS_API double Blend_SmoothRough(const FKeyBlendInfo& InBlendInfo);
+/**
+ * Blend values of -1 or 1, push adjacent blended keys further together or apart, respectively. -1 averages keys out while 1 increases jumps.
+ * @return The new Y value that the current key should have.
+ */
+TWEENINGUTILS_API double Blend_SmoothRough(double InBlendValue, 
+	const FVector2d& InBeforeCurrent, const FVector2d& InCurrent, const FVector2d& InAfterCurrent
+	);
 }
diff --git a/Engine/Plugins/Animation/TweeningUtils/Source/TweeningUtilsEditor/Private/Math/Abstraction/KeyBlendingAbstraction.cpp b/Engine/Plugins/Animation/TweeningUtils/Source/TweeningUtilsEditor/Private/Math/Abstraction/KeyBlendingAbstraction.cpp
--- a/Engine/Plugins/Animation/TweeningUtils/Source/TweeningUtilsEditor/Private/Math/Abstraction/KeyBlendingAbstraction.cpp
+++ b/Engine/Plugins/Animation/TweeningUtils/Source/TweeningUtilsEditor/Private/Math/Abstraction/KeyBlendingAbstraction.cpp
@@ -6,8 +6,6 @@
 #include "Internationalization/Internationalization.h"
 #include "Internationalization/Text.h"
 #include "Math/Color.h"
-#include "Math/KeyBlendInfo.h"
-#include "Math/KeyBlendingFunctions.h"
 #include "Misc/Optional.h"
 #include "Styling/SlateBrush.h"
 #include "Templates/SharedPointer.h"
@@ -20,10 +18,7 @@ namespace UE::TweeningUtilsEditor
 {
 struct FStaticBlendFunctionData
 {
-	using FBlendFunctionPtr = double(*)(const TweeningUtils::FKeyBlendInfo& InBlendInfo);
-	
 	const EBlendFunction Type;
-	const FBlendFunctionPtr EvaluationFunctionPtr;
 	const TWeakPtr<FUICommandInfo> Command; // Don't keep the command in a strong pointer case FTweeningUtilsCommands singleton gets destroyed
 	const FText Label;
 	const FText Abbreviation;
@@ -33,7 +28,6 @@ struct FStaticBlendFunctionData
 	
 	FStaticBlendFunctionData(
 		EBlendFunction InType,
-		FBlendFunctionPtr InEvaluationFunctionPtr,
 		TSharedPtr<FUICommandInfo> InCommand,
 		FText InLabel,
 		FText InAbbreviation,
@@ -41,7 +35,6 @@ struct FStaticBlendFunctionData
 		FString InStringEncoding
 		)
 		: Type(InType)
-		, EvaluationFunctionPtr(InEvaluationFunctionPtr)
 		, Command(MoveTemp(InCommand))
 		, Label(MoveTemp(InLabel))
 		, Abbreviation(MoveTemp(InAbbreviation))
@@ -72,42 +65,42 @@ static FStaticBlendFunctionData& GetFunctionData(EBlendFunction InFunction)
 	static FStaticBlendFunctionData Data[Num] =
 	{
 		FStaticBlendFunctionData(
-			EBlendFunction::BlendNeighbor, &TweeningUtils::Blend_Neighbor, FTweeningUtilsCommands::Get().SetTweenBlendNeighbor,
+			EBlendFunction::BlendNeighbor, FTweeningUtilsCommands::Get().SetTweenBlendNeighbor,
 			LOCTEXT("BlendNeighbor.Label", "Blend Neighbor"),
 			LOCTEXT("BlendNeighbour.Abbreviation", "BN"),
 			LOCTEXT("BlendNeighbour.Description", "Blend to the next or previous values for selected keys."),
 			TEXT("BlendNeighbour")
 		),
 		FStaticBlendFunctionData(
-			EBlendFunction::PushPull, &TweeningUtils::Blend_PushPull, FTweeningUtilsCommands::Get().SetTweenPushPull,
+			EBlendFunction::PushPull, FTweeningUtilsCommands::Get().SetTweenPushPull,
 			LOCTEXT("PushPull.Label", "Push / Pull"),
 			LOCTEXT("PushPull.Abbreviation", "PP"),
 			LOCTEXT("PushPull.Description", "Push or pull the values to the interpolation between the previous and next keys"),
 			TEXT("PushPull")
 		),
 		FStaticBlendFunctionData(
-			EBlendFunction::BlendEase, &TweeningUtils::Blend_Ease, FTweeningUtilsCommands::Get().SetTweenBlendEase,
+			EBlendFunction::BlendEase, FTweeningUtilsCommands::Get().SetTweenBlendEase,
 			LOCTEXT("BlendEase.Label", "Blend Ease"),
 			LOCTEXT("BlendEase.Abbreviation", "BE"),
 			LOCTEXT("BlendEase.Description", "Blend with an ease falloff to the next or previous value for selected keys"),
 			TEXT("BlendEase")
 		),
 		FStaticBlendFunctionData(
-			EBlendFunction::ControlsToTween, &TweeningUtils::Blend_ControlsToTween, FTweeningUtilsCommands::Get().SetControlsToTween,
+			EBlendFunction::ControlsToTween, FTweeningUtilsCommands::Get().SetControlsToTween,
 			LOCTEXT("ControlsToTween.Label", "Tween"),
 			LOCTEXT("ControlsToTween.Abbreviation", "TW"),
 			LOCTEXT("ControlsToTween.Description", "Interpolates between the previous and next keys"),
 			TEXT("Tween")
 		),
 		FStaticBlendFunctionData(
-			EBlendFunction::BlendRelative, &TweeningUtils::Blend_Relative, FTweeningUtilsCommands::Get().SetTweenBlendRelative,
+			EBlendFunction::BlendRelative, FTweeningUtilsCommands::Get().SetTweenBlendRelative,
 			LOCTEXT("BlendRelative.Label", "Blend Relative"),
 			LOCTEXT("BlendRelative.Abbreviation", "BR"),
 			LOCTEXT("BlendRelative.Description", "Blend relative to the next or previous value for selected keys"),
 			TEXT("BlendRelative")
 		),
 		FStaticBlendFunctionData(
-			EBlendFunction::SmoothRough, &TweeningUtils::Blend_SmoothRough, FTweeningUtilsCommands::Get().SetTweenSmoothRough,
+			EBlendFunction::SmoothRough, FTweeningUtilsCommands::Get().SetTweenSmoothRough,
 			LOCTEXT("SmoothRough.Label", "Smooth Rough"),
 			LOCTEXT("SmoothRough.Abbreviation", "SR"),
 			LOCTEXT("SmoothRough.Description", "Make selected keys smooth or rough"),
@@ -136,11 +129,6 @@ TOptional<EBlendFunction> LexBlendFunction(const FString& InString)
 	}
 	return {};
 }
-	
-double EvaluateBlend(EBlendFunction InFunction, const TweeningUtils::FKeyBlendInfo& InBlendInfo)
-{
-	return GetFunctionData(InFunction).EvaluationFunctionPtr(InBlendInfo);
-}
 
 TSharedPtr<FUICommandInfo> GetCommandForBlendFunction(EBlendFunction InFunction)
 {
diff --git a/Engine/Plugins/Animation/TweeningUtils/Source/TweeningUtilsEditor/Private/Math/Abstraction/TweenModelArray.cpp b/Engine/Plugins/Animation/TweeningUtils/Source/TweeningUtilsEditor/Private/Math/Abstraction/TweenModelArray.cpp
--- a/Engine/Plugins/Animation/TweeningUtils/Source/TweeningUtilsEditor/Private/Math/Abstraction/TweenModelArray.cpp
+++ b/Engine/Plugins/Animation/TweeningUtils/Source/TweeningUtilsEditor/Private/Math/Abstraction/TweenModelArray.cpp
@@ -9,74 +9,53 @@
 
 namespace UE::TweeningUtilsEditor
 {
-TOptional<EBlendFunction> FTweenModelArray::GetBlendFunctionType(const FTweenModel& InTweenModelToFind) const
-{
-	for (int32 Index = 0; Index < Models.size(); ++Index)
-	{
-		const TUniquePtr<FTweenModel>& TweenModel = Models[Index];
-		if (&InTweenModelToFind == TweenModel.Get())
-		{
-			return static_cast<EBlendFunction>(Index);
-		}
-	}
-	return {};
-}
-
-const FTweenModel* FTweenModelArray::GetBlendFunction(EBlendFunction InBlendFunction) const
-{
-	const int32 Index = static_cast<int32>(InBlendFunction);
-	return ensure(Index >= 0 && Index < Models.size()) ? Models[Index].Get() : nullptr;
-}
-
 void FTweenModelArray::ForEachModel(const TFunctionRef<void(FTweenModel&)> InConsumer) const
 {
-	for (const TUniquePtr<FTweenModel>& TweenModel : Models)
+	for (const FTweenModelUIEntry& DisplayInfo : TweenModels)
 	{
-		if (TweenModel) // Handle that case that the API user that constructed this FBuiltInTweenModelContainer, did not specify this built-in function.
-		{
-			InConsumer(*TweenModel.Get());
-		}
+		const TUniquePtr<FTweenModel>& TweenModel = DisplayInfo.TweenModel;
+		InConsumer(*TweenModel.Get());
 	}
 }
 
 FTweenModel* FTweenModelArray::GetModel(int32 InIndex) const
 {
-	return Models[InIndex].Get();
+	return TweenModels[InIndex].TweenModel.Get();
 }
 
-TSharedPtr<FUICommandInfo> FTweenModelArray::GetCommandForModel(const FTweenModel& InFunction) const
+TSharedPtr<FUICommandInfo> FTweenModelArray::GetCommandForModel(const FTweenModel& InModel) const
 {
-	const TOptional<EBlendFunction> Function = GetBlendFunctionType(InFunction);
-	return Function ? GetCommandForBlendFunction(*Function) : nullptr;
+	const int32 Index = IndexOf(InModel);
+	return TweenModels.IsValidIndex(Index) ? TweenModels[Index].DisplayInfo.Command : nullptr;
 }
 
-const FSlateBrush* FTweenModelArray::GetIconForModel(const FTweenModel& InFunction) const
+const FSlateBrush* FTweenModelArray::GetIconForModel(const FTweenModel& InModel) const
 {
-	const TOptional<EBlendFunction> Function = GetBlendFunctionType(InFunction);
-	return Function ? GetUntintedIconForTweenFunction(*Function) : nullptr;
+	const int32 Index = IndexOf(InModel);
+	return TweenModels.IsValidIndex(Index) ? TweenModels[Index].DisplayInfo.Brush : nullptr;
 }
 
-FLinearColor FTweenModelArray::GetColorForModel(const FTweenModel& InFunction) const
+FLinearColor FTweenModelArray::GetColorForModel(const FTweenModel& InModel) const
 {
-	const TOptional<EBlendFunction> Function = GetBlendFunctionType(InFunction);
-	return Function ? GetTintColorForTweenFunction(*Function) : FLinearColor::White;
+	const int32 Index = IndexOf(InModel);
+	return TweenModels.IsValidIndex(Index) ? TweenModels[Index].DisplayInfo.Color : FLinearColor::White;
 }
 
-FText FTweenModelArray::GetLabelForModel(const FTweenModel& InFunction) const
+FText FTweenModelArray::GetLabelForModel(const FTweenModel& InModel) const
 {
-	const TOptional<EBlendFunction> Function = GetBlendFunctionType(InFunction);
-	return Function ? GetLabelForBlendFunction(*Function) : FText::GetEmpty();
+	const int32 Index = IndexOf(InModel);
+	return TweenModels.IsValidIndex(Index) ? TweenModels[Index].DisplayInfo.Label : FText::GetEmpty();
 }
 	
-FText FTweenModelArray::GetToolTipForModel(const FTweenModel& InFunction) const
+FText FTweenModelArray::GetToolTipForModel(const FTweenModel& InModel) const
 {
-	const TOptional<EBlendFunction> Function = GetBlendFunctionType(InFunction);
-	return Function ? GetDescriptionForBlendFunction(*Function) : FText::GetEmpty();
+	const int32 Index = IndexOf(InModel);
+	return TweenModels.IsValidIndex(Index) ? TweenModels[Index].DisplayInfo.ToolTip : FText::GetEmpty();
 }
 
-FString FTweenModelArray::GetModelIdentifier(const FTweenModel& InFunction) const
+FString FTweenModelArray::GetModelIdentifier(const FTweenModel& InModel) const
 {
-	const TOptional<EBlendFunction> Function = GetBlendFunctionType(InFunction);
-	return Function ? BlendFunctionToString(*Function) : FString();
+	const int32 Index = IndexOf(InModel);
+	return TweenModels.IsValidIndex(Index) ? TweenModels[Index].DisplayInfo.Identifier : FString();
 }
 }
diff --git a/Engine/Plugins/Animation/TweeningUtils/Source/TweeningUtilsEditor/Private/Math/ContiguousKeyMapping.cpp b/Engine/Plugins/Animation/TweeningUtils/Source/TweeningUtilsEditor/Private/Math/ContiguousKeyMapping.cpp
--- a/Engine/Plugins/Animation/TweeningUtils/Source/TweeningUtilsEditor/Private/Math/ContiguousKeyMapping.cpp
+++ b/Engine/Plugins/Animation/TweeningUtils/Source/TweeningUtilsEditor/Private/Math/ContiguousKeyMapping.cpp
@@ -92,7 +92,7 @@ static void AppendKeyArray(
 				else //not contiguous, so need to add to the OutKeyMap
 				{
 					//now start new set with the new index
-					KeyArray.Add(ContiguousKeyIndices);
+					KeyArray.AddBlendRange(ContiguousKeyIndices);
 					ContiguousKeyIndices.Reset();
 					ContiguousKeyIndices.Add(AllKeysIndex);
 				}
@@ -104,7 +104,7 @@ static void AppendKeyArray(
 		}
 		if (ContiguousKeyIndices.Num() > 0)
 		{
-			KeyArray.Add(ContiguousKeyIndices);
+			KeyArray.AddBlendRange(ContiguousKeyIndices);
 			ContiguousKeyIndices.Reset();
 		}
 	}
@@ -121,28 +121,19 @@ static TMap<FCurveModelID, FContiguousKeyMapping::FContiguousKeysArray> ComputeK
 	return KeyMap;
 }
 
-FContiguousKeyMapping::FContiguousKeys::FContiguousKeys(
-	const TArray<FKeyHandle>& AllKeyHandles, const TArray<FVector2d>& AllKeyPositions, const TArray<int32>& InContiguousKeyIndices
-	)
-	: Indices(InContiguousKeyIndices)
-	, PreviousIndex(AllKeyHandles.IsValidIndex(Indices[0] - 1) ? Indices[0] - 1 : INDEX_NONE)
-	, NextIndex(AllKeyHandles.IsValidIndex(Indices[Indices.Num() -1] + 1) ? Indices[Indices.Num() -1] + 1 : INDEX_NONE)
+void FBlendRangesData::AddBlendRange(const TArray<int32>& InContiguousKeyIndices)
 {
-	check(AllKeyHandles.Num() > 0 && AllKeyPositions.Num() == AllKeyHandles.Num() && InContiguousKeyIndices.Num() > 0);
+	if (AllKeyPositions.Num() > 0 && InContiguousKeyIndices.Num() > 0)
+	{
+		KeysArray.Add(FContiguousKeys(AllKeyPositions, InContiguousKeyIndices));
+	}
 }
 
-FContiguousKeyMapping::FContiguousKeysArray::FContiguousKeysArray(
-	const TArray<FKeyHandle>& InAllKeyHandles, const TArray<FVector2d>& InAllKeyPositions
-	)
-	: AllKeyHandles(InAllKeyHandles)
-	, AllKeyPositions(InAllKeyPositions)
-{}
-
-void FContiguousKeyMapping::FContiguousKeysArray::Add(const TArray<int32>& InContiguousKeyIndices)
+void FContiguousKeyMapping::FContiguousKeysArray::AddBlendRange(const TArray<int32>& InContiguousKeyIndices)
 {
 	if (AllKeyHandles.Num() > 0 && AllKeyPositions.Num() == AllKeyHandles.Num() && InContiguousKeyIndices.Num() > 0)
 	{
-		KeysArray.Add(FContiguousKeys(AllKeyHandles, AllKeyPositions, InContiguousKeyIndices));
+		KeysArray.Add(FContiguousKeys(AllKeyPositions, InContiguousKeyIndices));
 	}
 }
 
diff --git a/Engine/Plugins/Animation/TweeningUtils/Source/TweeningUtilsEditor/Private/Math/Models/CurveTweeningModel.cpp b/Engine/Plugins/Animation/TweeningUtils/Source/TweeningUtilsEditor/Private/Math/Models/CurveTweeningModel.cpp
deleted file mode 100644
diff --git a/Engine/Plugins/Animation/TweeningUtils/Source/TweeningUtilsEditor/Public/Math/Abstraction/KeyBlendingAbstraction.h b/Engine/Plugins/Animation/TweeningUtils/Source/TweeningUtilsEditor/Public/Math/Abstraction/KeyBlendingAbstraction.h
--- a/Engine/Plugins/Animation/TweeningUtils/Source/TweeningUtilsEditor/Public/Math/Abstraction/KeyBlendingAbstraction.h
+++ b/Engine/Plugins/Animation/TweeningUtils/Source/TweeningUtilsEditor/Public/Math/Abstraction/KeyBlendingAbstraction.h
@@ -13,7 +13,6 @@ class FUICommandInfo;
 struct FLinearColor;
 struct FSlateBrush;
 template<typename T> struct TOptional;
-namespace UE::TweeningUtils { struct FKeyBlendInfo; }
 
 /** The following just contains some utils that is common in UI code. Tools can opt for their own description but using this helps with consistency. */
 namespace UE::TweeningUtilsEditor
@@ -47,9 +46,6 @@ enum class EBlendFunction : uint8
 TWEENINGUTILSEDITOR_API FString BlendFunctionToString(EBlendFunction InFunction);
 /** @return The EBlendFunction encoded by InString. */
 TWEENINGUTILSEDITOR_API TOptional<EBlendFunction> LexBlendFunction(const FString& InString);
-	
-/** Evaluates the function */
-TWEENINGUTILSEDITOR_API double EvaluateBlend(EBlendFunction InFunction, const TweeningUtils::FKeyBlendInfo& InBlendInfo);
 
 /** @return The command that is used to select the function in most UI. */
 TWEENINGUTILSEDITOR_API TSharedPtr<FUICommandInfo> GetCommandForBlendFunction(EBlendFunction InFunction);
diff --git a/Engine/Plugins/Animation/TweeningUtils/Source/TweeningUtilsEditor/Public/Math/Abstraction/TweenModelArray.h b/Engine/Plugins/Animation/TweeningUtils/Source/TweeningUtilsEditor/Public/Math/Abstraction/TweenModelArray.h
--- a/Engine/Plugins/Animation/TweeningUtils/Source/TweeningUtilsEditor/Public/Math/Abstraction/TweenModelArray.h
+++ b/Engine/Plugins/Animation/TweeningUtils/Source/TweeningUtilsEditor/Public/Math/Abstraction/TweenModelArray.h
@@ -3,37 +3,23 @@
 #pragma once
 
 #include "ITweenModelContainer.h"
-#include "KeyBlendingAbstraction.h"
 #include "Math/Models/TweenModel.h"
-#include "Templates/UniquePtr.h"
-
-#include <array>
-#include <type_traits>
+#include "TweenModelDisplayInfo.h"
+#include "Templates/UnrealTemplate.h"
 
 namespace UE::TweeningUtilsEditor
 {
-/** Contains all the functions from this module, i.e. those specified by EBlendFunction.  */
-class TWEENINGUTILSEDITOR_API FTweenModelArray : public ITweenModelContainer
+/** Backed by an array of tween models and their UI info. */
+class TWEENINGUTILSEDITOR_API FTweenModelArray : public ITweenModelContainer, public FNoncopyable
 {
 public:
 	
-	using FTweenArray = std::array<TUniquePtr<FTweenModel>, static_cast<int32>(EBlendFunction::Num)>;
-
-	explicit FTweenModelArray(FTweenArray InTweenModels) : Models(MoveTemp(InTweenModels)) {}
-	/** Util constructs the tween models using a factory function. */
-	template<typename TConstructTweenModel> requires std::is_invocable_r_v<TUniquePtr<FTweenModel>, TConstructTweenModel, EBlendFunction>
-	explicit FTweenModelArray(TConstructTweenModel&& InMakeTweenModel);
-
-	/** @return Gets the type of blend function the tween function corresponds to, or unset if you pass in a function not owned by this container. */
-	TOptional<EBlendFunction> GetBlendFunctionType(const FTweenModel& InTweenModelToFind) const;
-
-	/** @return The blend function identified by the enum. */
-	const FTweenModel* GetBlendFunction(EBlendFunction InBlendFunction) const;
+	explicit FTweenModelArray(TArray<FTweenModelUIEntry> InTweenModels) : TweenModels(MoveTemp(InTweenModels)) {}
 
 	//~ Begin FTweenModelContainer Interface
 	virtual void ForEachModel(const TFunctionRef<void(FTweenModel&)> InConsumer) const override;
 	virtual FTweenModel* GetModel(int32 InIndex) const override;
-	virtual int32 NumModels() const override { return Models.size(); }
+	virtual int32 NumModels() const override { return TweenModels.Num(); }
 	virtual TSharedPtr<FUICommandInfo> GetCommandForModel(const FTweenModel& InTweenModel) const override;
 	virtual const FSlateBrush* GetIconForModel(const FTweenModel& InTweenModel) const override;
 	virtual FLinearColor GetColorForModel(const FTweenModel& InTweenModel) const override;
@@ -44,22 +30,7 @@ public:
 
 private:
 
-	/** The built-in functions. */
-	const FTweenArray Models;
+	/** Info about the contained tween models. */
+	const TArray<FTweenModelUIEntry> TweenModels;
 };
-
-template <typename TConstructTweenModel> requires std::is_invocable_r_v<TUniquePtr<FTweenModel>, TConstructTweenModel, EBlendFunction>
-FTweenModelArray::FTweenModelArray(TConstructTweenModel&& InMakeTweenModel)
-	: FTweenModelArray([&InMakeTweenModel]
-	{
-		FTweenArray Array;
-		int32 Index = 0;
-		ForEachBlendFunction([&InMakeTweenModel, &Array, &Index](EBlendFunction BlendFunction)
-		{
-			Array[Index] = InMakeTweenModel(BlendFunction);
-			++Index;
-		});
-		return Array;
-	}())
-{}
 }
diff --git a/Engine/Plugins/Animation/TweeningUtils/Source/TweeningUtilsEditor/Public/Math/Abstraction/TweenModelDisplayInfo.h b/Engine/Plugins/Animation/TweeningUtils/Source/TweeningUtilsEditor/Public/Math/Abstraction/TweenModelDisplayInfo.h
new file mode 100644
--- /dev/null
+++ b/Engine/Plugins/Animation/TweeningUtils/Source/TweeningUtilsEditor/Public/Math/Abstraction/TweenModelDisplayInfo.h
@@ -0,0 +1,65 @@
+﻿// Copyright Epic Games, Inc. All Rights Reserved.
+
+#pragma once
+
+#include "Containers/UnrealString.h"
+#include "Internationalization/Text.h"
+#include "KeyBlendingAbstraction.h"
+#include "Math/Color.h"
+#include "Math/Models/TweenModel.h"
+#include "Templates/SharedPointer.h"
+#include "Templates/UniquePtr.h"
+#include "Templates/UnrealTemplate.h"
+
+class FUICommandInfo;
+
+namespace UE::TweeningUtilsEditor
+{
+class FTweenModel;
+
+/** Information about how a FTweenModel can be displayed in the UI. */
+struct FTweenModelDisplayInfo
+{
+	TSharedPtr<FUICommandInfo> Command;
+	const FSlateBrush* Brush;
+	FLinearColor Color;
+	FText Label;
+	FText ToolTip;
+	FString Identifier;
+
+	explicit FTweenModelDisplayInfo(
+		TSharedPtr<FUICommandInfo> Command, const FSlateBrush* Brush, const FLinearColor& Color,
+		const FText& Label, const FText& ToolTip, const FString& Identifier
+		)
+		: Command(MoveTemp(Command))
+		, Brush(Brush)
+		, Color(Color)
+		, Label(Label)
+		, ToolTip(ToolTip)
+		, Identifier(Identifier)
+	{}
+
+	explicit FTweenModelDisplayInfo(EBlendFunction InBlendFunction)
+		: FTweenModelDisplayInfo(
+			GetCommandForBlendFunction(InBlendFunction),
+			GetUntintedIconForTweenFunction(InBlendFunction),
+			GetTintColorForTweenFunction(InBlendFunction),
+			GetLabelForBlendFunction(InBlendFunction),
+			GetDescriptionForBlendFunction(InBlendFunction),
+			BlendFunctionToString(InBlendFunction)
+			)
+	{}
+};
+
+/** An association between a tween model and its UI display info. Useful for backing some UI, e.g. like combo button selection list. */
+struct FTweenModelUIEntry : FNoncopyable
+{
+	TUniquePtr<FTweenModel> TweenModel;
+	FTweenModelDisplayInfo DisplayInfo;
+
+	explicit FTweenModelUIEntry(TUniquePtr<FTweenModel> TweenModel, FTweenModelDisplayInfo DisplayInfo)
+		: TweenModel(MoveTemp(TweenModel))
+		, DisplayInfo(MoveTemp(DisplayInfo))
+	{}
+};
+}
diff --git a/Engine/Plugins/Animation/TweeningUtils/Source/TweeningUtilsEditor/Public/Math/Abstraction/TweenRangeTemplates.h b/Engine/Plugins/Animation/TweeningUtils/Source/TweeningUtilsEditor/Public/Math/Abstraction/TweenRangeTemplates.h
new file mode 100644
--- /dev/null
+++ b/Engine/Plugins/Animation/TweeningUtils/Source/TweeningUtilsEditor/Public/Math/Abstraction/TweenRangeTemplates.h
@@ -0,0 +1,176 @@
+﻿// Copyright Epic Games, Inc. All Rights Reserved.
+
+#pragma once
+
+#include "Delegates/IntegerSequence.h"
+#include "KeyBlendingAbstraction.h"
+#include "Math/ContiguousKeyMapping.h"
+#include "Math/KeyBlendingFunctions.h"
+
+namespace UE::TweeningUtilsEditor
+{
+/**
+ * Interface for applying "simple" blend function to range of keys.
+ * Template specializations extract the required parameters and passes them to the underlying blend function.
+ *
+ * This pattern is only possible with "simple" blend functions, i.e. those that only have the generic parameters described in KeyBlendingFunctions.h.
+ * For these simple functions, all you need to do is specialize TweenRange and your function should then show up automatically, provided tools are
+ * using ForEachCurveTweenable to go discover tween simple functions.
+ * 
+ * Some blend functions may require more arguments. For example, time offsetting / shifting would require caching underlying FCurveModel to be able
+ * to evaluate arbitrary function values between the keys, etc.
+ */
+template<EBlendFunction>
+double TweenRange(
+	double InBlendValue,
+	const FBlendRangesData& AllBlendedKeys, const FContiguousKeys& CurrentBlendRange,
+	int32 InCurrentKeyIndex
+); // Intentionally not implemented. This function is supposed to be specialiyed with the EBlendFunction types that support it.
+
+/**
+ * @return Whether this EBlendFunction supports simple tweening by using TweenRange.
+ * Some blend functions require more arguments, like a cached curve, etc.
+ */
+constexpr bool SupportsTweenRange(EBlendFunction BlendFunction);
+	
+/**
+ * Invokes InCallback for each EBlendFunction that TCurveTweener can be instantiated with.
+ * InCallback must be a templated lambda, e.g. []<EBlendFunction Function>(){}.
+ */
+template <typename TCallback>
+constexpr void ForEachCurveTweenable(TCallback&& InCallback);
+}
+
+namespace UE::TweeningUtilsEditor
+{
+constexpr bool SupportsTweenRange(EBlendFunction BlendFunction)
+{
+	// This function could use SFINAE but this way is easier to generate static_assert and remind developers to implement TweenRange with.
+	static_assert(
+		static_cast<int32>(EBlendFunction::Num) == 6,
+		"Does the blend function you added support TweenRange? If so, specialize TweenRange and update this function."
+		); // TweenRange is supported if all inputs to your tween function can be determined by one of the FContiguousKeyRanges::GetX functions.
+
+	switch (BlendFunction)
+	{
+	case EBlendFunction::BlendNeighbor: [[fallthrough]];
+	case EBlendFunction::PushPull: [[fallthrough]];
+	case EBlendFunction::BlendEase: [[fallthrough]];
+	case EBlendFunction::ControlsToTween: [[fallthrough]];
+	case EBlendFunction::BlendRelative: [[fallthrough]];
+	case EBlendFunction::SmoothRough:
+		return true;
+	default: return false;
+	}
+}
+
+template<>
+inline double TweenRange<EBlendFunction::ControlsToTween>(
+	double InBlendValue,
+	const FBlendRangesData& AllBlendedKeys, const FContiguousKeys& CurrentBlendRange,
+	int32 InCurrentKeyIndex // Fyi, this particular blend function does not care for the current key, i.e. InCurrentKeyIndex
+)
+{
+	return TweeningUtils::Blend_ControlsToTween(
+		InBlendValue, AllBlendedKeys.GetBeforeBlendRange(CurrentBlendRange), AllBlendedKeys.GetAfterBlendRange(CurrentBlendRange)
+		);
+}
+
+template<>
+inline double TweenRange<EBlendFunction::PushPull>(
+	double InBlendValue,
+	const FBlendRangesData& AllBlendedKeys, const FContiguousKeys& CurrentBlendRange,
+	int32 InCurrentKeyIndex
+)
+{
+	return TweeningUtils::Blend_PushPull(
+		InBlendValue,
+		AllBlendedKeys.GetBeforeBlendRange(CurrentBlendRange),
+		AllBlendedKeys.GetCurrent(CurrentBlendRange, InCurrentKeyIndex),
+		AllBlendedKeys.GetAfterBlendRange(CurrentBlendRange)
+		);
+}
+	
+template<>
+inline double TweenRange<EBlendFunction::BlendNeighbor>(
+	double InBlendValue,
+	const FBlendRangesData& AllBlendedKeys, const FContiguousKeys& CurrentBlendRange,
+	int32 InCurrentKeyIndex
+)
+{
+	return TweeningUtils::Blend_Neighbor(
+		InBlendValue,
+		AllBlendedKeys.GetBeforeBlendRange(CurrentBlendRange),
+		AllBlendedKeys.GetCurrent(CurrentBlendRange, InCurrentKeyIndex),
+		AllBlendedKeys.GetAfterBlendRange(CurrentBlendRange)
+		);
+}
+
+template<>
+inline double TweenRange<EBlendFunction::BlendRelative>(
+	double InBlendValue,
+	const FBlendRangesData& AllBlendedKeys, const FContiguousKeys& CurrentBlendRange,
+	int32 InCurrentKeyIndex
+)
+{
+	return TweeningUtils::Blend_Relative(
+		InBlendValue,
+		AllBlendedKeys.GetBeforeBlendRange(CurrentBlendRange),
+		AllBlendedKeys.GetFirstInBlendRange(CurrentBlendRange),
+		AllBlendedKeys.GetCurrent(CurrentBlendRange, InCurrentKeyIndex),
+		AllBlendedKeys.GetLastInBlendRange(CurrentBlendRange),
+		AllBlendedKeys.GetAfterBlendRange(CurrentBlendRange)
+		);
+}
+	
+template<>
+inline double TweenRange<EBlendFunction::BlendEase>(
+	double InBlendValue,
+	const FBlendRangesData& AllBlendedKeys, const FContiguousKeys& CurrentBlendRange,
+	int32 InCurrentKeyIndex
+)
+{
+	return TweeningUtils::Blend_Ease(
+		InBlendValue,
+		AllBlendedKeys.GetBeforeBlendRange(CurrentBlendRange),
+		AllBlendedKeys.GetCurrent(CurrentBlendRange, InCurrentKeyIndex),
+		AllBlendedKeys.GetAfterBlendRange(CurrentBlendRange)
+		);
+}
+
+template<>
+inline double TweenRange<EBlendFunction::SmoothRough>(
+	double InBlendValue,
+	const FBlendRangesData& AllBlendedKeys, const FContiguousKeys& CurrentBlendRange,
+	int32 InCurrentKeyIndex
+)
+{
+	return TweeningUtils::Blend_SmoothRough(
+		InBlendValue,
+		AllBlendedKeys.GetBeforeCurrent(CurrentBlendRange, InCurrentKeyIndex),
+		AllBlendedKeys.GetCurrent(CurrentBlendRange, InCurrentKeyIndex),
+		AllBlendedKeys.GetAfterCurrent(CurrentBlendRange, InCurrentKeyIndex)
+		);
+}
+	
+namespace Private
+{
+template <typename TCallback, int32... Indices>
+constexpr void ForEachCurveTweenableImpl(TCallback&& InCallback, TIntegerSequence<int32, Indices...>)
+{
+	// Fold expression... substitutes Indices for 1, 2, 3, etc., i.e. for each EBlendFunction entry.
+	((SupportsTweenRange(static_cast<EBlendFunction>(Indices))
+		? InCallback.template operator()<static_cast<EBlendFunction>(Indices)>() : void())
+	, ...);
+}
+}
+
+template <typename TCallback>
+constexpr void ForEachCurveTweenable(TCallback&& InCallback)
+{
+	Private::ForEachCurveTweenableImpl(
+		Forward<TCallback>(InCallback),
+		TMakeIntegerSequence<int32, static_cast<int32>(EBlendFunction::Num)>{}
+		);
+}
+}
\ No newline at end of file
diff --git a/Engine/Plugins/Animation/TweeningUtils/Source/TweeningUtilsEditor/Public/Math/ContiguousKeyMapping.h b/Engine/Plugins/Animation/TweeningUtils/Source/TweeningUtilsEditor/Public/Math/ContiguousKeyMapping.h
--- a/Engine/Plugins/Animation/TweeningUtils/Source/TweeningUtilsEditor/Public/Math/ContiguousKeyMapping.h
+++ b/Engine/Plugins/Animation/TweeningUtils/Source/TweeningUtilsEditor/Public/Math/ContiguousKeyMapping.h
@@ -7,46 +7,88 @@
 #include "Containers/Array.h"
 #include "Containers/Map.h"
 #include "HAL/Platform.h"
+#include "Misc/CoreMiscDefines.h"
 
 class FCurveEditor;
 
 namespace UE::TweeningUtilsEditor
 {
-/** Data structure for use by BlendCurves. */
-struct TWEENINGUTILSEDITOR_API FContiguousKeyMapping
+/** Unbroken chain of keys. */
+struct TWEENINGUTILSEDITOR_API FContiguousKeys
 {
-	// Set of keys in a row
-	struct FContiguousKeys
-	{
-		/** Indices of keys to blend */
-		const TArray<int32> Indices;
-		/** The index of the key before Indices[0]. INDEX_NONE if there is no previous key. */
-		const int32 PreviousIndex = INDEX_NONE;
-		/** The index of the key after Indidces' last element. INDEX_NONE if there is no next key. */
-		const int32 NextIndex = INDEX_NONE;
-		
-		explicit FContiguousKeys(
-			const TArray<FKeyHandle>& InAllKeyHandles, const TArray<FVector2d>& InAllKeyPositions, const TArray<int32>& ContiguousKeyIndices
-			);
-	};
+	/** Indices of keys to blend */
+	const TArray<int32> Indices;
+	/** The index of the key before Indices[0]. INDEX_NONE if there is no previous key. */
+	const int32 PreviousIndex = INDEX_NONE;
+	/** The index of the key after Indidces' last element. INDEX_NONE if there is no next key. */
+	const int32 NextIndex = INDEX_NONE;
+	
+	explicit FContiguousKeys(const TArray<FVector2d>& InAllKeyPositions, const TArray<int32>& InContiguousKeyIndices)
+		: Indices(InContiguousKeyIndices)
+		, PreviousIndex(InAllKeyPositions.IsValidIndex(Indices[0] - 1) ? Indices[0] - 1 : INDEX_NONE)
+		, NextIndex(InAllKeyPositions.IsValidIndex(Indices[Indices.Num() -1] + 1) ? Indices[Indices.Num() -1] + 1 : INDEX_NONE)
+		{}
+};
+
+/**
+ * Stores all key positions as they were before the blend, and all the chains of keys that are supposed to be blended.
+ * You can use this type independently from FContiguousKeyMapping. For example, you can use this to call TweeningUtilsEditor::TweenRange directly.
+ */
+struct TWEENINGUTILSEDITOR_API FBlendRangesData
+{
+	const TArray<FVector2d> AllKeyPositions;
+
+	/**
+	 * The user selected keys split up into consecutive ranges.
+	 * 
+	 * Let 'x' denote keys selected by the user and '.' keys not selected.
+	 * Suppose the user selection was '.xx...xxx.', i.e. the user selected keys and then shift+selected additional keys somewhere else on the curve.
+	 * Then this array would contain indices of the two ranges 'xx' and 'xxx' is not a blendable range.
+	 */
+	TArray<FContiguousKeys> KeysArray;
+
+	explicit FBlendRangesData(const TArray<FVector2d>& InAllKeyPositions) : AllKeyPositions(InAllKeyPositions) {}
 
-	// All the selected keys in a row for a particular curve and pointers to all of the keys
-	struct FContiguousKeysArray
+	/** @return The key before the blend range defined by InKeys */
+	const FVector2d& GetBeforeBlendRange(const FContiguousKeys& InKeys) const;
+	/** @retuirn The first key that is blended in the specified range. */
+	const FVector2D& GetFirstInBlendRange(const FContiguousKeys& InKeys) const;
+
+	/**
+	 * @param InIndex Index to InKeys.Indices.
+	 * @return Gets the key before the currently blended key. If the current key is the first in the blend range, this returns GetBeforeBlendRange.
+	 */
+	const FVector2D& GetBeforeCurrent(const FContiguousKeys& InKeys, int32 InIndex) const;
+	/**
+	 * @param InIndex Index to InKeys.Indices.
+	 * @return AllKeyPositions[Indices[Index]] */
+	const FVector2D& GetCurrent(const FContiguousKeys& InKeys, int32 InIndex) const;
+	/**
+	 * @param InIndex Index to InKeys.Indices.
+	 * @return Gets the key after the currently blended key. If the current key is the last in the blend range, this returns GetAfterBlendRange.
+	 */
+	const FVector2D& GetAfterCurrent(const FContiguousKeys& InKeys, int32 InIndex) const;
+	
+	/** @retuirn The last key that is blended in the specified range. */
+	const FVector2D& GetLastInBlendRange(const FContiguousKeys& InKeys) const;
+	/** @return The key after the blend range defined by InKeys */
+	const FVector2d& GetAfterBlendRange(const FContiguousKeys& InKeys) const;
+
+	void AddBlendRange(const TArray<int32>& InContiguousKeyIndices);
+};
+	
+/** Information for blending a map of curves and the keys in it to blend. Saves the old key values for the duration of the blend.*/
+struct TWEENINGUTILSEDITOR_API FContiguousKeyMapping
+{
+	/** Holds the corresponding key handles for all the keys. */
+	struct TWEENINGUTILSEDITOR_API FContiguousKeysArray : FBlendRangesData
 	{
 		const TArray<FKeyHandle> AllKeyHandles;
-		const TArray<FVector2d> AllKeyPositions;
-
-		/**
-		 * The user selected keys split up into consecutive ranges.
-		 * 
-		 * Let 'x' denote keys selected by the user and '.' keys not selected.
-		 * Suppose the user selection was '.xx...xxx.', i.e. the user selected keys and then shift+selected additional keys somewhere else on the curve.
-		 * Then this array would contain indices of the two ranges 'xx' and 'xxx' is not a blendable range.
-		 */
-		TArray<FContiguousKeys> KeysArray;
-
-		explicit FContiguousKeysArray(const TArray<FKeyHandle>& InAllKeyHandle, const TArray<FVector2d>& InAllKeyPositions);
-		void Add(const TArray<int32>& ContiguousKeyIndices);
+		
+		explicit FContiguousKeysArray(const TArray<FKeyHandle>& InAllKeyHandles, const TArray<FVector2d>& InAllKeyPositions)
+			: FBlendRangesData(InAllKeyPositions) , AllKeyHandles(InAllKeyHandles) {}
+		
+		void AddBlendRange(const TArray<int32>& ContiguousKeyIndices);
 	};
 
 	FContiguousKeyMapping() = default;
@@ -58,3 +100,43 @@ struct TWEENINGUTILSEDITOR_API FContiguousKeyMapping
 	TMap<FCurveModelID, FContiguousKeysArray> KeyMap;
 };
 }
+
+namespace UE::TweeningUtilsEditor
+{
+inline const FVector2d& FBlendRangesData::GetBeforeBlendRange(const FContiguousKeys& InKeys) const
+{
+	const int32 BeforeIndex = InKeys.PreviousIndex != INDEX_NONE ? InKeys.PreviousIndex : InKeys.Indices[0];
+	return AllKeyPositions[BeforeIndex];
+}
+	
+inline const FVector2D& FBlendRangesData::GetFirstInBlendRange(const FContiguousKeys& InKeys) const
+{
+	return AllKeyPositions[InKeys.Indices[0]];
+}
+	
+inline const FVector2D& FBlendRangesData::GetBeforeCurrent(const FContiguousKeys& InKeys, int32 InIndex) const
+{
+	return InIndex > 0 ? AllKeyPositions[InKeys.Indices[InIndex]] : GetBeforeBlendRange(InKeys);
+}
+
+inline const FVector2D& FBlendRangesData::GetCurrent(const FContiguousKeys& InKeys, int32 InIndex) const
+{
+	return AllKeyPositions[InKeys.Indices[InIndex]];
+}
+
+inline const FVector2D& FBlendRangesData::GetAfterCurrent(const FContiguousKeys& InKeys, int32 InIndex) const
+{
+	return InIndex < InKeys.Indices.Num() - 1 ? AllKeyPositions[InKeys.Indices[InIndex + 1]] : GetAfterBlendRange(InKeys);
+}
+	
+inline const FVector2D& FBlendRangesData::GetLastInBlendRange(const FContiguousKeys& InKeys) const
+{
+	return AllKeyPositions[InKeys.Indices[InKeys.Indices.Num() - 1]];
+}
+
+inline const FVector2d& FBlendRangesData::GetAfterBlendRange(const FContiguousKeys& InKeys) const
+{
+	const int32 AfterIndex = InKeys.NextIndex != INDEX_NONE ? InKeys.NextIndex : InKeys.Indices[InKeys.Indices.Num() - 1];
+	return AllKeyPositions[AfterIndex];
+}
+}
\ No newline at end of file
diff --git a/Engine/Plugins/Animation/TweeningUtils/Source/TweeningUtilsEditor/Public/Math/CurveBlending.h b/Engine/Plugins/Animation/TweeningUtils/Source/TweeningUtilsEditor/Public/Math/CurveBlending.h
--- a/Engine/Plugins/Animation/TweeningUtils/Source/TweeningUtilsEditor/Public/Math/CurveBlending.h
+++ b/Engine/Plugins/Animation/TweeningUtils/Source/TweeningUtilsEditor/Public/Math/CurveBlending.h
@@ -6,31 +6,75 @@
 #include "CurveDataAbstraction.h"
 #include "CurveEditor.h"
 #include "CurveModel.h"
-#include "Math/KeyBlendInfo.h"
 
 #include <type_traits>
 
 namespace UE::TweeningUtilsEditor
 {
 /**
- * Invokes callback InDoBlend on all keys in InKeySelection.
- *
- * @note This function will call Modify on each curve.
+ * A callable that knows how to blend a key at a time.
+ * Use AllBlendedKeys.AllKeyPositions[CurrentBlendRange.Indices[Index]] to get current key
+ */
+template<typename TCallback>
+concept CBlendKeyCallable = std::is_invocable_r_v<double, TCallback,
+	const FCurveModelID& /*BlendedCurve*/,
+	const FContiguousKeyMapping::FContiguousKeysArray& /*AllBlendedKeys*/,
+	const FContiguousKeys& /*CurrentBlendRange*/,
+	int32 /*Index, i.e. the index into FContiguousKeys::Indices currently being blended.*/
+	>;
+/** Blends the specified keys by using a InBlendKeyCallback, which returns the new key value each key should have. */
+template<CBlendKeyCallable TCallback>
+bool BlendCurves_BySingleKey(const FCurveEditor& InCurveEditor, const FContiguousKeyMapping& InKeySelection, TCallback&& InBlendKeyCallback);
+	
+/**
+ * A callable that knows how to blend an entire range of keys at a time.
  * 
- * @param InBlendValue The blend value to pass to InDoBlend
- * @param InCurveEditor Used to look up curves.
- * @param InKeySelection The selection to blend
- * @param InDoBlend	Does the blending
- * @return Whether InDoBlend was invokes at least once.
+ * Blend values are written out to OutKeyHandles & OutKeyPositions, which are preallocated to be
+ * OutKeyHandles.Num() == OutKeyPositions.Num() == CurrentBlendRange.Indices.Num().
  */
-template<typename TCallback> requires std::is_invocable_r_v<double, TCallback, const TweeningUtils::FKeyBlendInfo&>
-bool BlendCurves(float InBlendValue, const FCurveEditor& InCurveEditor, const FContiguousKeyMapping& InKeySelection, TCallback&& InDoBlend);
+template<typename TCallback>
+concept CBlendRangeCallback = std::is_invocable_v<TCallback,
+	const FCurveModelID& /*BlendedCurve*/,
+	const FContiguousKeyMapping::FContiguousKeysArray& /*AllBlendedKeys*/,
+	const FContiguousKeys& /*CurrentBlendRange*/,
+	TArray<FKeyHandle>& /*OutKeyHandles*/,
+	TArray<FKeyPosition>& /*OutKeyPositions*/
+	>;
+/**
+ * Generic helper function that invokes InBlendRangeCallback for each unbroken sub-range of keys in each curve.
+ * @return Whether any key value was changed.
+ */
+template<CBlendRangeCallback TCallback>
+bool BlendCurve_ByKeyRange(const FCurveEditor& InCurveEditor, const FContiguousKeyMapping& InKeySelection, TCallback&& InBlendRangeCallback);
 }
 
 namespace UE::TweeningUtilsEditor
 {
-template<typename TCallback> requires std::is_invocable_r_v<double, TCallback, const TweeningUtils::FKeyBlendInfo&>
-bool BlendCurves(float InBlendValue, const FCurveEditor& InCurveEditor, const FContiguousKeyMapping& InKeySelection, TCallback&& InDoBlend)
+template <CBlendKeyCallable TCallback>
+bool BlendCurves_BySingleKey(const FCurveEditor& InCurveEditor, const FContiguousKeyMapping& InKeySelection, TCallback&& InBlendKeyCallback)
+{
+	return BlendCurve_ByKeyRange(InCurveEditor, InKeySelection,
+		[&InBlendKeyCallback](
+			const FCurveModelID& CurveId,
+			const FContiguousKeyMapping::FContiguousKeysArray& AllBlendedKeys, const FContiguousKeys& CurrentBlendRange,
+			TArray<FKeyHandle>& OutKeyHandles, TArray<FKeyPosition>& OutKeyPositions
+		)
+		{
+			const TConstArrayView<FKeyHandle>& AllKeyHandles = AllBlendedKeys.AllKeyHandles;
+			const TConstArrayView<int32> Indices = CurrentBlendRange.Indices;
+
+			for (int32 Index = 0; Index < Indices.Num(); ++Index)
+			{
+				const double NewValue = InBlendKeyCallback(CurveId, AllBlendedKeys, CurrentBlendRange, Index);
+				const FVector2D& CurrentKey = AllBlendedKeys.GetCurrent(CurrentBlendRange, Index);
+				OutKeyPositions[Index] = FKeyPosition(CurrentKey.X, NewValue);
+				OutKeyHandles[Index] = AllKeyHandles[Indices[Index]];
+			}
+		});
+}
+
+template<CBlendRangeCallback TCallback>
+bool BlendCurve_ByKeyRange(const FCurveEditor& InCurveEditor, const FContiguousKeyMapping& InKeySelection, TCallback&& InBlendRangeCallback)
 {
 	bool bDidBlend = false;
 	
@@ -38,37 +82,22 @@ bool BlendCurves(float InBlendValue, const FCurveEditor& InCurveEditor, const FC
 	TArray<FKeyPosition> KeyPositions;
 	for (const TPair<FCurveModelID, FContiguousKeyMapping::FContiguousKeysArray>& KeysArray : InKeySelection.KeyMap)
 	{
-		FCurveModel* Curve = InCurveEditor.FindCurve(KeysArray.Key);
+		const FCurveModelID& ModelId = KeysArray.Key;
+		FCurveModel* Curve = InCurveEditor.FindCurve(ModelId);
 		if (!Curve)
 		{
 			continue;
 		}
 		Curve->Modify();
 		
-		const TConstArrayView<FKeyHandle>& AllKeyHandles = KeysArray.Value.AllKeyHandles;
-		const TConstArrayView<FVector2d>& AllKeyPositions = KeysArray.Value.AllKeyPositions;
-		for (const FContiguousKeyMapping::FContiguousKeys& Keys : KeysArray.Value.KeysArray)
+		const FContiguousKeyMapping::FContiguousKeysArray& CurveBlendedKeys = KeysArray.Value;
+		for (const FContiguousKeys& BlendRange : CurveBlendedKeys.KeysArray)
 		{
-			const TConstArrayView<int32>& Indices = Keys.Indices;
-				
-			// Get the key before & after the blended keys
-			const int32 BeforeIndex = Keys.PreviousIndex != INDEX_NONE ? Keys.PreviousIndex : Keys.Indices[0];
-			const FVector2d& BeforeBlendRange = AllKeyPositions[BeforeIndex];
-			const int32 AfterIndex = Keys.NextIndex != INDEX_NONE ? Keys.NextIndex : Indices[Indices.Num() - 1];
-			const FVector2d& AfterBlendRange = AllKeyPositions[AfterIndex];
-
-			const int32 NumIndices = Indices.Num();
+			const int32 NumIndices = BlendRange.Indices.Num();
 			KeyHandles.SetNum(NumIndices);
 			KeyPositions.SetNum(NumIndices);
 
-			// Move through every blended key and blend it
-			TweeningUtils::FKeyBlendInfo::ForEachKey(InBlendValue, BeforeBlendRange, AfterBlendRange, AllKeyPositions, Indices,
-				[&InDoBlend, &KeyPositions, &KeyHandles, &AllKeyHandles, &Indices](int32 Index, const TweeningUtils::FKeyBlendInfo& BlendInfo)
-				{
-					const double NewValue = InDoBlend(BlendInfo);
-					KeyPositions[Index] = FKeyPosition(BlendInfo.Current.X, NewValue);
-					KeyHandles[Index] = AllKeyHandles[Indices[Index]];
-				});
+			InBlendRangeCallback(ModelId, CurveBlendedKeys, BlendRange, KeyHandles, KeyPositions);
 				
 			Curve->SetKeyPositions(KeyHandles, KeyPositions);
 			bDidBlend = true;
diff --git a/Engine/Plugins/Animation/TweeningUtils/Source/TweeningUtilsEditor/Public/Math/Models/CurveTweenModel.h b/Engine/Plugins/Animation/TweeningUtils/Source/TweeningUtilsEditor/Public/Math/Models/CurveTweenModel.h
--- a/Engine/Plugins/Animation/TweeningUtils/Source/TweeningUtilsEditor/Public/Math/Models/CurveTweenModel.h
+++ b/Engine/Plugins/Animation/TweeningUtils/Source/TweeningUtilsEditor/Public/Math/Models/CurveTweenModel.h
@@ -2,21 +2,24 @@
 
 #pragma once
 
-#include "TweenModel.h"
+#include "CurveEditor.h"
+#include "Math/Abstraction/TweenRangeTemplates.h"
 #include "Math/ContiguousKeyMapping.h"
-#include "Math/Abstraction/KeyBlendingAbstraction.h"
+#include "Math/CurveBlending.h"
+#include "TweenModel.h"
 #include "Templates/SharedPointer.h"
 
 class FCurveEditor;
 
 namespace UE::TweeningUtilsEditor
 {
-/** Tween the curves of the curve editor's selection. */
-class TWEENINGUTILSEDITOR_API FCurveTweenModel : public FTweenModel
+/** Tweens the curves of the curve editor's selection. */
+template<EBlendFunction>
+class TCurveTweenModel : public FTweenModel
 {
 public:
 
-	explicit FCurveTweenModel(TWeakPtr<FCurveEditor> InWeakCurveEditor, EBlendFunction InBlendFunction);
+	explicit TCurveTweenModel(TWeakPtr<FCurveEditor> InWeakCurveEditor);
 
 	/** @return Whether it makes sense to call StartBlendOperation, etc. */
 	bool HasAnythingToBlend() const;
@@ -31,11 +34,59 @@ protected:
 	
 	/** The curve editor on which to tween the curves. */
 	const TWeakPtr<FCurveEditor> WeakCurveEditor;
-	/** The blend function to use for blending keys. */
-	const EBlendFunction BlendFunction;
 
 	/** Created in StartBlendOperation and used for the entirety of the blend operation. */
 	FContiguousKeyMapping ContiguousKeySelection;
 };
 }
 
+namespace UE::TweeningUtilsEditor
+{
+template <EBlendFunction BlendFunction>
+TCurveTweenModel<BlendFunction>::TCurveTweenModel(TWeakPtr<FCurveEditor> InWeakCurveEditor)
+	: WeakCurveEditor(MoveTemp(InWeakCurveEditor))
+{}
+
+template <EBlendFunction BlendFunction>
+bool TCurveTweenModel<BlendFunction>::HasAnythingToBlend() const
+{
+	const TSharedPtr<FCurveEditor> CurveEditor = WeakCurveEditor.Pin();
+	return CurveEditor && !CurveEditor->GetSelection().IsEmpty();
+}
+
+template <EBlendFunction BlendFunction>
+void TCurveTweenModel<BlendFunction>::StartBlendOperation()
+{
+	const TSharedPtr<FCurveEditor> CurveEditor = WeakCurveEditor.Pin();
+	ContiguousKeySelection = CurveEditor ? FContiguousKeyMapping(*CurveEditor) : FContiguousKeyMapping{};
+}
+
+template <EBlendFunction BlendFunction>
+void TCurveTweenModel<BlendFunction>::StopBlendOperation()
+{
+	ContiguousKeySelection.KeyMap.Empty(
+		ContiguousKeySelection.KeyMap.Num()
+		);
+}
+
+template <EBlendFunction BlendFunction>
+void TCurveTweenModel<BlendFunction>::BlendValues(float InNormalizedValue)
+{
+	const TSharedPtr<FCurveEditor> CurveEditor = WeakCurveEditor.Pin();
+	if (!CurveEditor)
+	{
+		return;
+	}
+
+	const float ScaledBlendValue = ScaleBlendValue(InNormalizedValue); 
+	TweeningUtilsEditor::BlendCurves_BySingleKey(*CurveEditor, ContiguousKeySelection,
+		[this, ScaledBlendValue](
+			const FCurveModelID&,
+			const FContiguousKeyMapping::FContiguousKeysArray& AllBlendedKeys, const FContiguousKeys& CurrentBlendRange,
+			int32 InCurrentKeyIndex
+		)
+	{
+		return TweenRange<BlendFunction>(ScaledBlendValue, AllBlendedKeys, CurrentBlendRange, InCurrentKeyIndex);
+	});
+}
+}
\ No newline at end of file
diff --git a/Engine/Plugins/Editor/CurveEditorTools/Source/CurveEditorTools/CurveEditorTools.Build.cs b/Engine/Plugins/Editor/CurveEditorTools/Source/CurveEditorTools/CurveEditorTools.Build.cs
--- a/Engine/Plugins/Editor/CurveEditorTools/Source/CurveEditorTools/CurveEditorTools.Build.cs
+++ b/Engine/Plugins/Editor/CurveEditorTools/Source/CurveEditorTools/CurveEditorTools.Build.cs
@@ -27,6 +27,7 @@ public class CurveEditorTools : ModuleRules
 				"AudioMixer",
 				"SignalProcessing",
 				"EditorStyle", 
+				"TweeningUtils",
 				"TweeningUtilsEditor",
 			}
 		);
diff --git a/Engine/Plugins/Editor/CurveEditorTools/Source/CurveEditorTools/Private/Extensions/Tweening/ContextAwareTweeningModel.cpp b/Engine/Plugins/Editor/CurveEditorTools/Source/CurveEditorTools/Private/Extensions/Tweening/ContextAwareTweeningModel.cpp
deleted file mode 100644
diff --git a/Engine/Plugins/Editor/CurveEditorTools/Source/CurveEditorTools/Private/Extensions/Tweening/ContextAwareTweeningModel.h b/Engine/Plugins/Editor/CurveEditorTools/Source/CurveEditorTools/Private/Extensions/Tweening/ContextAwareTweeningModel.h
--- a/Engine/Plugins/Editor/CurveEditorTools/Source/CurveEditorTools/Private/Extensions/Tweening/ContextAwareTweeningModel.h
+++ b/Engine/Plugins/Editor/CurveEditorTools/Source/CurveEditorTools/Private/Extensions/Tweening/ContextAwareTweeningModel.h
@@ -6,6 +6,8 @@
 
 #include <type_traits>
 
+#include "KeyInsertingTweenModel.h"
+
 class FCurveEditor;
 template<typename ObjectType> class TAttribute;
 
@@ -17,16 +19,13 @@ namespace UE::CurveEditorTools
  * - When no keys are selected, we'll insert keys at the time controller's scrub position to all visible curves (those selected in tree view).
  *   If there's no such controller, nothing happens.
  */
-class FContextAwareTweeningModel : public TweeningUtilsEditor::FTweenModel
+template<TweeningUtilsEditor::EBlendFunction BlendFunction>
+class TContextAwareTweeningModel : public TweeningUtilsEditor::FTweenModel
 {
 	using Super = FTweenModel;
 public:
 	
-	/**
-	 * @param InWeakCurveEditor Used to get key info
-	 * @param InBlendFunction The blend function to blend with
-	 */
-	explicit FContextAwareTweeningModel(TWeakPtr<FCurveEditor> InWeakCurveEditor, TweeningUtilsEditor::EBlendFunction InBlendFunction);
+	explicit TContextAwareTweeningModel(TWeakPtr<FCurveEditor> InWeakCurveEditor);
 
 	//~ Begin FTweenModel Interface
 	virtual void SetScaleMode(TweeningUtilsEditor::ETweenScaleMode InMode) override;
@@ -41,7 +40,7 @@ private:
 	const TWeakPtr<FCurveEditor> CurveEditor;
 
 	/** Tweens when there is a key selection. */
-	TweeningUtilsEditor::FCurveTweenModel SelectionBasedModel;
+	TweeningUtilsEditor::TCurveTweenModel<BlendFunction> SelectionBasedModel;
 	/** Injected model that tweens when there is no key selection. Can be nullptr. */
 	const TSharedPtr<FTweenModel> NoKeySelectionModel;
 
@@ -50,3 +49,65 @@ private:
 };
 }
 
+
+namespace UE::CurveEditorTools
+{
+template <TweeningUtilsEditor::EBlendFunction BlendFunction>
+TContextAwareTweeningModel<BlendFunction>::TContextAwareTweeningModel(TWeakPtr<FCurveEditor> InWeakCurveEditor)
+	: CurveEditor(InWeakCurveEditor)
+	, SelectionBasedModel(InWeakCurveEditor)
+	, NoKeySelectionModel(MakeShared<TKeyInsertingTweenModel<BlendFunction>>(InWeakCurveEditor))
+{}
+
+template <TweeningUtilsEditor::EBlendFunction BlendFunction>
+void TContextAwareTweeningModel<BlendFunction>::SetScaleMode(TweeningUtilsEditor::ETweenScaleMode InMode)
+{
+	Super::SetScaleMode(InMode);
+	SelectionBasedModel.SetScaleMode(InMode);
+	if (NoKeySelectionModel)
+	{
+		NoKeySelectionModel->SetScaleMode(InMode);
+	}
+}
+
+template <TweeningUtilsEditor::EBlendFunction BlendFunction>
+void TContextAwareTweeningModel<BlendFunction>::StartBlendOperation()
+{
+	AccessCurrentModel([this](FTweenModel& Model)
+	{
+		Model.StartBlendOperation();
+	});
+}
+
+template <TweeningUtilsEditor::EBlendFunction BlendFunction>
+void TContextAwareTweeningModel<BlendFunction>::StopBlendOperation()
+{
+	AccessCurrentModel([](FTweenModel& Model)
+	{
+		Model.StopBlendOperation();
+	});
+}
+
+template <TweeningUtilsEditor::EBlendFunction BlendFunction>
+void TContextAwareTweeningModel<BlendFunction>::BlendValues(float InNormalizedValue)
+{
+	AccessCurrentModel([InNormalizedValue](FTweenModel& Model)
+	{
+		Model.BlendValues(InNormalizedValue);
+	});
+}
+
+template <TweeningUtilsEditor::EBlendFunction BlendFunction>
+template <typename TCallback> requires std::is_invocable_v<TCallback, TweeningUtilsEditor::FTweenModel&>
+void TContextAwareTweeningModel<BlendFunction>::AccessCurrentModel(TCallback&& InCallback)
+{
+	if (SelectionBasedModel.HasAnythingToBlend())
+	{
+		InCallback(SelectionBasedModel);
+	}
+	else if (NoKeySelectionModel)
+	{
+		InCallback(*NoKeySelectionModel.Get());
+	}
+}
+}
\ No newline at end of file
diff --git a/Engine/Plugins/Editor/CurveEditorTools/Source/CurveEditorTools/Private/Extensions/Tweening/CurveEditorTweenModels.cpp b/Engine/Plugins/Editor/CurveEditorTools/Source/CurveEditorTools/Private/Extensions/Tweening/CurveEditorTweenModels.cpp
new file mode 100644
--- /dev/null
+++ b/Engine/Plugins/Editor/CurveEditorTools/Source/CurveEditorTools/Private/Extensions/Tweening/CurveEditorTweenModels.cpp
@@ -0,0 +1,35 @@
+﻿// Copyright Epic Games, Inc. All Rights Reserved.
+
+#include "CurveEditorTweenModels.h"
+
+#include "ContextAwareTweeningModel.h"
+#include "Math/Models/EditorTweenModel.h"
+
+namespace UE::CurveEditorTools
+{
+namespace TweenDetail
+{
+static TArray<TweeningUtilsEditor::FTweenModelUIEntry> MakeTweenModels(TWeakPtr<FCurveEditor> InCurveEditor)
+{
+	using namespace TweeningUtilsEditor;
+	
+	TArray<FTweenModelUIEntry> Result;
+	ForEachCurveTweenable([&InCurveEditor, &Result]<EBlendFunction TBlendFunction>
+	{
+		using FModel = TEditorTweenModel<TContextAwareTweeningModel<TBlendFunction>>;
+
+		TUniquePtr<FModel> TweenModel = MakeUnique<FModel>(
+			InCurveEditor,			// TTangentFlatteningTweenProxy arg
+			InCurveEditor			// FContextAwareTweeningModel args
+			);
+		Result.Emplace(MoveTemp(TweenModel), FTweenModelDisplayInfo(TBlendFunction));
+	});
+
+	return Result;
+}
+}
+	
+FCurveEditorTweenModels::FCurveEditorTweenModels(TWeakPtr<FCurveEditor> InCurveEditor)
+	: FTweenModelArray(TweenDetail::MakeTweenModels(InCurveEditor))
+{}
+}
diff --git a/Engine/Plugins/Editor/CurveEditorTools/Source/CurveEditorTools/Private/Extensions/Tweening/CurveEditorTweenModels.h b/Engine/Plugins/Editor/CurveEditorTools/Source/CurveEditorTools/Private/Extensions/Tweening/CurveEditorTweenModels.h
new file mode 100644
--- /dev/null
+++ b/Engine/Plugins/Editor/CurveEditorTools/Source/CurveEditorTools/Private/Extensions/Tweening/CurveEditorTweenModels.h
@@ -0,0 +1,19 @@
+﻿// Copyright Epic Games, Inc. All Rights Reserved.
+
+#pragma once
+
+#include "Math/Abstraction/TweenModelArray.h"
+
+class FCurveEditor;
+
+namespace UE::CurveEditorTools
+{
+/** Knows of the tween models used in curve editor. */
+class FCurveEditorTweenModels : public TweeningUtilsEditor::FTweenModelArray
+{
+public:
+
+	explicit FCurveEditorTweenModels(const TWeakPtr<FCurveEditor> InCurveEditor);
+};
+}
+
diff --git a/Engine/Plugins/Editor/CurveEditorTools/Source/CurveEditorTools/Private/Extensions/Tweening/KeyInsertingTweenModel.cpp b/Engine/Plugins/Editor/CurveEditorTools/Source/CurveEditorTools/Private/Extensions/Tweening/KeyInsertingTweenModel.cpp
deleted file mode 100644
diff --git a/Engine/Plugins/Editor/CurveEditorTools/Source/CurveEditorTools/Private/Extensions/Tweening/KeyInsertingTweenModel.h b/Engine/Plugins/Editor/CurveEditorTools/Source/CurveEditorTools/Private/Extensions/Tweening/KeyInsertingTweenModel.h
--- a/Engine/Plugins/Editor/CurveEditorTools/Source/CurveEditorTools/Private/Extensions/Tweening/KeyInsertingTweenModel.h
+++ b/Engine/Plugins/Editor/CurveEditorTools/Source/CurveEditorTools/Private/Extensions/Tweening/KeyInsertingTweenModel.h
@@ -2,6 +2,7 @@
 
 #pragma once
 
+#include "HAL/Platform.h"
 #include "Math/Models/CurveTweenModel.h"
 
 class FCurveEditor;
@@ -9,11 +10,13 @@ class FCurveEditor;
 namespace UE::CurveEditorTools
 {
 /** Inserts keys of visible curves into the current scrubber position. */
-class FKeyInsertingTweenModel : public TweeningUtilsEditor::FCurveTweenModel
+template<TweeningUtilsEditor::EBlendFunction BlendFunction>
+class TKeyInsertingTweenModel : public TweeningUtilsEditor::TCurveTweenModel<BlendFunction>
 {
+	using Super = TweeningUtilsEditor::TCurveTweenModel<BlendFunction>;
 public:
-	
-	explicit FKeyInsertingTweenModel(TWeakPtr<FCurveEditor> InWeakCurveEditor, TweeningUtilsEditor::EBlendFunction InBlendFunction);
+
+	explicit TKeyInsertingTweenModel(TWeakPtr<FCurveEditor> InWeakCurveEditor) : Super(MoveTemp(InWeakCurveEditor)) {}
 
 	//~ Begin FTweenModel
 	virtual void StartBlendOperation() override;
@@ -21,3 +24,62 @@ public:
 };
 }
 
+namespace UE::CurveEditorTools
+{
+template <TweeningUtilsEditor::EBlendFunction BlendFunction>
+void TKeyInsertingTweenModel<BlendFunction>::StartBlendOperation()
+{
+	const TSharedPtr<FCurveEditor> CurveEditor = this->WeakCurveEditor.Pin();
+	if (!ensure(CurveEditor))
+	{
+		return;
+	}
+
+	// Not all curve editors have a time slider controller
+	const TSharedPtr<ITimeSliderController> TimeSliderController = CurveEditor->GetTimeSliderController();
+	if (!TimeSliderController)
+	{
+		return;
+	}
+	
+	const double ScrubTime = TimeSliderController->GetTickResolution().AsSeconds(TimeSliderController->GetScrubPosition());
+	for (const TTuple<FCurveEditorTreeItemID, ECurveEditorTreeSelectionState>& Pair : CurveEditor->GetTreeSelection())
+	{
+		if (Pair.Value == ECurveEditorTreeSelectionState::None)
+		{
+			continue;
+		}
+		
+		const FCurveEditorTreeItem& TreeItem = CurveEditor->GetTreeItem(Pair.Key);
+		for (const FCurveModelID& CurveModelID : TreeItem.GetCurves())
+		{
+			FCurveModel* CurveModel = CurveEditor->FindCurve(CurveModelID);
+			if (!CurveModel)
+			{
+				return;
+			}
+
+			double ValueAtScrubTime;
+			if (!CurveModel->Evaluate(ScrubTime, ValueAtScrubTime))
+			{
+				continue;
+			}
+
+			// Try to insert the key with similar interpolation modes as its neighbours
+			FKeyAttributes KeyAttributes;
+			const TPair<ERichCurveInterpMode, ERichCurveTangentMode> Modes = CurveModel->GetInterpolationMode(ScrubTime, RCIM_Cubic, RCTM_SmartAuto);
+			KeyAttributes.SetInterpMode(Modes.Key);
+			KeyAttributes.SetTangentMode(Modes.Value);
+
+			const TOptional<FKeyHandle> KeyHandle = CurveModel->AddKey({ ScrubTime, ValueAtScrubTime }, KeyAttributes);
+			if (!KeyHandle)
+			{
+				continue;
+			}
+
+			// This will make Blend function blend the key we just inserted
+			this->ContiguousKeySelection.Append(*CurveEditor, CurveModelID, { *KeyHandle });
+		}
+	}
+}
+}
\ No newline at end of file
diff --git a/Engine/Plugins/Editor/CurveEditorTools/Source/CurveEditorTools/Private/Extensions/Tweening/TweenEditorExtension.cpp b/Engine/Plugins/Editor/CurveEditorTools/Source/CurveEditorTools/Private/Extensions/Tweening/TweenEditorExtension.cpp
--- a/Engine/Plugins/Editor/CurveEditorTools/Source/CurveEditorTools/Private/Extensions/Tweening/TweenEditorExtension.cpp
+++ b/Engine/Plugins/Editor/CurveEditorTools/Source/CurveEditorTools/Private/Extensions/Tweening/TweenEditorExtension.cpp
@@ -6,8 +6,6 @@
 #include "CurveEditor.h"
 #include "Framework/MultiBox/MultiBoxBuilder.h"
 #include "Framework/MultiBox/MultiBoxExtender.h"
-#include "Math/Abstraction/TweenModelArray.h"
-#include "Math/Models/EditorTweenModel.h"
 #include "Misc/Optional.h"
 #include "Widgets/MVC/TweenToolbarController.h"
 
@@ -15,16 +13,7 @@ namespace UE::CurveEditorTools
 {
 FTweenEditorExtension::FTweenEditorExtension(TWeakPtr<FCurveEditor> InCurveEditor)
 	: WeakCurveEditor(InCurveEditor)
-	, TweenModelContainer(
-		MakeShared<TweeningUtilsEditor::FTweenModelArray>([this](TweeningUtilsEditor::EBlendFunction BlendFunction)
-		{
-			using FModel = TweeningUtilsEditor::TEditorTweenModel<FContextAwareTweeningModel>;
-			return MakeUnique<FModel>(
-				WeakCurveEditor,							// TTangentFlatteningTweenProxy arg
-				WeakCurveEditor, BlendFunction			// FContextAwareTweeningModel args
-				);
-		})
-		)
+	, TweenModelContainer(MakeShared<FCurveEditorTweenModels>(InCurveEditor))
 {}
 
 void FTweenEditorExtension::BindCommands(TSharedRef<FUICommandList> InCommandList)
diff --git a/Engine/Plugins/Editor/CurveEditorTools/Source/CurveEditorTools/Private/Extensions/Tweening/TweenEditorExtension.h b/Engine/Plugins/Editor/CurveEditorTools/Source/CurveEditorTools/Private/Extensions/Tweening/TweenEditorExtension.h
--- a/Engine/Plugins/Editor/CurveEditorTools/Source/CurveEditorTools/Private/Extensions/Tweening/TweenEditorExtension.h
+++ b/Engine/Plugins/Editor/CurveEditorTools/Source/CurveEditorTools/Private/Extensions/Tweening/TweenEditorExtension.h
@@ -2,6 +2,7 @@
 
 #pragma once
 
+#include "CurveEditorTweenModels.h"
 #include "ICurveEditorExtension.h"
 #include "Misc/Optional.h"
 #include "Templates/SharedPointer.h"
@@ -31,7 +32,7 @@ private:
 	const TWeakPtr<FCurveEditor> WeakCurveEditor;
 
 	/** Holds the functions that ToolbarController shows: just the default, built-in ones.*/
-	const TSharedRef<TweeningUtilsEditor::FTweenModelArray> TweenModelContainer;
+	const TSharedRef<FCurveEditorTweenModels> TweenModelContainer;
 
 	/** Created by InitControllers. */
 	TOptional<TweeningUtilsEditor::FTweenControllers> TweenControllers;
