Control Rig: Improve SControlRigOutliner, only use full path for normal controls, add tooltip to show full path

#rnx

diff --git a/Engine/Plugins/Animation/ControlRig/Source/ControlRigEditor/Private/EditMode/SControlRigOutliner.cpp b/Engine/Plugins/Animation/ControlRig/Source/ControlRigEditor/Private/EditMode/SControlRigOutliner.cpp
--- a/Engine/Plugins/Animation/ControlRig/Source/ControlRigEditor/Private/EditMode/SControlRigOutliner.cpp
+++ b/Engine/Plugins/Animation/ControlRig/Source/ControlRigEditor/Private/EditMode/SControlRigOutliner.cpp
@@ -94,13 +94,28 @@ FText FMultiRigData::GetDisplayName(const FRigTreeDisplaySettings& InSettings) c
 			if (const URigHierarchy* Hierarchy = GetHierarchy())
 			{
 				const FRigElementKey& RigElementKey = Key.GetValue();
+				const FRigControlElement* ControlElement = Hierarchy->Find<FRigControlElement>(RigElementKey);
+				bool bIncludeNameSpace = true;
+				if(ControlElement)
+				{
+					// animation channels should not show their module name if they belong to a control
+					// in the same module.
+					if(ControlElement->Settings.AnimationType == ERigControlAnimationType::AnimationChannel)
+					{
+						const FRigElementKey ParentElementKey = Hierarchy->GetFirstParent(RigElementKey);
+						if(Hierarchy->GetModuleFName(ParentElementKey) == Hierarchy->GetModuleFName(RigElementKey))
+						{
+							bIncludeNameSpace = false;
+						}
+					}
+				}
 				
 				const FText DisplayNameForUI = Hierarchy->GetDisplayNameForUI(RigElementKey, !InSettings.bUseShortName);
 				if (!DisplayNameForUI.IsEmpty())
 				{
 					CachedDisplayName = DisplayNameForUI;
 				}
-				else if (const FRigControlElement* ControlElement = Hierarchy->Find<FRigControlElement>(RigElementKey))
+				else if (ControlElement)
 				{
 					if (!ControlElement->Settings.DisplayName.IsNone())
 					{
@@ -124,6 +139,15 @@ FText FMultiRigData::GetDisplayName(const FRigTreeDisplaySettings& InSettings) c
 	return *CachedDisplayName;
 }
 
+FText FMultiRigData::GetToolTipText() const
+{
+	if(Key.IsSet())
+	{
+		return FText::FromName(Key.GetValue().Name);
+	}
+	return FText();
+}
+
 bool FMultiRigData::operator== (const FMultiRigData & Other) const
 {
 	if (ControlRig.IsValid() == false && Other.ControlRig.IsValid() == false)
@@ -321,6 +345,7 @@ void SMultiRigHierarchyItem::Construct(const FArguments& InArgs, const TSharedRe
 			[
 				SAssignNew(InlineWidget, SInlineEditableTextBlock)
 				.Text(this, &SMultiRigHierarchyItem::GetDisplayName)
+				.ToolTipText(this, &SMultiRigHierarchyItem::GetToolTipText)
 				.MultiLine(false)
 				.ColorAndOpacity_Lambda([this]()
 				{
@@ -375,6 +400,11 @@ FText SMultiRigHierarchyItem::GetDisplayName() const
 	return (WeakRigTreeElement.Pin()->Data.GetDisplayName(Settings));
 }
 
+FText SMultiRigHierarchyItem::GetToolTipText() const
+{
+	return (WeakRigTreeElement.Pin()->Data.GetToolTipText());
+}
+
 TPair<const FSlateBrush*, FSlateColor> SMultiRigHierarchyItem::GetBrushForElementType(const URigHierarchy* InHierarchy, const FMultiRigData& InData)
 {
 	const FSlateBrush* Brush = nullptr;
diff --git a/Engine/Plugins/Animation/ControlRig/Source/ControlRigEditor/Private/EditMode/SControlRigOutliner.h b/Engine/Plugins/Animation/ControlRig/Source/ControlRigEditor/Private/EditMode/SControlRigOutliner.h
--- a/Engine/Plugins/Animation/ControlRig/Source/ControlRigEditor/Private/EditMode/SControlRigOutliner.h
+++ b/Engine/Plugins/Animation/ControlRig/Source/ControlRigEditor/Private/EditMode/SControlRigOutliner.h
@@ -92,6 +92,7 @@ public:
 	FMultiRigData(UControlRig* InControlRig, FName InModuleName) : ControlRig(InControlRig), ModuleName(InModuleName) {};
 	FText GetName() const;
 	FText GetDisplayName(const FRigTreeDisplaySettings& DisplaySettings) const;
+	FText GetToolTipText() const;
 	bool operator == (const FMultiRigData & Other) const;
 	bool IsValid() const;
 	URigHierarchy* GetHierarchy() const;
@@ -143,6 +144,7 @@ private:
 	FMultiRigTreeDelegates Delegates;
 
 	FText GetDisplayName() const;
+	FText GetToolTipText() const;
 	FReply OnGetSelectedClicked();
 
 };
