Fix an issue where changing the control rig settings placed into the level viewport toolbar were not correctly reflected in the viewport.

diff --git a/Engine/Plugins/Animation/ControlRig/Source/ControlRigEditor/Private/Editor/ControlRigViewportToolbarExtensions.cpp b/Engine/Plugins/Animation/ControlRig/Source/ControlRigEditor/Private/Editor/ControlRigViewportToolbarExtensions.cpp
--- a/Engine/Plugins/Animation/ControlRig/Source/ControlRigEditor/Private/Editor/ControlRigViewportToolbarExtensions.cpp
+++ b/Engine/Plugins/Animation/ControlRig/Source/ControlRigEditor/Private/Editor/ControlRigViewportToolbarExtensions.cpp
@@ -24,7 +24,9 @@
 #include "LevelEditor.h"
 #include "InteractiveToolManager.h"
 #include "EditorModeManager.h"
+#include "SEditorViewport.h"
 #include "Tools/EdModeInteractiveToolsContext.h"
+#include "ViewportToolbar/UnrealEdViewportToolbarContext.h"
 
 #define LOCTEXT_NAMESPACE "ControlRigViewportToolbar"
 
@@ -1098,15 +1100,21 @@ void UE::ControlRig::PopulateControlRigViewportToolbarShowSubmenu(const FName In
 			[](UToolMenu* AnimationShowFlagsSubmenu)
 			{
 				FToolMenuSection& UnnamedSection = AnimationShowFlagsSubmenu->FindOrAddSection(NAME_None);
-
+				
+				UUnrealEdViewportToolbarContext* Context = AnimationShowFlagsSubmenu->FindContext<UUnrealEdViewportToolbarContext>();
 				if (UControlRigEditModeSettings* const Settings = GetMutableDefault<UControlRigEditModeSettings>())
 				{
 					{
 						FUIAction Action;
 						Action.ExecuteAction = FExecuteAction::CreateLambda(
-							[Settings]() -> void
+							[Settings, Context]() -> void
 							{
 								Settings->bDisplayAxesOnSelection = !Settings->bDisplayAxesOnSelection;
+								UControlRigEditModeSettings::OnSettingsChange.Broadcast(Settings);
+								if (Context)
+								{
+									Context->RefreshViewport();
+								}
 							}
 						);
 						Action.GetActionCheckState = FGetActionCheckState::CreateLambda(
@@ -1123,7 +1131,7 @@ void UE::ControlRig::PopulateControlRigViewportToolbarShowSubmenu(const FName In
 							LOCTEXT("AxisOnSelectionTooltip", "Should we show axes for the selected elements"),
 							FSlateIcon(),
 							Action,
-							EUserInterfaceActionType::RadioButton
+							EUserInterfaceActionType::ToggleButton
 						);
 						UnnamedSection.AddEntry(Entry);
 
@@ -1131,9 +1139,14 @@ void UE::ControlRig::PopulateControlRigViewportToolbarShowSubmenu(const FName In
 					{
 						FUIAction Action;
 						Action.ExecuteAction = FExecuteAction::CreateLambda(
-							[Settings]() -> void
+							[Settings, Context]() -> void
 							{
 								Settings->bDisplayHierarchy = !Settings->bDisplayHierarchy;
+								UControlRigEditModeSettings::OnSettingsChange.Broadcast(Settings);
+								if (Context)
+								{
+									Context->RefreshViewport();
+								}
 							}
 						);
 						Action.GetActionCheckState = FGetActionCheckState::CreateLambda(
@@ -1150,16 +1163,21 @@ void UE::ControlRig::PopulateControlRigViewportToolbarShowSubmenu(const FName In
 							LOCTEXT("HierarchyTooltip", "Whether to show all bones in the hierarchy"),
 							FSlateIcon(),
 							Action,
-							EUserInterfaceActionType::RadioButton
+							EUserInterfaceActionType::ToggleButton
 						);
 						UnnamedSection.AddEntry(Entry);
 					}
 					{
 						FUIAction Action;
 						Action.ExecuteAction = FExecuteAction::CreateLambda(
-							[Settings]() -> void
+							[Settings, Context]() -> void
 							{
 								Settings->bShowControlsAsOverlay = !Settings->bShowControlsAsOverlay;
+								UControlRigEditModeSettings::OnSettingsChange.Broadcast(Settings);
+								if (Context)
+								{
+									Context->RefreshViewport();
+								}
 							}
 						);
 						Action.GetActionCheckState = FGetActionCheckState::CreateLambda(
@@ -1176,16 +1194,21 @@ void UE::ControlRig::PopulateControlRigViewportToolbarShowSubmenu(const FName In
 							LOCTEXT("ControlsAsOverlaylTooltip", "Whether to show controls as overlay"),
 							FSlateIcon(),
 							Action,
-							EUserInterfaceActionType::RadioButton
+							EUserInterfaceActionType::ToggleButton
 						);
 						UnnamedSection.AddEntry(Entry);
 					}
 					{
 						FUIAction Action;
 						Action.ExecuteAction = FExecuteAction::CreateLambda(
-							[Settings]() -> void
+							[Settings, Context]() -> void
 							{
 								Settings->bHideControlShapes = !Settings->bHideControlShapes;
+								UControlRigEditModeSettings::OnSettingsChange.Broadcast(Settings);
+								if (Context)
+								{
+									Context->RefreshViewport();
+								}
 							}
 						);
 						Action.GetActionCheckState = FGetActionCheckState::CreateLambda(
@@ -1202,16 +1225,21 @@ void UE::ControlRig::PopulateControlRigViewportToolbarShowSubmenu(const FName In
 							LOCTEXT("ControlShapesTooltip", "Should we always hide control shapes in viewport"),
 							FSlateIcon(),
 							Action,
-							EUserInterfaceActionType::RadioButton
+							EUserInterfaceActionType::ToggleButton
 						);
 						UnnamedSection.AddEntry(Entry);
 					}
 					{
 						FUIAction Action;
 						Action.ExecuteAction = FExecuteAction::CreateLambda(
-							[Settings]() -> void
+							[Settings, Context]() -> void
 							{
 								Settings->bDisplayNulls = !Settings->bDisplayNulls;
+								UControlRigEditModeSettings::OnSettingsChange.Broadcast(Settings);
+								if (Context)
+								{
+									Context->RefreshViewport();
+								}
 							}
 						);
 						Action.GetActionCheckState = FGetActionCheckState::CreateLambda(
@@ -1228,16 +1256,21 @@ void UE::ControlRig::PopulateControlRigViewportToolbarShowSubmenu(const FName In
 							LOCTEXT("NullTooltip", "Whether to show all nulls in the hierarchy"),
 							FSlateIcon(),
 							Action,
-							EUserInterfaceActionType::RadioButton
+							EUserInterfaceActionType::ToggleButton
 						);
 						UnnamedSection.AddEntry(Entry);
 					}
 					{
 						FUIAction Action;
 						Action.ExecuteAction = FExecuteAction::CreateLambda(
-							[Settings]() -> void
+							[Settings, Context]() -> void
 							{
 								Settings->bShowAllProxyControls = !Settings->bShowAllProxyControls;
+								UControlRigEditModeSettings::OnSettingsChange.Broadcast(Settings);
+								if (Context)
+								{
+									Context->RefreshViewport();
+								}
 							}
 						);
 						Action.GetActionCheckState = FGetActionCheckState::CreateLambda(
@@ -1254,16 +1287,21 @@ void UE::ControlRig::PopulateControlRigViewportToolbarShowSubmenu(const FName In
 							LOCTEXT("ProxyControlsTooltip", "Whether to show Proxy Controls"),
 							FSlateIcon(),
 							Action,
-							EUserInterfaceActionType::RadioButton
+							EUserInterfaceActionType::ToggleButton
 						);
 						UnnamedSection.AddEntry(Entry);
 					}
 					{
 						FUIAction Action;
 						Action.ExecuteAction = FExecuteAction::CreateLambda(
-							[Settings]() -> void
+							[Settings, Context]() -> void
 							{
 								Settings->bDisplaySockets = !Settings->bDisplaySockets;
+								UControlRigEditModeSettings::OnSettingsChange.Broadcast(Settings);
+								if (Context)
+								{
+									Context->RefreshViewport();
+								}
 							}
 						);
 						Action.GetActionCheckState = FGetActionCheckState::CreateLambda(
@@ -1280,7 +1318,7 @@ void UE::ControlRig::PopulateControlRigViewportToolbarShowSubmenu(const FName In
 							LOCTEXT("SocketsTooltip", "Whether to show Sockets"),
 							FSlateIcon(),
 							Action,
-							EUserInterfaceActionType::RadioButton
+							EUserInterfaceActionType::ToggleButton
 						);
 						UnnamedSection.AddEntry(Entry);
 					}
diff --git a/Engine/Source/Editor/UnrealEd/Private/ViewportToolbar/UnrealEdViewportToolbarContext.cpp b/Engine/Source/Editor/UnrealEd/Private/ViewportToolbar/UnrealEdViewportToolbarContext.cpp
--- a/Engine/Source/Editor/UnrealEd/Private/ViewportToolbar/UnrealEdViewportToolbarContext.cpp
+++ b/Engine/Source/Editor/UnrealEd/Private/ViewportToolbar/UnrealEdViewportToolbarContext.cpp
@@ -17,6 +17,17 @@ TSharedPtr<IPreviewProfileController> UUnrealEdViewportToolbarContext::GetPrevie
 	return nullptr;
 }
 
+void UUnrealEdViewportToolbarContext::RefreshViewport()
+{
+	if (TSharedPtr<SEditorViewport> EditorViewport = Viewport.Pin())
+	{
+		if (TSharedPtr<FEditorViewportClient> Client = EditorViewport->GetViewportClient())
+		{
+			Client->Invalidate();
+		}
+	}
+}
+
 FText UUnrealEdViewportToolbarContext::GetGridSnapLabel() const
 {
 	return FText::AsNumber(GEditor->GetGridSize());
diff --git a/Engine/Source/Editor/UnrealEd/Public/ViewportToolbar/UnrealEdViewportToolbarContext.h b/Engine/Source/Editor/UnrealEd/Public/ViewportToolbar/UnrealEdViewportToolbarContext.h
--- a/Engine/Source/Editor/UnrealEd/Public/ViewportToolbar/UnrealEdViewportToolbarContext.h
+++ b/Engine/Source/Editor/UnrealEd/Public/ViewportToolbar/UnrealEdViewportToolbarContext.h
@@ -60,6 +60,8 @@ public:
 	
 	bool bShowSurfaceSnap = true;
 	
+	virtual void RefreshViewport();
+	
 	virtual FText GetGridSnapLabel() const;
 	virtual TArray<float> GetGridSnapSizes() const;
 	virtual bool IsGridSnapSizeActive(int32 GridSizeIndex) const;
