Control Rig: Fix crash when accessing attach parent actor on a shape actor
#rnx

diff --git a/Engine/Plugins/Animation/ControlRig/Source/ControlRigEditor/Private/EditMode/ControlRigEditMode.cpp b/Engine/Plugins/Animation/ControlRig/Source/ControlRigEditor/Private/EditMode/ControlRigEditMode.cpp
--- a/Engine/Plugins/Animation/ControlRig/Source/ControlRigEditor/Private/EditMode/ControlRigEditMode.cpp
+++ b/Engine/Plugins/Animation/ControlRig/Source/ControlRigEditor/Private/EditMode/ControlRigEditMode.cpp
@@ -5345,9 +5345,16 @@ void FControlRigEditMode::MoveControlShape( AControlRigShapeActor* ShapeActor,
 		// Update ShapeActor in case we are moving multiple shapes, which affect one another
 		const FTransform Transform = ControlRig->GetHierarchy()->GetTransform(ControlElement, ERigTransformType::CurrentGlobal);
 		ShapeActor->SetActorTransform(Transform * ToWorldTransform);
-		
-		const FTransform& ParentTransform = ShapeActor->GetAttachParentActor()->GetTransform();
-		CurrentTransform = ShapeActor->GetGlobalTransform() * ParentTransform;
+
+		if (const AActor* AttachParentActor = ShapeActor->GetAttachParentActor())
+		{
+			const FTransform& ParentTransform = AttachParentActor->GetTransform();
+			CurrentTransform = ShapeActor->GetGlobalTransform() * ParentTransform;
+		}
+		else
+		{
+			CurrentTransform = ShapeActor->GetGlobalTransform();
+		}
 	}
 	else
 	{
