Constraints & Modular rig

#rnx

diff --git a/Engine/Plugins/Animation/ControlRig/Source/ControlRig/Private/Constraints/ControlRigTransformableHandle.cpp b/Engine/Plugins/Animation/ControlRig/Source/ControlRig/Private/Constraints/ControlRigTransformableHandle.cpp
--- a/Engine/Plugins/Animation/ControlRig/Source/ControlRig/Private/Constraints/ControlRigTransformableHandle.cpp
+++ b/Engine/Plugins/Animation/ControlRig/Source/ControlRig/Private/Constraints/ControlRigTransformableHandle.cpp
@@ -8,6 +8,7 @@
 #include "ControlRigComponent.h"
 #include "IControlRigObjectBinding.h"
 #include "ControlRigObjectBinding.h"
+#include "ModularRig.h"
 #include "Transform/TransformableHandleUtils.h"
 #include "Rigs/RigHierarchyElements.h"
 #include "Sequencer/MovieSceneControlRigParameterSection.h"
@@ -288,13 +289,15 @@ bool UTransformableControlHandle::HasDirectDependencyWith(const UTransformableHa
 		return false;
 	}
 
+	UControlRig* ControlRigPtr = ControlRig.Get();
+	
 	// check whether the other handle is one of the control parent within the CR hierarchy
 	static constexpr bool bRecursive = true;
 	const URigHierarchy* Hierarchy = ControlRig->GetHierarchy();
 	const FRigBaseElementParentArray AllParents = Hierarchy->GetParents(ControlElement, bRecursive);
-	const bool bIsParent = AllParents.ContainsByPredicate([this, OtherHash](const FRigBaseElement* Parent)
+	const bool bIsParent = AllParents.ContainsByPredicate([ControlRigPtr, OtherHash](const FRigBaseElement* Parent)
 	{
-		const uint32 ParentHash = ComputeHash(ControlRig.Get(), Parent->GetFName());
+		const uint32 ParentHash = ComputeHash(ControlRigPtr, Parent->GetFName());
 		return ParentHash == OtherHash;		
 	});
 
@@ -305,24 +308,78 @@ bool UTransformableControlHandle::HasDirectDependencyWith(const UTransformableHa
 
 	// otherwise, check if there are any dependency in the graph that would cause a cycle
 	const TArray<FRigControlElement*> AllControls = Hierarchy->GetControls();
-	const int32 IndexOfPossibleParent = AllControls.IndexOfByPredicate([this, OtherHash](const FRigBaseElement* Parent)
+	const int32 IndexOfPossibleParent = AllControls.IndexOfByPredicate([ControlRigPtr, OtherHash](const FRigBaseElement* Parent)
 	{
-		const uint32 ChildHash = ComputeHash(ControlRig.Get(), Parent->GetFName());
+		const uint32 ChildHash = ComputeHash(ControlRigPtr, Parent->GetFName());
 		return ChildHash == OtherHash;
 	});
 
 	if (IndexOfPossibleParent != INDEX_NONE)
 	{
-		FRigControlElement* Parent = AllControls[IndexOfPossibleParent];
+		// at this point, we know that both handles belong to the same rig
+		FRigControlElement* PossibleParent = AllControls[IndexOfPossibleParent];
 
+		auto IsParentedTo = [ControlElement, PossibleParent](UControlRig* InControlRig, const URigHierarchy* InRigHierarchy)
+		{
+			if (InControlRig && InRigHierarchy)
+			{
 #if WITH_EDITOR
-		const FRigVMDependenciesProvider DependencyProvider(Hierarchy, ControlRig->GetVM());
-		const bool bIsParentedTo = Hierarchy->IsParentedTo(ControlElement, Parent, DependencyProvider);
+				const FRigVMDependenciesProvider SameModuleDependencyProvider(InRigHierarchy, InControlRig->GetVM());
+				return InRigHierarchy->IsParentedTo(ControlElement, PossibleParent, SameModuleDependencyProvider);
 #else
-		const bool bIsParentedTo = Hierarchy->IsParentedTo(ControlElement, Parent);
-#endif
+				return InRigHierarchy->IsParentedTo(ControlElement, PossibleParent);
+#endif			
+			}
+			return false;
+		};
 		
-		if (bIsParentedTo)
+		// modular test
+		if (UModularRig* ModularRig = Cast<UModularRig>(ControlRigPtr))
+		{
+			const FName ModuleName = Hierarchy->GetModuleFName(ControlElement->GetKey());
+			const FRigModuleInstance* ModuleInstance = ModuleName != NAME_None ? ModularRig->FindModule(ModuleName) : nullptr;
+
+			const FName ParentModuleName = Hierarchy->GetModuleFName(PossibleParent->GetKey());
+			const FRigModuleInstance* ParentModuleInstance = ParentModuleName != NAME_None ? ModularRig->FindModule(ParentModuleName) : nullptr;
+
+			if (ModuleInstance && ParentModuleInstance)
+			{
+				if (ModuleInstance == ParentModuleInstance)
+				{
+					// both handles are under the same module so check dependencies within that module
+					UControlRig* ModuleRig = ModuleInstance->GetRig();
+					const URigHierarchy* ModuleHierarchy = ModuleRig ? ModuleRig->GetHierarchy() : nullptr;
+					
+					const FRigBaseElementParentArray ParentParents = ModuleHierarchy->GetParents(PossibleParent, bRecursive);
+					if (ParentParents.Contains(ControlElement))
+					{
+						// if ControlElement is a parent of PossibleParent then it can't be one of its children
+						return false;		
+					}
+
+					// NOTE: we'd like to call this here but read URigHierarchy::GetDependenciesForVM about using this function with modular rigs
+					// if (IsParentedTo(ModuleRig, ModuleHierarchy))
+					// {
+					// 	return true;
+					// }
+				}
+				else
+				{
+					FRigModuleInstance* ParentModule = ModuleInstance->CachedParentModule;
+					while (ParentModule)
+					{
+						if (ParentModule == ParentModuleInstance)
+						{
+							return true;	
+						}
+						ParentModule = ParentModule->CachedParentModule; 
+					}
+				}
+			}
+		}
+
+		// default control rig test
+		if (IsParentedTo(ControlRigPtr, Hierarchy))
 		{
 			return true;
 		}
@@ -508,8 +565,10 @@ void UTransformableControlHandle::OnControlModified(
 					ControlHandleLocals::NotifyingRigs.Remove(ControlRig);
 				});
 				ControlHandleLocals::NotifyingRigs.Add(ControlRig.Get());
-				
-				const bool bPreTick = !ControlRig->IsAdditive();
+
+				const UTickableConstraint* TickableConstraint = GetTypedOuter<UTickableConstraint>();
+				const bool bIsConstraintActive = TickableConstraint && TickableConstraint->IsFullyActive();
+				const bool bPreTick = !ControlRig->IsAdditive() && bIsConstraintActive;
 				Notify(EHandleEvent::UpperDependencyUpdated, bPreTick);
 			}
 		}
diff --git a/Engine/Plugins/Animation/ControlRig/Source/ControlRig/Public/Constraints/ControlRigTransformableHandle.h b/Engine/Plugins/Animation/ControlRig/Source/ControlRig/Public/Constraints/ControlRigTransformableHandle.h
--- a/Engine/Plugins/Animation/ControlRig/Source/ControlRig/Public/Constraints/ControlRigTransformableHandle.h
+++ b/Engine/Plugins/Animation/ControlRig/Source/ControlRig/Public/Constraints/ControlRigTransformableHandle.h
@@ -142,10 +142,10 @@ private:
 	static FControlEvaluationGraphBinding& GetEvaluationBinding();
 
 #if WITH_EDITOR
-	/** @todo document */
+	/** Tracks if the control rig that this handle wraps has been replaced. (see FCoreUObjectDelegates::FOnObjectsReplaced for more info). */
 	void OnObjectsReplaced(const TMap<UObject*, UObject*>& InOldToNewInstances);
 #endif
 
-	/** @todo document */
+	/** Returns the control that this handle wraps. */
 	FRigControlElement* GetControlElement() const;
 };
