ControlRig: Maintain control visibility after replacing a rig in the ControlRigEditMode
#rnx

diff --git a/Engine/Plugins/Animation/ControlRig/Source/ControlRig/Public/ControlRig.h b/Engine/Plugins/Animation/ControlRig/Source/ControlRig/Public/ControlRig.h
--- a/Engine/Plugins/Animation/ControlRig/Source/ControlRig/Public/ControlRig.h
+++ b/Engine/Plugins/Animation/ControlRig/Source/ControlRig/Public/ControlRig.h
@@ -694,7 +694,7 @@ private:
 
 	/** Whether controls are visible */
 	UPROPERTY(transient)
-	bool bControlsVisible;
+	bool bControlsVisible = true;
 
 #endif
 
diff --git a/Engine/Plugins/Animation/ControlRig/Source/ControlRigEditor/Private/EditMode/ControlRigEditMode.cpp b/Engine/Plugins/Animation/ControlRig/Source/ControlRigEditor/Private/EditMode/ControlRigEditMode.cpp
--- a/Engine/Plugins/Animation/ControlRig/Source/ControlRigEditor/Private/EditMode/ControlRigEditMode.cpp
+++ b/Engine/Plugins/Animation/ControlRig/Source/ControlRigEditor/Private/EditMode/ControlRigEditMode.cpp
@@ -4472,6 +4472,8 @@ void FControlRigEditMode::OnObjectsReplaced(const TMap<UObject*, UObject*>& OldT
 				AddControlRigInternal(NewControlRig);
 				NewControlRig->Initialize();
 
+				CopyControlsVisibility(OldControlRig, NewControlRig);
+
 				bHasAnyRigBeenReplaced = true;
 			}
 		}
@@ -4483,6 +4485,29 @@ void FControlRigEditMode::OnObjectsReplaced(const TMap<UObject*, UObject*>& OldT
 	}
 }
 
+void FControlRigEditMode::CopyControlsVisibility(const UControlRig* SourceRig, UControlRig* TargetRig)
+{
+	TargetRig->bControlsVisible = SourceRig->bControlsVisible;
+    if (const UModularRig* SourceModularRig = Cast<UModularRig>(SourceRig))
+    {
+    	if (UModularRig* TargetModularRig = Cast<UModularRig>(TargetRig))
+    	{
+    		for (const FRigModuleInstance& SourceModule : SourceModularRig->Modules)
+    		{
+    			if (FRigModuleInstance* TargetModule = TargetModularRig->FindModule(SourceModule.Name))
+    			{
+    				UControlRig* SourceModuleRig = SourceModule.GetRig();
+    				UControlRig* TargetModuleRig = TargetModule->GetRig();
+    				if (SourceModuleRig && TargetModuleRig)
+    				{
+    					TargetModuleRig->bControlsVisible = SourceModuleRig->bControlsVisible;
+    				}
+    			}
+    		}
+    	}
+    }
+}
+
 bool FControlRigEditMode::IsTransformDelegateAvailable() const
 {
 	return (OnGetRigElementTransformDelegate.IsBound() && OnSetRigElementTransformDelegate.IsBound());
@@ -4577,6 +4602,7 @@ void FControlRigEditMode::ReplaceControlRig(UControlRig* OldControlRig, UControl
 	SetObjects_Internal();
 	RequestToRecreateControlShapeActors(NewControlRig);
 
+	CopyControlsVisibility(OldControlRig, NewControlRig);
 }
 void FControlRigEditMode::OnHierarchyModified(ERigHierarchyNotification InNotif, URigHierarchy* InHierarchy, const FRigNotificationSubject& InSubject)
 {
diff --git a/Engine/Plugins/Animation/ControlRig/Source/ControlRigEditor/Public/EditMode/ControlRigEditMode.h b/Engine/Plugins/Animation/ControlRig/Source/ControlRigEditor/Public/EditMode/ControlRigEditMode.h
--- a/Engine/Plugins/Animation/ControlRig/Source/ControlRigEditor/Public/EditMode/ControlRigEditMode.h
+++ b/Engine/Plugins/Animation/ControlRig/Source/ControlRigEditor/Public/EditMode/ControlRigEditMode.h
@@ -410,6 +410,9 @@ protected:
 	/** Called from the editor when a blueprint object replacement has occurred */
 	void OnObjectsReplaced(const TMap<UObject*, UObject*>& OldToNewInstanceMap);
 
+	/** Copy control visibility from the source rig (and its modules if necessary) to the target rig */ 
+	void CopyControlsVisibility(const UControlRig* SourceRig, UControlRig* TargetRig);
+
 	/** Return true if transform setter/getter delegates are available */
 	bool IsTransformDelegateAvailable() const;
 
