Control Rig: Do not autoexpand subtrees in the anim outliner when a parent is selected. Also keep inactive highlighting when a descendant of an element is selected.
#rnx

diff --git a/Engine/Plugins/Animation/ControlRig/Source/ControlRigEditor/Private/EditMode/SControlRigOutliner.cpp b/Engine/Plugins/Animation/ControlRig/Source/ControlRigEditor/Private/EditMode/SControlRigOutliner.cpp
--- a/Engine/Plugins/Animation/ControlRig/Source/ControlRigEditor/Private/EditMode/SControlRigOutliner.cpp
+++ b/Engine/Plugins/Animation/ControlRig/Source/ControlRigEditor/Private/EditMode/SControlRigOutliner.cpp
@@ -252,6 +252,7 @@ void FMultiRigTreeElement::RefreshDisplaySettings(const URigHierarchy* InHierarc
 void SMultiRigHierarchyItem::Construct(const FArguments& InArgs, const TSharedRef<STableViewBase>& OwnerTable, TSharedRef<FMultiRigTreeElement> InRigTreeElement, TSharedPtr<SMultiRigHierarchyTreeView> InTreeView, const FRigTreeDisplaySettings& InSettings, bool bPinned)
 {
 	WeakRigTreeElement = InRigTreeElement;
+	TreeView = InTreeView;
 	Delegates = InTreeView->GetTreeDelegates();
 
 	
@@ -436,6 +437,47 @@ TPair<const FSlateBrush*, FSlateColor> SMultiRigHierarchyItem::GetBrushForElemen
 	return TPair<const FSlateBrush*, FSlateColor>(Brush, Color);
 }
 
+const FSlateBrush* SMultiRigHierarchyItem::GetBorder() const
+{
+	if (!TreeView.IsValid())
+	{
+		return STableRow<TSharedPtr<FMultiRigTreeElement>>::GetBorder(); 
+	}
+	
+	TArray<TSharedPtr<FMultiRigTreeElement>> SelectedElements = TreeView->GetSelectedItems();
+	for (TSharedPtr<FMultiRigTreeElement> Selected : SelectedElements)
+	{
+		TSharedPtr<FMultiRigTreeElement> Cur = Selected;
+		if (Cur == WeakRigTreeElement)
+		{
+			// item selected
+			return &Style->ActiveHighlightedBrush;
+		}
+		else
+		{
+			bool bAnyDescendantSelected = false;
+			while (Cur)
+			{
+				if (Cur == WeakRigTreeElement)
+				{
+					bAnyDescendantSelected = true;
+					break;
+				}
+				else
+				{
+					Cur = TreeView->GetParentElement(Cur);
+				}
+			}
+			if (bAnyDescendantSelected)
+			{
+				// Descendant selected
+				return &Style->InactiveHighlightedBrush;
+			}
+		}
+	}
+	return STableRow<TSharedPtr<FMultiRigTreeElement>>::GetBorder();
+}
+
 //////////////////////////////////////////////////////////////
 /// SRigHierarchyTreeView
 ///////////////////////////////////////////////////////////
@@ -861,6 +903,24 @@ bool SMultiRigHierarchyTreeView::RemoveElement(const FMultiRigData& InData)
 	return ElementMap.Remove(InData) > 0;
 }
 
+TSharedPtr<FMultiRigTreeElement> SMultiRigHierarchyTreeView::GetParentElement(TSharedPtr<FMultiRigTreeElement> InElement) const
+{
+	if (!InElement.IsValid())
+	{
+		return nullptr;
+	}
+
+	if (const FMultiRigData* Parent = ParentMap.Find(InElement->Data))
+	{
+		if (const TSharedPtr<FMultiRigTreeElement>* ParentElement = ElementMap.Find(*Parent))
+		{
+			return *ParentElement;
+		}
+	}
+
+	return nullptr;
+}
+
 void SMultiRigHierarchyTreeView::RequestTreeViewRefresh(TFunction<void()>&& InRefreshFunction)
 {
 	UnregisterPendingRefresh();
@@ -1483,15 +1543,6 @@ void SControlRigOutliner::HandleControlSelected(UControlRig* Subject, FRigContro
 			HierarchyTreeView->GetTreeView()->SetItemSelection(Found, bSelected, ESelectInfo::OnNavigation);
 
 			TArray<TSharedPtr<FMultiRigTreeElement>> SelectedItems = HierarchyTreeView->GetTreeView()->GetSelectedItems();
-			for (TSharedPtr<FMultiRigTreeElement> SelectedItem : SelectedItems)
-			{
-				if (SelectedItem->Data.ModuleName.IsSet())
-				{
-					continue;
-				}
-				HierarchyTreeView->GetTreeView()->SetExpansionRecursive(SelectedItem, false, true);
-			}
-
 			if (SelectedItems.Num() > 0)
 			{
 				HierarchyTreeView->GetTreeView()->RequestScrollIntoView(SelectedItems.Last());
diff --git a/Engine/Plugins/Animation/ControlRig/Source/ControlRigEditor/Private/EditMode/SControlRigOutliner.h b/Engine/Plugins/Animation/ControlRig/Source/ControlRigEditor/Private/EditMode/SControlRigOutliner.h
--- a/Engine/Plugins/Animation/ControlRig/Source/ControlRigEditor/Private/EditMode/SControlRigOutliner.h
+++ b/Engine/Plugins/Animation/ControlRig/Source/ControlRigEditor/Private/EditMode/SControlRigOutliner.h
@@ -110,7 +110,7 @@ class FMultiRigTreeElement : public TSharedFromThis<FMultiRigTreeElement>
 {
 public:
 	FMultiRigTreeElement(const FMultiRigData& InData, TWeakPtr<SMultiRigHierarchyTreeView> InTreeView,ERigTreeFilterResult InFilterResult);
-public:
+
 	/** Element Data to display */
 	FMultiRigData Data;
 	TArray<TSharedPtr<FMultiRigTreeElement>> Children;
@@ -138,9 +138,11 @@ public:
 
 	void Construct(const FArguments& InArgs, const TSharedRef<STableViewBase>& OwnerTable, TSharedRef<FMultiRigTreeElement> InRigTreeElement, TSharedPtr<SMultiRigHierarchyTreeView> InTreeView, const FRigTreeDisplaySettings& InSettings, bool bPinned);
 	static TPair<const FSlateBrush*, FSlateColor> GetBrushForElementType(const URigHierarchy* InHierarchy, const FMultiRigData& InData);
+	virtual const FSlateBrush* GetBorder() const override;
 
 private:
 	TWeakPtr<FMultiRigTreeElement> WeakRigTreeElement;
+	TSharedPtr<SMultiRigHierarchyTreeView> TreeView;
 	FMultiRigTreeDelegates Delegates;
 
 	FText GetDisplayName() const;
@@ -214,6 +216,8 @@ public:
 	 * Note that the refresh function will be set to RefreshTreeView if no function is passed. 
 	 * InRefreshFunction can capture what's necessary if anything has to be done once the tree rebuilt. */
 	void RequestTreeViewRefresh(TFunction<void()>&& InRefreshFunction = nullptr);
+
+	TSharedPtr<FMultiRigTreeElement> GetParentElement(TSharedPtr<FMultiRigTreeElement> InElement) const;
 	
 private:
 
