Control Rig: Fix anim outliner not behaving correctly during shift+click on modules. Also recursively selecting elements under a module, including submodules.
#rnx

diff --git a/Engine/Plugins/Animation/ControlRig/Source/ControlRigEditor/Private/EditMode/SControlRigOutliner.cpp b/Engine/Plugins/Animation/ControlRig/Source/ControlRigEditor/Private/EditMode/SControlRigOutliner.cpp
--- a/Engine/Plugins/Animation/ControlRig/Source/ControlRigEditor/Private/EditMode/SControlRigOutliner.cpp
+++ b/Engine/Plugins/Animation/ControlRig/Source/ControlRigEditor/Private/EditMode/SControlRigOutliner.cpp
@@ -1264,21 +1264,17 @@ void SSearchableMultiRigHierarchyTreeView::Construct(const FArguments& InArgs)
 				]
 
 			+ SVerticalBox::Slot()
-				.AutoHeight()
+				.FillHeight(1.f)
 				.VAlign(VAlign_Top)
 				.HAlign(HAlign_Fill)
 				.Padding(0.0f, 0.0f)
 				[
-					SNew(SScrollBox)
-					+ SScrollBox::Slot()
+					SNew(SBorder)
+					.Padding(2.0f)
+					.BorderImage(FAppStyle::GetBrush("SCSEditor.TreePanel"))
 					[
-						SNew(SBorder)
-						.Padding(2.0f)
-						.BorderImage(FAppStyle::GetBrush("SCSEditor.TreePanel"))
-						[
-							SAssignNew(TreeView, SMultiRigHierarchyTreeView)
-							.RigTreeDelegates(TreeDelegates)
-						]
+						SAssignNew(TreeView, SMultiRigHierarchyTreeView)
+						.RigTreeDelegates(TreeDelegates)
 					]
 				]
 		];
@@ -1459,26 +1455,12 @@ void SControlRigOutliner::Construct(const FArguments& InArgs, FControlRigEditMod
 	
 	ChildSlot
 		[
-			SNew(SScrollBox)
-			+ SScrollBox::Slot()
-			[
 			SNew(SVerticalBox)
 			+ SVerticalBox::Slot()
-				.AutoHeight()
-				[
-					SAssignNew(PickerExpander, SExpandableArea)
-					.InitiallyCollapsed(false)
-					//.AreaTitle(AreaTitle)
-					//.AreaTitleFont(FAppStyle::GetFontStyle("DetailsView.CategoryFontStyle"))
-					.BorderBackgroundColor(FLinearColor(.6f, .6f, .6f))
-					.BodyContent()
-					[
-						SAssignNew(HierarchyTreeView, SSearchableMultiRigHierarchyTreeView)
-						.RigTreeDelegates(RigTreeDelegates)
-					]
-					
-				]
-		
+			.FillHeight(1.0f)
+			[
+				SAssignNew(HierarchyTreeView, SSearchableMultiRigHierarchyTreeView)
+				.RigTreeDelegates(RigTreeDelegates)
 			]
 		];
 	SetEditMode(InEditMode);
@@ -1540,7 +1522,7 @@ void SControlRigOutliner::HandleControlSelected(UControlRig* Subject, FRigContro
 		{
 			TGuardValue<bool> GuardRigHierarchyChanges(bIsChangingRigHierarchy, true);
 
-			HierarchyTreeView->GetTreeView()->SetItemSelection(Found, bSelected, ESelectInfo::OnNavigation);
+			HierarchyTreeView->GetTreeView()->SetItemSelection(Found, bSelected, ESelectInfo::Direct);
 
 			TArray<TSharedPtr<FMultiRigTreeElement>> SelectedItems = HierarchyTreeView->GetTreeView()->GetSelectedItems();
 			if (SelectedItems.Num() > 0)
@@ -1575,13 +1557,15 @@ void SControlRigOutliner::HandleSelectionChanged(TSharedPtr<FMultiRigTreeElement
 				if (Settings.FilterText.IsEmpty() || !Settings.bFlattenHierarchyOnFilter)
 				{
 					TSharedPtr<FMultiRigTreeElement> Element = HierarchyTreeView->GetTreeView()->FindElement(Data);
-					for (TSharedPtr<FMultiRigTreeElement> Child : Element->Children)
+					TArray<TSharedPtr<FMultiRigTreeElement>> Descendants = Element->Children;
+					for (int32 i=0; i<Descendants.Num(); ++i)
 					{
-						if (!Child->Data.Key.IsSet())
+						TSharedPtr<FMultiRigTreeElement>& Child = Descendants[i];
+						if (Child->Data.Key.IsSet())
 						{
-							break;
+							Keys.AddUnique(Child->Data.Key.GetValue());
 						}
-						Keys.AddUnique(Child->Data.Key.GetValue());
+						Descendants.Append(Child->Children);
 					}
 				}
 				else
diff --git a/Engine/Plugins/Animation/ControlRig/Source/ControlRigEditor/Private/EditMode/SControlRigOutliner.h b/Engine/Plugins/Animation/ControlRig/Source/ControlRigEditor/Private/EditMode/SControlRigOutliner.h
--- a/Engine/Plugins/Animation/ControlRig/Source/ControlRigEditor/Private/EditMode/SControlRigOutliner.h
+++ b/Engine/Plugins/Animation/ControlRig/Source/ControlRigEditor/Private/EditMode/SControlRigOutliner.h
@@ -313,7 +313,6 @@ private:
 	FRigTreeDisplaySettings DisplaySettings;
 	const FRigTreeDisplaySettings& GetDisplaySettings() const { return DisplaySettings; }
 	bool bIsChangingRigHierarchy = false;
-	TSharedPtr<SExpandableArea> PickerExpander;
 
 	//set of control rigs we are bound too and need to clear delegates from
 	TArray<TWeakObjectPtr<UControlRig>> BoundControlRigs;
