Avoid a crash in ControlRig  Trait when we make a public variable in the CR but we don't recompile (so it comes as a external variable without memory).

diff --git a/Engine/Plugins/Animation/ControlRig/Source/ControlRigDeveloper/Private/ControlRigIOMapping.cpp b/Engine/Plugins/Animation/ControlRig/Source/ControlRigDeveloper/Private/ControlRigIOMapping.cpp
--- a/Engine/Plugins/Animation/ControlRig/Source/ControlRigDeveloper/Private/ControlRigIOMapping.cpp
+++ b/Engine/Plugins/Animation/ControlRig/Source/ControlRigDeveloper/Private/ControlRigIOMapping.cpp
@@ -112,10 +112,10 @@ void FControlRigIOMapping::RebuildExposedProperties()
 	CustomPinProperties.Reset();
 
 	// go through exposed properties, and clean up
-	GetVariables(TargetClass, true, InputVariables);
+	GetVariables(TargetClass, true, bIgnoreVariablesWithNoMemory, InputVariables);
 	// we still update OUtputvariables
 	// we don't want output to be exposed
-	GetVariables(TargetClass, false, OutputVariables);
+	GetVariables(TargetClass, false, bIgnoreVariablesWithNoMemory, OutputVariables);
 
 	// clear IO variables that don't exist anymore
 	auto ClearInvalidMapping = [](const TMap<FName, FRigVMExternalVariable>& InVariables, TMap<FName, FName>& InOutMapping)
@@ -344,7 +344,7 @@ const TArray<FControlRigIOMapping::FControlsInfo>& FControlRigIOMapping::FRigCon
 	return ControlsInfo;
 }
 
-void FControlRigIOMapping::GetVariables(const UClass* TargetClass, bool bInput, TMap<FName, FRigVMExternalVariable>& OutVariables)
+void FControlRigIOMapping::GetVariables(const UClass* TargetClass, bool bInput, bool bIgnoreVariablesWithNoMemory, TMap<FName, FRigVMExternalVariable>& OutVariables)
 {
 	DECLARE_SCOPE_HIERARCHICAL_COUNTER_FUNC()
 
@@ -363,6 +363,10 @@ void FControlRigIOMapping::GetVariables(const UClass* TargetClass, bool bInput,
 				{
 					if (!bInput || !PublicVariable.bIsReadOnly)
 					{
+						if (bIgnoreVariablesWithNoMemory && PublicVariable.Memory == nullptr)
+						{
+							continue;
+						}
 						OutVariables.Add(PublicVariable.Name, PublicVariable);
 					}
 				}
diff --git a/Engine/Plugins/Animation/ControlRig/Source/ControlRigDeveloper/Public/ControlRigIOMapping.h b/Engine/Plugins/Animation/ControlRig/Source/ControlRigDeveloper/Public/ControlRigIOMapping.h
--- a/Engine/Plugins/Animation/ControlRig/Source/ControlRigDeveloper/Public/ControlRigIOMapping.h
+++ b/Engine/Plugins/Animation/ControlRig/Source/ControlRigDeveloper/Public/ControlRigIOMapping.h
@@ -87,6 +87,8 @@ struct CONTROLRIGDEVELOPER_API FControlRigIOMapping : public TSharedFromThis<FCo
 	void SetIOMapping(bool bInput, const FName& SourceProperty, const FName& TargetCurve);
 	FName GetIOMapping(bool bInput, const FName& SourceProperty) const;
 
+	void SetIgnoreVariablesWithNoMemory(bool bIgnore) { bIgnoreVariablesWithNoMemory = bIgnore; }
+
 private:
 	FOnGetTargetSkeleton OnGetTargetSkeletonDelegate;
 	FOnGetTargetClass OnGetTargetClassDelegate;
@@ -100,8 +102,10 @@ private:
 	// Controls data
 	FRigControlsData RigControlsData;
 
+	bool bIgnoreVariablesWithNoMemory = false;
+
 	// property related things
-	static void GetVariables(const UClass* TargetClass, bool bInput, TMap<FName, FRigVMExternalVariable>& OutParameters);
+	static void GetVariables(const UClass* TargetClass, bool bInput, bool bIgnoreVariablesWithNoMemory, TMap<FName, FRigVMExternalVariable>& OutParameters);
 
 	TMap<FName, FRigVMExternalVariable> InputVariables;
 	TMap<FName, FRigVMExternalVariable> OutputVariables;
diff --git a/Engine/Plugins/Experimental/AnimNextControlRig/Source/AnimNextControlRig/Private/ControlRigTrait.cpp b/Engine/Plugins/Experimental/AnimNextControlRig/Source/AnimNextControlRig/Private/ControlRigTrait.cpp
--- a/Engine/Plugins/Experimental/AnimNextControlRig/Source/AnimNextControlRig/Private/ControlRigTrait.cpp
+++ b/Engine/Plugins/Experimental/AnimNextControlRig/Source/AnimNextControlRig/Private/ControlRigTrait.cpp
@@ -622,6 +622,11 @@ void FControlRigTrait::GetProgrammaticPins(FAnimNextTraitSharedData* InSharedDat
 							continue;
 						}
 
+						if (Variable.Memory == nullptr) // if we make a variable public but don't recompile, we have to skip, as it comes without memory
+						{
+							continue;
+						}
+
 						if (const URigVMPin* SubPin = InTraitPin->FindSubPin(Variable.Name.ToString()))
 						{
 							const FString PinDefaultValue = SubPin->GetDefaultValue();
diff --git a/Engine/Plugins/Experimental/AnimNextControlRig/Source/AnimNextControlRigEditor/Private/ControlRigTraitCustomization.cpp b/Engine/Plugins/Experimental/AnimNextControlRig/Source/AnimNextControlRigEditor/Private/ControlRigTraitCustomization.cpp
--- a/Engine/Plugins/Experimental/AnimNextControlRig/Source/AnimNextControlRigEditor/Private/ControlRigTraitCustomization.cpp
+++ b/Engine/Plugins/Experimental/AnimNextControlRig/Source/AnimNextControlRigEditor/Private/ControlRigTraitCustomization.cpp
@@ -70,6 +70,7 @@ void FControlRigTraitSharedDataCustomization::CustomizeHeader(TSharedRef<IProper
 	ControlRigIOMapping->GetOnVariableMappingChanged().BindSP(this, &FControlRigTraitSharedDataCustomization::OnVariableMappingChanged);
 	ControlRigIOMapping->GetOnGetTargetSkeletonDelegate().BindSP(this, &FControlRigTraitSharedDataCustomization::GetTargetSkeleton);
 	ControlRigIOMapping->GetOnGetTargetClassDelegate().BindSP(this, &FControlRigTraitSharedDataCustomization::GetTargetClass);
+	ControlRigIOMapping->SetIgnoreVariablesWithNoMemory(true); // AnimBP creates it's own variables, AnimNext needs the memory to create the pin
 
 	OnObjectsReinstancedHandle = UE::AnimNext::ControlRig::FAnimNextControlRigModule::OnObjectsReinstanced.AddRaw(this, &FControlRigTraitSharedDataCustomization::OnObjectsReinstanced);
 
