ControlRig & Sequencer: make sure this is safe to use when OnActivateSequence delegate is triggered

#rnx

diff --git a/Engine/Plugins/Animation/ControlRig/Source/ControlRigEditor/Private/Sequencer/ControlRigParameterTrackEditor.cpp b/Engine/Plugins/Animation/ControlRig/Source/ControlRigEditor/Private/Sequencer/ControlRigParameterTrackEditor.cpp
--- a/Engine/Plugins/Animation/ControlRig/Source/ControlRigEditor/Private/Sequencer/ControlRigParameterTrackEditor.cpp
+++ b/Engine/Plugins/Animation/ControlRig/Source/ControlRigEditor/Private/Sequencer/ControlRigParameterTrackEditor.cpp
@@ -2725,62 +2725,71 @@ void FControlRigParameterTrackEditor::OnActivateSequenceChanged(FMovieSceneSeque
 	TSharedPtr<ISequencer> Sequencer = GetSequencer();
 	if (bControlRigEditModeWasOpen && Sequencer && Sequencer->IsLevelEditorSequencer())
 	{
-		GEditor->GetTimerManager()->SetTimerForNextTick([this]()
+		GEditor->GetTimerManager()->SetTimerForNextTick([this, WeakThis = this->AsWeak()]()
 		{
+			if (!WeakThis.IsValid())
+			{
+				return;
+			}
+			
 			//we need to make sure pending deactivated edit modes, including a possible control rig edit mode
 			//get totally removed which only happens on a tick 
-				if (GLevelEditorModeTools().HasToolkitHost())
+			if (GLevelEditorModeTools().HasToolkitHost())
+			{
+				if (GEditor->GetActiveViewport() && GEditor->GetActiveViewport()->GetClient())
 				{
-					if (GEditor->GetActiveViewport() && GEditor->GetActiveViewport()->GetClient())
+					if (FEditorModeTools* EditorModeTools = GetEditorModeTools())
 					{
-						if (FEditorModeTools* EditorModeTools = GetEditorModeTools())
-						{
-							FViewport* ActiveViewport = GEditor->GetActiveViewport();
-							FEditorViewportClient* EditorViewClient = (FEditorViewportClient*)ActiveViewport->GetClient();
-							EditorModeTools->Tick(EditorViewClient, 0.033f);
-						}
+						FViewport* ActiveViewport = GEditor->GetActiveViewport();
+						FEditorViewportClient* EditorViewClient = (FEditorViewportClient*)ActiveViewport->GetClient();
+						EditorModeTools->Tick(EditorViewClient, 0.033f);
 					}
+				}
 
-					GEditor->GetTimerManager()->SetTimerForNextTick([this]()
+				GEditor->GetTimerManager()->SetTimerForNextTick([this, WeakThis]()
+				{
+					if (!WeakThis.IsValid())
 					{
-						//now we can recreate it
-						FControlRigEditMode* ControlRigEditMode = GetEditMode(true);
-						if (ControlRigEditMode)
-						{
-							bool bSequencerSet = false;
+						return;
+					}
+					
+					//now we can recreate it
+					if (FControlRigEditMode* ControlRigEditMode = GetEditMode(true))
+					{
+						bool bSequencerSet = false;
 
-							TSharedPtr<ISequencer> Sequencer = GetSequencer();
-							for (TWeakObjectPtr<UControlRig>& ControlRig : BoundControlRigs)
+						TSharedPtr<ISequencer> Sequencer = GetSequencer();
+						for (TWeakObjectPtr<UControlRig>& ControlRig : BoundControlRigs)
+						{
+							if (ControlRig.IsValid())
 							{
-								if (ControlRig.IsValid())
-								{
-									ControlRigEditMode->AddControlRigObject(ControlRig.Get(), Sequencer);
-									bSequencerSet = true;
+								ControlRigEditMode->AddControlRigObject(ControlRig.Get(), Sequencer);
+								bSequencerSet = true;
 
-									for (int32 Index = 0; Index < PreviousSelectedControlRigs.Num(); ++Index)
+								for (int32 Index = 0; Index < PreviousSelectedControlRigs.Num(); ++Index)
+								{
+									if (ControlRig.Get()->GetClass() == PreviousSelectedControlRigs[Index].Key)
 									{
-										if (ControlRig.Get()->GetClass() == PreviousSelectedControlRigs[Index].Key)
+										for (const FName& ControlName : PreviousSelectedControlRigs[Index].Value)
 										{
-											for (const FName& ControlName : PreviousSelectedControlRigs[Index].Value)
-											{
-												ControlRig.Get()->SelectControl(ControlName, true);
-											}
-											PreviousSelectedControlRigs.RemoveAt(Index);
-											break;
+											ControlRig.Get()->SelectControl(ControlName, true);
 										}
+										PreviousSelectedControlRigs.RemoveAt(Index);
+										break;
 									}
 								}
 							}
+						}
 
-							if (!bSequencerSet)
-							{
-								ControlRigEditMode->SetSequencer(Sequencer);
-							}
-
+						if (!bSequencerSet)
+						{
+							ControlRigEditMode->SetSequencer(Sequencer);
 						}
-						PreviousSelectedControlRigs.Reset();
-					});
-				}
+
+					}
+					PreviousSelectedControlRigs.Reset();
+				});
+			}
 		});
 	}
 }
