Control Rig: Selecting the control rig root element in the anim outliner selects all controls
#rnx

diff --git a/Engine/Plugins/Animation/ControlRig/Source/ControlRigEditor/Private/EditMode/SControlRigOutliner.cpp b/Engine/Plugins/Animation/ControlRig/Source/ControlRigEditor/Private/EditMode/SControlRigOutliner.cpp
--- a/Engine/Plugins/Animation/ControlRig/Source/ControlRigEditor/Private/EditMode/SControlRigOutliner.cpp
+++ b/Engine/Plugins/Animation/ControlRig/Source/ControlRigEditor/Private/EditMode/SControlRigOutliner.cpp
@@ -291,7 +291,7 @@ void SMultiRigHierarchyItem::Construct(const FArguments& InArgs, const TSharedRe
 			[
 				SNew(SButton)
 				.ButtonStyle(FAppStyle::Get(), "NoBorder")
-				.OnClicked(this, &SMultiRigHierarchyItem::OnGetSelectedClicked)
+				.OnClicked(this, &SMultiRigHierarchyItem::OnToggleVisibilityClicked)
 				[
 					SNew(SImage)
 					.Image_Lambda([this]() -> const FSlateBrush*
@@ -361,35 +361,68 @@ void SMultiRigHierarchyItem::Construct(const FArguments& InArgs, const TSharedRe
 
 }
 
-FReply SMultiRigHierarchyItem::OnGetSelectedClicked()
+FReply SMultiRigHierarchyItem::OnToggleVisibilityClicked()
 {
 	if (TSharedPtr<FMultiRigTreeElement> Element = WeakRigTreeElement.Pin())
 	{
-		if (Element->Data.Key.IsSet() == false)
+		TArray<TSharedPtr<FMultiRigTreeElement>> ElementsToToggle;
+		ElementsToToggle.Add(Element);
+		if (FSlateApplication::Get().GetModifierKeys().IsShiftDown())
 		{
-			UControlRig* Rig = nullptr;
-			if (Element->Data.ModuleName.IsSet())
+			TArray<TSharedPtr<FMultiRigTreeElement>> SelectedItems = TreeView->GetSelectedItems();
+			for (TSharedPtr<FMultiRigTreeElement> Selected : SelectedItems)
 			{
-				if (UModularRig* ModularRig = Cast<UModularRig>(Element->Data.ControlRig))
+				if (Selected->Data.ModuleName.IsSet())
 				{
-					if (FRigModuleInstance* Module = ModularRig->FindModule(Element->Data.ModuleName.GetValue()))
-					{
-						Rig = Module->GetRig();
-					}
+					ElementsToToggle.AddUnique(Selected);
 				}
 			}
-			else
+		}
+
+		TArray<UControlRig*> RigsToToggle;
+
+		TOptional<bool> bSetControlsVisible;
+		for (TSharedPtr<FMultiRigTreeElement> ElementToToggle : ElementsToToggle)
+		{
+			if (ElementToToggle->Data.Key.IsSet() == false)
 			{
-				Rig  = Element->Data.ControlRig.Get();
-			}
+				UControlRig* Rig = nullptr;
+				if (ElementToToggle->Data.ModuleName.IsSet())
+				{
+					if (UModularRig* ModularRig = Cast<UModularRig>(Element->Data.ControlRig))
+					{
+						if (FRigModuleInstance* Module = ModularRig->FindModule(ElementToToggle->Data.ModuleName.GetValue()))
+						{
+							Rig = Module->GetRig();
+						}
+					}
+				}
+				else
+				{
+					Rig  = ElementToToggle->Data.ControlRig.Get();
+				}
 			
-			if (Rig)
+				if (Rig)
+				{
+					if (!bSetControlsVisible.IsSet())
+					{
+						bSetControlsVisible = !Rig->GetControlsVisible();
+					}
+
+					RigsToToggle.AddUnique(Rig);
+				}
+			}
+		}
+
+		if (!RigsToToggle.IsEmpty())
+		{
+			FScopedTransaction ScopedTransaction(LOCTEXT("ToggleControlsVisibility", "Toggle Controls Visibility"), !GIsTransacting);
+			for (UControlRig* Rig : RigsToToggle)
 			{
-				FScopedTransaction ScopedTransaction(LOCTEXT("ToggleControlsVisibility", "Toggle Controls Visibility"), !GIsTransacting);
 				Rig->Modify();
-				Rig->ToggleControlsVisible();
-				return FReply::Handled();
+				Rig->SetControlsVisible(bSetControlsVisible.GetValue());
 			}
+			return FReply::Handled();
 		}
 	}
 	return FReply::Unhandled();  //should flow to selection instead
@@ -1520,14 +1553,18 @@ void SControlRigOutliner::HandleControlSelected(UControlRig* Subject, FRigContro
 		TSharedPtr<FMultiRigTreeElement> Found = HierarchyTreeView->GetTreeView()->FindElement(Data, HierarchyTreeView->GetTreeView()->GetRootElements()[RootIndex]);
 		if (Found.IsValid())
 		{
+			bool bScrollToElement = !bIsChangingRigHierarchy;
 			TGuardValue<bool> GuardRigHierarchyChanges(bIsChangingRigHierarchy, true);
 
 			HierarchyTreeView->GetTreeView()->SetItemSelection(Found, bSelected, ESelectInfo::Direct);
 
-			TArray<TSharedPtr<FMultiRigTreeElement>> SelectedItems = HierarchyTreeView->GetTreeView()->GetSelectedItems();
-			if (SelectedItems.Num() > 0)
+			if (bScrollToElement)
 			{
-				HierarchyTreeView->GetTreeView()->RequestScrollIntoView(SelectedItems.Last());
+				TArray<TSharedPtr<FMultiRigTreeElement>> SelectedItems = HierarchyTreeView->GetTreeView()->GetSelectedItems();
+				if (SelectedItems.Num() > 0)
+				{
+					HierarchyTreeView->GetTreeView()->RequestScrollIntoView(SelectedItems.Last());
+				}
 			}
 		}
 	}
@@ -1541,6 +1578,7 @@ void SControlRigOutliner::HandleSelectionChanged(TSharedPtr<FMultiRigTreeElement
 	}
 	const TArray<FMultiRigData> NewSelection = HierarchyTreeView->GetTreeView()->GetSelectedData();
 	TMap<UControlRig*, TArray<FRigElementKey>> SelectedRigAndKeys;
+	TArray<TSharedPtr<FMultiRigTreeElement>> AddToSelection;
 	for (const FMultiRigData& Data : NewSelection)
 	{
 		if (Data.ControlRig.IsValid())
@@ -1550,7 +1588,7 @@ void SControlRigOutliner::HandleSelectionChanged(TSharedPtr<FMultiRigTreeElement
 				SelectedRigAndKeys.FindOrAdd(Data.ControlRig.Get()).AddUnique(Data.Key.GetValue());
 			}
 
-			if (Data.ModuleName.IsSet())
+			if (!Data.Key.IsSet())
 			{
 				TArray<FRigElementKey>& Keys = SelectedRigAndKeys.FindOrAdd(Data.ControlRig.Get());
 				const FRigTreeDisplaySettings& Settings = HierarchyTreeView->GetDisplaySettings();
@@ -1565,10 +1603,14 @@ void SControlRigOutliner::HandleSelectionChanged(TSharedPtr<FMultiRigTreeElement
 						{
 							Keys.AddUnique(Child->Data.Key.GetValue());
 						}
+						else if (Child->Data.ModuleName.IsSet())
+						{
+							AddToSelection.AddUnique(Child);
+						}
 						Descendants.Append(Child->Children);
 					}
 				}
-				else
+				else if (Data.ModuleName.IsSet())
 				{
 					// If we have flatten the hierarchy due to a search, we cannot rely on the children of the module
 					Data.GetHierarchy()->ForEach<FRigControlElement>([&](FRigControlElement* ControlElement)
@@ -1580,10 +1622,20 @@ void SControlRigOutliner::HandleSelectionChanged(TSharedPtr<FMultiRigTreeElement
 						return true;
 					});
 				}
+				else
+				{
+					// If the key is not set, and the module name is not set, we are selecting the root of the rig. Select all controls
+					for (const FRigElementKey ControlKey : Data.GetHierarchy()->GetControlKeys())
+					{
+						Keys.AddUnique(ControlKey);
+					}
+				}
 			}
 		}
 	}
+
 	TGuardValue<bool> GuardRigHierarchyChanges(bIsChangingRigHierarchy, true);
+	
 	FControlRigEditMode* EditMode = static_cast<FControlRigEditMode*>(ModeTools->GetActiveMode(FControlRigEditMode::ModeName));
 	bool bEndTransaction = false;
 	if (GEditor && !GIsTransacting && EditMode && EditMode->IsInLevelEditor())
@@ -1591,6 +1643,14 @@ void SControlRigOutliner::HandleSelectionChanged(TSharedPtr<FMultiRigTreeElement
 		GEditor->BeginTransaction(LOCTEXT("SelectControl", "Select Control"));
 		bEndTransaction = true;
 	}
+
+	if (!AddToSelection.IsEmpty())
+	{
+		TArray<TSharedPtr<FMultiRigTreeElement>> NewSelectedElements = HierarchyTreeView->GetTreeView()->GetSelectedItems();
+		NewSelectedElements.Append(AddToSelection);
+		HierarchyTreeView->GetTreeView()->SetItemSelection(AddToSelection, true);
+	}
+	
 	//due to how Sequencer Tree View will redo selection on next tick if we aren't keeping or toggling selection we need to clear it out
 	if (FSlateApplication::Get().GetModifierKeys().IsShiftDown() == false || FSlateApplication::Get().GetModifierKeys().IsControlDown() == false)
 	{
diff --git a/Engine/Plugins/Animation/ControlRig/Source/ControlRigEditor/Private/EditMode/SControlRigOutliner.h b/Engine/Plugins/Animation/ControlRig/Source/ControlRigEditor/Private/EditMode/SControlRigOutliner.h
--- a/Engine/Plugins/Animation/ControlRig/Source/ControlRigEditor/Private/EditMode/SControlRigOutliner.h
+++ b/Engine/Plugins/Animation/ControlRig/Source/ControlRigEditor/Private/EditMode/SControlRigOutliner.h
@@ -147,7 +147,7 @@ private:
 
 	FText GetDisplayName() const;
 	FText GetToolTipText() const;
-	FReply OnGetSelectedClicked();
+	FReply OnToggleVisibilityClicked();
 
 };
 
