Control Rig Edit Mode: fixed FK additive rigs not updating their rotation values when manipulated

#rnx

diff --git a/Engine/Plugins/Animation/ControlRig/Source/ControlRigEditor/Private/EditMode/ControlRigEditMode.cpp b/Engine/Plugins/Animation/ControlRig/Source/ControlRigEditor/Private/EditMode/ControlRigEditMode.cpp
--- a/Engine/Plugins/Animation/ControlRig/Source/ControlRigEditor/Private/EditMode/ControlRigEditMode.cpp
+++ b/Engine/Plugins/Animation/ControlRig/Source/ControlRigEditor/Private/EditMode/ControlRigEditMode.cpp
@@ -5500,8 +5500,8 @@ void FControlRigEditMode::MoveControlShape( AControlRigShapeActor* ShapeActor,
 		const TFunction<void()> SetControlShapeTask = [this, ShapeActor, ControlRig, ControlElement, NewTransform, ToWorldTransformCopy, Context,
 			bPrintPythonCommands, bIsTransientControl, InTransformContext, bQueueTasks]()
 		{
-			//fix flips and do rotation orders only if not additive
-			const bool bFixEulerFlips = !ControlRig->IsAdditive() && InTransformContext.bRotation;
+			//fix flips and do rotation orders only if not additive or fk rig
+			const bool bFixEulerFlips = (!ControlRig->IsAdditive() || ControlRig->IsA<UFKControlRig>()) && InTransformContext.bRotation;
 			SetControlShapeTransform(ShapeActor, NewTransform, ToWorldTransformCopy, Context, bPrintPythonCommands, InTransformContext, bFixEulerFlips);
 			NotifyDrivenControls(ControlRig, ShapeActor->GetElementKey(),Context);
 
