Exposing FRigUnit_PositionConstraintLocalSpaceOffset for use in game code.

diff --git a/Engine/Plugins/Animation/ControlRig/Source/ControlRig/Public/Units/Highlevel/Hierarchy/RigUnit_TransformConstraint.h b/Engine/Plugins/Animation/ControlRig/Source/ControlRig/Public/Units/Highlevel/Hierarchy/RigUnit_TransformConstraint.h
--- a/Engine/Plugins/Animation/ControlRig/Source/ControlRig/Public/Units/Highlevel/Hierarchy/RigUnit_TransformConstraint.h
+++ b/Engine/Plugins/Animation/ControlRig/Source/ControlRig/Public/Units/Highlevel/Hierarchy/RigUnit_TransformConstraint.h
@@ -344,7 +344,7 @@ struct CONTROLRIG_API FRigUnit_PositionConstraint : public FRigUnit_HighlevelBas
 * Constrains an item's position to multiple items' positions 
 */
 USTRUCT(meta=(DisplayName="Position Constraint", Category="Constraints", Keywords = "Parent,Translation"))
-struct FRigUnit_PositionConstraintLocalSpaceOffset : public FRigUnit_HighlevelBaseMutable
+struct CONTROLRIG_API FRigUnit_PositionConstraintLocalSpaceOffset : public FRigUnit_HighlevelBaseMutable
 {
 	GENERATED_BODY()
 
