Control Rig Edit Mode: flush any existing pending keyframes when sequencer's playback status changes

#rnx

diff --git a/Engine/Plugins/Animation/ControlRig/Source/ControlRigEditor/Private/EditMode/ControlRigEditMode.cpp b/Engine/Plugins/Animation/ControlRig/Source/ControlRigEditor/Private/EditMode/ControlRigEditMode.cpp
--- a/Engine/Plugins/Animation/ControlRig/Source/ControlRigEditor/Private/EditMode/ControlRigEditMode.cpp
+++ b/Engine/Plugins/Animation/ControlRig/Source/ControlRigEditor/Private/EditMode/ControlRigEditMode.cpp
@@ -1245,12 +1245,7 @@ bool FControlRigEditMode::HandleEndTransform(FEditorViewportClient* InViewportCl
 		return false;
 	}
 
-	const UE::Widget::EWidgetMode WidgetMode = InViewportClient->GetWidgetMode();
-	FTransformContext TransformContext;
-	TransformContext.bTranslation = WidgetMode == UE::Widget::WM_Translate || WidgetMode == UE::Widget::WM_TranslateRotateZ;
-	TransformContext.bRotation = WidgetMode == UE::Widget::WM_Rotate || WidgetMode == UE::Widget::WM_TranslateRotateZ;
-	TransformContext.bScale = WidgetMode == UE::Widget::WM_Scale;
-
+	const FTransformContext TransformContext(InViewportClient->GetWidgetMode());
 	const bool bWasInteracting = bManipulatorMadeChange && InteractionType != (uint8)EControlRigInteractionType::None;
 	
 	InteractionType = (uint8)EControlRigInteractionType::None;
@@ -4025,6 +4020,26 @@ void FControlRigEditMode::UpdateSequencerStatus()
 
     	// update keyframer state
     	const bool bDeferAutokeyOnMouseRelease = !bSequencerPlaying && IsInLevelEditor();
+
+    	// flush any existing pending keyframes
+    	const bool HasPendingKeyframes = Keyframer.IsEnabled() && !bDeferAutokeyOnMouseRelease && !InteractionScopes.IsEmpty();
+    	if (HasPendingKeyframes)
+    	{
+    		if (FEditorModeTools* ModeTools = GetModeManager())
+    		{
+    			const FTransformContext TransformContext(ModeTools->GetWidgetMode());
+    			for (const auto& [ControlRig, InteractionScope] : InteractionScopes)
+    			{
+    				if (InteractionScope)
+    				{
+    					Keyframer.Apply(*InteractionScope, TransformContext);	
+    				}
+    			}
+    			Keyframer.Finalize(ModeTools->GetWorld());
+    		}
+    	}
+
+    	// set new state
     	Keyframer.Enable(bDeferAutokeyOnMouseRelease);
     }
 }
diff --git a/Engine/Plugins/Animation/ControlRig/Source/ControlRigEditor/Public/EditMode/ControlRigEditModeUtil.h b/Engine/Plugins/Animation/ControlRig/Source/ControlRigEditor/Public/EditMode/ControlRigEditModeUtil.h
--- a/Engine/Plugins/Animation/ControlRig/Source/ControlRigEditor/Public/EditMode/ControlRigEditModeUtil.h
+++ b/Engine/Plugins/Animation/ControlRig/Source/ControlRigEditor/Public/EditMode/ControlRigEditModeUtil.h
@@ -9,6 +9,7 @@
 #include "Rigs/RigHierarchy.h"
 #include "Rigs/RigHierarchyDefines.h"
 #include "UObject/WeakObjectPtr.h"
+#include "UnrealWidgetFwd.h"
 
 class FEditorViewportClient;
 class UControlRig;
@@ -38,6 +39,14 @@ enum class ETransformSpace
 
 struct FTransformContext
 {
+	FTransformContext() = default;
+	
+	FTransformContext(const UE::Widget::EWidgetMode& InWidgetMode)
+		: bTranslation(InWidgetMode == UE::Widget::WM_Translate || InWidgetMode == UE::Widget::WM_TranslateRotateZ)
+		, bRotation(InWidgetMode == UE::Widget::WM_Rotate || InWidgetMode == UE::Widget::WM_TranslateRotateZ)
+		, bScale(InWidgetMode == UE::Widget::WM_Scale)
+	{}
+	
 	bool CanTransform() const
 	{
 		return bTranslation || bRotation || bScale;
