Control Rig: Move visibility button in anim outliner into its own column
#rnx

diff --git a/Engine/Plugins/Animation/ControlRig/Source/ControlRig/Private/Settings/ControlRigSettings.cpp b/Engine/Plugins/Animation/ControlRig/Source/ControlRig/Private/Settings/ControlRigSettings.cpp
--- a/Engine/Plugins/Animation/ControlRig/Source/ControlRig/Private/Settings/ControlRigSettings.cpp
+++ b/Engine/Plugins/Animation/ControlRig/Source/ControlRig/Private/Settings/ControlRigSettings.cpp
@@ -35,6 +35,7 @@ UControlRigEditorSettings::UControlRigEditorSettings(const FObjectInitializer& O
 	bUseShortName = false;
 	bArrangeByModules = true;
 	bShowAnimationChannels = true;
+	bFocusOnSelection = true;
 #endif
 }
 
diff --git a/Engine/Plugins/Animation/ControlRig/Source/ControlRig/Public/Settings/ControlRigSettings.h b/Engine/Plugins/Animation/ControlRig/Source/ControlRig/Public/Settings/ControlRigSettings.h
--- a/Engine/Plugins/Animation/ControlRig/Source/ControlRig/Public/Settings/ControlRigSettings.h
+++ b/Engine/Plugins/Animation/ControlRig/Source/ControlRig/Public/Settings/ControlRigSettings.h
@@ -150,6 +150,12 @@ public:
 	UPROPERTY(EditAnywhere, config, Category = Outliner)
 	bool bUseShortName;
 
+	/**
+	 * If turned on the Anim Outliner will focus on the selection.
+	 */
+	UPROPERTY(EditAnywhere, config, Category = Outliner)
+	bool bFocusOnSelection;
+
 #endif
 #if WITH_EDITOR
 	void PostEditChangeProperty(struct FPropertyChangedEvent& PropertyChangedEvent)
diff --git a/Engine/Plugins/Animation/ControlRig/Source/ControlRigEditor/Private/ControlRigHierarchyCommands.cpp b/Engine/Plugins/Animation/ControlRig/Source/ControlRigEditor/Private/ControlRigHierarchyCommands.cpp
--- a/Engine/Plugins/Animation/ControlRig/Source/ControlRigEditor/Private/ControlRigHierarchyCommands.cpp
+++ b/Engine/Plugins/Animation/ControlRig/Source/ControlRigEditor/Private/ControlRigHierarchyCommands.cpp
@@ -34,6 +34,7 @@ void FControlRigHierarchyCommands::RegisterCommands()
 	UI_COMMAND(FilteringFlattensHierarchy, "Filtering Flattens Hierarchy", "Whether to keep the hierarchy or flatten it when searching for tree items", EUserInterfaceActionType::ToggleButton, FInputChord());
 	UI_COMMAND(HideParentsWhenFiltering, "Hide Parents When Filtering", "Whether to show parent items grayed out, or hide them entirely when filtering", EUserInterfaceActionType::ToggleButton, FInputChord());
 	UI_COMMAND(ArrangeByModules, "Arrange by Modules", "Arrange the hierarchy by the modules that spawn each control", EUserInterfaceActionType::ToggleButton, FInputChord());
+	UI_COMMAND(FocusOnSelection, "Focus On Selection", "Focus on the selection when it changes", EUserInterfaceActionType::ToggleButton, FInputChord());
 	UI_COMMAND(ShowFullPaths, "Show Full Paths", "Whether to show the full path or the short name of each element", EUserInterfaceActionType::ToggleButton, FInputChord());
 	UI_COMMAND(ShowImportedBones, "Show Imported Bones", "Whether to show or hide imported bones", EUserInterfaceActionType::ToggleButton, FInputChord());
 	UI_COMMAND(ShowBones, "Show Bones", "Whether to show or hide bones", EUserInterfaceActionType::ToggleButton, FInputChord());
diff --git a/Engine/Plugins/Animation/ControlRig/Source/ControlRigEditor/Private/ControlRigHierarchyCommands.h b/Engine/Plugins/Animation/ControlRig/Source/ControlRigEditor/Private/ControlRigHierarchyCommands.h
--- a/Engine/Plugins/Animation/ControlRig/Source/ControlRigEditor/Private/ControlRigHierarchyCommands.h
+++ b/Engine/Plugins/Animation/ControlRig/Source/ControlRigEditor/Private/ControlRigHierarchyCommands.h
@@ -138,6 +138,9 @@ public:
 	/** Whether to tint the icons with the element color */
 	TSharedPtr< FUICommandInfo > ShowIconColors;
 
+	/* Scroll to selection when it changes */
+	TSharedPtr< FUICommandInfo > FocusOnSelection;
+
 	/**
 	 * Initialize commands
 	 */
diff --git a/Engine/Plugins/Animation/ControlRig/Source/ControlRigEditor/Private/EditMode/SControlRigOutliner.cpp b/Engine/Plugins/Animation/ControlRig/Source/ControlRigEditor/Private/EditMode/SControlRigOutliner.cpp
--- a/Engine/Plugins/Animation/ControlRig/Source/ControlRigEditor/Private/EditMode/SControlRigOutliner.cpp
+++ b/Engine/Plugins/Animation/ControlRig/Source/ControlRigEditor/Private/EditMode/SControlRigOutliner.cpp
@@ -34,6 +34,7 @@
 #include "Framework/Application/SlateApplication.h"
 #include "Editor/EditorEngine.h"
 #include "ScopedTransaction.h"
+#include "Blueprint/WidgetBlueprintLibrary.h"
 #include "Editor/SRigHierarchy.h"
 #include "Rigs/FKControlRig.h"
 
@@ -235,20 +236,33 @@ void FMultiRigTreeElement::RefreshDisplaySettings(const URigHierarchy* InHierarc
 
 	IconBrush = Result.Key;
 	IconColor = Result.Value;
-	if (IconColor.IsColorSpecified() && InSettings.bShowIconColors)
+
+	if (FilterResult == ERigTreeFilterResult::Shown)
 	{
-		IconColor = FilterResult == ERigTreeFilterResult::Shown ? Result.Value : FSlateColor(Result.Value.GetSpecifiedColor() * 0.5f);
+		IconColor = IconColor.IsColorSpecified() && InSettings.bShowIconColors ? Result.Value : FSlateColor::UseForeground();
+		const bool bIsKey = Data.Key.IsSet();
+		const bool bIsModule = Data.ModuleName.IsSet();
+		const bool bIsRig = !bIsModule && !bIsKey;
+		if (bIsRig || bIsModule)
+		{
+			TextColor = FSlateColor(FLinearColor::White);
+		}
+		else
+		{
+			TextColor = FSlateColor::UseForeground();
+		}
 	}
 	else
 	{
-		IconColor = FilterResult == ERigTreeFilterResult::Shown ? FSlateColor::UseForeground() : FSlateColor(FLinearColor::Gray * 0.5f);
+		IconColor = IconColor.IsColorSpecified() && InSettings.bShowIconColors ? FSlateColor(Result.Value.GetSpecifiedColor() * 0.5f) : FSlateColor(FLinearColor::Gray * 0.5f);
+		TextColor = FSlateColor::UseForeground();
 	}
-	TextColor = FilterResult == ERigTreeFilterResult::Shown ? FSlateColor::UseForeground() : FSlateColor(FLinearColor::Gray * 0.5f);
 
 }
 //////////////////////////////////////////////////////////////
 /// SMultiRigHierarchyItem
 ///////////////////////////////////////////////////////////
+TMap<FSoftObjectPath, TSharedPtr<FSlateBrush>> SMultiRigHierarchyItem::IconPathToBrush;
 void SMultiRigHierarchyItem::Construct(const FArguments& InArgs, const TSharedRef<STableViewBase>& OwnerTable, TSharedRef<FMultiRigTreeElement> InRigTreeElement, TSharedPtr<SMultiRigHierarchyTreeView> InTreeView, const FRigTreeDisplaySettings& InSettings, bool bPinned)
 {
 	WeakRigTreeElement = InRigTreeElement;
@@ -259,8 +273,8 @@ void SMultiRigHierarchyItem::Construct(const FArguments& InArgs, const TSharedRe
 	if (InRigTreeElement->Data.ControlRig.IsValid() == false || (InRigTreeElement->Data.Key.IsSet() &&
 		InRigTreeElement->Data.Key.GetValue().IsValid() == false))
 	{
-		STableRow<TSharedPtr<FMultiRigTreeElement>>::Construct(
-			STableRow<TSharedPtr<FMultiRigTreeElement>>::FArguments()
+		SMultiColumnTableRow<TSharedPtr<FMultiRigTreeElement>>::Construct(
+			SMultiColumnTableRow<TSharedPtr<FMultiRigTreeElement>>::FArguments()
 			.ShowSelection(false)
 			.Content()
 			[
@@ -274,91 +288,9 @@ void SMultiRigHierarchyItem::Construct(const FArguments& InArgs, const TSharedRe
 		return;
 	}
 
-	TSharedPtr< SInlineEditableTextBlock > InlineWidget;
-
-	STableRow<TSharedPtr<FMultiRigTreeElement>>::Construct(
-		STableRow<TSharedPtr<FMultiRigTreeElement>>::FArguments()
-		.ShowWires(true)
-		.Content()
-		[
-			SNew(SHorizontalBox)
-			+ SHorizontalBox::Slot()
-			.MaxWidth(24)
-			.FillWidth(1.0)
-			.HAlign(HAlign_Left)
-			.VAlign(VAlign_Center)
-			.Padding(FMargin(0.f, 0.f, 3.f, 0.f))
-			[
-				SNew(SButton)
-				.ButtonStyle(FAppStyle::Get(), "NoBorder")
-				.OnClicked(this, &SMultiRigHierarchyItem::OnToggleVisibilityClicked)
-				[
-					SNew(SImage)
-					.Image_Lambda([this]() -> const FSlateBrush*
-					{
-						if (TSharedPtr<FMultiRigTreeElement> SharedElement = WeakRigTreeElement.Pin())
-						{
-							//if no key is set then it's the control rig or a module so we get that based upon it's state
-							if (!SharedElement->Data.Key.IsSet())
-							{
-								UControlRig* Rig = nullptr;
-								// If we have a module name, we are dealing with a module, get the rig from there
-								if (SharedElement->Data.ModuleName.IsSet())
-								{
-									if (UModularRig* ModularRig = Cast<UModularRig>(SharedElement->Data.ControlRig))
-									{
-										if (FRigModuleInstance* Module = ModularRig->FindModule(SharedElement->Data.ModuleName.GetValue()))
-										{
-											Rig = Module->GetRig();
-										}
-									}
-								}
-								else
-								{
-									Rig = SharedElement->Data.ControlRig.Get();
-								}
-
-								if (Rig)
-								{
-									if (Rig->GetControlsVisible())
-									{
-										return (FAppStyle::GetBrush("Level.VisibleIcon16x"));
-									}
-									else
-									{
-										return  (FAppStyle::GetBrush("Level.NotVisibleIcon16x"));
-									}
-								}
-							}
-						}
-						return WeakRigTreeElement.Pin()->IconBrush;
-					})
-					.ColorAndOpacity_Lambda([this]()
-					{
-						return WeakRigTreeElement.Pin()->IconColor;
-					})
-					.DesiredSizeOverride(FVector2D(16, 16))
-				]
-			]
-			+ SHorizontalBox::Slot()
-			.AutoWidth()
-			.VAlign(VAlign_Center)
-			[
-				SAssignNew(InlineWidget, SInlineEditableTextBlock)
-				.Text(this, &SMultiRigHierarchyItem::GetDisplayName)
-				.ToolTipText(this, &SMultiRigHierarchyItem::GetToolTipText)
-				.MultiLine(false)
-				.ColorAndOpacity_Lambda([this]()
-				{
-					if (WeakRigTreeElement.IsValid())
-					{
-						return WeakRigTreeElement.Pin()->TextColor;
-					}
-					return FSlateColor::UseForeground();
-				})
-			]
-		], OwnerTable);
-
+	SMultiColumnTableRow<TSharedPtr<FMultiRigTreeElement>>::Construct(
+		SMultiColumnTableRow<TSharedPtr<FMultiRigTreeElement>>::FArguments()
+		.ShowWires(true), OwnerTable);
 }
 
 FReply SMultiRigHierarchyItem::OnToggleVisibilityClicked()
@@ -370,11 +302,29 @@ FReply SMultiRigHierarchyItem::OnToggleVisibilityClicked()
 		if (FSlateApplication::Get().GetModifierKeys().IsShiftDown())
 		{
 			TArray<TSharedPtr<FMultiRigTreeElement>> SelectedItems = TreeView->GetSelectedItems();
-			for (TSharedPtr<FMultiRigTreeElement> Selected : SelectedItems)
+
+			// If the element toggled belongs to a selection, toggle all the modules selected
+			if (SelectedItems.Contains(Element))
+			{
+				for (TSharedPtr<FMultiRigTreeElement> Selected : SelectedItems)
+				{
+					if (Selected->Data.ModuleName.IsSet())
+					{
+						ElementsToToggle.AddUnique(Selected);
+					}
+				}
+			}
+			else
 			{
-				if (Selected->Data.ModuleName.IsSet())
+				// If the module toggled does not belong to a selection, toggle all submodules
+				TArray<TSharedPtr<FMultiRigTreeElement>> Descendants = Element->Children;
+				for (int32 i=0; i<Descendants.Num(); i++)
 				{
-					ElementsToToggle.AddUnique(Selected);
+					if (Descendants[i]->Data.ModuleName.IsSet())
+					{
+						ElementsToToggle.AddUnique(Descendants[i]);
+						Descendants.Append(Descendants[i]->Children);
+					}
 				}
 			}
 		}
@@ -409,6 +359,7 @@ FReply SMultiRigHierarchyItem::OnToggleVisibilityClicked()
 						bSetControlsVisible = !Rig->GetControlsVisible();
 					}
 
+					ElementToToggle->bAreControlsVisible = bSetControlsVisible.GetValue();
 					RigsToToggle.AddUnique(Rig);
 				}
 			}
@@ -484,7 +435,7 @@ const FSlateBrush* SMultiRigHierarchyItem::GetBorder() const
 		if (Cur == WeakRigTreeElement)
 		{
 			// item selected
-			return &Style->ActiveHighlightedBrush;
+			return &Style->ActiveBrush;
 		}
 		else
 		{
@@ -511,6 +462,212 @@ const FSlateBrush* SMultiRigHierarchyItem::GetBorder() const
 	return STableRow<TSharedPtr<FMultiRigTreeElement>>::GetBorder();
 }
 
+TSharedRef<SWidget> SMultiRigHierarchyItem::GenerateWidgetForColumn(const FName& InColumnName)
+{
+	if (InColumnName == "Visibility")
+	{
+		return SNew(SButton)
+			.ButtonStyle(FAppStyle::Get(), "NoBorder")
+			.OnClicked(this, &SMultiRigHierarchyItem::OnToggleVisibilityClicked)
+			.OnHovered_Lambda([this]()
+			{
+				if (TSharedPtr<FMultiRigTreeElement> Element = WeakRigTreeElement.Pin())
+				{
+					Element->bIsEyeballIconHovered = true;
+				}
+			})
+			.OnUnhovered_Lambda([this]()
+			{
+				if (TSharedPtr<FMultiRigTreeElement> Element = WeakRigTreeElement.Pin())
+				{
+					Element->bIsEyeballIconHovered = false;
+				}
+			})
+			[
+				SNew(SImage)
+				.Image_Lambda([this]() -> const FSlateBrush*
+				{
+					if (TSharedPtr<FMultiRigTreeElement> SharedElement = WeakRigTreeElement.Pin())
+					{
+						//if no key is set then it's the control rig or a module so we get that based upon it's state
+						if (!SharedElement->Data.Key.IsSet())
+						{
+							UControlRig* Rig = nullptr;
+							// If we have a module name, we are dealing with a module, get the rig from there
+							if (SharedElement->Data.ModuleName.IsSet())
+							{
+								if (UModularRig* ModularRig = Cast<UModularRig>(SharedElement->Data.ControlRig))
+								{
+									if (FRigModuleInstance* Module = ModularRig->FindModule(SharedElement->Data.ModuleName.GetValue()))
+									{
+										Rig = Module->GetRig();
+									}
+								}
+							}
+							else
+							{
+								Rig = SharedElement->Data.ControlRig.Get();
+							}
+
+							if (Rig)
+							{
+								if (Rig->GetControlsVisible())
+								{
+									return (FAppStyle::GetBrush("Level.VisibleIcon16x"));
+								}
+								else
+								{
+									return  (FAppStyle::GetBrush("Level.NotVisibleIcon16x"));
+								}
+							}
+						}
+					}
+					return nullptr;
+				})
+				.ColorAndOpacity_Lambda([this]()
+				{
+					if (TSharedPtr<FMultiRigTreeElement> Element = WeakRigTreeElement.Pin())
+					{
+						if (Element->bAreControlsVisible && !Element->bIsRowHovered && !Element->bIsEyeballIconHovered)
+						{
+							return FSlateColor(FLinearColor::Transparent);
+						}
+						if (Element->bIsEyeballIconHovered)
+						{
+							return FSlateColor::UseForeground();
+						}
+					}
+					return FSlateColor::UseSubduedForeground();
+				})
+				.DesiredSizeOverride(FVector2D(16, 16))
+			];
+	}
+	else if (InColumnName == "Name")
+	{
+		return SNew(SHorizontalBox)
+		+SHorizontalBox::Slot()
+		.AutoWidth()
+		.Padding(4, 2, 2, 0)
+		.VAlign(VAlign_Fill)
+		[
+			SNew(SExpanderArrow, SharedThis(this))
+			.IndentAmount(12)
+			.ShouldDrawWires(true)
+		]
+			
+		+SHorizontalBox::Slot()
+		.AutoWidth()
+		[
+			SNew(SButton)
+			.ButtonStyle(FAppStyle::Get(), "NoBorder")
+			[
+				SNew(SImage)
+				.Visibility_Lambda([this]() -> EVisibility
+				{
+					if (TSharedPtr<FMultiRigTreeElement> SharedElement = WeakRigTreeElement.Pin())
+					{
+						if (SharedElement->Data.ModuleName.IsSet() || SharedElement->Data.Key.IsSet())
+						{
+							return EVisibility::Visible;
+						}
+					}
+					return EVisibility::Collapsed;
+				})
+				.Image_Lambda([this]() -> const FSlateBrush*
+				{
+					if (TSharedPtr<FMultiRigTreeElement> SharedElement = WeakRigTreeElement.Pin())
+					{
+						// If we have a module name, we are dealing with a module, get the rig from there
+						if (SharedElement->Data.ModuleName.IsSet())
+						{
+							UControlRig* Rig = nullptr;
+							if (UModularRig* ModularRig = Cast<UModularRig>(SharedElement->Data.ControlRig))
+							{
+								if (FRigModuleInstance* Module = ModularRig->FindModule(SharedElement->Data.ModuleName.GetValue()))
+								{
+									Rig = Module->GetRig();
+								}
+							}
+							
+							if (Rig)
+							{
+								Rig->GetRigModuleSettings().Icon;
+								const FSoftObjectPath& IconPath = Rig->GetRigModuleSettings().Icon;
+								const TSharedPtr<FSlateBrush>* ExistingBrush = IconPathToBrush.Find(IconPath);
+								if(ExistingBrush && ExistingBrush->IsValid())
+								{
+									return ExistingBrush->Get();
+								}
+								if(UTexture2D* Icon = Cast<UTexture2D>(IconPath.TryLoad()))
+								{
+									const TSharedPtr<FSlateBrush> NewBrush = MakeShareable(new FSlateBrush(UWidgetBlueprintLibrary::MakeBrushFromTexture(Icon, 16.0f, 16.0f)));
+									IconPathToBrush.FindOrAdd(IconPath) = NewBrush;
+									return NewBrush.Get();
+								}
+							}
+						}
+					}
+					return WeakRigTreeElement.Pin()->IconBrush;
+				})
+				.ColorAndOpacity_Lambda([this]()
+				{
+					return WeakRigTreeElement.Pin()->IconColor;
+				})
+				.DesiredSizeOverride(FVector2D(16, 16))
+			]
+		]
+		+SHorizontalBox::Slot()
+		.FillWidth(1.0)
+		[
+			SNew(SInlineEditableTextBlock)
+				.Text(this, &SMultiRigHierarchyItem::GetDisplayName)
+				.ToolTipText(this, &SMultiRigHierarchyItem::GetToolTipText)
+				.MultiLine(false)
+				.Font_Lambda([this]()
+				{
+					if (TSharedPtr<FMultiRigTreeElement> Element = WeakRigTreeElement.Pin())
+					{
+						if (!Element->Data.ModuleName.IsSet() && !Element->Data.Key.IsSet())
+						{
+							return FCoreStyle::GetDefaultFontStyle("Bold", 10);
+						}
+					}
+					return FCoreStyle::GetDefaultFontStyle("Regular", 10);
+				})
+				.ColorAndOpacity_Lambda([this]()
+				{
+					if (TSharedPtr<FMultiRigTreeElement> Element = WeakRigTreeElement.Pin())
+					{
+						bool bIsSelected = false;
+						if (TreeView.IsValid())
+						{
+							bIsSelected = TreeView->GetSelectedData().Contains(Element->Data);
+						}
+						return bIsSelected ? FSlateColor(FLinearColor::White) : Element->TextColor;
+					}
+					return FSlateColor::UseForeground();
+				})
+		];
+	}
+	return SNullWidget::NullWidget;
+}
+
+void SMultiRigHierarchyItem::OnMouseEnter(const FGeometry& MyGeometry, const FPointerEvent& MouseEvent) 
+{
+	if (TSharedPtr<FMultiRigTreeElement> Element = WeakRigTreeElement.Pin())
+	{
+		Element->bIsRowHovered = true;
+	}
+}
+
+void SMultiRigHierarchyItem::OnMouseLeave(const FPointerEvent& MouseEvent) 
+{
+	if (TSharedPtr<FMultiRigTreeElement> Element = WeakRigTreeElement.Pin())
+	{
+		Element->bIsRowHovered = false;
+	}
+}
+
 //////////////////////////////////////////////////////////////
 /// SRigHierarchyTreeView
 ///////////////////////////////////////////////////////////
@@ -529,6 +686,22 @@ void SMultiRigHierarchyTreeView::Construct(const FArguments& InArgs)
 	SuperArgs.SelectionMode(ESelectionMode::Multi);
 	SuperArgs.OnGenerateRow(this, &SMultiRigHierarchyTreeView::MakeTableRowWidget, false);
 	SuperArgs.OnGetChildren(this, &SMultiRigHierarchyTreeView::HandleGetChildrenForTree);
+	SuperArgs.HeaderRow(
+		SNew(SHeaderRow)
+
+		+ SHeaderRow::Column(FName(TEXT("Visibility")))
+		.DefaultLabel(FText::GetEmpty())
+		.FixedWidth(30)
+		.HAlignCell(HAlign_Left)
+		.HAlignHeader(HAlign_Left)
+		.VAlignCell(VAlign_Top)
+
+		+ SHeaderRow::Column(FName(TEXT("Name")))
+		.DefaultLabel(LOCTEXT("ElementName", "Name"))
+		.HAlignCell(HAlign_Fill)
+		.HAlignHeader(HAlign_Fill)
+	);
+
 	SuperArgs.OnSelectionChanged(FOnMultiRigTreeSelectionChanged::CreateRaw(&Delegates, &FMultiRigTreeDelegates::HandleSelectionChanged));
 	SuperArgs.OnContextMenuOpening(Delegates.OnContextMenuOpening);
 	SuperArgs.OnMouseButtonClick(Delegates.OnMouseButtonClick);
@@ -1323,6 +1496,7 @@ const FRigTreeDisplaySettings& SSearchableMultiRigHierarchyTreeView::GetDisplayS
 	Settings.bShowAnimationChannels = UControlRigEditorSettings::Get()->bShowAnimationChannels;
 	Settings.bArrangeByModules = UControlRigEditorSettings::Get()->bArrangeByModules;
 	Settings.bUseShortName = UControlRigEditorSettings::Get()->bUseShortName;
+	Settings.bFocusOnSelection = UControlRigEditorSettings::Get()->bFocusOnSelection;
 	return Settings;
 }
 
@@ -1338,6 +1512,7 @@ TSharedRef<SWidget> SSearchableMultiRigHierarchyTreeView::OnGetOptionsMenu()
 		MenuBuilder.AddMenuEntry(Actions.ArrangeByModules);
 		MenuBuilder.AddMenuEntry(Actions.ShowAnimationChannels);
 		MenuBuilder.AddMenuEntry(Actions.ShowFullPaths);
+		MenuBuilder.AddMenuEntry(Actions.FocusOnSelection);
 		MenuBuilder.AddMenuEntry(EditModeActions.ToggleModuleManipulators);
 	}
 	MenuBuilder.EndSection();
@@ -1354,6 +1529,12 @@ void SSearchableMultiRigHierarchyTreeView::BindCommands()
 		FExecuteAction::CreateSP(this, &SSearchableMultiRigHierarchyTreeView::ToggleArrangeByModules),
 		FCanExecuteAction(),
 		FIsActionChecked::CreateSP(this, &SSearchableMultiRigHierarchyTreeView::IsArrangedByModules));
+
+	CommandList->MapAction(
+		Commands.FocusOnSelection,
+		FExecuteAction::CreateSP(this, &SSearchableMultiRigHierarchyTreeView::ToggleFocusOnSelection),
+		FCanExecuteAction(),
+		FIsActionChecked::CreateSP(this, &SSearchableMultiRigHierarchyTreeView::IsFocusingOnSelection));
 	
 	CommandList->MapAction(
 		Commands.ShowAnimationChannels,
@@ -1454,6 +1635,23 @@ bool SSearchableMultiRigHierarchyTreeView::CanToggleModuleManipulators() const
 	return false;
 }
 
+bool SSearchableMultiRigHierarchyTreeView::IsFocusingOnSelection() const
+{
+	return UControlRigEditorSettings::Get()->bFocusOnSelection;
+}
+
+void SSearchableMultiRigHierarchyTreeView::ToggleFocusOnSelection()
+{
+	UControlRigEditorSettings* EditorSettings = UControlRigEditorSettings::Get();
+	EditorSettings->bFocusOnSelection = !EditorSettings->bFocusOnSelection;
+
+	if (FProperty* Property = UControlRigEditorSettings::StaticClass()->FindPropertyByName(GET_MEMBER_NAME_CHECKED(UControlRigEditorSettings, bFocusOnSelection)))
+	{
+		FPropertyChangedEvent PropertyChangedEvent(Property, EPropertyChangeType::ValueSet);
+		EditorSettings->PostEditChangeProperty(PropertyChangedEvent);
+	}
+}
+
 void SSearchableMultiRigHierarchyTreeView::OnFilterTextChanged(const FText& SearchText)
 {
 	FilterText = SearchText;
@@ -1480,6 +1678,7 @@ void SControlRigOutliner::Construct(const FArguments& InArgs, FControlRigEditMod
 	DisplaySettings.bShowAnimationChannels = true;
 	DisplaySettings.bArrangeByModules = true;
 	DisplaySettings.bUseShortName = true;
+	DisplaySettings.bFocusOnSelection = true;
 
 	FMultiRigTreeDelegates RigTreeDelegates;
 	RigTreeDelegates.OnGetDisplaySettings = FOnGetRigTreeDisplaySettings::CreateSP(this, &SControlRigOutliner::GetDisplaySettings);
@@ -1553,16 +1752,26 @@ void SControlRigOutliner::HandleControlSelected(UControlRig* Subject, FRigContro
 		TSharedPtr<FMultiRigTreeElement> Found = HierarchyTreeView->GetTreeView()->FindElement(Data, HierarchyTreeView->GetTreeView()->GetRootElements()[RootIndex]);
 		if (Found.IsValid())
 		{
-			bool bScrollToElement = !bIsChangingRigHierarchy;
+			bool bScrollToElement = !bIsChangingRigHierarchy && UControlRigEditorSettings::Get()->bFocusOnSelection;
+			if (!bIsChangingRigHierarchy)
+			{
+				HierarchyTreeView->GetTreeView()->ClearSelection();
+			}
 			TGuardValue<bool> GuardRigHierarchyChanges(bIsChangingRigHierarchy, true);
 
 			HierarchyTreeView->GetTreeView()->SetItemSelection(Found, bSelected, ESelectInfo::Direct);
 
-			if (bScrollToElement)
+			if (bSelected && bScrollToElement)
 			{
 				TArray<TSharedPtr<FMultiRigTreeElement>> SelectedItems = HierarchyTreeView->GetTreeView()->GetSelectedItems();
 				if (SelectedItems.Num() > 0)
 				{
+					TSharedPtr<FMultiRigTreeElement> Parent = HierarchyTreeView->GetTreeView()->GetParentElement(Found);
+					while (Parent)
+					{
+						HierarchyTreeView->GetTreeView()->SetItemExpansion(Parent, true);
+						Parent = HierarchyTreeView->GetTreeView()->GetParentElement(Parent);
+					}
 					HierarchyTreeView->GetTreeView()->RequestScrollIntoView(SelectedItems.Last());
 				}
 			}
diff --git a/Engine/Plugins/Animation/ControlRig/Source/ControlRigEditor/Private/EditMode/SControlRigOutliner.h b/Engine/Plugins/Animation/ControlRig/Source/ControlRigEditor/Private/EditMode/SControlRigOutliner.h
--- a/Engine/Plugins/Animation/ControlRig/Source/ControlRigEditor/Private/EditMode/SControlRigOutliner.h
+++ b/Engine/Plugins/Animation/ControlRig/Source/ControlRigEditor/Private/EditMode/SControlRigOutliner.h
@@ -123,22 +123,35 @@ public:
 	ERigTreeFilterResult FilterResult;
 
 	/** The brush to use when rendering an icon */
-	const FSlateBrush* IconBrush;;
+	const FSlateBrush* IconBrush;
 
 	/** The color to use when rendering an icon */
 	FSlateColor IconColor;
 
 	/** The color to use when rendering the label text */
 	FSlateColor TextColor;
+
+	/** Whether or not the controls of this element are visible */
+	bool bAreControlsVisible = true;
+
+	/** Whether or not this row is being hovered */
+	bool bIsRowHovered = false;
+
+	/** Whether or not the eyeball icon is being hovered */
+	bool bIsEyeballIconHovered = false;
+
 };
 
-class SMultiRigHierarchyItem : public STableRow<TSharedPtr<FMultiRigTreeElement>>
+class SMultiRigHierarchyItem : public SMultiColumnTableRow<TSharedPtr<FMultiRigTreeElement>>
 {
 public:
 
 	void Construct(const FArguments& InArgs, const TSharedRef<STableViewBase>& OwnerTable, TSharedRef<FMultiRigTreeElement> InRigTreeElement, TSharedPtr<SMultiRigHierarchyTreeView> InTreeView, const FRigTreeDisplaySettings& InSettings, bool bPinned);
 	static TPair<const FSlateBrush*, FSlateColor> GetBrushForElementType(const URigHierarchy* InHierarchy, const FMultiRigData& InData);
 	virtual const FSlateBrush* GetBorder() const override;
+	virtual TSharedRef<SWidget> GenerateWidgetForColumn( const FName& InColumnName ) override;
+	virtual void OnMouseEnter(const FGeometry& MyGeometry, const FPointerEvent& MouseEvent) override;
+	virtual void OnMouseLeave(const FPointerEvent& MouseEvent) override;
 
 private:
 	TWeakPtr<FMultiRigTreeElement> WeakRigTreeElement;
@@ -149,6 +162,8 @@ private:
 	FText GetToolTipText() const;
 	FReply OnToggleVisibilityClicked();
 
+	static TMap<FSoftObjectPath, TSharedPtr<FSlateBrush>> IconPathToBrush;
+
 };
 
 class SMultiRigHierarchyTreeView : public STreeView<TSharedPtr<FMultiRigTreeElement>>
@@ -270,6 +285,8 @@ public:
 	void ToggleShortNames();
 	void ToggleModuleManipulators();
 	bool CanToggleModuleManipulators() const;
+	bool IsFocusingOnSelection() const;
+	void ToggleFocusOnSelection();
 
 private:
 
diff --git a/Engine/Plugins/Animation/ControlRig/Source/ControlRigEditor/Private/Editor/SRigHierarchyTreeView.h b/Engine/Plugins/Animation/ControlRig/Source/ControlRigEditor/Private/Editor/SRigHierarchyTreeView.h
--- a/Engine/Plugins/Animation/ControlRig/Source/ControlRigEditor/Private/Editor/SRigHierarchyTreeView.h
+++ b/Engine/Plugins/Animation/ControlRig/Source/ControlRigEditor/Private/Editor/SRigHierarchyTreeView.h
@@ -72,9 +72,14 @@ struct CONTROLRIGEDITOR_API FRigTreeDisplaySettings
 	/** Whether to tint the icons with the element color */
 	bool bShowIconColors;
 
+	/** Whether or not to show animation channels */
 	bool bShowAnimationChannels;
 
+	/** Whether or not to arrange the controls into modules (only for modular rigs) */
 	bool bArrangeByModules;
+
+	/** Whether or not to expand and focus the selection when this changes */
+	bool bFocusOnSelection;
 };
 
 DECLARE_DELEGATE_RetVal(const URigHierarchy*, FOnGetRigTreeHierarchy);
