Control Rig: Fix compilation error in linux
#rnx

diff --git a/Engine/Plugins/Animation/ControlRig/Source/ControlRigEditor/Private/EditMode/SControlRigOutliner.cpp b/Engine/Plugins/Animation/ControlRig/Source/ControlRigEditor/Private/EditMode/SControlRigOutliner.cpp
--- a/Engine/Plugins/Animation/ControlRig/Source/ControlRigEditor/Private/EditMode/SControlRigOutliner.cpp
+++ b/Engine/Plugins/Animation/ControlRig/Source/ControlRigEditor/Private/EditMode/SControlRigOutliner.cpp
@@ -591,7 +591,6 @@ TSharedRef<SWidget> SMultiRigHierarchyItem::GenerateWidgetForColumn(const FName&
 							
 							if (Rig)
 							{
-								Rig->GetRigModuleSettings().Icon;
 								const FSoftObjectPath& IconPath = Rig->GetRigModuleSettings().Icon;
 								const TSharedPtr<FSlateBrush>* ExistingBrush = IconPathToBrush.Find(IconPath);
 								if(ExistingBrush && ExistingBrush->IsValid())
