Restore ControlRig Input / Output bones to transfer arrays in editor, as these are used in some rigs to decrease the amount of bones processed

diff --git a/Engine/Plugins/Animation/ControlRig/Source/ControlRig/Public/AnimNode_ControlRigBase.h b/Engine/Plugins/Animation/ControlRig/Source/ControlRig/Public/AnimNode_ControlRigBase.h
--- a/Engine/Plugins/Animation/ControlRig/Source/ControlRig/Public/AnimNode_ControlRigBase.h
+++ b/Engine/Plugins/Animation/ControlRig/Source/ControlRig/Public/AnimNode_ControlRigBase.h
@@ -96,7 +96,7 @@ protected:
 	 * of the input pose transfer phase.
 	 * If this list is empty all bones will be transferred.
 	 */
-	UPROPERTY()
+	UPROPERTY(EditAnywhere, AdvancedDisplay, Category = Settings)
 	TArray<FBoneReference> InputBonesToTransfer;
 
 	/**
@@ -104,7 +104,7 @@ protected:
 	 * of the output pose transfer phase.
 	 * If this list is empty all bones will be transferred.
 	 */
-	UPROPERTY()
+	UPROPERTY(EditAnywhere, AdvancedDisplay, Category = Settings)
 	TArray<FBoneReference> OutputBonesToTransfer;
 
 	UPROPERTY(EditAnywhere, AdvancedDisplay, Category = Settings)
