Control Rig Modular Rig Selection broken.

diff --git a/Engine/Plugins/Animation/ControlRig/Source/ControlRigEditor/Private/EditMode/SControlRigOutliner.cpp b/Engine/Plugins/Animation/ControlRig/Source/ControlRigEditor/Private/EditMode/SControlRigOutliner.cpp
--- a/Engine/Plugins/Animation/ControlRig/Source/ControlRigEditor/Private/EditMode/SControlRigOutliner.cpp
+++ b/Engine/Plugins/Animation/ControlRig/Source/ControlRigEditor/Private/EditMode/SControlRigOutliner.cpp
@@ -1754,7 +1754,7 @@ void SControlRigOutliner::HandleControlSelected(UControlRig* Subject, FRigContro
 			bool bScrollToElement = !bIsChangingRigHierarchy && UControlRigEditorSettings::Get()->bFocusOnSelection;
 			if (!bIsChangingRigHierarchy)
 			{
-				HierarchyTreeView->GetTreeView()->ClearSelection();
+				//mz nope HierarchyTreeView->GetTreeView()->ClearSelection();
 			}
 			TGuardValue<bool> GuardRigHierarchyChanges(bIsChangingRigHierarchy, true);
 
