Control Rig: Fix undo/redo creation of ubergraphs
#rnx

diff --git a/Engine/Plugins/Animation/ControlRig/Source/ControlRigEditor/Private/Editor/ControlRigNewEditor.cpp b/Engine/Plugins/Animation/ControlRig/Source/ControlRigEditor/Private/Editor/ControlRigNewEditor.cpp
--- a/Engine/Plugins/Animation/ControlRig/Source/ControlRigEditor/Private/Editor/ControlRigNewEditor.cpp
+++ b/Engine/Plugins/Animation/ControlRig/Source/ControlRigEditor/Private/Editor/ControlRigNewEditor.cpp
@@ -123,6 +123,17 @@ bool FControlRigEditor::NewDocument_IsVisibleForType(FRigVMEditorBase::ECreatedD
 	return true;
 }
 
+void FControlRigEditor::PostUndo(bool bSuccess)
+{
+	IControlRigNewEditor::PostUndo(bSuccess);
+}
+
+void FControlRigEditor::PostTransaction(bool bSuccess, const FTransaction* Transaction, bool bIsRedo)
+{
+	FRigVMEditorBase::PostTransaction(bSuccess, Transaction, bIsRedo);
+	return FControlRigBaseEditor::PostTransactionImpl(bSuccess, Transaction, bIsRedo);
+}
+
 void FControlRigEditor::Tick(float DeltaTime)
 {
 	PreviewScene.UpdateCaptureContents();
diff --git a/Engine/Plugins/Animation/ControlRig/Source/ControlRigEditor/Private/Editor/ControlRigNewEditor.h b/Engine/Plugins/Animation/ControlRig/Source/ControlRigEditor/Private/Editor/ControlRigNewEditor.h
--- a/Engine/Plugins/Animation/ControlRig/Source/ControlRigEditor/Private/Editor/ControlRigNewEditor.h
+++ b/Engine/Plugins/Animation/ControlRig/Source/ControlRigEditor/Private/Editor/ControlRigNewEditor.h
@@ -89,7 +89,8 @@ public:
 	virtual bool NewDocument_IsVisibleForType(FRigVMEditorBase::ECreatedDocumentType GraphType) const override;
 
 
-	virtual void PostTransaction(bool bSuccess, const FTransaction* Transaction, bool bIsRedo) override { return FControlRigBaseEditor::PostTransactionImpl(bSuccess, Transaction, bIsRedo); }
+	virtual void PostUndo(bool bSuccess) override;
+	virtual void PostTransaction(bool bSuccess, const FTransaction* Transaction, bool bIsRedo) override;
 
 	//  FTickableEditorObject Interface
 	virtual void Tick(float DeltaTime) override;
diff --git a/Engine/Plugins/Runtime/RigVM/Source/RigVMDeveloper/Private/RigVMBlueprint.cpp b/Engine/Plugins/Runtime/RigVM/Source/RigVMDeveloper/Private/RigVMBlueprint.cpp
--- a/Engine/Plugins/Runtime/RigVM/Source/RigVMDeveloper/Private/RigVMBlueprint.cpp
+++ b/Engine/Plugins/Runtime/RigVM/Source/RigVMDeveloper/Private/RigVMBlueprint.cpp
@@ -495,7 +495,8 @@ void URigVMBlueprint::HandleRigVMGraphAdded(const FRigVMClient* InClient, const
 	if(URigVMGraph* Model = InClient->GetModel(InNodePath))
 	{
 		if(!HasAnyFlags(RF_ClassDefaultObject | RF_NeedInitialization | RF_NeedLoad | RF_NeedPostLoad) &&
-			GetOuter() != GetTransientPackage())
+			GetOuter() != GetTransientPackage() &&
+			!GIsTransacting)
 		{
 			CreateEdGraph(Model, true);
 			RecompileVM();
diff --git a/Engine/Plugins/Runtime/RigVM/Source/RigVMEditor/Private/Editor/RigVMEditor.cpp b/Engine/Plugins/Runtime/RigVM/Source/RigVMEditor/Private/Editor/RigVMEditor.cpp
--- a/Engine/Plugins/Runtime/RigVM/Source/RigVMEditor/Private/Editor/RigVMEditor.cpp
+++ b/Engine/Plugins/Runtime/RigVM/Source/RigVMEditor/Private/Editor/RigVMEditor.cpp
@@ -884,6 +884,11 @@ void FRigVMEditorBase::PostRedoImpl(bool bSuccess)
 	PostTransaction(bSuccess, Transaction, true);
 }
 
+void FRigVMEditorBase::PostTransaction(bool bSuccess, const FTransaction* Transaction, bool bIsRedo)
+{
+	ForceEditorRefresh(ERefreshRigVMEditorReason::PostUndo);
+}
+
 void FRigVMEditorBase::OnStartWatchingPin()
 {
 	if (UEdGraphPin* Pin = GetCurrentlySelectedPin())
diff --git a/Engine/Plugins/Runtime/RigVM/Source/RigVMEditor/Private/Editor/RigVMNewEditor.cpp b/Engine/Plugins/Runtime/RigVM/Source/RigVMEditor/Private/Editor/RigVMNewEditor.cpp
--- a/Engine/Plugins/Runtime/RigVM/Source/RigVMEditor/Private/Editor/RigVMNewEditor.cpp
+++ b/Engine/Plugins/Runtime/RigVM/Source/RigVMEditor/Private/Editor/RigVMNewEditor.cpp
@@ -283,6 +283,12 @@ void FRigVMNewEditor::CommonInitialization(const TArray<UBlueprint*>& InitBluepr
 	// FCoreUObjectDelegates::OnObjectPropertyChanged.AddSP(this, &FRigVMNewEditor::OnPostObjectPropertyChanged); // TODO sara-s
 
 	// FKismetEditorUtilities::OnBlueprintUnloaded.AddSP(this, &FRigVMNewEditor::OnBlueprintUnloaded); // TODO sara-s
+
+	UEditorEngine* Editor = Cast<UEditorEngine>(GEngine);
+	if (Editor != nullptr)
+	{
+		Editor->RegisterForUndo(this);
+	}
 }
 
 void FRigVMNewEditor::OnBlueprintChanged(UBlueprint* InBlueprint)
diff --git a/Engine/Plugins/Runtime/RigVM/Source/RigVMEditor/Private/Widgets/SRigVMEditorGraphExplorer.cpp b/Engine/Plugins/Runtime/RigVM/Source/RigVMEditor/Private/Widgets/SRigVMEditorGraphExplorer.cpp
--- a/Engine/Plugins/Runtime/RigVM/Source/RigVMEditor/Private/Widgets/SRigVMEditorGraphExplorer.cpp
+++ b/Engine/Plugins/Runtime/RigVM/Source/RigVMEditor/Private/Widgets/SRigVMEditorGraphExplorer.cpp
@@ -680,8 +680,10 @@ FText SRigVMEditorGraphExplorer::GetGraphDisplayName(const FString& InGraphPath)
 	{
 		const URigVMEdGraphSchema* EdSchema = GetDefault<URigVMEdGraphSchema>();
 		FGraphDisplayInfo DisplayInfo;
-		UEdGraph* EdGraph = Editor->GetRigVMBlueprint()->GetEdGraph(Graph);
-		EdSchema->GetGraphDisplayInformation(*EdGraph, DisplayInfo);
+		if (UEdGraph* EdGraph = Editor->GetRigVMBlueprint()->GetEdGraph(Graph))
+		{
+			EdSchema->GetGraphDisplayInformation(*EdGraph, DisplayInfo);
+		}
 
 		DisplayName = DisplayInfo.DisplayName;
 	}
diff --git a/Engine/Plugins/Runtime/RigVM/Source/RigVMEditor/Public/Editor/RigVMEditor.h b/Engine/Plugins/Runtime/RigVM/Source/RigVMEditor/Public/Editor/RigVMEditor.h
--- a/Engine/Plugins/Runtime/RigVM/Source/RigVMEditor/Public/Editor/RigVMEditor.h
+++ b/Engine/Plugins/Runtime/RigVM/Source/RigVMEditor/Public/Editor/RigVMEditor.h
@@ -343,7 +343,7 @@ protected:
 	void PostUndoImpl(bool bSuccess);
 	void PostRedoImpl(bool bSuccess);
 
-	virtual void PostTransaction(bool bSuccess, const FTransaction* Transaction, bool bIsRedo) {}
+	virtual void PostTransaction(bool bSuccess, const FTransaction* Transaction, bool bIsRedo);
 
 
 	// IToolkitHost Interface
