Control Rig: avoid send event nodes + autokey leading to re-entrant evaluation

- NOTE: the autokey system highlights the problem, but this could easily happen in other scenarios

#rnx

diff --git a/Engine/Plugins/Animation/ControlRig/Source/ControlRigEditor/Private/Sequencer/ControlRigSpaceChannelEditors.cpp b/Engine/Plugins/Animation/ControlRig/Source/ControlRigEditor/Private/Sequencer/ControlRigSpaceChannelEditors.cpp
--- a/Engine/Plugins/Animation/ControlRig/Source/ControlRigEditor/Private/Sequencer/ControlRigSpaceChannelEditors.cpp
+++ b/Engine/Plugins/Animation/ControlRig/Source/ControlRigEditor/Private/Sequencer/ControlRigSpaceChannelEditors.cpp
@@ -1331,6 +1331,12 @@ void FControlRigSpaceChannelHelpers::CompensateIfNeeded(
 		return;
 	}
 
+	if (!ControlRig || ControlRig->IsEvaluating())
+	{
+		// avoid reentrant evaluation: space switching evaluates the rig at T-1, but this must not happen if it's already being evaluated.
+		return;
+	}
+
 	// Frames to compensate
 	TArray<FFrameNumber> OptionalTimeArray;
 	if(OptionalTime.IsSet())
