Control Rig: Add alt+click option to select subtree of an element in the anim outliner
#rnx

diff --git a/Engine/Plugins/Animation/ControlRig/Source/ControlRigEditor/Private/EditMode/SControlRigOutliner.cpp b/Engine/Plugins/Animation/ControlRig/Source/ControlRigEditor/Private/EditMode/SControlRigOutliner.cpp
--- a/Engine/Plugins/Animation/ControlRig/Source/ControlRigEditor/Private/EditMode/SControlRigOutliner.cpp
+++ b/Engine/Plugins/Animation/ControlRig/Source/ControlRigEditor/Private/EditMode/SControlRigOutliner.cpp
@@ -429,35 +429,33 @@ const FSlateBrush* SMultiRigHierarchyItem::GetBorder() const
 	}
 	
 	TArray<TSharedPtr<FMultiRigTreeElement>> SelectedElements = TreeView->GetSelectedItems();
+	if (SelectedElements.Contains(WeakRigTreeElement))
+	{
+		// item selected
+		return &Style->ActiveBrush;
+	}
+	
 	for (TSharedPtr<FMultiRigTreeElement> Selected : SelectedElements)
 	{
 		TSharedPtr<FMultiRigTreeElement> Cur = Selected;
-		if (Cur == WeakRigTreeElement)
-		{
-			// item selected
-			return &Style->ActiveBrush;
-		}
-		else
+		bool bAnyDescendantSelected = false;
+		while (Cur)
 		{
-			bool bAnyDescendantSelected = false;
-			while (Cur)
+			if (Cur == WeakRigTreeElement)
 			{
-				if (Cur == WeakRigTreeElement)
-				{
-					bAnyDescendantSelected = true;
-					break;
-				}
-				else
-				{
-					Cur = TreeView->GetParentElement(Cur);
-				}
+				bAnyDescendantSelected = true;
+				break;
 			}
-			if (bAnyDescendantSelected)
+			else
 			{
-				// Descendant selected
-				return &Style->InactiveHighlightedBrush;
+				Cur = TreeView->GetParentElement(Cur);
 			}
 		}
+		if (bAnyDescendantSelected)
+		{
+			// Descendant selected
+			return &Style->InactiveHighlightedBrush;
+		}
 	}
 	return STableRow<TSharedPtr<FMultiRigTreeElement>>::GetBorder();
 }
@@ -545,6 +543,21 @@ TSharedRef<SWidget> SMultiRigHierarchyItem::GenerateWidgetForColumn(const FName&
 	else if (InColumnName == "Name")
 	{
 		return SNew(SHorizontalBox)
+		.ToolTipText_Lambda([this]()
+		{
+			if (TSharedPtr<FMultiRigTreeElement> Element = WeakRigTreeElement.Pin())
+			{
+				if (Element->Data.Key.IsSet())
+				{
+					return FText::Format(LOCTEXT("SMultiRigHierarchyItemControlTooltip", "{0}"), FText::FromName(Element->Data.Key.GetValue().Name));
+				}
+				else
+				{
+					return FText::Format(LOCTEXT("SMultiRigHierarchyModuleControlTooltip", "{0}\n\nUse Alt+Click to select subtree."), GetDisplayName());
+				}
+			}
+			return FText::GetEmpty();
+		})
 		+SHorizontalBox::Slot()
 		.AutoWidth()
 		.Padding(4, 2, 2, 0)
@@ -703,7 +716,7 @@ void SMultiRigHierarchyTreeView::Construct(const FArguments& InArgs)
 
 	SuperArgs.OnSelectionChanged(FOnMultiRigTreeSelectionChanged::CreateRaw(&Delegates, &FMultiRigTreeDelegates::HandleSelectionChanged));
 	SuperArgs.OnContextMenuOpening(Delegates.OnContextMenuOpening);
-	SuperArgs.OnMouseButtonClick(Delegates.OnMouseButtonClick);
+	SuperArgs.OnMouseButtonClick(FOnMouseButtonClick::CreateSP(this, &SMultiRigHierarchyTreeView::HandleMouseClicked));
 	SuperArgs.OnMouseButtonDoubleClick(Delegates.OnMouseButtonDoubleClick);
 	SuperArgs.OnSetExpansionRecursive(FOnMultiRigTreeSetExpansionRecursive::CreateSP(this, &SMultiRigHierarchyTreeView::SetExpansionRecursive));
 	SuperArgs.HighlightParentNodesForSelection(true);
@@ -1377,6 +1390,16 @@ TSharedPtr<FMultiRigTreeElement> SMultiRigHierarchyTreeView::FindElement(const F
 	return nullptr;
 }
 
+void SMultiRigHierarchyTreeView::HandleMouseClicked(TSharedPtr<FMultiRigTreeElement> InElement)
+{
+	// When alt+clicking an element that is already selected, the actual selection is not changed, so HandleSelectionChanged is not called
+	// However, we want to handle this case to select the subtree of that element
+	if (FSlateApplication::Get().GetModifierKeys().IsAltDown())
+	{
+		Delegates.HandleSelectionChanged(InElement, ESelectInfo::Type::OnMouseClick);
+	}
+}
+
 TArray<FMultiRigData> SMultiRigHierarchyTreeView::GetSelectedData() const
 {
 	TArray<FMultiRigData> Keys;
@@ -1744,19 +1767,28 @@ SControlRigOutliner::~SControlRigOutliner()
 void SControlRigOutliner::HandleControlSelected(UControlRig* Subject, FRigControlElement* ControlElement, bool bSelected)
 {
 	FControlRigBaseDockableView::HandleControlSelected(Subject, ControlElement, bSelected);
+
+	if (bIsChangingRigHierarchy)
+	{
+		// This action was initiated by the outliner
+		return;
+	}
+	
 	const FRigElementKey Key = ControlElement->GetKey();
 	FMultiRigData Data(Subject, Key);
+
+	bool bScrollToElement = UControlRigEditorSettings::Get()->bFocusOnSelection;
 	for (int32 RootIndex = 0; RootIndex < HierarchyTreeView->GetTreeView()->GetRootElements().Num(); ++RootIndex)
 	{
 		TSharedPtr<FMultiRigTreeElement> Found = HierarchyTreeView->GetTreeView()->FindElement(Data, HierarchyTreeView->GetTreeView()->GetRootElements()[RootIndex]);
 		if (Found.IsValid())
 		{
-			bool bScrollToElement = !bIsChangingRigHierarchy && UControlRigEditorSettings::Get()->bFocusOnSelection;
-			if (!bIsChangingRigHierarchy)
+			TGuardValue<bool> GuardRigHierarchyChanges(bIsChangingRigHierarchy, true);
+
+			if (!FSlateApplication::Get().GetModifierKeys().IsShiftDown() && !FSlateApplication::Get().GetModifierKeys().IsControlDown())
 			{
-				//mz nope HierarchyTreeView->GetTreeView()->ClearSelection();
+				HierarchyTreeView->GetTreeView()->ClearSelection();
 			}
-			TGuardValue<bool> GuardRigHierarchyChanges(bIsChangingRigHierarchy, true);
 
 			HierarchyTreeView->GetTreeView()->SetItemSelection(Found, bSelected, ESelectInfo::Direct);
 
@@ -1796,46 +1828,47 @@ void SControlRigOutliner::HandleSelectionChanged(TSharedPtr<FMultiRigTreeElement
 				SelectedRigAndKeys.FindOrAdd(Data.ControlRig.Get()).AddUnique(Data.Key.GetValue());
 			}
 
-			if (!Data.Key.IsSet())
+			if (FSlateApplication::Get().GetModifierKeys().IsAltDown())
 			{
-				TArray<FRigElementKey>& Keys = SelectedRigAndKeys.FindOrAdd(Data.ControlRig.Get());
-				const FRigTreeDisplaySettings& Settings = HierarchyTreeView->GetDisplaySettings();
-				if (Settings.FilterText.IsEmpty() || !Settings.bFlattenHierarchyOnFilter)
+				if (!Data.Key.IsSet())
 				{
-					TSharedPtr<FMultiRigTreeElement> Element = HierarchyTreeView->GetTreeView()->FindElement(Data);
-					TArray<TSharedPtr<FMultiRigTreeElement>> Descendants = Element->Children;
-					for (int32 i=0; i<Descendants.Num(); ++i)
+					TArray<FRigElementKey>& Keys = SelectedRigAndKeys.FindOrAdd(Data.ControlRig.Get());
+					const FRigTreeDisplaySettings& Settings = HierarchyTreeView->GetDisplaySettings();
+					if (Settings.FilterText.IsEmpty() || !Settings.bFlattenHierarchyOnFilter)
 					{
-						TSharedPtr<FMultiRigTreeElement>& Child = Descendants[i];
-						if (Child->Data.Key.IsSet())
-						{
-							Keys.AddUnique(Child->Data.Key.GetValue());
-						}
-						else if (Child->Data.ModuleName.IsSet())
+						TSharedPtr<FMultiRigTreeElement> Element = HierarchyTreeView->GetTreeView()->FindElement(Data);
+						TArray<TSharedPtr<FMultiRigTreeElement>> Descendants = Element->Children;
+						for (int32 i=0; i<Descendants.Num(); ++i)
 						{
+							TSharedPtr<FMultiRigTreeElement>& Child = Descendants[i];
+							if (Child->Data.Key.IsSet())
+							{
+								Keys.AddUnique(Child->Data.Key.GetValue());
+							}
+							
 							AddToSelection.AddUnique(Child);
+							Descendants.Append(Child->Children);
 						}
-						Descendants.Append(Child->Children);
 					}
-				}
-				else if (Data.ModuleName.IsSet())
-				{
-					// If we have flatten the hierarchy due to a search, we cannot rely on the children of the module
-					Data.GetHierarchy()->ForEach<FRigControlElement>([&](FRigControlElement* ControlElement)
+					else if (Data.ModuleName.IsSet())
 					{
-						if (Data.GetHierarchy()->GetModuleName(ControlElement->GetKey()) == Data.ModuleName.GetValue())
+						// If we have flatten the hierarchy due to a search, we cannot rely on the children of the module
+						Data.GetHierarchy()->ForEach<FRigControlElement>([&](FRigControlElement* ControlElement)
 						{
-							Keys.AddUnique(ControlElement->GetKey());
-						}
-						return true;
-					});
-				}
-				else
-				{
-					// If the key is not set, and the module name is not set, we are selecting the root of the rig. Select all controls
-					for (const FRigElementKey ControlKey : Data.GetHierarchy()->GetControlKeys())
+							if (Data.GetHierarchy()->GetModuleName(ControlElement->GetKey()) == Data.ModuleName.GetValue())
+							{
+								Keys.AddUnique(ControlElement->GetKey());
+							}
+							return true;
+						});
+					}
+					else
 					{
-						Keys.AddUnique(ControlKey);
+						// If the key is not set, and the module name is not set, we are selecting the root of the rig. Select all controls
+						for (const FRigElementKey ControlKey : Data.GetHierarchy()->GetControlKeys())
+						{
+							Keys.AddUnique(ControlKey);
+						}
 					}
 				}
 			}
diff --git a/Engine/Plugins/Animation/ControlRig/Source/ControlRigEditor/Private/EditMode/SControlRigOutliner.h b/Engine/Plugins/Animation/ControlRig/Source/ControlRigEditor/Private/EditMode/SControlRigOutliner.h
--- a/Engine/Plugins/Animation/ControlRig/Source/ControlRigEditor/Private/EditMode/SControlRigOutliner.h
+++ b/Engine/Plugins/Animation/ControlRig/Source/ControlRigEditor/Private/EditMode/SControlRigOutliner.h
@@ -219,6 +219,7 @@ public:
 	TSharedRef<ITableRow> MakeTableRowWidget(TSharedPtr<FMultiRigTreeElement> InItem, const TSharedRef<STableViewBase>& OwnerTable, bool bPinned);
 	void HandleGetChildrenForTree(TSharedPtr<FMultiRigTreeElement> InItem, TArray<TSharedPtr<FMultiRigTreeElement>>& OutChildren);
 	TSharedPtr<FMultiRigTreeElement> FindElement(const FMultiRigData& InData);
+	void HandleMouseClicked(TSharedPtr<FMultiRigTreeElement> InElement);
 
 	TArray<FMultiRigData> GetSelectedData() const;
 	const TArray<TSharedPtr<FMultiRigTreeElement>>& GetRootElements() const { return RootElements; }
