Control Rig: Settings for element name display mode

#rnx

diff --git a/Engine/Plugins/Animation/ControlRig/Source/ControlRig/Private/Rigs/RigHierarchy.cpp b/Engine/Plugins/Animation/ControlRig/Source/ControlRig/Private/Rigs/RigHierarchy.cpp
--- a/Engine/Plugins/Animation/ControlRig/Source/ControlRig/Private/Rigs/RigHierarchy.cpp
+++ b/Engine/Plugins/Animation/ControlRig/Source/ControlRig/Private/Rigs/RigHierarchy.cpp
@@ -783,6 +783,7 @@ void URigHierarchy::Reset_Impl(bool bResetElements)
 	TopologyVersion = 0;
 	MetadataVersion = 0;
 	bEnableDirtyPropagation = true;
+	NonUniqueShortNamesCache.Reset();
 
 	if(bResetElements)
 	{
@@ -1275,6 +1276,105 @@ uint32 URigHierarchy::GetTopologyHash(bool bIncludeTopologyVersion, bool bInclud
 }
 
 #if WITH_EDITOR
+
+bool URigHierarchy::HasOnlyUniqueShortNames(ERigElementType InElementType) const
+{
+	if(!NonUniqueShortNamesCache.IsValid(TopologyVersion))
+	{
+		TSet<FRigElementKey> NonUniqueShortNames;
+		int32 MaxNum = 0;
+		for(int32 FlatIndex = 0; FlatIndex < ElementsPerType.Num(); FlatIndex++)
+		{
+			MaxNum = FMath::Max(MaxNum, ElementsPerType[FlatIndex].Num());
+		}
+
+		TSet<FName> UniqueNames;
+		UniqueNames.Reserve(MaxNum);
+
+		for(int32 FlatIndex = 0; FlatIndex < ElementsPerType.Num(); FlatIndex++)
+		{
+			UniqueNames.Reset();
+			
+			for(const FRigBaseElement* Element : ElementsPerType[FlatIndex])
+			{
+				FName Name = Element->GetFName();
+				const FRigHierarchyModulePath Path(Name);
+				if(Path.IsValid())
+				{
+					Name = Path.GetElementFName();
+				}
+				if(UniqueNames.Contains(Name))
+				{
+					NonUniqueShortNames.Add({Name, Element->GetType()});
+				}
+				UniqueNames.Add(Name);
+			}
+		}
+
+		NonUniqueShortNamesCache.Set(NonUniqueShortNames, TopologyVersion);
+	}
+
+	// if we don't have any non-unique short names - exit early
+	if(NonUniqueShortNamesCache.Get().IsEmpty())
+	{
+		return true;
+	}
+
+	// if there's any non-unique short names and we queried for all element types - exit early
+	if(InElementType == ERigElementType::All)
+	{
+		return false;
+	}
+
+	// if we find any non-unique short name for the given element type...
+	const uint8 InElementTypeAsUInt8 = static_cast<uint8>(InElementType);
+	for(const FRigElementKey& Key : NonUniqueShortNamesCache.Get())
+	{
+		if ((InElementTypeAsUInt8 & static_cast<uint8>(Key.Type)) != 0)
+		{
+			return false;
+		}
+	}
+
+	return true;
+}
+
+bool URigHierarchy::HasUniqueShortName(ERigElementType InElementType, const FRigName& InName) const
+{
+	FName Name = InName.GetFName();
+	if(Name.IsNone())
+	{
+		return true;
+	}
+	
+	if(HasOnlyUniqueShortNames(ERigElementType::All))
+	{
+		return true;
+	}
+
+	const FRigHierarchyModulePath Path(InName.ToString());
+	if(Path.IsValid())
+	{
+		Name = Path.GetElementFName(); 
+	}
+
+	return !NonUniqueShortNamesCache.Get().Contains({Name, InElementType});
+}
+
+bool URigHierarchy::HasUniqueShortName(ERigElementType InElementType, const FString& InName) const
+{
+	if(InName.IsEmpty())
+	{
+		return true;
+	}
+	return HasUniqueShortName(InElementType, FRigName(InName));
+}
+
+bool URigHierarchy::HasUniqueShortName(const FRigBaseElement* InElement) const
+{
+	return HasUniqueShortName(InElement->GetType(), FRigName(InElement->GetFName()));
+}
+
 void URigHierarchy::RegisterListeningHierarchy(URigHierarchy* InHierarchy)
 {
 	if (InHierarchy)
@@ -2838,16 +2938,48 @@ FRigName URigHierarchy::GetSafeNewComponentName(const FRigElementKey& InElementK
 	return FRigName(UniqueName);
 }
 
-FText URigHierarchy::GetDisplayNameForUI(const FRigBaseElement* InElement, bool bIncludeModuleName) const
+FText URigHierarchy::GetDisplayNameForUI(const FRigBaseElement* InElement, EElementNameDisplayMode InNameMode) const
 {
 	check(InElement);
 
+#if WITH_EDITOR
+	
 	const FName& DisplayName = InElement->GetDisplayName();
 	FString DisplayNameString = DisplayName.ToString();
 	
 	const FRigHierarchyModulePath ModulePath(DisplayNameString);
 	(void)ModulePath.Split(nullptr, &DisplayNameString);
-	
+
+	if(InNameMode == EElementNameDisplayMode::AssetDefault)
+	{
+		InNameMode = EElementNameDisplayMode::Auto;
+
+		if(const UControlRig* ControlRig = Cast<UControlRig>(GetOuter()))
+		{
+			InNameMode = ControlRig->HierarchySettings.ElementNameDisplayMode;
+		}
+	}
+
+	bool bIncludeModuleName;
+	switch(InNameMode)
+	{
+		case EElementNameDisplayMode::Auto:
+		{
+			bIncludeModuleName = !HasUniqueShortName(InElement);
+			break;
+		}
+		case EElementNameDisplayMode::ForceShort:
+		{
+			bIncludeModuleName = false;
+			break;
+		}
+		case EElementNameDisplayMode::ForceLong:
+		default:
+		{
+			bIncludeModuleName = true;
+			break;
+		}
+	}
 	if(bIncludeModuleName)
 	{
 		const FName ModuleName = GetNameMetadata(InElement->Key, ModuleMetadataName, NAME_None);
@@ -2859,13 +2991,18 @@ FText URigHierarchy::GetDisplayNameForUI(const FRigBaseElement* InElement, bool
 	}
 
 	return FText::FromString(*DisplayNameString);
+#else
+
+	return FText::FromName(InElement->GetFName());
+
+#endif
 }
 
-FText URigHierarchy::GetDisplayNameForUI(const FRigElementKey& InKey, bool bIncludeModuleName) const
+FText URigHierarchy::GetDisplayNameForUI(const FRigElementKey& InKey, EElementNameDisplayMode InNameMode) const
 {
 	if(const FRigBaseElement* Element = Find(InKey))
 	{
-		return GetDisplayNameForUI(Element, bIncludeModuleName);
+		return GetDisplayNameForUI(Element, InNameMode);
 	}
 	return FText();
 }
diff --git a/Engine/Plugins/Animation/ControlRig/Source/ControlRig/Private/Sequencer/MovieSceneControlRigParameterSection.cpp b/Engine/Plugins/Animation/ControlRig/Source/ControlRig/Private/Sequencer/MovieSceneControlRigParameterSection.cpp
--- a/Engine/Plugins/Animation/ControlRig/Source/ControlRig/Private/Sequencer/MovieSceneControlRigParameterSection.cpp
+++ b/Engine/Plugins/Animation/ControlRig/Source/ControlRig/Private/Sequencer/MovieSceneControlRigParameterSection.cpp
@@ -34,6 +34,7 @@
 
 #if WITH_EDITOR
 #include "Misc/MessageDialog.h"
+#include "Settings/ControlRigSettings.h"
 #endif
 
 #include UE_INLINE_GENERATED_CPP_BY_NAME(MovieSceneControlRigParameterSection)
@@ -866,6 +867,8 @@ UMovieSceneControlRigParameterSection::UMovieSceneControlRigParameterSection() :
 	static const FMovieSceneChannelMetaData MetaData("Weight", LOCTEXT("WeightChannelText", "Weight"));
 	ChannelProxy = MakeShared<FMovieSceneChannelProxy>(Weight, MetaData, TMovieSceneExternalValue<float>());
 
+	UControlRigEditorSettings::Get()->OnSettingChanged().AddUObject(this, &UMovieSceneControlRigParameterSection::OnControlRigEditorSettingChanged);
+
 #else
 
 	ChannelProxy = MakeShared<FMovieSceneChannelProxy>(Weight);
@@ -1801,6 +1804,12 @@ EMovieSceneChannelProxyType UMovieSceneControlRigParameterSection::CacheChannelP
 	FMovieSceneChannelProxyData Channels;
 	ControlChannelMap.Empty();
 
+#if WITH_EDITOR
+	const EElementNameDisplayMode ElementNameDisplayMode = UControlRigEditorSettings::Get()->ElementNameDisplayMode;
+#else
+	const EElementNameDisplayMode ElementNameDisplayMode = EElementNameDisplayMode::AssetDefault;
+#endif
+
 	// Need to create the channels in sorted orders, only if we have controls
 	if (ControlRig)
 	{
@@ -1914,7 +1923,7 @@ EMovieSceneChannelProxyType UMovieSceneControlRigParameterSection::CacheChannelP
 					if (const FRigControlElement* ParentControlElement = Cast<FRigControlElement>(Hierarchy->GetFirstParent(ControlElement)))
 					{
 						ParentControlName = ParentControlElement->GetFName();
-						Group = Hierarchy->GetDisplayNameForUI(ParentControlElement);
+						Group = Hierarchy->GetDisplayNameForUI(ParentControlElement, ElementNameDisplayMode);
 					}
 				}
 
@@ -1933,7 +1942,7 @@ EMovieSceneChannelProxyType UMovieSceneControlRigParameterSection::CacheChannelP
 							if (Group.IsEmpty())
 							{
 								ControlChannelMap.Add(Scalar.ParameterName, FChannelMapInfo(ControlIndex, SortOrder, FloatChannelIndex, INDEX_NONE, NAME_None, MaskIndex, CategoryIndex));
-								Group = Hierarchy->GetDisplayNameForUI(ControlElement);
+								Group = Hierarchy->GetDisplayNameForUI(ControlElement, ElementNameDisplayMode);
 								if (bEnabled)
 								{
 									++CategoryIndex;
@@ -1947,7 +1956,7 @@ EMovieSceneChannelProxyType UMovieSceneControlRigParameterSection::CacheChannelP
 							}
 
 							FParameterFloatChannelEditorData EditorData(ControlRig, Scalar.ParameterName, bEnabled, Group, SortOrder);
-							EditorData.MetaData.DisplayText = Hierarchy->GetDisplayNameForUI(ControlElement);
+							EditorData.MetaData.DisplayText = Hierarchy->GetDisplayNameForUI(ControlElement, ElementNameDisplayMode);
 							EditorData.MetaData.PropertyMetaData.Add(UIMin, FString::SanitizeFloat(ControlElement->Settings.MinimumValue.Get<float>()));
 							EditorData.MetaData.PropertyMetaData.Add(UIMax, FString::SanitizeFloat(ControlElement->Settings.MaximumValue.Get<float>()));
 							Channels.Add(Scalar.ParameterCurve, EditorData.MetaData, EditorData.ExternalValues);
@@ -1968,7 +1977,7 @@ EMovieSceneChannelProxyType UMovieSceneControlRigParameterSection::CacheChannelP
 							if (Group.IsEmpty())
 							{
 								ControlChannelMap.Add(Bool.ParameterName, FChannelMapInfo(ControlIndex, SortOrder, BoolChannelIndex, INDEX_NONE, BoolChannelTypeName, MaskIndex, CategoryIndex));
-								Group = Hierarchy->GetDisplayNameForUI(ControlElement);
+								Group = Hierarchy->GetDisplayNameForUI(ControlElement, ElementNameDisplayMode);
 								if (bEnabled)
 								{
 									++CategoryIndex;
@@ -1982,7 +1991,7 @@ EMovieSceneChannelProxyType UMovieSceneControlRigParameterSection::CacheChannelP
 							}
 
 							FMovieSceneChannelMetaData MetaData(Bool.ParameterName, Group, Group, bEnabled);
-							MetaData.DisplayText = Hierarchy->GetDisplayNameForUI(ControlElement);
+							MetaData.DisplayText = Hierarchy->GetDisplayNameForUI(ControlElement, ElementNameDisplayMode);
 							MetaData.SortOrder = SortOrder++;
 							BoolChannelIndex += 1;
 							ControlIndex += 1;
@@ -2005,7 +2014,7 @@ EMovieSceneChannelProxyType UMovieSceneControlRigParameterSection::CacheChannelP
 								if (Group.IsEmpty())
 								{
 									ControlChannelMap.Add(Enum.ParameterName, FChannelMapInfo(ControlIndex, SortOrder, EnumChannelIndex, INDEX_NONE, EnumChannelTypeName, MaskIndex, CategoryIndex));
-									Group = Hierarchy->GetDisplayNameForUI(ControlElement);
+									Group = Hierarchy->GetDisplayNameForUI(ControlElement, ElementNameDisplayMode);
 									if (bEnabled)
 									{
 										++CategoryIndex;
@@ -2019,7 +2028,7 @@ EMovieSceneChannelProxyType UMovieSceneControlRigParameterSection::CacheChannelP
 								}
 
 								FMovieSceneChannelMetaData MetaData(Enum.ParameterName, Group, Group, bEnabled);
-								MetaData.DisplayText = Hierarchy->GetDisplayNameForUI(ControlElement);
+								MetaData.DisplayText = Hierarchy->GetDisplayNameForUI(ControlElement, ElementNameDisplayMode);
 								EnumChannelIndex += 1;
 								ControlIndex += 1;
 								MetaData.SortOrder = SortOrder++;
@@ -2039,7 +2048,7 @@ EMovieSceneChannelProxyType UMovieSceneControlRigParameterSection::CacheChannelP
 								if (Group.IsEmpty())
 								{
 									ControlChannelMap.Add(Integer.ParameterName, FChannelMapInfo(ControlIndex, SortOrder, IntegerChannelIndex, INDEX_NONE, IntegerChannelTypeName, MaskIndex, CategoryIndex));
-									Group = Hierarchy->GetDisplayNameForUI(ControlElement);
+									Group = Hierarchy->GetDisplayNameForUI(ControlElement, ElementNameDisplayMode);
 									if (bEnabled)
 									{
 										++CategoryIndex;
@@ -2053,7 +2062,7 @@ EMovieSceneChannelProxyType UMovieSceneControlRigParameterSection::CacheChannelP
 								}
 
 								FMovieSceneChannelMetaData MetaData(Integer.ParameterName, Group, Group, bEnabled);
-								MetaData.DisplayText = Hierarchy->GetDisplayNameForUI(ControlElement);
+								MetaData.DisplayText = Hierarchy->GetDisplayNameForUI(ControlElement, ElementNameDisplayMode);
 								IntegerChannelIndex += 1;
 								ControlIndex += 1;
 								MetaData.SortOrder = SortOrder++;
@@ -2082,7 +2091,7 @@ EMovieSceneChannelProxyType UMovieSceneControlRigParameterSection::CacheChannelP
 								{
 									++CategoryIndex;
 								}
-								Group = Hierarchy->GetDisplayNameForUI(ControlElement);
+								Group = Hierarchy->GetDisplayNameForUI(ControlElement, ElementNameDisplayMode);
 							}
 							else
 							{
@@ -2118,7 +2127,7 @@ EMovieSceneChannelProxyType UMovieSceneControlRigParameterSection::CacheChannelP
 								{
 									++CategoryIndex;
 								}
-								Group = Hierarchy->GetDisplayNameForUI(ControlElement);
+								Group = Hierarchy->GetDisplayNameForUI(ControlElement, ElementNameDisplayMode);
 							}
 							else
 							{
@@ -2197,7 +2206,7 @@ EMovieSceneChannelProxyType UMovieSceneControlRigParameterSection::CacheChannelP
 								{
 									++CategoryIndex;
 								}
-								Group = Hierarchy->GetDisplayNameForUI(ControlElement);
+								Group = Hierarchy->GetDisplayNameForUI(ControlElement, ElementNameDisplayMode);
 							}
 							else
 							{
@@ -2802,6 +2811,21 @@ bool UMovieSceneControlRigParameterSection::RenameParameterName(const FName& Old
 	return bWasReplaced;
 }
 
+#if WITH_EDITOR
+
+void UMovieSceneControlRigParameterSection::OnControlRigEditorSettingChanged(UObject* InSettingsChanged, struct FPropertyChangedEvent& InPropertyChangedEvent)
+{
+	if(InPropertyChangedEvent.Property)
+	{
+		if(InPropertyChangedEvent.Property->GetFName() == GET_MEMBER_NAME_CHECKED(UControlRigEditorSettings, ElementNameDisplayMode))
+		{
+			ReconstructChannelProxy();
+		}
+	}
+}
+
+#endif
+
 void UMovieSceneControlRigParameterSection::RecreateWithThisControlRig(UControlRig* InControlRig, bool bSetDefault)
 {
 	bool bSameControlRig = (ControlRig == InControlRig);
diff --git a/Engine/Plugins/Animation/ControlRig/Source/ControlRig/Private/Settings/ControlRigSettings.cpp b/Engine/Plugins/Animation/ControlRig/Source/ControlRig/Private/Settings/ControlRigSettings.cpp
--- a/Engine/Plugins/Animation/ControlRig/Source/ControlRig/Private/Settings/ControlRigSettings.cpp
+++ b/Engine/Plugins/Animation/ControlRig/Source/ControlRig/Private/Settings/ControlRigSettings.cpp
@@ -32,10 +32,10 @@ UControlRigEditorSettings::UControlRigEditorSettings(const FObjectInitializer& O
 	bShowStackedHierarchy = false;
 	MaxStackSize = 16;
 	bLeftMouseDragDoesMarquee = true;
-	bUseShortName = false;
 	bArrangeByModules = true;
 	bShowAnimationChannels = true;
 	bFocusOnSelection = true;
+	ElementNameDisplayMode = EElementNameDisplayMode::AssetDefault;
 #endif
 }
 
diff --git a/Engine/Plugins/Animation/ControlRig/Source/ControlRig/Private/Units/Execution/RigUnit_RigModules.cpp b/Engine/Plugins/Animation/ControlRig/Source/ControlRig/Private/Units/Execution/RigUnit_RigModules.cpp
--- a/Engine/Plugins/Animation/ControlRig/Source/ControlRig/Private/Units/Execution/RigUnit_RigModules.cpp
+++ b/Engine/Plugins/Animation/ControlRig/Source/ControlRig/Private/Units/Execution/RigUnit_RigModules.cpp
@@ -76,7 +76,7 @@ FRigUnit_GetItemShortName_Execute()
 
 	if(const URigHierarchy* Hierarchy = ExecuteContext.Hierarchy)
 	{
-		ShortName = *Hierarchy->GetDisplayNameForUI(Item).ToString();
+		ShortName = *Hierarchy->GetDisplayNameForUI(Item, EElementNameDisplayMode::ForceShort).ToString();
 	}
 
 	if(ShortName.IsNone())
diff --git a/Engine/Plugins/Animation/ControlRig/Source/ControlRig/Public/Rigs/RigHierarchy.h b/Engine/Plugins/Animation/ControlRig/Source/ControlRig/Public/Rigs/RigHierarchy.h
--- a/Engine/Plugins/Animation/ControlRig/Source/ControlRig/Public/Rigs/RigHierarchy.h
+++ b/Engine/Plugins/Animation/ControlRig/Source/ControlRig/Public/Rigs/RigHierarchy.h
@@ -231,6 +231,18 @@ public:
 	uint32 GetTopologyHash(bool bIncludeTopologyVersion = true, bool bIncludeTransientControls = false) const;
 
 #if WITH_EDITOR
+	/**
+	 * Returns true if this hierarchy has unique short names for a given element type.
+	 */
+	bool HasOnlyUniqueShortNames(ERigElementType InElementType) const;
+
+	/**
+	 * Returns true if in this hierarchy a given short name is unique
+	 */
+	bool HasUniqueShortName(ERigElementType InElementType, const FRigName& InName) const;
+	bool HasUniqueShortName(ERigElementType InElementType, const FString& InName) const;
+	bool HasUniqueShortName(const FRigBaseElement* InElement) const;
+	
 	/**
 	* Add dependent hierarchies that listens to changes made to this hierarchy
 	* Note: By default, only changes to the initial states of this hierarchy is mirrored to the listening hierarchies
@@ -2154,8 +2166,8 @@ public:
 	/**
 	 * Returns the display label for an element to be used for the UI
 	 */
-	FText GetDisplayNameForUI(const FRigBaseElement* InElement, bool bIncludeModuleName = true) const;
-	FText GetDisplayNameForUI(const FRigElementKey& InKey, bool bIncludeModuleName = true) const;
+	FText GetDisplayNameForUI(const FRigBaseElement* InElement, EElementNameDisplayMode InNameMode = EElementNameDisplayMode::AssetDefault) const;
+	FText GetDisplayNameForUI(const FRigElementKey& InKey, EElementNameDisplayMode InNameMode = EElementNameDisplayMode::AssetDefault) const;
 
 	/**
 	 * Returns the modified event, which can be used to 
@@ -5497,6 +5509,8 @@ private:
 #if WITH_EDITOR
 	mutable TArray<FRigElementKey> ReceivedNameSpaceBasedKeys;
 #endif
+
+	mutable THierarchyCache< TSet< FRigElementKey > > NonUniqueShortNamesCache;
 	
 	friend class URigHierarchyController;
 	friend class UControlRig;
diff --git a/Engine/Plugins/Animation/ControlRig/Source/ControlRig/Public/Rigs/RigHierarchyDefines.h b/Engine/Plugins/Animation/ControlRig/Source/ControlRig/Public/Rigs/RigHierarchyDefines.h
--- a/Engine/Plugins/Animation/ControlRig/Source/ControlRig/Public/Rigs/RigHierarchyDefines.h
+++ b/Engine/Plugins/Animation/ControlRig/Source/ControlRig/Public/Rigs/RigHierarchyDefines.h
@@ -2572,3 +2572,25 @@ private:
 	friend struct FRigUnit_DiscardMatches;
 	friend struct FRigUnit_SetDefaultMatch;
 };
+
+/* 
+ * Defines how to retrieve the UI name for an element
+ */
+UENUM(BlueprintType)
+enum class EElementNameDisplayMode : uint8
+{
+	// Relies on the setting in the referenced asset. The setting for each asset can be changed
+	// in the class defaults - HierarchySettings of the Control Rig.
+	// With this setting you can have different name display mode for each control rig in the sequencer.
+	AssetDefault,
+	
+	// Shows full paths only for elements that need it
+	Auto,
+
+	// Always shows short names
+	// (potentially resulting in clashing labels)
+	ForceShort,
+
+	// Always shows full paths
+	ForceLong,
+};
diff --git a/Engine/Plugins/Animation/ControlRig/Source/ControlRig/Public/Sequencer/MovieSceneControlRigParameterSection.h b/Engine/Plugins/Animation/ControlRig/Source/ControlRig/Public/Sequencer/MovieSceneControlRigParameterSection.h
--- a/Engine/Plugins/Animation/ControlRig/Source/ControlRig/Public/Sequencer/MovieSceneControlRigParameterSection.h
+++ b/Engine/Plugins/Animation/ControlRig/Source/ControlRig/Public/Sequencer/MovieSceneControlRigParameterSection.h
@@ -278,6 +278,9 @@ public:
 
 	bool RenameParameterName(const FName& OldParameterName, const FName& NewParameterName, TOptional<ERigControlType> ControlType);
 
+#if WITH_EDITOR
+	void OnControlRigEditorSettingChanged(UObject* InSettingsChanged, struct FPropertyChangedEvent& InPropertyChangedEvent);
+#endif
 
 private:
 
diff --git a/Engine/Plugins/Animation/ControlRig/Source/ControlRig/Public/Settings/ControlRigSettings.h b/Engine/Plugins/Animation/ControlRig/Source/ControlRig/Public/Settings/ControlRigSettings.h
--- a/Engine/Plugins/Animation/ControlRig/Source/ControlRig/Public/Settings/ControlRigSettings.h
+++ b/Engine/Plugins/Animation/ControlRig/Source/ControlRig/Public/Settings/ControlRigSettings.h
@@ -8,6 +8,7 @@
 
 #include "CoreMinimal.h"
 #include "RigVMSettings.h"
+#include "Rigs/RigHierarchyDefines.h"
 #include "ControlRigGizmoLibrary.h"
 
 #if WITH_EDITOR
@@ -145,10 +146,10 @@ public:
 	bool bShowAnimationChannels;
 
 	/**
-	 * If turned on the controls in the Anim Outliner will be displayed using their short name, without their NameSpace.
+	 * Defines how the element names will be displayed in treeviews, anim outliner etc
 	 */
 	UPROPERTY(EditAnywhere, config, Category = Outliner)
-	bool bUseShortName;
+	EElementNameDisplayMode ElementNameDisplayMode;
 
 	/**
 	 * If turned on the Anim Outliner will focus on the selection.
diff --git a/Engine/Plugins/Animation/ControlRig/Source/ControlRig/Public/Units/RigUnitContext.h b/Engine/Plugins/Animation/ControlRig/Source/ControlRig/Public/Units/RigUnitContext.h
--- a/Engine/Plugins/Animation/ControlRig/Source/ControlRig/Public/Units/RigUnitContext.h
+++ b/Engine/Plugins/Animation/ControlRig/Source/ControlRig/Public/Units/RigUnitContext.h
@@ -48,10 +48,15 @@ struct CONTROLRIG_API FRigHierarchySettings
 	GENERATED_BODY();
 
 	FRigHierarchySettings()
-		: ProceduralElementLimit(2000)
+		: ElementNameDisplayMode(EElementNameDisplayMode::Auto)
+		, ProceduralElementLimit(2000)
 	{
 	}
 
+	// The way to display this hierarchy's element names in the user interface
+	UPROPERTY(EditAnywhere, Category = "Hierarchy Settings")
+	EElementNameDisplayMode ElementNameDisplayMode;
+
 	// Sets the limit for the number of elements to create procedurally
 	UPROPERTY(EditAnywhere, Category = "Hierarchy Settings")
 	int32 ProceduralElementLimit;
diff --git a/Engine/Plugins/Animation/ControlRig/Source/ControlRigEditor/Private/ControlRigElementDetails.cpp b/Engine/Plugins/Animation/ControlRig/Source/ControlRigEditor/Private/ControlRigElementDetails.cpp
--- a/Engine/Plugins/Animation/ControlRig/Source/ControlRigEditor/Private/ControlRigElementDetails.cpp
+++ b/Engine/Plugins/Animation/ControlRig/Source/ControlRigEditor/Private/ControlRigElementDetails.cpp
@@ -4260,7 +4260,7 @@ void FRigControlElementDetails::CustomizeAvailableSpaces(IDetailLayoutBuilder& D
 	DisplaySettings.bShowIconColors = true;
 	DisplaySettings.bShowAnimationChannels = true;
 	DisplaySettings.bArrangeByModules = false;
-	DisplaySettings.bUseShortName = false;
+	DisplaySettings.NameDisplayMode = EElementNameDisplayMode::AssetDefault;
 
 	const TSharedRef<SHorizontalBox> HeaderContentWidget = SNew(SHorizontalBox);
 	HeaderContentWidget->AddSlot()
diff --git a/Engine/Plugins/Animation/ControlRig/Source/ControlRigEditor/Private/ControlRigHierarchyCommands.cpp b/Engine/Plugins/Animation/ControlRig/Source/ControlRigEditor/Private/ControlRigHierarchyCommands.cpp
--- a/Engine/Plugins/Animation/ControlRig/Source/ControlRigEditor/Private/ControlRigHierarchyCommands.cpp
+++ b/Engine/Plugins/Animation/ControlRig/Source/ControlRigEditor/Private/ControlRigHierarchyCommands.cpp
@@ -35,7 +35,6 @@ void FControlRigHierarchyCommands::RegisterCommands()
 	UI_COMMAND(HideParentsWhenFiltering, "Hide Parents When Filtering", "Whether to show parent items grayed out, or hide them entirely when filtering", EUserInterfaceActionType::ToggleButton, FInputChord());
 	UI_COMMAND(ArrangeByModules, "Arrange by Modules", "Arrange the hierarchy by the modules that spawn each control", EUserInterfaceActionType::ToggleButton, FInputChord());
 	UI_COMMAND(FocusOnSelection, "Focus On Selection", "Focus on the selection when it changes", EUserInterfaceActionType::ToggleButton, FInputChord());
-	UI_COMMAND(ShowFullPaths, "Show Full Paths", "Whether to show the full path or the short name of each element", EUserInterfaceActionType::ToggleButton, FInputChord());
 	UI_COMMAND(ShowImportedBones, "Show Imported Bones", "Whether to show or hide imported bones", EUserInterfaceActionType::ToggleButton, FInputChord());
 	UI_COMMAND(ShowBones, "Show Bones", "Whether to show or hide bones", EUserInterfaceActionType::ToggleButton, FInputChord());
 	UI_COMMAND(ShowControls, "Show Controls", "Whether to show or hide controls", EUserInterfaceActionType::ToggleButton, FInputChord());
diff --git a/Engine/Plugins/Animation/ControlRig/Source/ControlRigEditor/Private/ControlRigHierarchyCommands.h b/Engine/Plugins/Animation/ControlRig/Source/ControlRigEditor/Private/ControlRigHierarchyCommands.h
--- a/Engine/Plugins/Animation/ControlRig/Source/ControlRigEditor/Private/ControlRigHierarchyCommands.h
+++ b/Engine/Plugins/Animation/ControlRig/Source/ControlRigEditor/Private/ControlRigHierarchyCommands.h
@@ -102,9 +102,6 @@ public:
 	/* Arrange the controls by modules in a modular rig */
 	TSharedPtr< FUICommandInfo > ArrangeByModules;
 
-	/* Show full paths */
-	TSharedPtr< FUICommandInfo > ShowFullPaths;
-
 	/* Show imported bones */
 	TSharedPtr< FUICommandInfo > ShowImportedBones;
 
diff --git a/Engine/Plugins/Animation/ControlRig/Source/ControlRigEditor/Private/EditMode/SControlRigEditModeTools.cpp b/Engine/Plugins/Animation/ControlRig/Source/ControlRigEditor/Private/EditMode/SControlRigEditModeTools.cpp
--- a/Engine/Plugins/Animation/ControlRig/Source/ControlRigEditor/Private/EditMode/SControlRigEditModeTools.cpp
+++ b/Engine/Plugins/Animation/ControlRig/Source/ControlRigEditor/Private/EditMode/SControlRigEditModeTools.cpp
@@ -222,7 +222,7 @@ void SControlRigEditModeTools::Construct(const FArguments& InArgs, TSharedPtr<FC
 	DisplaySettings.bShowIconColors = true;
 	DisplaySettings.bShowAnimationChannels = true;
 	DisplaySettings.bArrangeByModules = false;
-	DisplaySettings.bUseShortName = false;
+	DisplaySettings.NameDisplayMode = EElementNameDisplayMode::AssetDefault;
 #if USE_LOCAL_DETAILS
 	FRigTreeDelegates RigTreeDelegates;
 	RigTreeDelegates.OnGetHierarchy = FOnGetRigTreeHierarchy::CreateSP(this, &SControlRigEditModeTools::GetHierarchy);
diff --git a/Engine/Plugins/Animation/ControlRig/Source/ControlRigEditor/Private/EditMode/SControlRigOutliner.cpp b/Engine/Plugins/Animation/ControlRig/Source/ControlRigEditor/Private/EditMode/SControlRigOutliner.cpp
--- a/Engine/Plugins/Animation/ControlRig/Source/ControlRigEditor/Private/EditMode/SControlRigOutliner.cpp
+++ b/Engine/Plugins/Animation/ControlRig/Source/ControlRigEditor/Private/EditMode/SControlRigOutliner.cpp
@@ -37,6 +37,7 @@
 #include "Blueprint/WidgetBlueprintLibrary.h"
 #include "Editor/SRigHierarchy.h"
 #include "Rigs/FKControlRig.h"
+#include "SEnumCombo.h"
 
 #define LOCTEXT_NAMESPACE "ControlRigOutliner"
 
@@ -96,7 +97,7 @@ FText FMultiRigData::GetDisplayName(const FRigTreeDisplaySettings& InSettings) c
 			{
 				const FRigElementKey& RigElementKey = Key.GetValue();
 				const FRigControlElement* ControlElement = Hierarchy->Find<FRigControlElement>(RigElementKey);
-				bool bIncludeNameSpace = true;
+				EElementNameDisplayMode ElementNameDisplayMode = InSettings.NameDisplayMode;
 				if(ControlElement)
 				{
 					// animation channels should not show their module name if they belong to a control
@@ -106,12 +107,12 @@ FText FMultiRigData::GetDisplayName(const FRigTreeDisplaySettings& InSettings) c
 						const FRigElementKey ParentElementKey = Hierarchy->GetFirstParent(RigElementKey);
 						if(Hierarchy->GetModuleFName(ParentElementKey) == Hierarchy->GetModuleFName(RigElementKey))
 						{
-							bIncludeNameSpace = false;
+							ElementNameDisplayMode = EElementNameDisplayMode::ForceShort;
 						}
 					}
 				}
 				
-				const FText DisplayNameForUI = Hierarchy->GetDisplayNameForUI(RigElementKey, !InSettings.bUseShortName);
+				const FText DisplayNameForUI = Hierarchy->GetDisplayNameForUI(RigElementKey, ElementNameDisplayMode);
 				if (!DisplayNameForUI.IsEmpty())
 				{
 					CachedDisplayName = DisplayNameForUI;
@@ -1517,8 +1518,8 @@ const FRigTreeDisplaySettings& SSearchableMultiRigHierarchyTreeView::GetDisplayS
 	Settings.FilterText = FilterText;
 	Settings.bShowAnimationChannels = UControlRigEditorSettings::Get()->bShowAnimationChannels;
 	Settings.bArrangeByModules = UControlRigEditorSettings::Get()->bArrangeByModules;
-	Settings.bUseShortName = UControlRigEditorSettings::Get()->bUseShortName;
 	Settings.bFocusOnSelection = UControlRigEditorSettings::Get()->bFocusOnSelection;
+	Settings.NameDisplayMode = UControlRigEditorSettings::Get()->ElementNameDisplayMode;
 	return Settings;
 }
 
@@ -1533,7 +1534,28 @@ TSharedRef<SWidget> SSearchableMultiRigHierarchyTreeView::OnGetOptionsMenu()
 		const FControlRigEditModeCommands& EditModeActions = FControlRigEditModeCommands::Get();
 		MenuBuilder.AddMenuEntry(Actions.ArrangeByModules);
 		MenuBuilder.AddMenuEntry(Actions.ShowAnimationChannels);
-		MenuBuilder.AddMenuEntry(Actions.ShowFullPaths);
+
+		MenuBuilder.AddWidget(
+			SNew(SVerticalBox)
+
+			+ SVerticalBox::Slot()
+			.AutoHeight()
+			.Padding(3)
+			[
+				SNew(SEnumComboBox, StaticEnum<EElementNameDisplayMode>())
+				.CurrentValue_Lambda([this]() -> int32
+				{
+					return static_cast<int32>(GetElementNameDisplayMode());
+				})
+				.OnEnumSelectionChanged_Lambda([this](int32 InEnumValue, ESelectInfo::Type)
+				{
+					SetElementNameDisplayMode(static_cast<EElementNameDisplayMode>(InEnumValue));
+				})
+			]
+			,
+			LOCTEXT("ElementNameDisplayMode", "Name Mode")
+		);
+		
 		MenuBuilder.AddMenuEntry(Actions.FocusOnSelection);
 		MenuBuilder.AddMenuEntry(EditModeActions.ToggleModuleManipulators);
 	}
@@ -1564,12 +1586,6 @@ void SSearchableMultiRigHierarchyTreeView::BindCommands()
 		FCanExecuteAction(),
 		FIsActionChecked::CreateSP(this, &SSearchableMultiRigHierarchyTreeView::IsShowingAnimationChannels));
 
-	CommandList->MapAction(
-		Commands.ShowFullPaths,
-		FExecuteAction::CreateSP(this, &SSearchableMultiRigHierarchyTreeView::ToggleShortNames),
-		FCanExecuteAction(),
-		FIsActionChecked::CreateSP(this, &SSearchableMultiRigHierarchyTreeView::IsShowingFullPaths));
-
 	CommandList->MapAction(
 		FControlRigEditModeCommands::Get().ToggleModuleManipulators,
 		FExecuteAction::CreateSP(this, &SSearchableMultiRigHierarchyTreeView::ToggleModuleManipulators),
@@ -1617,22 +1633,22 @@ void SSearchableMultiRigHierarchyTreeView::ToggleShowAnimationChannels()
 	TreeView->RequestTreeViewRefresh();
 }
 
-bool SSearchableMultiRigHierarchyTreeView::IsShowingFullPaths() const
+EElementNameDisplayMode SSearchableMultiRigHierarchyTreeView::GetElementNameDisplayMode() const
 {
-	return !UControlRigEditorSettings::Get()->bUseShortName;
+	UControlRigEditorSettings* EditorSettings = UControlRigEditorSettings::Get();
+	return EditorSettings->ElementNameDisplayMode;
 }
 
-void SSearchableMultiRigHierarchyTreeView::ToggleShortNames()
+void SSearchableMultiRigHierarchyTreeView::SetElementNameDisplayMode(EElementNameDisplayMode InElementNameDisplayMode)
 {
 	UControlRigEditorSettings* EditorSettings = UControlRigEditorSettings::Get();
-	EditorSettings->bUseShortName = !EditorSettings->bUseShortName;
+	EditorSettings->ElementNameDisplayMode = InElementNameDisplayMode;
 
-	if (FProperty* Property = UControlRigEditorSettings::StaticClass()->FindPropertyByName(GET_MEMBER_NAME_CHECKED(UControlRigEditorSettings, bUseShortName)))
+	if (FProperty* Property = UControlRigEditorSettings::StaticClass()->FindPropertyByName(GET_MEMBER_NAME_CHECKED(UControlRigEditorSettings, ElementNameDisplayMode)))
 	{
 		FPropertyChangedEvent PropertyChangedEvent(Property, EPropertyChangeType::ValueSet);
 		EditorSettings->PostEditChangeProperty(PropertyChangedEvent);
 	}
-	
 	TreeView->RequestTreeViewRefresh();
 }
 
@@ -1699,8 +1715,8 @@ void SControlRigOutliner::Construct(const FArguments& InArgs, FControlRigEditMod
 	DisplaySettings.bShowConnectors = false;
 	DisplaySettings.bShowAnimationChannels = true;
 	DisplaySettings.bArrangeByModules = true;
-	DisplaySettings.bUseShortName = true;
 	DisplaySettings.bFocusOnSelection = true;
+	DisplaySettings.NameDisplayMode = EElementNameDisplayMode::AssetDefault;
 
 	FMultiRigTreeDelegates RigTreeDelegates;
 	RigTreeDelegates.OnGetDisplaySettings = FOnGetRigTreeDisplaySettings::CreateSP(this, &SControlRigOutliner::GetDisplaySettings);
diff --git a/Engine/Plugins/Animation/ControlRig/Source/ControlRigEditor/Private/EditMode/SControlRigOutliner.h b/Engine/Plugins/Animation/ControlRig/Source/ControlRigEditor/Private/EditMode/SControlRigOutliner.h
--- a/Engine/Plugins/Animation/ControlRig/Source/ControlRigEditor/Private/EditMode/SControlRigOutliner.h
+++ b/Engine/Plugins/Animation/ControlRig/Source/ControlRigEditor/Private/EditMode/SControlRigOutliner.h
@@ -265,8 +265,9 @@ class SSearchableMultiRigHierarchyTreeView : public SCompoundWidget
 {
 public:
 
-	SLATE_BEGIN_ARGS(SSearchableMultiRigHierarchyTreeView) {}
-	SLATE_ARGUMENT(FMultiRigTreeDelegates, RigTreeDelegates)
+	SLATE_BEGIN_ARGS(SSearchableMultiRigHierarchyTreeView)
+	{}
+		SLATE_ARGUMENT(FMultiRigTreeDelegates, RigTreeDelegates)
 		SLATE_ARGUMENT(FText, InitialFilterText)
 	SLATE_END_ARGS()
 
@@ -282,8 +283,8 @@ public:
 	void ToggleArrangeByModules();
 	bool IsShowingAnimationChannels() const;
 	void ToggleShowAnimationChannels();
-	bool IsShowingFullPaths() const;
-	void ToggleShortNames();
+	EElementNameDisplayMode GetElementNameDisplayMode() const;
+	void SetElementNameDisplayMode(EElementNameDisplayMode InElementNameDisplayMode);
 	void ToggleModuleManipulators();
 	bool CanToggleModuleManipulators() const;
 	bool IsFocusingOnSelection() const;
diff --git a/Engine/Plugins/Animation/ControlRig/Source/ControlRigEditor/Private/Editor/ControlRigEditor.cpp b/Engine/Plugins/Animation/ControlRig/Source/ControlRigEditor/Private/Editor/ControlRigEditor.cpp
--- a/Engine/Plugins/Animation/ControlRig/Source/ControlRigEditor/Private/Editor/ControlRigEditor.cpp
+++ b/Engine/Plugins/Animation/ControlRig/Source/ControlRigEditor/Private/Editor/ControlRigEditor.cpp
@@ -3572,7 +3572,12 @@ void FControlRigBaseEditor::OnFinishedChangingPropertiesImpl(const FPropertyChan
 
 	if (ControlRigBP)
 	{
-		if (PropertyChangedEvent.MemberProperty->GetNameCPP() == GET_MEMBER_NAME_STRING_CHECKED(UControlRigBlueprint, HierarchySettings))
+		if (PropertyChangedEvent.Property->GetNameCPP() == GET_MEMBER_NAME_STRING_CHECKED(FRigHierarchySettings, ElementNameDisplayMode))
+		{
+			Compile();
+		}
+		
+		else if (PropertyChangedEvent.MemberProperty->GetNameCPP() == GET_MEMBER_NAME_STRING_CHECKED(UControlRigBlueprint, HierarchySettings))
 		{
 			ControlRigBP->PropagateHierarchyFromBPToInstances();
 		}
diff --git a/Engine/Plugins/Animation/ControlRig/Source/ControlRigEditor/Private/Editor/ControlRigSchematicModel.cpp b/Engine/Plugins/Animation/ControlRig/Source/ControlRigEditor/Private/Editor/ControlRigSchematicModel.cpp
--- a/Engine/Plugins/Animation/ControlRig/Source/ControlRigEditor/Private/Editor/ControlRigSchematicModel.cpp
+++ b/Engine/Plugins/Animation/ControlRig/Source/ControlRigEditor/Private/Editor/ControlRigSchematicModel.cpp
@@ -75,7 +75,7 @@ const FText& FControlRigSchematicRigElementKeyNode::GetLabel() const
 				for(const FRigElementKey& Connector : Connectors)
 				{
 					static const FText NodeLabelConnectionFormat = LOCTEXT("NodeLabelConnectionFormat", "{0}\nConnection: {1}");
-					const FText ConnectorShortestPath = Hierarchy->GetDisplayNameForUI(Connector, true);
+					const FText ConnectorShortestPath = Hierarchy->GetDisplayNameForUI(Connector, EElementNameDisplayMode::ForceLong);
 					LabelText = FText::Format(NodeLabelConnectionFormat, LabelText, ConnectorShortestPath);
 				}
 			}
@@ -1172,7 +1172,7 @@ bool FControlRigSchematicModel::GetContextMenuForNode(const FSchematicGraphNode*
 				// note: this is a copy on purpose since it is passed into the lambda
 				for(const FRigElementKey Connector : Connectors)
 				{
-					FText Label = ModularRig->GetHierarchy()->GetDisplayNameForUI(Connector, false);
+					FText Label = ModularRig->GetHierarchy()->GetDisplayNameForUI(Connector, EElementNameDisplayMode::ForceShort);
 					const FText Description = FText::FromString(FString::Printf(TEXT("Disconnect %s"), *Label.ToString()));
 					
 					OutMenu.AddMenuEntry(Description, Description, FSlateIcon(), FUIAction(
diff --git a/Engine/Plugins/Animation/ControlRig/Source/ControlRigEditor/Private/Editor/SRigConnectorTargetWidget.cpp b/Engine/Plugins/Animation/ControlRig/Source/ControlRigEditor/Private/Editor/SRigConnectorTargetWidget.cpp
--- a/Engine/Plugins/Animation/ControlRig/Source/ControlRigEditor/Private/Editor/SRigConnectorTargetWidget.cpp
+++ b/Engine/Plugins/Animation/ControlRig/Source/ControlRigEditor/Private/Editor/SRigConnectorTargetWidget.cpp
@@ -153,7 +153,8 @@ void SRigConnectorTargetComboButton::PopulateButtonBox()
 						SNew( STextBlock )
 						.Text_Lambda([this, Hierarchy]()
 						{
-							return Hierarchy->GetDisplayNameForUI(TargetKey.Get(), true);
+							const EElementNameDisplayMode NameDisplayMode = RigTreeDelegates.GetDisplaySettings().NameDisplayMode;
+							return Hierarchy->GetDisplayNameForUI(TargetKey.Get(), NameDisplayMode);
 						})
 						.Font( IDetailLayoutBuilder::GetDetailFont() )
 					]
diff --git a/Engine/Plugins/Animation/ControlRig/Source/ControlRigEditor/Private/Editor/SRigHierarchy.cpp b/Engine/Plugins/Animation/ControlRig/Source/ControlRigEditor/Private/Editor/SRigHierarchy.cpp
--- a/Engine/Plugins/Animation/ControlRig/Source/ControlRigEditor/Private/Editor/SRigHierarchy.cpp
+++ b/Engine/Plugins/Animation/ControlRig/Source/ControlRigEditor/Private/Editor/SRigHierarchy.cpp
@@ -13,6 +13,7 @@
 #include "Widgets/Layout/SSpacer.h"
 #include "Widgets/Images/SImage.h"
 #include "Widgets/Text/STextBlock.h"
+#include "SEnumCombo.h"
 #include "K2Node_VariableGet.h"
 #include "RigVMBlueprintUtils.h"
 #include "ControlRigHierarchyCommands.h"
@@ -401,12 +402,6 @@ void SRigHierarchy::BindCommands()
 		FCanExecuteAction(),
 		FIsActionChecked::CreateLambda([this]() { return DisplaySettings.bHideParentsOnFilter; }));
 
-	CommandList->MapAction(
-		Commands.ShowFullPaths,
-		FExecuteAction::CreateLambda([this]() { DisplaySettings.bUseShortName = !DisplaySettings.bUseShortName; RefreshTreeView(); }),
-		FCanExecuteAction(),
-		FIsActionChecked::CreateLambda([this]() { return !DisplaySettings.bUseShortName; }));
-
 	CommandList->MapAction(
 		Commands.ShowImportedBones,
 		FExecuteAction::CreateLambda([this]() { DisplaySettings.bShowImportedBones = !DisplaySettings.bShowImportedBones; RefreshTreeView(); }),
@@ -1195,7 +1190,46 @@ TSharedRef< SWidget > SRigHierarchy::CreateFilterMenu()
 	{
 		MenuBuilder.AddMenuEntry(Actions.FilteringFlattensHierarchy);
 		MenuBuilder.AddMenuEntry(Actions.HideParentsWhenFiltering);
-		MenuBuilder.AddMenuEntry(Actions.ShowFullPaths);
+
+		MenuBuilder.AddWidget(
+			SNew(SVerticalBox)
+
+			+ SVerticalBox::Slot()
+			.AutoHeight()
+			.Padding(3)
+			[
+				SNew(SEnumComboBox, StaticEnum<EElementNameDisplayMode>())
+				.CurrentValue_Lambda([this]() -> int32
+				{
+					const IControlRigBaseEditor* StrongEditor = ControlRigEditor.Pin().Get();
+					check(StrongEditor);
+					
+					if (UControlRigBlueprint* Blueprint = StrongEditor->GetControlRigBlueprint())
+					{
+						return static_cast<int32>(Blueprint->HierarchySettings.ElementNameDisplayMode);
+					}
+					
+					return static_cast<int32>(EElementNameDisplayMode::AssetDefault);
+				})
+				.OnEnumSelectionChanged_Lambda([this](int32 InEnumValue, ESelectInfo::Type)
+				{
+					IControlRigBaseEditor* StrongEditor = ControlRigEditor.Pin().Get();
+					check(StrongEditor);
+										
+					if (UControlRigBlueprint* Blueprint = StrongEditor->GetControlRigBlueprint())
+					{
+						FScopedTransaction Transaction(LOCTEXT("HierarchySetElementNameDisplayMode", "Change Name Mode"));
+						Blueprint->Modify();
+						Blueprint->HierarchySettings.ElementNameDisplayMode = static_cast<EElementNameDisplayMode>(InEnumValue);
+						StrongEditor->Compile();
+					}
+				})
+			],
+			LOCTEXT("ElementNameDisplayMode", "Name Mode"),
+			false,
+			true,
+			LOCTEXT("NameModeToolTip", "Defines how the names of the elements will be shown in the tree view")
+		);
 	}
 	MenuBuilder.EndSection();
 
diff --git a/Engine/Plugins/Animation/ControlRig/Source/ControlRigEditor/Private/Editor/SRigHierarchyTreeView.cpp b/Engine/Plugins/Animation/ControlRig/Source/ControlRigEditor/Private/Editor/SRigHierarchyTreeView.cpp
--- a/Engine/Plugins/Animation/ControlRig/Source/ControlRigEditor/Private/Editor/SRigHierarchyTreeView.cpp
+++ b/Engine/Plugins/Animation/ControlRig/Source/ControlRigEditor/Private/Editor/SRigHierarchyTreeView.cpp
@@ -51,8 +51,8 @@ FRigTreeElement::FRigTreeElement(const FRigHierarchyKey& InKey, TWeakPtr<SRigHie
 		{
 			if(InKey.IsElement())
 			{
-				LongName = Hierarchy->GetDisplayNameForUI(InKey.GetElement(), true);
-				ShortName = Hierarchy->GetDisplayNameForUI(InKey.GetElement(), false);
+				LongName = Hierarchy->GetDisplayNameForUI(InKey.GetElement(), EElementNameDisplayMode::ForceLong);
+				ShortName = Hierarchy->GetDisplayNameForUI(InKey.GetElement(), EElementNameDisplayMode::ForceShort);
 			}
 			
 			const FRigTreeDisplaySettings& Settings = InTreeView.Pin()->GetRigTreeDelegates().GetDisplaySettings();
@@ -256,10 +256,10 @@ void SRigHierarchyItem::Construct(const FArguments& InArgs, const TSharedRef<STa
 
 FText SRigHierarchyItem::GetNameForUI() const
 {
-	return GetName(Delegates.GetDisplaySettings().bUseShortName);
+	return GetName(Delegates.GetDisplaySettings().NameDisplayMode);
 }
 
-FText SRigHierarchyItem::GetName(bool bUseShortName) const
+FText SRigHierarchyItem::GetName(EElementNameDisplayMode InNameDisplayMode) const
 {
 	if(WeakRigTreeElement.Pin()->bIsTransient)
 	{
@@ -270,10 +270,32 @@ FText SRigHierarchyItem::GetName(bool bUseShortName) const
 	{
 		return FText::FromName(WeakRigTreeElement.Pin()->ChannelName);
 	}
-	if(bUseShortName)
+
+	if(InNameDisplayMode == EElementNameDisplayMode::AssetDefault)
+	{
+		if(const URigHierarchy* Hierarchy = Delegates.GetHierarchy())
+		{
+			if(const UControlRig* ControlRig = Cast<UControlRig>(Hierarchy->GetOuter()))
+			{
+				InNameDisplayMode = ControlRig->HierarchySettings.ElementNameDisplayMode;
+			}
+		}
+	}
+
+	if(InNameDisplayMode == EElementNameDisplayMode::Auto)
+	{
+		if(const URigHierarchy* Hierarchy = Delegates.GetHierarchy())
+		{
+			const FRigElementKey& Key = WeakRigTreeElement.Pin()->Key.GetElement();
+			InNameDisplayMode = Hierarchy->HasUniqueShortName(Key.Type, Key.Name) ? EElementNameDisplayMode::ForceShort : EElementNameDisplayMode::ForceLong;
+		}
+	}
+
+	if(InNameDisplayMode == EElementNameDisplayMode::ForceShort)
 	{
 		return WeakRigTreeElement.Pin()->ShortName;
 	}
+	
 	return WeakRigTreeElement.Pin()->LongName;
 }
 
@@ -287,8 +309,8 @@ FText SRigHierarchyItem::GetItemTooltip() const
 			return ToolTip.GetValue();
 		}
 	}
-	const FText FullName = GetName(false);
-	const FText ShortName = GetName(true);
+	const FText FullName = GetName(EElementNameDisplayMode::ForceLong);
+	const FText ShortName = GetName(EElementNameDisplayMode::ForceShort);
 	if(FullName.EqualTo(ShortName))
 	{
 		return FText();
@@ -656,8 +678,8 @@ bool SRigHierarchyTreeView::AddElement(const FRigBaseElement* InElement)
 								}
 							}
 
-							const bool bUseShortName = GetRigTreeDelegates().GetDisplaySettings().bUseShortName;
-							const FText Label = Hierarchy->GetDisplayNameForUI(ConnectorKey, !bUseShortName);
+							const EElementNameDisplayMode NameDisplayMode = GetRigTreeDelegates().GetDisplaySettings().NameDisplayMode;
+							const FText Label = Hierarchy->GetDisplayNameForUI(ConnectorKey, NameDisplayMode);
 							TagArguments.Text(Label);
 							TagArguments.TooltipText(FText::FromName(ConnectorKey.Name));
 							TagArguments.Color(BackgroundColor);
diff --git a/Engine/Plugins/Animation/ControlRig/Source/ControlRigEditor/Private/Editor/SRigHierarchyTreeView.h b/Engine/Plugins/Animation/ControlRig/Source/ControlRigEditor/Private/Editor/SRigHierarchyTreeView.h
--- a/Engine/Plugins/Animation/ControlRig/Source/ControlRigEditor/Private/Editor/SRigHierarchyTreeView.h
+++ b/Engine/Plugins/Animation/ControlRig/Source/ControlRigEditor/Private/Editor/SRigHierarchyTreeView.h
@@ -20,7 +20,7 @@ struct CONTROLRIGEDITOR_API FRigTreeDisplaySettings
 
 		bFlattenHierarchyOnFilter = false;
 		bHideParentsOnFilter = false;
-		bUseShortName = false;
+		NameDisplayMode = EElementNameDisplayMode::AssetDefault;
 		bShowImportedBones = true;
 		bShowBones = true;
 		bShowControls = true;
@@ -42,8 +42,8 @@ struct CONTROLRIGEDITOR_API FRigTreeDisplaySettings
 	/** Hide parents when text filtering is active */
 	bool bHideParentsOnFilter;
 
-	/** When true, the elements will show their short name */
-	bool bUseShortName;
+	/** The mode used to determine how names are displayed */
+	EElementNameDisplayMode NameDisplayMode;
 
 	/** Whether or not to show imported bones in the hierarchy */
 	bool bShowImportedBones;
@@ -297,7 +297,7 @@ private:
  	FRigTreeDelegates Delegates;
 
 	FText GetNameForUI() const;
-	FText GetName(bool bUseShortName) const;
+	FText GetName(EElementNameDisplayMode InNameDisplayMode) const;
 	FText GetItemTooltip() const;
 
 	friend class SRigHierarchyTreeView; 
