Animation Mode: allow custom sequence to be used in animation mode

- ULevelSequence are the default sequences used, but custom UMovieSceneSequence subclasses can now be registered as a fallback solution
- this gives the oportunity for constraints, snapper, animation tools, etc. to use custom UMovieSceneSequence

diff --git a/Engine/Plugins/Animation/ControlRig/Source/ControlRigEditor/Private/EditMode/ControlRigEditMode.cpp b/Engine/Plugins/Animation/ControlRig/Source/ControlRigEditor/Private/EditMode/ControlRigEditMode.cpp
--- a/Engine/Plugins/Animation/ControlRig/Source/ControlRigEditor/Private/EditMode/ControlRigEditMode.cpp
+++ b/Engine/Plugins/Animation/ControlRig/Source/ControlRigEditor/Private/EditMode/ControlRigEditMode.cpp
@@ -76,7 +76,6 @@
 #include "Editor/Sequencer/Private/SSequencer.h"
 #include "Slate/SceneViewport.h"
 #include "EditorGizmos/EditorTransformGizmoUtil.h"
-#include "Tools/BakingHelper.h"
 #include "Sequencer/AnimLayers/AnimLayers.h"
 
 #include UE_INLINE_GENERATED_CPP_BY_NAME(ControlRigEditMode)
@@ -454,7 +453,7 @@ void FControlRigEditMode::Enter()
 		{
 			if (WeakSequencer.IsValid() == false)
 			{
-				SetSequencer(FBakingHelper::GetSequencer());
+				SetSequencer(UE::AnimationEditMode::GetSequencer());
 			}
 		}
 		if (!Toolkit.IsValid())
diff --git a/Engine/Plugins/Animation/ControlRig/Source/ControlRigEditor/Private/Editor/Constraints/SConstraintsWidget.cpp b/Engine/Plugins/Animation/ControlRig/Source/ControlRigEditor/Private/Editor/Constraints/SConstraintsWidget.cpp
--- a/Engine/Plugins/Animation/ControlRig/Source/ControlRigEditor/Private/Editor/Constraints/SConstraintsWidget.cpp
+++ b/Engine/Plugins/Animation/ControlRig/Source/ControlRigEditor/Private/Editor/Constraints/SConstraintsWidget.cpp
@@ -20,13 +20,13 @@
 #include "Tools/ConstraintBaker.h"
 #include "ScopedTransaction.h"
 #include "ISequencer.h"
-#include "Tools/BakingHelper.h"
 #include "ConstraintsManager.h"
 #include "Styling/SlateIconFinder.h"
 #include "Widgets/Input/NumericTypeInterface.h"
 #include "FrameNumberDetailsCustomization.h"
 #include "SAssetDropTarget.h"
 #include "Sequencer/ControlRigParameterTrackEditor.h"
+#include "Sequencer/EditModeAnimationUtil.h"
 #include "Transform/TransformConstraint.h"
 #include "Transform/TransformConstraintUtil.h"
 
@@ -150,7 +150,7 @@ namespace UE::Private
 
 	static TWeakPtr<ISequencer> GetSequencerChecked()
 	{
-		const TWeakPtr<ISequencer> WeakSequencer = FBakingHelper::GetSequencer();
+		const TWeakPtr<ISequencer> WeakSequencer = AnimationEditMode::GetSequencer();
 		if (!WeakSequencer.IsValid() || !WeakSequencer.Pin()->GetFocusedMovieSceneSequence())
 		{
 			return nullptr;
diff --git a/Engine/Plugins/Animation/ControlRig/Source/ControlRigEditor/Private/Sequencer/EditModeAnimationUtil.cpp b/Engine/Plugins/Animation/ControlRig/Source/ControlRigEditor/Private/Sequencer/EditModeAnimationUtil.cpp
--- a/Engine/Plugins/Animation/ControlRig/Source/ControlRigEditor/Private/Sequencer/EditModeAnimationUtil.cpp
+++ b/Engine/Plugins/Animation/ControlRig/Source/ControlRigEditor/Private/Sequencer/EditModeAnimationUtil.cpp
@@ -4,15 +4,62 @@
 
 #include "ControlRig.h"
 #include "ConstraintsManager.h"
+#include "LevelSequenceEditorBlueprintLibrary.h"
 #include "Constraints/ControlRigTransformableHandle.h"
 #include "Constraints/MovieSceneConstraintChannelHelper.h"
 #include "EditMode/ControlRigEditModeUtil.h"
+#include "ILevelSequenceEditorToolkit.h"
+#include "LevelEditorSequencerIntegration.h"
+#include "LevelSequence.h"
 #include "Rigs/RigHierarchyDefines.h"
 #include "Sequencer/AnimationAuthoringSettings.h"
 
 namespace UE::AnimationEditMode
 {
 
+TWeakPtr<ISequencer> GetSequencer()
+{
+	// if getting sequencer from level sequence need to use the current(leader), not the focused
+	if (ULevelSequence* LevelSequence = ULevelSequenceEditorBlueprintLibrary::GetCurrentLevelSequence())
+	{
+		if (UAssetEditorSubsystem* AssetEditorSubsystem = GEditor->GetEditorSubsystem<UAssetEditorSubsystem>())
+		{
+			static constexpr bool bFocusIfOpen = false;
+			IAssetEditorInstance* AssetEditor = AssetEditorSubsystem->FindEditorForAsset(LevelSequence, bFocusIfOpen);
+			const ILevelSequenceEditorToolkit* LevelSequenceEditor = static_cast<ILevelSequenceEditorToolkit*>(AssetEditor);
+			return LevelSequenceEditor ? LevelSequenceEditor->GetSequencer() : nullptr;
+		}
+	}
+
+	// look for custom UMovieSceneSequence as a fallback.
+	const FCustomMovieSceneRegistry& Registry = FCustomMovieSceneRegistry::Get();
+	const TArray<TWeakPtr<ISequencer>> Sequencers = FLevelEditorSequencerIntegration::Get().GetSequencers();
+	const int32 Found = Sequencers.IndexOfByPredicate([&Registry](const TWeakPtr<ISequencer>& WeakSequence)
+	{
+		if (const TSharedPtr<ISequencer> Sequencer = WeakSequence.IsValid() ? WeakSequence.Pin() : nullptr)
+		{
+			if (UMovieSceneSequence* MovieSceneSequence = Sequencer->GetRootMovieSceneSequence())
+			{
+				return Registry.IsSequenceSupported(MovieSceneSequence->GetClass());
+			}
+		}
+		return false;
+	});
+		
+	return Found == INDEX_NONE ? nullptr : Sequencers[Found];
+}
+
+FCustomMovieSceneRegistry& FCustomMovieSceneRegistry::Get()
+{
+	static FCustomMovieSceneRegistry Singleton;
+	return Singleton;
+}
+
+bool FCustomMovieSceneRegistry::IsSequenceSupported(const UClass* InSequenceClass) const
+{
+	return InSequenceClass ? SupportedSequenceTypes.Contains(InSequenceClass) : false;
+}
+	
 namespace Private
 {
 
diff --git a/Engine/Plugins/Animation/ControlRig/Source/ControlRigEditor/Private/Tools/BakingHelper.cpp b/Engine/Plugins/Animation/ControlRig/Source/ControlRigEditor/Private/Tools/BakingHelper.cpp
--- a/Engine/Plugins/Animation/ControlRig/Source/ControlRigEditor/Private/Tools/BakingHelper.cpp
+++ b/Engine/Plugins/Animation/ControlRig/Source/ControlRigEditor/Private/Tools/BakingHelper.cpp
@@ -3,26 +3,9 @@
 
 #include "Tools/BakingHelper.h"
 
-#include "ILevelSequenceEditorToolkit.h"
-#include "ISequencer.h"
-#include "LevelSequence.h"
-#include "LevelSequenceEditorBlueprintLibrary.h"
-#include "MovieSceneSequence.h"
-#include "MovieSceneToolHelpers.h"
-#include "Channels/MovieSceneChannelProxy.h"
+#include "Sequencer/EditModeAnimationUtil.h"
 
 TWeakPtr<ISequencer> FBakingHelper::GetSequencer()
 {
-	// if getting sequencer from level sequence need to use the current(leader), not the focused
-	if (ULevelSequence* LevelSequence = ULevelSequenceEditorBlueprintLibrary::GetCurrentLevelSequence())
-	{
-		if (UAssetEditorSubsystem* AssetEditorSubsystem = GEditor->GetEditorSubsystem<UAssetEditorSubsystem>())
-		{
-			static constexpr bool bFocusIfOpen = false;
-			IAssetEditorInstance* AssetEditor = AssetEditorSubsystem->FindEditorForAsset(LevelSequence, bFocusIfOpen);
-			const ILevelSequenceEditorToolkit* LevelSequenceEditor = static_cast<ILevelSequenceEditorToolkit*>(AssetEditor);
-			return LevelSequenceEditor ? LevelSequenceEditor->GetSequencer() : nullptr;
-		}
-	}
-	return nullptr;
+	return UE::AnimationEditMode::GetSequencer();
 }
diff --git a/Engine/Plugins/Animation/ControlRig/Source/ControlRigEditor/Private/Tools/ControlRigSnapper.cpp b/Engine/Plugins/Animation/ControlRig/Source/ControlRigEditor/Private/Tools/ControlRigSnapper.cpp
--- a/Engine/Plugins/Animation/ControlRig/Source/ControlRigEditor/Private/Tools/ControlRigSnapper.cpp
+++ b/Engine/Plugins/Animation/ControlRig/Source/ControlRigEditor/Private/Tools/ControlRigSnapper.cpp
@@ -29,9 +29,9 @@
 #include "MovieSceneToolHelpers.h"
 #include "LevelSequencePlayer.h"
 #include "LevelSequenceActor.h"
-#include "Tools/BakingHelper.h"
 #include "Containers/SortedMap.h"
 #include "Sequencer/ControlRigSequencerHelpers.h"
+#include "Sequencer/EditModeAnimationUtil.h"
 #include "Transform/TransformConstraint.h"
 #include "Transform/TransformConstraintUtil.h"
 #include "Constraints/MovieSceneConstraintChannelHelper.inl"
@@ -109,7 +109,7 @@ int32 FControlRigSnapperSelection::NumSelected() const
 
 TWeakPtr<ISequencer> FControlRigSnapper::GetSequencer()
 {
-	return FBakingHelper::GetSequencer();
+	return UE::AnimationEditMode::GetSequencer();
 }
 
 static bool LocalGetControlRigControlTransforms(IMovieScenePlayer* Player, const TOptional<FFrameNumber>& CurrentFrame, UMovieSceneSequence* MovieSceneSequence, FMovieSceneSequenceIDRef Template, FMovieSceneSequenceTransform& RootToLocalTransform,
diff --git a/Engine/Plugins/Animation/ControlRig/Source/ControlRigEditor/Public/Sequencer/EditModeAnimationUtil.h b/Engine/Plugins/Animation/ControlRig/Source/ControlRigEditor/Public/Sequencer/EditModeAnimationUtil.h
--- a/Engine/Plugins/Animation/ControlRig/Source/ControlRigEditor/Public/Sequencer/EditModeAnimationUtil.h
+++ b/Engine/Plugins/Animation/ControlRig/Source/ControlRigEditor/Public/Sequencer/EditModeAnimationUtil.h
@@ -3,10 +3,13 @@
 #pragma once
 
 #include "Containers/Map.h"
+#include "Containers/Set.h"
 #include "Math/Transform.h"
 #include "Misc/EnumClassFlags.h"
+#include "MovieSceneSequence.h"
 
 class FControlRigInteractionScope;
+class ISequencer;
 class UAnimationAuthoringSettings;
 class UControlRig;
 
@@ -20,6 +23,48 @@ namespace UE::ControlRigEditMode
 namespace UE::AnimationEditMode
 {
 
+// Returns the current sequencer.
+TWeakPtr<ISequencer> GetSequencer();
+	
+/**
+ * FCustomMovieSceneRegistry contains custom UMovieSceneSequence that support constraints (among other things)
+ * This allows other types than ULevelSequence to manage constraints.
+ * Registration can be done at module startup (for example) as follows:
+ * FCustomMovieSceneRegistry& Registry = FCustomMovieSceneRegistry::Get();
+ * Registry.RegisterSequence<UMyCustomSequence>();
+*/
+
+class FCustomMovieSceneRegistry
+{
+public:
+	~FCustomMovieSceneRegistry() = default;
+
+	static CONTROLRIGEDITOR_API FCustomMovieSceneRegistry& Get();
+
+	// Registers a particular UMovieSceneSequence subclass to support constraints.
+	template<typename SequenceType>
+	void RegisterSequence()
+	{
+		static_assert(TIsDerivedFrom<SequenceType, UMovieSceneSequence>::Value,
+			"The template class SequenceType must be a subclass of UMovieSceneSequence.");
+	
+		UClass* SequenceClass = SequenceType::StaticClass();
+		if (ensureAlways(!SequenceClass->HasAnyClassFlags(CLASS_Abstract)))
+		{
+			SupportedSequenceTypes.Add(SequenceClass);
+		}
+	}
+
+	// Whether a particular UMovieSceneSequence subclass is supported.
+	bool CONTROLRIGEDITOR_API IsSequenceSupported(const UClass* InSequenceClass) const;
+
+private:
+	FCustomMovieSceneRegistry() = default;
+
+	// List of supported UMovieSceneSequence classes.
+	TSet<UClass*> SupportedSequenceTypes;
+};
+
 /**
  * FControlKeyframeData provides a way of passing the various keyframe parameters a control needs to set / know about.
  * Extend it if necessary to pass in more data to the keyframer.
diff --git a/Engine/Plugins/Animation/ControlRig/Source/ControlRigEditor/Public/Tools/BakingHelper.h b/Engine/Plugins/Animation/ControlRig/Source/ControlRigEditor/Public/Tools/BakingHelper.h
--- a/Engine/Plugins/Animation/ControlRig/Source/ControlRigEditor/Public/Tools/BakingHelper.h
+++ b/Engine/Plugins/Animation/ControlRig/Source/ControlRigEditor/Public/Tools/BakingHelper.h
@@ -2,15 +2,9 @@
 
 #pragma once
 
-#include "Misc/FrameNumber.h"
 #include "Templates/SharedPointer.h"
-#include "Math/Transform.h"
-
 
 class ISequencer;
-class UControlRig;
-struct FFrameRate;
-enum class EMovieSceneTransformChannel : uint32;
 
 /**
  * FBakingHelper
@@ -18,6 +12,6 @@ enum class EMovieSceneTransformChannel : uint32;
 
 struct FBakingHelper
 {
-	/** Returns the current sequencer. */
+	UE_DEPRECATED(5.6, "This function was moved to the UE::AnimationEditMode namespace.")
 	static TWeakPtr<ISequencer> GetSequencer();
 };
