Control Rig: Anim outliner showing actor and components. Also remove the show animation channels option.
#rnx

diff --git a/Engine/Plugins/Animation/ControlRig/Source/ControlRig/Private/Settings/ControlRigSettings.cpp b/Engine/Plugins/Animation/ControlRig/Source/ControlRig/Private/Settings/ControlRigSettings.cpp
--- a/Engine/Plugins/Animation/ControlRig/Source/ControlRig/Private/Settings/ControlRigSettings.cpp
+++ b/Engine/Plugins/Animation/ControlRig/Source/ControlRig/Private/Settings/ControlRigSettings.cpp
@@ -33,7 +33,6 @@ UControlRigEditorSettings::UControlRigEditorSettings(const FObjectInitializer& O
 	MaxStackSize = 16;
 	bLeftMouseDragDoesMarquee = true;
 	bArrangeByModules = true;
-	bShowAnimationChannels = true;
 	bFocusOnSelection = true;
 	ElementNameDisplayMode = EElementNameDisplayMode::AssetDefault;
 #endif
diff --git a/Engine/Plugins/Animation/ControlRig/Source/ControlRig/Public/Settings/ControlRigSettings.h b/Engine/Plugins/Animation/ControlRig/Source/ControlRig/Public/Settings/ControlRigSettings.h
--- a/Engine/Plugins/Animation/ControlRig/Source/ControlRig/Public/Settings/ControlRigSettings.h
+++ b/Engine/Plugins/Animation/ControlRig/Source/ControlRig/Public/Settings/ControlRigSettings.h
@@ -139,12 +139,6 @@ public:
 	UPROPERTY(EditAnywhere, config, Category = Outliner)
 	bool bArrangeByModules;
 
-	/**
-	 * If turned on the animation channel controls in the Anim Outliner will be displayed.
-	 */
-	UPROPERTY(EditAnywhere, config, Category = Outliner)
-	bool bShowAnimationChannels;
-
 	/**
 	 * Defines how the element names will be displayed in treeviews, anim outliner etc
 	 */
diff --git a/Engine/Plugins/Animation/ControlRig/Source/ControlRigEditor/Private/ControlRigElementDetails.cpp b/Engine/Plugins/Animation/ControlRig/Source/ControlRigEditor/Private/ControlRigElementDetails.cpp
--- a/Engine/Plugins/Animation/ControlRig/Source/ControlRigEditor/Private/ControlRigElementDetails.cpp
+++ b/Engine/Plugins/Animation/ControlRig/Source/ControlRigEditor/Private/ControlRigElementDetails.cpp
@@ -4258,7 +4258,6 @@ void FRigControlElementDetails::CustomizeAvailableSpaces(IDetailLayoutBuilder& D
 	DisplaySettings.bHideParentsOnFilter = true;
 	DisplaySettings.bFlattenHierarchyOnFilter = true;
 	DisplaySettings.bShowIconColors = true;
-	DisplaySettings.bShowAnimationChannels = true;
 	DisplaySettings.bArrangeByModules = false;
 	DisplaySettings.NameDisplayMode = EElementNameDisplayMode::AssetDefault;
 
diff --git a/Engine/Plugins/Animation/ControlRig/Source/ControlRigEditor/Private/ControlRigHierarchyCommands.cpp b/Engine/Plugins/Animation/ControlRig/Source/ControlRigEditor/Private/ControlRigHierarchyCommands.cpp
--- a/Engine/Plugins/Animation/ControlRig/Source/ControlRigEditor/Private/ControlRigHierarchyCommands.cpp
+++ b/Engine/Plugins/Animation/ControlRig/Source/ControlRigEditor/Private/ControlRigHierarchyCommands.cpp
@@ -38,7 +38,6 @@ void FControlRigHierarchyCommands::RegisterCommands()
 	UI_COMMAND(ShowImportedBones, "Show Imported Bones", "Whether to show or hide imported bones", EUserInterfaceActionType::ToggleButton, FInputChord());
 	UI_COMMAND(ShowBones, "Show Bones", "Whether to show or hide bones", EUserInterfaceActionType::ToggleButton, FInputChord());
 	UI_COMMAND(ShowControls, "Show Controls", "Whether to show or hide controls", EUserInterfaceActionType::ToggleButton, FInputChord());
-	UI_COMMAND(ShowAnimationChannels, "Show Animation Channels", "Whether to show or hide animation channels", EUserInterfaceActionType::ToggleButton, FInputChord());
 	UI_COMMAND(ShowNulls, "Show Nulls", "Whether to show or hide nulls", EUserInterfaceActionType::ToggleButton, FInputChord());
 	UI_COMMAND(ShowReferences, "Show References", "Whether to show or hide references", EUserInterfaceActionType::ToggleButton, FInputChord());
 	UI_COMMAND(ShowSockets, "Show Sockets", "Whether to show or hide sockets", EUserInterfaceActionType::ToggleButton, FInputChord());
diff --git a/Engine/Plugins/Animation/ControlRig/Source/ControlRigEditor/Private/ControlRigHierarchyCommands.h b/Engine/Plugins/Animation/ControlRig/Source/ControlRigEditor/Private/ControlRigHierarchyCommands.h
--- a/Engine/Plugins/Animation/ControlRig/Source/ControlRigEditor/Private/ControlRigHierarchyCommands.h
+++ b/Engine/Plugins/Animation/ControlRig/Source/ControlRigEditor/Private/ControlRigHierarchyCommands.h
@@ -111,9 +111,6 @@ public:
 	/* Show controls */
 	TSharedPtr< FUICommandInfo > ShowControls;
 
-	/* Show animation channels */
-	TSharedPtr< FUICommandInfo > ShowAnimationChannels;
-
 	/* Show spaces */
 	TSharedPtr< FUICommandInfo > ShowNulls;
 
diff --git a/Engine/Plugins/Animation/ControlRig/Source/ControlRigEditor/Private/EditMode/SControlRigEditModeTools.cpp b/Engine/Plugins/Animation/ControlRig/Source/ControlRigEditor/Private/EditMode/SControlRigEditModeTools.cpp
--- a/Engine/Plugins/Animation/ControlRig/Source/ControlRigEditor/Private/EditMode/SControlRigEditModeTools.cpp
+++ b/Engine/Plugins/Animation/ControlRig/Source/ControlRigEditor/Private/EditMode/SControlRigEditModeTools.cpp
@@ -220,7 +220,6 @@ void SControlRigEditModeTools::Construct(const FArguments& InArgs, TSharedPtr<FC
 	DisplaySettings.bHideParentsOnFilter = true;
 	DisplaySettings.bFlattenHierarchyOnFilter = true;
 	DisplaySettings.bShowIconColors = true;
-	DisplaySettings.bShowAnimationChannels = true;
 	DisplaySettings.bArrangeByModules = false;
 	DisplaySettings.NameDisplayMode = EElementNameDisplayMode::AssetDefault;
 #if USE_LOCAL_DETAILS
diff --git a/Engine/Plugins/Animation/ControlRig/Source/ControlRigEditor/Private/EditMode/SControlRigOutliner.cpp b/Engine/Plugins/Animation/ControlRig/Source/ControlRigEditor/Private/EditMode/SControlRigOutliner.cpp
--- a/Engine/Plugins/Animation/ControlRig/Source/ControlRigEditor/Private/EditMode/SControlRigOutliner.cpp
+++ b/Engine/Plugins/Animation/ControlRig/Source/ControlRigEditor/Private/EditMode/SControlRigOutliner.cpp
@@ -38,6 +38,7 @@
 #include "Editor/SRigHierarchy.h"
 #include "Rigs/FKControlRig.h"
 #include "SEnumCombo.h"
+#include "Components/SkeletalMeshComponent.h"
 
 #define LOCTEXT_NAMESPACE "ControlRigOutliner"
 
@@ -50,38 +51,34 @@ FRigTreeDisplaySettings FMultiRigTreeDelegates::DefaultDisplaySettings;
 
 uint32 GetTypeHash(const FMultiRigData& Data)
 {
-	return GetTypeHash(TTuple<const UControlRig*, FRigElementKey>(Data.ControlRig.Get(), (Data.Key.IsSet() ? Data.Key.GetValue() : FRigElementKey())));
+	return GetTypeHash(TTuple<FMultiRigData::EMultiRigDataType, FName, const UControlRig*, FRigElementKey> (
+		Data.Type,
+		Data.GetItemName(),
+		Data.WeakControlRig.Get(),
+		Data.GetElementKey())
+	);
 }
 
 FText FMultiRigData::GetName() const
 {
-	if (Key.IsSet())
+	if (IsControlElement())
 	{
 		return FText::FromName(Key->Name);
 	}
 
-	if (ModuleName.IsSet())
+	if (IsModule() || IsActor() || IsComponent())
 	{
-		return FText::FromName(ModuleName.GetValue());
+		return FText::FromName(GetItemName());
 	}
 	
-	if (UControlRig* ControlRigPtr = ControlRig.Get())
+	if (UControlRig* ControlRigPtr = WeakControlRig.Get())
 	{
 		FString ControlRigName = ControlRigPtr->GetName();
-		FString BoundObjectName;
-
-		if (TSharedPtr<IControlRigObjectBinding> ObjectBinding = ControlRigPtr->GetObjectBinding())
+		if (ControlRigPtr->IsAdditive())
 		{
-			if (UObject* BoundObject = ObjectBinding->GetBoundObject())
-			{
-				if (AActor* Actor = BoundObject->GetTypedOuter<AActor>())
-				{
-					BoundObjectName = Actor->GetActorLabel();
-				}
-			}
+			return FText::Format(LOCTEXT("ControlTitle", "{0} (layered)"), FText::AsCultureInvariant(ControlRigName));
 		}
-		
-		return FText::Format(LOCTEXT("ControlTitle", "{0}  ({1})"), FText::AsCultureInvariant(ControlRigName), FText::AsCultureInvariant(BoundObjectName));
+		return FText::Format(LOCTEXT("ControlTitle", "{0}"), FText::AsCultureInvariant(ControlRigName));
 	}
 	
 	return FText::FromName(NAME_None);
@@ -127,11 +124,6 @@ FText FMultiRigData::GetDisplayName(const FRigTreeDisplaySettings& InSettings) c
 			}
 		}
 
-		if (ModuleName.IsSet())
-		{
-			CachedDisplayName = FText::FromName(ModuleName.GetValue());
-		}
-
 		if (!CachedDisplayName.IsSet())
 		{
 			CachedDisplayName = GetName();
@@ -141,71 +133,83 @@ FText FMultiRigData::GetDisplayName(const FRigTreeDisplaySettings& InSettings) c
 	return *CachedDisplayName;
 }
 
-FText FMultiRigData::GetToolTipText() const
+FText FMultiRigData::GetToolTipText(const FRigTreeDisplaySettings& Settings) const
 {
-	if(Key.IsSet())
+	if (IsControlElement())
+	{
+		return FText::Format(LOCTEXT("SMultiRigHierarchyItemControlTooltip", "{0}"), FText::FromName(GetElementKey().Name));
+	}
+	else
 	{
-		return FText::FromName(Key.GetValue().Name);
+		return FText::Format(LOCTEXT("SMultiRigHierarchyModuleControlTooltip", "{0}\n\nUse Alt+Click to select subtree."), GetDisplayName(Settings));
 	}
 	return FText();
 }
 
 bool FMultiRigData::operator== (const FMultiRigData & Other) const
 {
-	if (ControlRig.IsValid() == false && Other.ControlRig.IsValid() == false)
+	if (WeakControlRig == Other.WeakControlRig
+		&& Type == Other.Type
+		&& Key == Other.Key
+		&& Name == Other.Name)
 	{
 		return true;
 	}
-	else if (ControlRig.IsValid() == true && Other.ControlRig.IsValid() == true)
-	{
-		if (ControlRig.Get() == Other.ControlRig.Get())
-		{
-			if (Key.IsSet() == false && Other.Key.IsSet() == false)
-			{
-				if (ModuleName.IsSet() == false && Other.ModuleName.IsSet() == false)
-				{
-					return true;
-				}
-				else if (ModuleName.IsSet() == true && Other.ModuleName.IsSet() == true)
-				{
-					return ModuleName.GetValue() == Other.ModuleName.GetValue();
-				}
-			}
-			else if (Key.IsSet() == true && Other.Key.IsSet() == true)
-			{
-				return Key.GetValue() == Other.Key.GetValue();
-			}
-		}
-	}
+
 	return false;
 }
 
 bool FMultiRigData::IsValid() const
 {
-	if (ControlRig.IsValid())
+	if (WeakControlRig.IsValid())
 	{
 		if (Key.IsSet())
 		{
-			return Key.GetValue().IsValid();
+			return Key->IsValid();
 		}
-		if (ModuleName.IsSet())
+		if (Name.IsSet())
 		{
-			return !ModuleName->IsNone();
+			return !Name->IsNone();
 		}
 		return true;
 	}
-	return false;
+	
+	return (IsActor() || IsComponent());
 }
 
 URigHierarchy* FMultiRigData::GetHierarchy() const
 {
-	if (UControlRig* ControlRigPtr = ControlRig.Get())
+	if (UControlRig* ControlRigPtr = WeakControlRig.Get())
 	{
 		return ControlRigPtr->GetHierarchy();
 	}
 	return nullptr;
 }
 
+bool FMultiRigData::IsModularRig() const
+{
+	if (TStrongObjectPtr<UControlRig> Rig = WeakControlRig.Pin())
+	{
+		return Rig->IsModularRig();
+	}
+	return false;
+}
+
+FRigModuleInstance* FMultiRigData::GetModuleInstance() const
+{
+	if (IsModularRig() && IsModule())
+	{
+		if (TStrongObjectPtr<UControlRig> Rig = WeakControlRig.Pin())
+		{
+			if (UModularRig* ModularRig = Cast<UModularRig>(Rig.Get()))
+			{
+				return ModularRig->FindModule(GetItemName());
+			}
+		}
+	}
+	return nullptr;
+}
+
 //////////////////////////////////////////////////////////////
 /// FMultiRigTreeElement
 ///////////////////////////////////////////////////////////
@@ -215,13 +219,10 @@ FMultiRigTreeElement::FMultiRigTreeElement(const FMultiRigData& InData, TWeakPtr
 	Data = InData;
 	FilterResult = InFilterResult;
 
-	if (InTreeView.IsValid() && Data.ControlRig.IsValid())
+	if (InTreeView.IsValid() && Data.IsValid())
 	{
-		if (const URigHierarchy* Hierarchy = Data.GetHierarchy())
-		{
-			const FRigTreeDisplaySettings& Settings = InTreeView.Pin()->GetTreeDelegates().GetDisplaySettings();
-			RefreshDisplaySettings(Hierarchy, Settings);
-		}
+		const FRigTreeDisplaySettings& Settings = InTreeView.Pin()->GetTreeDelegates().GetDisplaySettings();
+		RefreshDisplaySettings(Data.GetHierarchy(), Settings);
 	}
 }
 
@@ -233,7 +234,11 @@ TSharedRef<ITableRow> FMultiRigTreeElement::MakeTreeRowWidget(const TSharedRef<S
 
 void FMultiRigTreeElement::RefreshDisplaySettings(const URigHierarchy* InHierarchy, const FRigTreeDisplaySettings& InSettings)
 {
-	const TPair<const FSlateBrush*, FSlateColor> Result = SMultiRigHierarchyItem::GetBrushForElementType(InHierarchy, Data);
+	TPair<const FSlateBrush*, FSlateColor> Result;
+	if (InHierarchy)
+	{
+		Result = SMultiRigHierarchyItem::GetBrushForElementType(InHierarchy, Data);
+	}
 
 	IconBrush = Result.Key;
 	IconColor = Result.Value;
@@ -241,10 +246,10 @@ void FMultiRigTreeElement::RefreshDisplaySettings(const URigHierarchy* InHierarc
 	if (FilterResult == ERigTreeFilterResult::Shown)
 	{
 		IconColor = IconColor.IsColorSpecified() && InSettings.bShowIconColors ? Result.Value : FSlateColor::UseForeground();
-		const bool bIsKey = Data.Key.IsSet();
-		const bool bIsModule = Data.ModuleName.IsSet();
-		const bool bIsRig = !bIsModule && !bIsKey;
-		if (bIsRig || bIsModule)
+		const bool bIsKey = Data.IsControlElement();
+		const bool bIsModule = Data.IsModule();
+		const bool bIsSuperItem = !bIsModule && !bIsKey;
+		if (bIsSuperItem || bIsModule)
 		{
 			TextColor = FSlateColor(FLinearColor::White);
 		}
@@ -271,8 +276,7 @@ void SMultiRigHierarchyItem::Construct(const FArguments& InArgs, const TSharedRe
 	Delegates = InTreeView->GetTreeDelegates();
 
 	
-	if (InRigTreeElement->Data.ControlRig.IsValid() == false || (InRigTreeElement->Data.Key.IsSet() &&
-		InRigTreeElement->Data.Key.GetValue().IsValid() == false))
+	if (!InRigTreeElement->Data.IsValid())
 	{
 		SMultiColumnTableRow<TSharedPtr<FMultiRigTreeElement>>::Construct(
 			SMultiColumnTableRow<TSharedPtr<FMultiRigTreeElement>>::FArguments()
@@ -309,7 +313,7 @@ FReply SMultiRigHierarchyItem::OnToggleVisibilityClicked()
 			{
 				for (TSharedPtr<FMultiRigTreeElement> Selected : SelectedItems)
 				{
-					if (Selected->Data.ModuleName.IsSet())
+					if (Selected->Data.IsModule())
 					{
 						ElementsToToggle.AddUnique(Selected);
 					}
@@ -321,7 +325,7 @@ FReply SMultiRigHierarchyItem::OnToggleVisibilityClicked()
 				TArray<TSharedPtr<FMultiRigTreeElement>> Descendants = Element->Children;
 				for (int32 i=0; i<Descendants.Num(); i++)
 				{
-					if (Descendants[i]->Data.ModuleName.IsSet())
+					if (Descendants[i]->Data.IsModule())
 					{
 						ElementsToToggle.AddUnique(Descendants[i]);
 						Descendants.Append(Descendants[i]->Children);
@@ -335,22 +339,19 @@ FReply SMultiRigHierarchyItem::OnToggleVisibilityClicked()
 		TOptional<bool> bSetControlsVisible;
 		for (TSharedPtr<FMultiRigTreeElement> ElementToToggle : ElementsToToggle)
 		{
-			if (ElementToToggle->Data.Key.IsSet() == false)
+			if (!ElementToToggle->Data.IsControlElement())
 			{
 				UControlRig* Rig = nullptr;
-				if (ElementToToggle->Data.ModuleName.IsSet())
+				if (ElementToToggle->Data.IsModule())
 				{
-					if (UModularRig* ModularRig = Cast<UModularRig>(Element->Data.ControlRig))
+					if (FRigModuleInstance* Module = ElementToToggle->Data.GetModuleInstance())
 					{
-						if (FRigModuleInstance* Module = ModularRig->FindModule(ElementToToggle->Data.ModuleName.GetValue()))
-						{
-							Rig = Module->GetRig();
-						}
+						Rig = Module->GetRig();
 					}
 				}
 				else
 				{
-					Rig  = ElementToToggle->Data.ControlRig.Get();
+					Rig  = ElementToToggle->Data.WeakControlRig.Get();
 				}
 			
 				if (Rig)
@@ -388,36 +389,19 @@ FText SMultiRigHierarchyItem::GetDisplayName() const
 
 FText SMultiRigHierarchyItem::GetToolTipText() const
 {
-	return (WeakRigTreeElement.Pin()->Data.GetToolTipText());
+	const FRigTreeDisplaySettings Settings = Delegates.GetDisplaySettings();
+	return (WeakRigTreeElement.Pin()->Data.GetToolTipText(Settings));
 }
 
 TPair<const FSlateBrush*, FSlateColor> SMultiRigHierarchyItem::GetBrushForElementType(const URigHierarchy* InHierarchy, const FMultiRigData& InData)
 {
 	const FSlateBrush* Brush = nullptr;
 	FSlateColor Color = FSlateColor::UseForeground();
-	if (InData.Key.IsSet())
+	if (InData.IsControlElement())
 	{
-		const FRigElementKey Key = InData.Key.GetValue();
+		const FRigElementKey Key = InData.GetElementKey();
 		return SRigHierarchyItem::GetBrushForElementType(InHierarchy, Key);
 	}
-	else
-	{
-		if (UControlRig* ControlRig = InData.ControlRig.Get())
-		{
-			if (ControlRig->GetControlsVisible())
-			{
-				Brush = FAppStyle::GetBrush("Level.VisibleIcon16x");
-			}
-			else
-			{
-				Brush = FAppStyle::GetBrush("Level.NotVisibleIcon16x");
-			}
-		}
-		else
-		{
-			Brush = FAppStyle::GetBrush("Level.NotVisibleIcon16x");
-		}
-	}
 
 	return TPair<const FSlateBrush*, FSlateColor>(Brush, Color);
 }
@@ -488,36 +472,29 @@ TSharedRef<SWidget> SMultiRigHierarchyItem::GenerateWidgetForColumn(const FName&
 				{
 					if (TSharedPtr<FMultiRigTreeElement> SharedElement = WeakRigTreeElement.Pin())
 					{
-						//if no key is set then it's the control rig or a module so we get that based upon it's state
-						if (!SharedElement->Data.Key.IsSet())
+						UControlRig* Rig = nullptr;
+						// If we have a module name, we are dealing with a module, get the rig from there
+						if (SharedElement->Data.IsModule())
 						{
-							UControlRig* Rig = nullptr;
-							// If we have a module name, we are dealing with a module, get the rig from there
-							if (SharedElement->Data.ModuleName.IsSet())
+							if (FRigModuleInstance* Module = SharedElement->Data.GetModuleInstance())
 							{
-								if (UModularRig* ModularRig = Cast<UModularRig>(SharedElement->Data.ControlRig))
-								{
-									if (FRigModuleInstance* Module = ModularRig->FindModule(SharedElement->Data.ModuleName.GetValue()))
-									{
-										Rig = Module->GetRig();
-									}
-								}
+								Rig = Module->GetRig();
 							}
-							else
+						}
+						else if (SharedElement->Data.IsControlRig())
+						{
+							Rig = SharedElement->Data.WeakControlRig.Get();
+						}
+
+						if (Rig)
+						{
+							if (Rig->GetControlsVisible())
 							{
-								Rig = SharedElement->Data.ControlRig.Get();
+								return (FAppStyle::GetBrush("Level.VisibleIcon16x"));
 							}
-
-							if (Rig)
+							else
 							{
-								if (Rig->GetControlsVisible())
-								{
-									return (FAppStyle::GetBrush("Level.VisibleIcon16x"));
-								}
-								else
-								{
-									return  (FAppStyle::GetBrush("Level.NotVisibleIcon16x"));
-								}
+								return  (FAppStyle::GetBrush("Level.NotVisibleIcon16x"));
 							}
 						}
 					}
@@ -544,21 +521,7 @@ TSharedRef<SWidget> SMultiRigHierarchyItem::GenerateWidgetForColumn(const FName&
 	else if (InColumnName == "Name")
 	{
 		return SNew(SHorizontalBox)
-		.ToolTipText_Lambda([this]()
-		{
-			if (TSharedPtr<FMultiRigTreeElement> Element = WeakRigTreeElement.Pin())
-			{
-				if (Element->Data.Key.IsSet())
-				{
-					return FText::Format(LOCTEXT("SMultiRigHierarchyItemControlTooltip", "{0}"), FText::FromName(Element->Data.Key.GetValue().Name));
-				}
-				else
-				{
-					return FText::Format(LOCTEXT("SMultiRigHierarchyModuleControlTooltip", "{0}\n\nUse Alt+Click to select subtree."), GetDisplayName());
-				}
-			}
-			return FText::GetEmpty();
-		})
+		.ToolTipText(GetToolTipText())
 		+SHorizontalBox::Slot()
 		.AutoWidth()
 		.Padding(4, 2, 2, 0)
@@ -580,7 +543,7 @@ TSharedRef<SWidget> SMultiRigHierarchyItem::GenerateWidgetForColumn(const FName&
 				{
 					if (TSharedPtr<FMultiRigTreeElement> SharedElement = WeakRigTreeElement.Pin())
 					{
-						if (SharedElement->Data.ModuleName.IsSet() || SharedElement->Data.Key.IsSet())
+						if (SharedElement->Data.IsModule() || SharedElement->Data.IsControlElement())
 						{
 							return EVisibility::Visible;
 						}
@@ -592,15 +555,12 @@ TSharedRef<SWidget> SMultiRigHierarchyItem::GenerateWidgetForColumn(const FName&
 					if (TSharedPtr<FMultiRigTreeElement> SharedElement = WeakRigTreeElement.Pin())
 					{
 						// If we have a module name, we are dealing with a module, get the rig from there
-						if (SharedElement->Data.ModuleName.IsSet())
+						if (SharedElement->Data.IsModule())
 						{
 							UControlRig* Rig = nullptr;
-							if (UModularRig* ModularRig = Cast<UModularRig>(SharedElement->Data.ControlRig))
+							if (FRigModuleInstance* Module = SharedElement->Data.GetModuleInstance())
 							{
-								if (FRigModuleInstance* Module = ModularRig->FindModule(SharedElement->Data.ModuleName.GetValue()))
-								{
-									Rig = Module->GetRig();
-								}
+								Rig = Module->GetRig();
 							}
 							
 							if (Rig)
@@ -640,7 +600,7 @@ TSharedRef<SWidget> SMultiRigHierarchyItem::GenerateWidgetForColumn(const FName&
 				{
 					if (TSharedPtr<FMultiRigTreeElement> Element = WeakRigTreeElement.Pin())
 					{
-						if (!Element->Data.ModuleName.IsSet() && !Element->Data.Key.IsSet())
+						if (!Element->Data.IsModule() && !Element->Data.IsControlElement())
 						{
 							return FCoreStyle::GetDefaultFontStyle("Bold", 10);
 						}
@@ -801,13 +761,13 @@ bool SMultiRigHierarchyTreeView::AddElement(const FMultiRigData& InData, const F
 			{
 				if (const URigHierarchy* Hierarchy = InData.GetHierarchy())
 				{
-					if (InData.Key.IsSet())
+					if (InData.IsControlElement())
 					{
 						TSharedPtr<FMultiRigTreeElement> ChildItem = NewItem;
-						FRigElementKey ParentKey = Hierarchy->GetFirstParent(InData.Key.GetValue());
+						FRigElementKey ParentKey = Hierarchy->GetFirstParent(InData.GetElementKey());
 						while (ParentKey.IsValid())
 						{
-							FMultiRigData ParentData(InData.ControlRig.Get(), ParentKey);
+							FMultiRigData ParentData(InData.WeakControlRig.Get(), ParentKey);
 							if (!ElementMap.Contains(ParentData))
 							{
 								TSharedPtr<FMultiRigTreeElement> ParentItem = MakeShared<FMultiRigTreeElement>(ParentData, SharedThis(this), ERigTreeFilterResult::ShownDescendant);
@@ -881,7 +841,7 @@ bool SMultiRigHierarchyTreeView::AddElement(UControlRig* InControlRig, const FRi
 				{
 					return false;
 				}
-				if (!Settings.bShowAnimationChannels && ControlElement->Settings.AnimationType == ERigControlAnimationType::AnimationChannel)
+				if (ControlElement->Settings.AnimationType == ERigControlAnimationType::AnimationChannel)
 				{
 					return false;
 				}
@@ -933,9 +893,9 @@ bool SMultiRigHierarchyTreeView::AddElement(UControlRig* InControlRig, const FRi
 	}
 
 	FMultiRigData ParentData;
-	ParentData.ControlRig = InControlRig;
+	ParentData.WeakControlRig = InControlRig;
 
-	if (!AddElement(Data,ParentData))
+	if (!AddElement(Data, ParentData))
 	{
 		return false;
 	}
@@ -948,11 +908,12 @@ bool SMultiRigHierarchyTreeView::AddElement(UControlRig* InControlRig, const FRi
 		{
 			if (InControlRig->IsModularRig() && Settings.bArrangeByModules)
 			{
-				ParentData.ModuleName = *Hierarchy->GetModuleName(InElement->GetKey());
+				ParentData.SetItemName(FMultiRigData::EMultiRigDataType_Module, *Hierarchy->GetModuleName(InElement->GetKey()));
 
 				FMultiRigData NewData;
-				NewData.ControlRig = InControlRig;
-				NewData.Key = InElement->GetKey();
+				NewData.Type = FMultiRigData::EMultiRigDataType_Element;
+				NewData.WeakControlRig = InControlRig;
+				NewData.SetElementKey(InElement->GetKey());
 
 				if (TSharedPtr<FMultiRigTreeElement>* ParentElementPtr = ElementMap.Find(ParentData))
 				{
@@ -966,7 +927,7 @@ bool SMultiRigHierarchyTreeView::AddElement(UControlRig* InControlRig, const FRi
 								int32 InsertIndex = INDEX_NONE;
 								for (int32 i=Parent->Children.Num()-2; i >= 0; i--)
 								{
-									if (!Parent->Children[i]->Data.ModuleName.IsSet())
+									if (!Parent->Children[i]->Data.IsModule())
 									{
 										InsertIndex = i;
 										break;
@@ -1200,47 +1161,81 @@ void SMultiRigHierarchyTreeView::RefreshTreeView()
 	{
 		if (UControlRig* ControlRig  = ControlRigPtr.Get())
 		{
-			FMultiRigData CRData;
-			CRData.ControlRig = ControlRig; //leave key unset so it's valid
-
-			// add root element
-			AddElement(CRData, EmptyParentData);
-
-			const FRigTreeDisplaySettings& Settings = Delegates.GetDisplaySettings();
-			if (ControlRig->IsModularRig() && Settings.bArrangeByModules)
+			if (TSharedPtr<IControlRigObjectBinding> ObjectBinding = ControlRigPtr->GetObjectBinding())
 			{
-				if (UModularRig* ModularRig = Cast<UModularRig>(ControlRig))
+				TArray<FMultiRigData> DataHierarchy;
+				UObject* ParentObject = ObjectBinding->GetBoundObject();
+				while (ParentObject)
 				{
-					ModularRig->ForEachModule([&](FRigModuleInstance* Module)
+					if (USkeletalMeshComponent* Component = Cast<USkeletalMeshComponent>(ParentObject))
 					{
-						FMultiRigData ParentData;
-						ParentData.ControlRig = ControlRig;
-						if (Module->ParentModuleName.IsNone())
+						if (!Component->HasAnyFlags(RF_DefaultSubObject))
 						{
-							ParentData = CRData;
+							FMultiRigData Data;
+							Data.SetItemName(FMultiRigData::EMultiRigDataType_Component, *Component->GetName());
+							DataHierarchy.Add(Data);
 						}
-						else
+					}
+					else if (AActor* Actor = Cast<AActor>(ParentObject))
+					{
+						FMultiRigData Data;
+						Data.SetItemName(FMultiRigData::EMultiRigDataType_Actor, *Actor->GetActorLabel());
+						DataHierarchy.Add(Data);
+					}
+					ParentObject = ParentObject->GetOuter();
+				}
+
+				FMultiRigData* ParentData = &EmptyParentData;
+				for (int32 i=DataHierarchy.Num()-1; i>=0; --i)
+				{
+					AddElement(DataHierarchy[i], *ParentData);
+					ParentData = &DataHierarchy[i];
+				}
+
+				FMultiRigData CRData;
+				CRData.Type = FMultiRigData::EMultiRigDataType_ControlRig;
+				CRData.WeakControlRig = ControlRig; //leave key unset so it's valid
+
+				// add root element
+				AddElement(CRData, *ParentData);
+
+
+				const FRigTreeDisplaySettings& Settings = Delegates.GetDisplaySettings();
+				if (ControlRig->IsModularRig() && Settings.bArrangeByModules)
+				{
+					if (UModularRig* ModularRig = Cast<UModularRig>(ControlRig))
+					{
+						ModularRig->ForEachModule([&](FRigModuleInstance* Module)
 						{
-							ParentData.ModuleName = Module->ParentModuleName;
-						}
+							FMultiRigData ParentData;
+							ParentData.WeakControlRig = ControlRig;
+							if (Module->ParentModuleName.IsNone())
+							{
+								ParentData = CRData;
+							}
+							else
+							{
+								ParentData.SetItemName(FMultiRigData::EMultiRigDataType_Module, Module->ParentModuleName);
+							}
 
-						FMultiRigData NewData;
-						NewData.ControlRig = ControlRig;
-						NewData.ModuleName = Module->Name;
-						AddElement(NewData, ParentData);
-						return true;
-					});
+							FMultiRigData NewData;
+							NewData.WeakControlRig = ControlRig;
+							NewData.SetItemName(FMultiRigData::EMultiRigDataType_Module, Module->Name);
+							AddElement(NewData, ParentData);
+							return true;
+						});
+					}
 				}
-			}
 
-			// add children
-			if (const URigHierarchy* Hierarchy = ControlRig->GetHierarchy())
-			{
-				Hierarchy->Traverse([&](FRigBaseElement* Element, bool& bContinue)
+				// add children
+				if (const URigHierarchy* Hierarchy = ControlRig->GetHierarchy())
 				{
-					AddElement(ControlRig, Element);
-					bContinue = true;
-				});
+					Hierarchy->Traverse([&](FRigBaseElement* Element, bool& bContinue)
+					{
+						AddElement(ControlRig, Element);
+						bContinue = true;
+					});
+				}
 			}
 		}
 	}
@@ -1298,7 +1293,7 @@ void SMultiRigHierarchyTreeView::RefreshTreeView()
 		
 		// look for the root item referencing this rig
 		FMultiRigData ControlRigData;
-		ControlRigData.ControlRig = ControlRig;
+		ControlRigData.WeakControlRig = ControlRig;
 
 		const int32 RootElementIndex = RootElements.IndexOfByPredicate([ControlRigData](const TSharedPtr<FMultiRigTreeElement>& Root)
 		{
@@ -1407,7 +1402,10 @@ TArray<FMultiRigData> SMultiRigHierarchyTreeView::GetSelectedData() const
 	TArray<TSharedPtr<FMultiRigTreeElement>> SelectedElements = GetSelectedItems();
 	for (const TSharedPtr<FMultiRigTreeElement>& SelectedElement : SelectedElements)
 	{
-		Keys.Add(SelectedElement->Data);
+		if (SelectedElement)
+		{
+			Keys.Add(SelectedElement->Data);
+		}
 	}
 	return Keys;
 }
@@ -1516,7 +1514,6 @@ const FRigTreeDisplaySettings& SSearchableMultiRigHierarchyTreeView::GetDisplayS
 		Settings = SuperGetRigTreeDisplaySettings.Execute();
 	}
 	Settings.FilterText = FilterText;
-	Settings.bShowAnimationChannels = UControlRigEditorSettings::Get()->bShowAnimationChannels;
 	Settings.bArrangeByModules = UControlRigEditorSettings::Get()->bArrangeByModules;
 	Settings.bFocusOnSelection = UControlRigEditorSettings::Get()->bFocusOnSelection;
 	Settings.NameDisplayMode = UControlRigEditorSettings::Get()->ElementNameDisplayMode;
@@ -1533,7 +1530,6 @@ TSharedRef<SWidget> SSearchableMultiRigHierarchyTreeView::OnGetOptionsMenu()
 		const FControlRigHierarchyCommands& Actions = FControlRigHierarchyCommands::Get();
 		const FControlRigEditModeCommands& EditModeActions = FControlRigEditModeCommands::Get();
 		MenuBuilder.AddMenuEntry(Actions.ArrangeByModules);
-		MenuBuilder.AddMenuEntry(Actions.ShowAnimationChannels);
 
 		MenuBuilder.AddWidget(
 			SNew(SVerticalBox)
@@ -1580,12 +1576,6 @@ void SSearchableMultiRigHierarchyTreeView::BindCommands()
 		FCanExecuteAction(),
 		FIsActionChecked::CreateSP(this, &SSearchableMultiRigHierarchyTreeView::IsFocusingOnSelection));
 	
-	CommandList->MapAction(
-		Commands.ShowAnimationChannels,
-		FExecuteAction::CreateSP(this, &SSearchableMultiRigHierarchyTreeView::ToggleShowAnimationChannels),
-		FCanExecuteAction(),
-		FIsActionChecked::CreateSP(this, &SSearchableMultiRigHierarchyTreeView::IsShowingAnimationChannels));
-
 	CommandList->MapAction(
 		FControlRigEditModeCommands::Get().ToggleModuleManipulators,
 		FExecuteAction::CreateSP(this, &SSearchableMultiRigHierarchyTreeView::ToggleModuleManipulators),
@@ -1615,24 +1605,6 @@ void SSearchableMultiRigHierarchyTreeView::ToggleArrangeByModules()
 	TreeView->RequestTreeViewRefresh();
 }
 
-bool SSearchableMultiRigHierarchyTreeView::IsShowingAnimationChannels() const
-{
-	return UControlRigEditorSettings::Get()->bShowAnimationChannels;
-}
-
-void SSearchableMultiRigHierarchyTreeView::ToggleShowAnimationChannels()
-{
-	UControlRigEditorSettings* EditorSettings = UControlRigEditorSettings::Get();
-	EditorSettings->bShowAnimationChannels = !EditorSettings->bShowAnimationChannels;
-
-	if (FProperty* Property = UControlRigEditorSettings::StaticClass()->FindPropertyByName(GET_MEMBER_NAME_CHECKED(UControlRigEditorSettings, bShowAnimationChannels)))
-	{
-		FPropertyChangedEvent PropertyChangedEvent(Property, EPropertyChangeType::ValueSet);
-		EditorSettings->PostEditChangeProperty(PropertyChangedEvent);
-	}
-	TreeView->RequestTreeViewRefresh();
-}
-
 EElementNameDisplayMode SSearchableMultiRigHierarchyTreeView::GetElementNameDisplayMode() const
 {
 	UControlRigEditorSettings* EditorSettings = UControlRigEditorSettings::Get();
@@ -1665,7 +1637,7 @@ bool SSearchableMultiRigHierarchyTreeView::CanToggleModuleManipulators() const
 	const TArray<FMultiRigData> Selection = TreeView->GetSelectedData();
 	for (const FMultiRigData& Data : Selection)
 	{
-		if (Data.ControlRig->IsModularRig())
+		if (Data.WeakControlRig->IsModularRig())
 		{
 			return true;
 		}
@@ -1713,7 +1685,6 @@ void SControlRigOutliner::Construct(const FArguments& InArgs, FControlRigEditMod
 	DisplaySettings.bHideParentsOnFilter = true;
 	DisplaySettings.bFlattenHierarchyOnFilter = true;
 	DisplaySettings.bShowConnectors = false;
-	DisplaySettings.bShowAnimationChannels = true;
 	DisplaySettings.bArrangeByModules = true;
 	DisplaySettings.bFocusOnSelection = true;
 	DisplaySettings.NameDisplayMode = EElementNameDisplayMode::AssetDefault;
@@ -1832,62 +1803,88 @@ void SControlRigOutliner::HandleSelectionChanged(TSharedPtr<FMultiRigTreeElement
 	{
 		return;
 	}
+
 	const TArray<FMultiRigData> NewSelection = HierarchyTreeView->GetTreeView()->GetSelectedData();
 	TMap<UControlRig*, TArray<FRigElementKey>> SelectedRigAndKeys;
 	TArray<TSharedPtr<FMultiRigTreeElement>> AddToSelection;
-	for (const FMultiRigData& Data : NewSelection)
+	auto SelectRigDescendants = [this, &SelectedRigAndKeys, &AddToSelection](const FMultiRigData& Data)
 	{
-		if (Data.ControlRig.IsValid())
+		if (!Data.IsControlElement())
 		{
-			if (Data.Key.IsSet() && Data.Key.GetValue() != FRigElementKey())
+				
+			TArray<FRigElementKey>& Keys = SelectedRigAndKeys.FindOrAdd(Data.WeakControlRig.Get());
+			const FRigTreeDisplaySettings& Settings = HierarchyTreeView->GetDisplaySettings();
+			if (Settings.FilterText.IsEmpty() || !Settings.bFlattenHierarchyOnFilter)
 			{
-				SelectedRigAndKeys.FindOrAdd(Data.ControlRig.Get()).AddUnique(Data.Key.GetValue());
+				TSharedPtr<FMultiRigTreeElement> Element = HierarchyTreeView->GetTreeView()->FindElement(Data);
+				TArray<TSharedPtr<FMultiRigTreeElement>> Descendants = Element->Children;
+				for (int32 i=0; i<Descendants.Num(); ++i)
+				{
+					TSharedPtr<FMultiRigTreeElement>& Child = Descendants[i];
+					if (Child->Data.IsControlElement())
+					{
+						Keys.AddUnique(Child->Data.GetElementKey());
+					}
+						
+					AddToSelection.AddUnique(Child);
+					Descendants.Append(Child->Children);
+				}
 			}
-
-			if (FSlateApplication::Get().GetModifierKeys().IsAltDown())
+			else if (Data.IsModule())
 			{
-				if (!Data.Key.IsSet())
+				// If we have flatten the hierarchy due to a search, we cannot rely on the children of the module
+				Data.GetHierarchy()->ForEach<FRigControlElement>([&](FRigControlElement* ControlElement)
 				{
-					TArray<FRigElementKey>& Keys = SelectedRigAndKeys.FindOrAdd(Data.ControlRig.Get());
-					const FRigTreeDisplaySettings& Settings = HierarchyTreeView->GetDisplaySettings();
-					if (Settings.FilterText.IsEmpty() || !Settings.bFlattenHierarchyOnFilter)
+					if (Data.GetHierarchy()->GetModuleName(ControlElement->GetKey()) == Data.GetItemName())
 					{
-						TSharedPtr<FMultiRigTreeElement> Element = HierarchyTreeView->GetTreeView()->FindElement(Data);
-						TArray<TSharedPtr<FMultiRigTreeElement>> Descendants = Element->Children;
-						for (int32 i=0; i<Descendants.Num(); ++i)
-						{
-							TSharedPtr<FMultiRigTreeElement>& Child = Descendants[i];
-							if (Child->Data.Key.IsSet())
-							{
-								Keys.AddUnique(Child->Data.Key.GetValue());
-							}
-							
-							AddToSelection.AddUnique(Child);
-							Descendants.Append(Child->Children);
-						}
+						Keys.AddUnique(ControlElement->GetKey());
 					}
-					else if (Data.ModuleName.IsSet())
+					return true;
+				});
+			}
+			else
+			{
+				// If the key is not set, and the module name is not set, we are selecting the root of the rig. Select all controls
+				for (const FRigElementKey ControlKey : Data.GetHierarchy()->GetControlKeys())
+				{
+					Keys.AddUnique(ControlKey);
+				}
+			}
+		}
+	};
+	
+	
+	for (const FMultiRigData& Data : NewSelection)
+	{
+		if (Data.IsControlElement() && Data.IsValid())
+		{
+			SelectedRigAndKeys.FindOrAdd(Data.WeakControlRig.Get()).AddUnique(Data.GetElementKey());
+		}
+
+		if (FSlateApplication::Get().GetModifierKeys().IsAltDown())
+		{
+			if (Data.IsActor() || Data.IsComponent())
+			{
+				TSharedPtr<FMultiRigTreeElement> Element = HierarchyTreeView->GetTreeView()->FindElement(Data);
+				TArray<TSharedPtr<FMultiRigTreeElement>> Descendants = Element->Children;
+				for (int32 i=0; i<Descendants.Num(); ++i)
+				{
+					TSharedPtr<FMultiRigTreeElement> Child = Descendants[i];
+					if (Child->Data.IsControlRig())
 					{
-						// If we have flatten the hierarchy due to a search, we cannot rely on the children of the module
-						Data.GetHierarchy()->ForEach<FRigControlElement>([&](FRigControlElement* ControlElement)
-						{
-							if (Data.GetHierarchy()->GetModuleName(ControlElement->GetKey()) == Data.ModuleName.GetValue())
-							{
-								Keys.AddUnique(ControlElement->GetKey());
-							}
-							return true;
-						});
+						SelectRigDescendants(Child->Data);
 					}
 					else
 					{
-						// If the key is not set, and the module name is not set, we are selecting the root of the rig. Select all controls
-						for (const FRigElementKey ControlKey : Data.GetHierarchy()->GetControlKeys())
-						{
-							Keys.AddUnique(ControlKey);
-						}
+						Descendants.Append(Child->Children);
 					}
+					AddToSelection.AddUnique(Child);
 				}
 			}
+			else
+			{
+				SelectRigDescendants(Data);
+			}
 		}
 	}
 
diff --git a/Engine/Plugins/Animation/ControlRig/Source/ControlRigEditor/Private/EditMode/SControlRigOutliner.h b/Engine/Plugins/Animation/ControlRig/Source/ControlRigEditor/Private/EditMode/SControlRigOutliner.h
--- a/Engine/Plugins/Animation/ControlRig/Source/ControlRigEditor/Private/EditMode/SControlRigOutliner.h
+++ b/Engine/Plugins/Animation/ControlRig/Source/ControlRigEditor/Private/EditMode/SControlRigOutliner.h
@@ -87,21 +87,49 @@ class FMultiRigData
 {
 
 public:
-	FMultiRigData() {};
-	FMultiRigData(UControlRig* InControlRig, FRigElementKey InKey) : ControlRig(InControlRig), Key(InKey) {};
-	FMultiRigData(UControlRig* InControlRig, FName InModuleName) : ControlRig(InControlRig), ModuleName(InModuleName) {};
+	enum EMultiRigDataType
+	{
+		EMultiRigDataType_ControlRig,
+		EMultiRigDataType_Element,
+		EMultiRigDataType_Module,
+		EMultiRigDataType_Component,
+		EMultiRigDataType_Actor,
+
+		EMultiRigDataType_MAX
+	};
+	
+	FMultiRigData() { Type = EMultiRigDataType_ControlRig; };
+	FMultiRigData(UControlRig* InControlRig, FRigElementKey InKey) : WeakControlRig(InControlRig), Type(EMultiRigDataType_Element), Key(InKey) {};
+	FMultiRigData(UControlRig* InControlRig, EMultiRigDataType InType, FName InName) : WeakControlRig(InControlRig), Type(InType), Name(InName) {};
 	FText GetName() const;
 	FText GetDisplayName(const FRigTreeDisplaySettings& DisplaySettings) const;
-	FText GetToolTipText() const;
+	FText GetToolTipText(const FRigTreeDisplaySettings& Settings) const;
 	bool operator == (const FMultiRigData & Other) const;
 	bool IsValid() const;
 	URigHierarchy* GetHierarchy() const;
+	bool IsModularRig() const;
+	FRigModuleInstance* GetModuleInstance() const;
 
-	TWeakObjectPtr<UControlRig> ControlRig;
-	TOptional<FRigElementKey> Key;
-	TOptional<FName> ModuleName;
+	bool IsControlElement() const { return Type == EMultiRigDataType_Element && Key.IsSet() && Key.GetValue().IsValid(); }
+	bool IsModule() const { return Type == EMultiRigDataType_Module && Name.IsSet() && !Name.GetValue().IsNone(); }
+	bool IsControlRig() const { return Type == EMultiRigDataType_ControlRig && WeakControlRig.IsValid(); }
+	bool IsActor() const { return Type == EMultiRigDataType_Actor && !Name->IsNone(); }
+	bool IsComponent() const { return Type == EMultiRigDataType_Component && !Name->IsNone(); }
+
+	FRigElementKey GetElementKey() const { return Key.IsSet() ? Key.GetValue() : FRigElementKey(); }
+	FName GetItemName() const { return Name.IsSet() ? Name.GetValue() : NAME_None; }
+	
+	void SetElementKey(const FRigElementKey& InKey) { Type = EMultiRigDataType_Element; Key = InKey; }
+	void SetItemName(EMultiRigDataType InType, const FName& InName) { Type = InType; Name = InName; }
+	
+	TWeakObjectPtr<UControlRig> WeakControlRig;
+	EMultiRigDataType Type;
 
 private:
+	
+	TOptional<FName> Name;
+	TOptional<FRigElementKey> Key;
+	
 	mutable TOptional<FText> CachedDisplayName;
 };
 
@@ -281,8 +309,6 @@ public:
 
 	bool IsArrangedByModules() const;
 	void ToggleArrangeByModules();
-	bool IsShowingAnimationChannels() const;
-	void ToggleShowAnimationChannels();
 	EElementNameDisplayMode GetElementNameDisplayMode() const;
 	void SetElementNameDisplayMode(EElementNameDisplayMode InElementNameDisplayMode);
 	void ToggleModuleManipulators();
diff --git a/Engine/Plugins/Animation/ControlRig/Source/ControlRigEditor/Private/Editor/SRigHierarchyTreeView.h b/Engine/Plugins/Animation/ControlRig/Source/ControlRigEditor/Private/Editor/SRigHierarchyTreeView.h
--- a/Engine/Plugins/Animation/ControlRig/Source/ControlRigEditor/Private/Editor/SRigHierarchyTreeView.h
+++ b/Engine/Plugins/Animation/ControlRig/Source/ControlRigEditor/Private/Editor/SRigHierarchyTreeView.h
@@ -30,7 +30,6 @@ struct CONTROLRIGEDITOR_API FRigTreeDisplaySettings
 		bShowConnectors = true;
 		bShowComponents = false;
 		bShowIconColors = true;
-		bShowAnimationChannels = true;
 		bArrangeByModules = false;
 	}
 	
@@ -72,9 +71,6 @@ struct CONTROLRIGEDITOR_API FRigTreeDisplaySettings
 	/** Whether to tint the icons with the element color */
 	bool bShowIconColors;
 
-	/** Whether or not to show animation channels */
-	bool bShowAnimationChannels;
-
 	/** Whether or not to arrange the controls into modules (only for modular rigs) */
 	bool bArrangeByModules;
 
