Control Rig: Add option to flatten modules in the anim outliner. Also get the name of actor and component from sequencer.
#rnx

diff --git a/Engine/Plugins/Animation/ControlRig/Source/ControlRig/Private/Settings/ControlRigSettings.cpp b/Engine/Plugins/Animation/ControlRig/Source/ControlRig/Private/Settings/ControlRigSettings.cpp
--- a/Engine/Plugins/Animation/ControlRig/Source/ControlRig/Private/Settings/ControlRigSettings.cpp
+++ b/Engine/Plugins/Animation/ControlRig/Source/ControlRig/Private/Settings/ControlRigSettings.cpp
@@ -33,6 +33,7 @@ UControlRigEditorSettings::UControlRigEditorSettings(const FObjectInitializer& O
 	MaxStackSize = 16;
 	bLeftMouseDragDoesMarquee = true;
 	bArrangeByModules = true;
+	bFlattenModules = false;
 	bFocusOnSelection = true;
 	ElementNameDisplayMode = EElementNameDisplayMode::AssetDefault;
 #endif
diff --git a/Engine/Plugins/Animation/ControlRig/Source/ControlRig/Public/Settings/ControlRigSettings.h b/Engine/Plugins/Animation/ControlRig/Source/ControlRig/Public/Settings/ControlRigSettings.h
--- a/Engine/Plugins/Animation/ControlRig/Source/ControlRig/Public/Settings/ControlRigSettings.h
+++ b/Engine/Plugins/Animation/ControlRig/Source/ControlRig/Public/Settings/ControlRigSettings.h
@@ -139,6 +139,12 @@ public:
 	UPROPERTY(EditAnywhere, config, Category = Outliner)
 	bool bArrangeByModules;
 
+	/**
+	 * If turned on the modules in the Anim Outliner will be arranged in a flat list.
+	 */
+	UPROPERTY(EditAnywhere, config, Category = Outliner)
+	bool bFlattenModules;
+
 	/**
 	 * Defines how the element names will be displayed in treeviews, anim outliner etc
 	 */
diff --git a/Engine/Plugins/Animation/ControlRig/Source/ControlRigEditor/Private/ControlRigElementDetails.cpp b/Engine/Plugins/Animation/ControlRig/Source/ControlRigEditor/Private/ControlRigElementDetails.cpp
--- a/Engine/Plugins/Animation/ControlRig/Source/ControlRigEditor/Private/ControlRigElementDetails.cpp
+++ b/Engine/Plugins/Animation/ControlRig/Source/ControlRigEditor/Private/ControlRigElementDetails.cpp
@@ -4259,6 +4259,7 @@ void FRigControlElementDetails::CustomizeAvailableSpaces(IDetailLayoutBuilder& D
 	DisplaySettings.bFlattenHierarchyOnFilter = true;
 	DisplaySettings.bShowIconColors = true;
 	DisplaySettings.bArrangeByModules = false;
+	DisplaySettings.bFlattenModules = false;
 	DisplaySettings.NameDisplayMode = EElementNameDisplayMode::AssetDefault;
 
 	const TSharedRef<SHorizontalBox> HeaderContentWidget = SNew(SHorizontalBox);
diff --git a/Engine/Plugins/Animation/ControlRig/Source/ControlRigEditor/Private/ControlRigHierarchyCommands.cpp b/Engine/Plugins/Animation/ControlRig/Source/ControlRigEditor/Private/ControlRigHierarchyCommands.cpp
--- a/Engine/Plugins/Animation/ControlRig/Source/ControlRigEditor/Private/ControlRigHierarchyCommands.cpp
+++ b/Engine/Plugins/Animation/ControlRig/Source/ControlRigEditor/Private/ControlRigHierarchyCommands.cpp
@@ -34,6 +34,7 @@ void FControlRigHierarchyCommands::RegisterCommands()
 	UI_COMMAND(FilteringFlattensHierarchy, "Filtering Flattens Hierarchy", "Whether to keep the hierarchy or flatten it when searching for tree items", EUserInterfaceActionType::ToggleButton, FInputChord());
 	UI_COMMAND(HideParentsWhenFiltering, "Hide Parents When Filtering", "Whether to show parent items grayed out, or hide them entirely when filtering", EUserInterfaceActionType::ToggleButton, FInputChord());
 	UI_COMMAND(ArrangeByModules, "Arrange by Modules", "Arrange the hierarchy by the modules that spawn each control", EUserInterfaceActionType::ToggleButton, FInputChord());
+	UI_COMMAND(FlattenModules, "Flatten Modules", "Flatten the modular rig into a list of modules", EUserInterfaceActionType::ToggleButton, FInputChord());
 	UI_COMMAND(FocusOnSelection, "Focus On Selection", "Focus on the selection when it changes", EUserInterfaceActionType::ToggleButton, FInputChord());
 	UI_COMMAND(ShowImportedBones, "Show Imported Bones", "Whether to show or hide imported bones", EUserInterfaceActionType::ToggleButton, FInputChord());
 	UI_COMMAND(ShowBones, "Show Bones", "Whether to show or hide bones", EUserInterfaceActionType::ToggleButton, FInputChord());
diff --git a/Engine/Plugins/Animation/ControlRig/Source/ControlRigEditor/Private/ControlRigHierarchyCommands.h b/Engine/Plugins/Animation/ControlRig/Source/ControlRigEditor/Private/ControlRigHierarchyCommands.h
--- a/Engine/Plugins/Animation/ControlRig/Source/ControlRigEditor/Private/ControlRigHierarchyCommands.h
+++ b/Engine/Plugins/Animation/ControlRig/Source/ControlRigEditor/Private/ControlRigHierarchyCommands.h
@@ -102,6 +102,9 @@ public:
 	/* Arrange the controls by modules in a modular rig */
 	TSharedPtr< FUICommandInfo > ArrangeByModules;
 
+	/* Flatten the modules in a modular rig */
+	TSharedPtr< FUICommandInfo > FlattenModules;
+
 	/* Show imported bones */
 	TSharedPtr< FUICommandInfo > ShowImportedBones;
 
diff --git a/Engine/Plugins/Animation/ControlRig/Source/ControlRigEditor/Private/EditMode/SControlRigEditModeTools.cpp b/Engine/Plugins/Animation/ControlRig/Source/ControlRigEditor/Private/EditMode/SControlRigEditModeTools.cpp
--- a/Engine/Plugins/Animation/ControlRig/Source/ControlRigEditor/Private/EditMode/SControlRigEditModeTools.cpp
+++ b/Engine/Plugins/Animation/ControlRig/Source/ControlRigEditor/Private/EditMode/SControlRigEditModeTools.cpp
@@ -221,6 +221,7 @@ void SControlRigEditModeTools::Construct(const FArguments& InArgs, TSharedPtr<FC
 	DisplaySettings.bFlattenHierarchyOnFilter = true;
 	DisplaySettings.bShowIconColors = true;
 	DisplaySettings.bArrangeByModules = false;
+	DisplaySettings.bFlattenModules = false;
 	DisplaySettings.NameDisplayMode = EElementNameDisplayMode::AssetDefault;
 #if USE_LOCAL_DETAILS
 	FRigTreeDelegates RigTreeDelegates;
diff --git a/Engine/Plugins/Animation/ControlRig/Source/ControlRigEditor/Private/EditMode/SControlRigOutliner.cpp b/Engine/Plugins/Animation/ControlRig/Source/ControlRigEditor/Private/EditMode/SControlRigOutliner.cpp
--- a/Engine/Plugins/Animation/ControlRig/Source/ControlRigEditor/Private/EditMode/SControlRigOutliner.cpp
+++ b/Engine/Plugins/Animation/ControlRig/Source/ControlRigEditor/Private/EditMode/SControlRigOutliner.cpp
@@ -27,6 +27,8 @@
 #include "Rigs/RigHierarchyController.h"
 #include "ControlRigObjectBinding.h"
 #include "ModularRig.h"
+#include "MovieScene.h"
+#include "MovieSceneNameableTrack.h"
 #include "Settings/ControlRigSettings.h"
 #include "Widgets/Layout/SSpacer.h"
 #include "Widgets/Images/SImage.h"
@@ -74,10 +76,11 @@ FText FMultiRigData::GetName() const
 	if (UControlRig* ControlRigPtr = WeakControlRig.Get())
 	{
 		FString ControlRigName = ControlRigPtr->GetName();
-		if (ControlRigPtr->IsAdditive())
+		if (UMovieSceneNameableTrack* Track = Cast<UMovieSceneNameableTrack>(ControlRigPtr->GetOuter()))
 		{
-			return FText::Format(LOCTEXT("ControlTitle", "{0} (layered)"), FText::AsCultureInvariant(ControlRigName));
+			ControlRigName = Track->GetDisplayName().ToString();
 		}
+		
 		return FText::Format(LOCTEXT("ControlTitle", "{0}"), FText::AsCultureInvariant(ControlRigName));
 	}
 	
@@ -1157,6 +1160,41 @@ void SMultiRigHierarchyTreeView::RefreshTreeView()
 
 	// rebuild elements
 	static FMultiRigData EmptyParentData(nullptr, FRigElementKey());
+
+	FControlRigEditMode* EditMode = Delegates.GetEditMode();
+	const TWeakPtr<ISequencer>& WeakSequencer = EditMode->GetWeakSequencer();
+	
+	TSharedPtr<ISequencer> Sequencer = WeakSequencer.Pin();
+
+	auto GetAscestorName = [Sequencer](UObject* Object, FString& InOutName)
+	{
+		if (Sequencer && Object)
+		{
+			FGuid Id = Sequencer->FindCachedObjectId(*Object, Sequencer->GetFocusedTemplateID());
+			if (UMovieScene* MovieScene = Sequencer->GetFocusedMovieSceneSequence()->GetMovieScene())
+			{
+				const FString* Name = nullptr;
+				if (FMovieSceneSpawnable* Spawnable = MovieScene->FindSpawnable(Id))
+				{
+					Name = &Spawnable->GetName();
+				}
+				else if (FMovieScenePossessable* Possesable = MovieScene->FindPossessable(Id))
+				{
+					Name = &Possesable->GetName();
+				}
+				else if (FMovieSceneBinding* Binding = MovieScene->FindBinding(Id))
+				{
+					Name = &Binding->GetName();
+				}
+
+				if (Name && !Name->IsEmpty())
+				{
+					InOutName = *Name;
+				}
+			}
+		}
+	};
+	
 	for (const TWeakObjectPtr<UControlRig>& ControlRigPtr : ControlRigs)
 	{
 		if (UControlRig* ControlRig  = ControlRigPtr.Get())
@@ -1172,14 +1210,22 @@ void SMultiRigHierarchyTreeView::RefreshTreeView()
 						if (!Component->HasAnyFlags(RF_DefaultSubObject))
 						{
 							FMultiRigData Data;
-							Data.SetItemName(FMultiRigData::EMultiRigDataType_Component, *Component->GetName());
+
+							FString ComponentName = Component->GetName();
+							GetAscestorName(Component, ComponentName);
+							
+							Data.SetItemName(FMultiRigData::EMultiRigDataType_Component, *ComponentName);
 							DataHierarchy.Add(Data);
 						}
 					}
 					else if (AActor* Actor = Cast<AActor>(ParentObject))
 					{
 						FMultiRigData Data;
-						Data.SetItemName(FMultiRigData::EMultiRigDataType_Actor, *Actor->GetActorLabel());
+
+						FString ActorName = Actor->GetActorLabel();
+						GetAscestorName(Actor, ActorName);
+	
+						Data.SetItemName(FMultiRigData::EMultiRigDataType_Actor, *ActorName);
 						DataHierarchy.Add(Data);
 					}
 					ParentObject = ParentObject->GetOuter();
@@ -1209,7 +1255,7 @@ void SMultiRigHierarchyTreeView::RefreshTreeView()
 						{
 							FMultiRigData ParentData;
 							ParentData.WeakControlRig = ControlRig;
-							if (Module->ParentModuleName.IsNone())
+							if (Settings.bFlattenModules || Module->ParentModuleName.IsNone())
 							{
 								ParentData = CRData;
 							}
@@ -1515,6 +1561,7 @@ const FRigTreeDisplaySettings& SSearchableMultiRigHierarchyTreeView::GetDisplayS
 	}
 	Settings.FilterText = FilterText;
 	Settings.bArrangeByModules = UControlRigEditorSettings::Get()->bArrangeByModules;
+	Settings.bFlattenModules = UControlRigEditorSettings::Get()->bFlattenModules;
 	Settings.bFocusOnSelection = UControlRigEditorSettings::Get()->bFocusOnSelection;
 	Settings.NameDisplayMode = UControlRigEditorSettings::Get()->ElementNameDisplayMode;
 	return Settings;
@@ -1530,6 +1577,7 @@ TSharedRef<SWidget> SSearchableMultiRigHierarchyTreeView::OnGetOptionsMenu()
 		const FControlRigHierarchyCommands& Actions = FControlRigHierarchyCommands::Get();
 		const FControlRigEditModeCommands& EditModeActions = FControlRigEditModeCommands::Get();
 		MenuBuilder.AddMenuEntry(Actions.ArrangeByModules);
+		MenuBuilder.AddMenuEntry(Actions.FlattenModules);
 
 		MenuBuilder.AddWidget(
 			SNew(SVerticalBox)
@@ -1569,6 +1617,12 @@ void SSearchableMultiRigHierarchyTreeView::BindCommands()
 		FExecuteAction::CreateSP(this, &SSearchableMultiRigHierarchyTreeView::ToggleArrangeByModules),
 		FCanExecuteAction(),
 		FIsActionChecked::CreateSP(this, &SSearchableMultiRigHierarchyTreeView::IsArrangedByModules));
+		
+	CommandList->MapAction(
+		Commands.FlattenModules,
+		FExecuteAction::CreateSP(this, &SSearchableMultiRigHierarchyTreeView::ToggleFlattenModules),
+		FCanExecuteAction(),
+		FIsActionChecked::CreateSP(this, &SSearchableMultiRigHierarchyTreeView::IsShowingFlatModules));
 
 	CommandList->MapAction(
 		Commands.FocusOnSelection,
@@ -1605,6 +1659,24 @@ void SSearchableMultiRigHierarchyTreeView::ToggleArrangeByModules()
 	TreeView->RequestTreeViewRefresh();
 }
 
+bool SSearchableMultiRigHierarchyTreeView::IsShowingFlatModules() const
+{
+	return UControlRigEditorSettings::Get()->bFlattenModules;
+}
+
+void SSearchableMultiRigHierarchyTreeView::ToggleFlattenModules()
+{
+	UControlRigEditorSettings* EditorSettings = UControlRigEditorSettings::Get();
+	EditorSettings->bFlattenModules = !EditorSettings->bFlattenModules;
+
+	if (FProperty* Property = UControlRigEditorSettings::StaticClass()->FindPropertyByName(GET_MEMBER_NAME_CHECKED(UControlRigEditorSettings, bFlattenModules)))
+	{
+		FPropertyChangedEvent PropertyChangedEvent(Property, EPropertyChangeType::ValueSet);
+		EditorSettings->PostEditChangeProperty(PropertyChangedEvent);
+	}
+	TreeView->RequestTreeViewRefresh();
+}
+
 EElementNameDisplayMode SSearchableMultiRigHierarchyTreeView::GetElementNameDisplayMode() const
 {
 	UControlRigEditorSettings* EditorSettings = UControlRigEditorSettings::Get();
@@ -1686,6 +1758,7 @@ void SControlRigOutliner::Construct(const FArguments& InArgs, FControlRigEditMod
 	DisplaySettings.bFlattenHierarchyOnFilter = true;
 	DisplaySettings.bShowConnectors = false;
 	DisplaySettings.bArrangeByModules = true;
+	DisplaySettings.bFlattenModules = false;
 	DisplaySettings.bFocusOnSelection = true;
 	DisplaySettings.NameDisplayMode = EElementNameDisplayMode::AssetDefault;
 
@@ -1705,6 +1778,11 @@ void SControlRigOutliner::Construct(const FArguments& InArgs, FControlRigEditMod
 			]
 		];
 	SetEditMode(InEditMode);
+
+	if (TSharedPtr<ISequencer> Sequencer = InEditMode.GetWeakSequencer().Pin())
+	{
+		Sequencer->OnMovieSceneDataChanged().AddSP(this, &SControlRigOutliner::OnSequencerTreeViewChanged);
+	}
 }
 
 void SControlRigOutliner::OnObjectsReplaced(const TMap<UObject*, UObject*>& OldToNewInstanceMap)
@@ -1728,6 +1806,19 @@ void SControlRigOutliner::OnObjectsReplaced(const TMap<UObject*, UObject*>& OldT
 	}
 }
 
+void SControlRigOutliner::OnSequencerTreeViewChanged(EMovieSceneDataChangeType MovieSceneDataChange)
+{
+	if (MovieSceneDataChange == EMovieSceneDataChangeType::MovieSceneStructureItemAdded ||
+		MovieSceneDataChange == EMovieSceneDataChangeType::MovieSceneStructureItemRemoved ||
+		MovieSceneDataChange == EMovieSceneDataChangeType::MovieSceneStructureItemsChanged ||
+		MovieSceneDataChange == EMovieSceneDataChangeType::RefreshAllImmediately ||
+		MovieSceneDataChange == EMovieSceneDataChangeType::RefreshTree
+		)
+	{
+		HierarchyTreeView->GetTreeView()->RequestTreeViewRefresh();
+	}
+}
+
 SControlRigOutliner::SControlRigOutliner()
 {
 	FCoreUObjectDelegates::OnObjectsReplaced.AddRaw(this, &SControlRigOutliner::OnObjectsReplaced);
@@ -1811,7 +1902,6 @@ void SControlRigOutliner::HandleSelectionChanged(TSharedPtr<FMultiRigTreeElement
 	{
 		if (!Data.IsControlElement())
 		{
-				
 			TArray<FRigElementKey>& Keys = SelectedRigAndKeys.FindOrAdd(Data.WeakControlRig.Get());
 			const FRigTreeDisplaySettings& Settings = HierarchyTreeView->GetDisplaySettings();
 			if (Settings.FilterText.IsEmpty() || !Settings.bFlattenHierarchyOnFilter)
@@ -1838,6 +1928,10 @@ void SControlRigOutliner::HandleSelectionChanged(TSharedPtr<FMultiRigTreeElement
 					if (Data.GetHierarchy()->GetModuleName(ControlElement->GetKey()) == Data.GetItemName())
 					{
 						Keys.AddUnique(ControlElement->GetKey());
+
+						FMultiRigData ChildData(Data.WeakControlRig.Get(), ControlElement->GetKey());
+						TSharedPtr<FMultiRigTreeElement> ChildElement = HierarchyTreeView->GetTreeView()->FindElement(ChildData);
+						AddToSelection.AddUnique(ChildElement);
 					}
 					return true;
 				});
@@ -1848,6 +1942,10 @@ void SControlRigOutliner::HandleSelectionChanged(TSharedPtr<FMultiRigTreeElement
 				for (const FRigElementKey ControlKey : Data.GetHierarchy()->GetControlKeys())
 				{
 					Keys.AddUnique(ControlKey);
+
+					FMultiRigData ChildData(Data.WeakControlRig.Get(), ControlKey);
+					TSharedPtr<FMultiRigTreeElement> ChildElement = HierarchyTreeView->GetTreeView()->FindElement(ChildData);
+					AddToSelection.AddUnique(ChildElement);
 				}
 			}
 		}
diff --git a/Engine/Plugins/Animation/ControlRig/Source/ControlRigEditor/Private/EditMode/SControlRigOutliner.h b/Engine/Plugins/Animation/ControlRig/Source/ControlRigEditor/Private/EditMode/SControlRigOutliner.h
--- a/Engine/Plugins/Animation/ControlRig/Source/ControlRigEditor/Private/EditMode/SControlRigOutliner.h
+++ b/Engine/Plugins/Animation/ControlRig/Source/ControlRigEditor/Private/EditMode/SControlRigOutliner.h
@@ -9,6 +9,7 @@
 #include "Widgets/SWidget.h"
 #include "Widgets/SCompoundWidget.h"
 #include "ControlRig.h"
+#include "ISequencer.h"
 #include "Rigs/RigHierarchy.h"
 #include "Editor/SRigHierarchyTreeView.h"
 #include "Widgets/SBoxPanel.h"
@@ -309,6 +310,8 @@ public:
 
 	bool IsArrangedByModules() const;
 	void ToggleArrangeByModules();
+	bool IsShowingFlatModules() const;
+	void ToggleFlattenModules();
 	EElementNameDisplayMode GetElementNameDisplayMode() const;
 	void SetElementNameDisplayMode(EElementNameDisplayMode InElementNameDisplayMode);
 	void ToggleModuleManipulators();
@@ -350,7 +353,8 @@ private:
 	void HandleOnObjectBoundToControlRig(UObject* InObject);
 
 	void OnObjectsReplaced(const TMap<UObject*, UObject*>& OldToNewInstanceMap);
-
+	void OnSequencerTreeViewChanged(EMovieSceneDataChangeType MovieSceneDataChange);
+	
 	void HandleSelectionChanged(TSharedPtr<FMultiRigTreeElement> Selection, ESelectInfo::Type SelectInfo);
 
 	/** Hierarchy picker for controls*/
diff --git a/Engine/Plugins/Animation/ControlRig/Source/ControlRigEditor/Private/Editor/SRigHierarchyTreeView.h b/Engine/Plugins/Animation/ControlRig/Source/ControlRigEditor/Private/Editor/SRigHierarchyTreeView.h
--- a/Engine/Plugins/Animation/ControlRig/Source/ControlRigEditor/Private/Editor/SRigHierarchyTreeView.h
+++ b/Engine/Plugins/Animation/ControlRig/Source/ControlRigEditor/Private/Editor/SRigHierarchyTreeView.h
@@ -31,6 +31,7 @@ struct CONTROLRIGEDITOR_API FRigTreeDisplaySettings
 		bShowComponents = false;
 		bShowIconColors = true;
 		bArrangeByModules = false;
+		bFlattenModules = false;
 	}
 	
 	FText FilterText;
@@ -74,6 +75,9 @@ struct CONTROLRIGEDITOR_API FRigTreeDisplaySettings
 	/** Whether or not to arrange the controls into modules (only for modular rigs) */
 	bool bArrangeByModules;
 
+	/** Whether or not to arrange the modules in a flat list (only for modular rigs) */
+	bool bFlattenModules;
+
 	/** Whether or not to expand and focus the selection when this changes */
 	bool bFocusOnSelection;
 };
