Control Rig: New hotkey to reset all, make option previous reset hotkeys only work for controls not channels. Changed some functions to static's so they can be used elsewhere as needed (e.g anim layers or scrpting).  And also add hotkey for showing overlay and made sure pose hotkeys work on invisible rigs.

diff --git a/Engine/Plugins/Animation/ControlRig/Source/ControlRigEditor/Private/EditMode/ControlRigEditMode.cpp b/Engine/Plugins/Animation/ControlRig/Source/ControlRigEditor/Private/EditMode/ControlRigEditMode.cpp
--- a/Engine/Plugins/Animation/ControlRig/Source/ControlRigEditor/Private/EditMode/ControlRigEditMode.cpp
+++ b/Engine/Plugins/Animation/ControlRig/Source/ControlRigEditor/Private/EditMode/ControlRigEditMode.cpp
@@ -2377,15 +2377,24 @@ void FControlRigEditMode::SetRigElementsSelectionInternal(const TMap<UControlRig
 	HandleSelectionChanged();
 }
 
-TArray<FRigElementKey> FControlRigEditMode::GetSelectedRigElements(UControlRig* ControlRig) const
+TArray<FRigElementKey> FControlRigEditMode::GetSelectedRigElements() const
 {
-	if (ControlRig == nullptr && GetControlRigs().Num() > 0)
+	if (GetControlRigs().Num() > 0)
 	{
-		ControlRig = GetControlRigs()[0].Get();
+		UControlRig* ControlRig = GetControlRigs()[0].Get();
+		return GetSelectedRigElements(ControlRig);
 	}
+	return TArray<FRigElementKey>();
 
-	TArray<FRigElementKey> SelectedKeys;
+}
 
+TArray<FRigElementKey> FControlRigEditMode::GetSelectedRigElements(UControlRig* ControlRig) 
+{
+	TArray<FRigElementKey> SelectedKeys;
+	if (ControlRig == nullptr)
+	{
+		return SelectedKeys;
+	}
 	if (ControlRig->GetHierarchy())
 	{
 		SelectedKeys = ControlRig->GetHierarchy()->GetSelectedKeys();
@@ -2441,7 +2450,10 @@ bool FControlRigEditMode::AreRigElementsSelected(uint32 InTypes, UControlRig* In
 			return FRigElementTypeHelper::DoesHave(InTypes, InSelectedKey.Type);
 		});
 	}
-
+	if (InControlRig == nullptr && GetControlRigs().Num() > 0)
+	{
+		InControlRig = GetControlRigs()[0].Get();
+	}
 	const TArray<FRigElementKey> SelectedRigElements = GetSelectedRigElements(InControlRig);
 	for (const FRigElementKey& Ele : SelectedRigElements)
 	{
@@ -2907,18 +2919,27 @@ void FControlRigEditMode::BindCommands()
 	CommandBindings->MapAction(
 		Commands.ToggleAllManipulators,
 		FExecuteAction::CreateRaw(this, &FControlRigEditMode::ToggleAllManipulators));
+	CommandBindings->MapAction(
+		Commands.ToggleControlsAsOverlay,
+		FExecuteAction::CreateRaw(this, &FControlRigEditMode::ToggleControlsAsOverlay));
 	CommandBindings->MapAction(
 		Commands.ZeroTransforms,
-		FExecuteAction::CreateRaw(this, &FControlRigEditMode::ZeroTransforms, true));
+		FExecuteAction::CreateRaw(this, &FControlRigEditMode::ZeroTransforms, true, false));
 	CommandBindings->MapAction(
 		Commands.ZeroAllTransforms,
-		FExecuteAction::CreateRaw(this, &FControlRigEditMode::ZeroTransforms, false));
+		FExecuteAction::CreateRaw(this, &FControlRigEditMode::ZeroTransforms, false, false));
 	CommandBindings->MapAction(
 		Commands.InvertTransforms,
-		FExecuteAction::CreateRaw(this, &FControlRigEditMode::InvertInputPose, true));
+		FExecuteAction::CreateRaw(this, &FControlRigEditMode::InvertInputPose, true, false));
 	CommandBindings->MapAction(
 		Commands.InvertAllTransforms,
-		FExecuteAction::CreateRaw(this, &FControlRigEditMode::InvertInputPose, false));
+		FExecuteAction::CreateRaw(this, &FControlRigEditMode::InvertInputPose, false, false));
+	CommandBindings->MapAction(
+		Commands.InvertTransformsAndChannels,
+		FExecuteAction::CreateRaw(this, &FControlRigEditMode::InvertInputPose, true, true));
+	CommandBindings->MapAction(
+		Commands.InvertAllTransformsAndChannels,
+		FExecuteAction::CreateRaw(this, &FControlRigEditMode::InvertInputPose, false, true));
 	CommandBindings->MapAction(
 		Commands.ClearSelection,
 		FExecuteAction::CreateRaw(this, &FControlRigEditMode::ClearSelection));
@@ -3505,6 +3526,13 @@ void FControlRigEditMode::ToggleAllManipulators()
 	}
 }
 
+void FControlRigEditMode::ToggleControlsAsOverlay()
+{
+	UControlRigEditModeSettings* Settings = GetMutableSettings();
+	Settings->bShowControlsAsOverlay = !Settings->bShowControlsAsOverlay;
+	OnSettingsChanged(Settings);
+}
+
 bool FControlRigEditMode::AreControlsVisible() const
 {
 	if (!AreEditingControlRigDirectly())
@@ -3530,7 +3558,57 @@ bool FControlRigEditMode::AreControlsVisible() const
 	return !Settings->bHideControlShapes;
 }
 
-void FControlRigEditMode::ZeroTransforms(bool bSelectionOnly)
+TArray<FRigElementKey> FControlRigEditMode::GetRigElementsForSettingTransforms(UControlRig* ControlRig, bool bSelectionOnly, bool bIncludeChannels)
+{
+	TArray<FRigElementKey> RigElements;
+	if (bSelectionOnly)
+	{
+		RigElements = GetSelectedRigElements(ControlRig);
+		if (ControlRig->IsAdditive())
+		{
+			// For additive rigs, ignore boolean controls
+			RigElements = RigElements.FilterByPredicate([ControlRig](const FRigElementKey& Key)
+				{
+					if (FRigControlElement* Element = ControlRig->FindControl(Key.Name))
+					{
+						return Element->CanTreatAsAdditive();
+					}
+					return true;
+				});
+		}
+	}
+	else
+	{
+		TArray<FRigBaseElement*> Elements = ControlRig->GetHierarchy()->GetElementsOfType<FRigBaseElement>(true);
+		for (const FRigBaseElement* Element : Elements)
+		{
+			// For additive rigs, ignore non-additive controls
+			if (const FRigControlElement* Control = Cast<FRigControlElement>(Element))
+			{
+				if (ControlRig->IsAdditive() && !Control->CanTreatAsAdditive())
+				{
+					continue;
+				}
+			}
+			RigElements.Add(Element->GetKey());
+		}
+	}
+
+	if (bIncludeChannels == false)
+	{
+		RigElements = RigElements.FilterByPredicate([ControlRig](const FRigElementKey& Key)
+		{
+			if (FRigControlElement* Element = ControlRig->FindControl(Key.Name))
+			{
+				return Element->IsAnimationChannel() == false;
+			}
+			return true;
+		});
+	}
+	return RigElements;
+}
+
+void FControlRigEditMode::ZeroTransforms(bool bSelectionOnly, bool bIncludeChannels)
 {
 	// Gather up the control rigs for the selected controls
 	TArray<UControlRig*> ControlRigs;
@@ -3555,215 +3633,204 @@ void FControlRigEditMode::ZeroTransforms(bool bSelectionOnly)
 
 	for (UControlRig* ControlRig : ControlRigs)
 	{
-		TArray<FRigElementKey> SelectedRigElements = GetSelectedRigElements(ControlRig);
-		if (ControlRig->IsAdditive())
+		ZeroTransforms(ControlRig, Context, bSelectionOnly, bIncludeChannels);
+	}
+}
+
+void FControlRigEditMode::ZeroTransforms(UControlRig* ControlRig, const FRigControlModifiedContext& Context,  bool bSelectionOnly, bool bIncludeChannels)
+{
+	TArray<FRigElementKey> ControlsToReset;
+	TArray<FRigElementKey> ControlsInteracting;
+	TArray<FRigElementKey> TransformElementsToReset;
+
+	TArray<FRigElementKey> SelectedRigElements = FControlRigEditMode::GetRigElementsForSettingTransforms(ControlRig, bSelectionOnly, bIncludeChannels);
+	if (bSelectionOnly)
+	{
+		ControlsToReset = SelectedRigElements;
+		ControlsInteracting = SelectedRigElements;
+		TransformElementsToReset = SelectedRigElements;
+	}
+	else 
+	{
+		TransformElementsToReset = SelectedRigElements;
+		TArray<FRigControlElement*> Controls;
+		ControlRig->GetControlsInOrder(Controls);
+		ControlsToReset.SetNum(0);
+		ControlsInteracting.SetNum(0);
+		for (const FRigControlElement* Control : Controls)
 		{
 			// For additive rigs, ignore boolean controls
-			SelectedRigElements = SelectedRigElements.FilterByPredicate([ControlRig](const FRigElementKey& Key)
+			if (ControlRig->IsAdditive() && Control->Settings.ControlType == ERigControlType::Bool)
 			{
-				if (FRigControlElement* Element = ControlRig->FindControl(Key.Name))
-				{
-					return Element->CanTreatAsAdditive();
-				}
-				return true;
-			});
-		}
-		TArray<FRigElementKey> ControlsToReset = SelectedRigElements;
-		TArray<FRigElementKey> ControlsInteracting = SelectedRigElements;
-		TArray<FRigElementKey> TransformElementsToReset = SelectedRigElements;
-		if (!bSelectionOnly)
-		{
-			TArray<FRigBaseElement*> Elements = ControlRig->GetHierarchy()->GetElementsOfType<FRigBaseElement>(true);
-			TransformElementsToReset.Reset();
-			TransformElementsToReset.Reserve(Elements.Num());
-			for (const FRigBaseElement* Element : Elements)
+				continue;
+			}
+			if (bIncludeChannels == false && Control->IsAnimationChannel())
 			{
-				// For additive rigs, ignore non-additive controls
-				if (const FRigControlElement* Control = Cast<FRigControlElement>(Element))
-				{
-					if (ControlRig->IsAdditive() && !Control->CanTreatAsAdditive())
-					{
-						continue;
-					}
-				}
-				TransformElementsToReset.Add(Element->GetKey());
+				continue;
 			}
-			
-			TArray<FRigControlElement*> Controls;
-			ControlRig->GetControlsInOrder(Controls);
-			ControlsToReset.SetNum(0);
-			ControlsInteracting.SetNum(0);
-			for (const FRigControlElement* Control : Controls)
-			{
-				// For additive rigs, ignore boolean controls
-				if (ControlRig->IsAdditive() && Control->Settings.ControlType == ERigControlType::Bool)
-				{
-					continue;
-				}
-				ControlsToReset.Add(Control->GetKey());
-				if(Control->Settings.AnimationType == ERigControlAnimationType::AnimationControl ||
-					Control->IsAnimationChannel())
-				{
-					ControlsInteracting.Add(Control->GetKey());
-				}
+			ControlsToReset.Add(Control->GetKey());
+			if(Control->Settings.AnimationType == ERigControlAnimationType::AnimationControl ||
+				Control->IsAnimationChannel())
+			{
+				ControlsInteracting.Add(Control->GetKey());
 			}
 		}
-		bool bHasNonDefaultParent = false;
-		TMap<FRigElementKey, FRigElementKey> Parents;
-		for (const FRigElementKey& Key : TransformElementsToReset)
+	}
+	bool bHasNonDefaultParent = false;
+	TMap<FRigElementKey, FRigElementKey> Parents;
+	for (const FRigElementKey& Key : TransformElementsToReset)
+	{
+		FRigElementKey SpaceKey = ControlRig->GetHierarchy()->GetActiveParent(Key);
+		Parents.Add(Key, SpaceKey);
+		if (!bHasNonDefaultParent && SpaceKey != ControlRig->GetHierarchy()->GetDefaultParentKey())
 		{
-			FRigElementKey SpaceKey = ControlRig->GetHierarchy()->GetActiveParent(Key);
-			Parents.Add(Key, SpaceKey);
-			if (!bHasNonDefaultParent && SpaceKey != ControlRig->GetHierarchy()->GetDefaultParentKey())
-			{
-				bHasNonDefaultParent = true;
-			}
+			bHasNonDefaultParent = true;
 		}
+	}
 
-		FControlRigInteractionScope InteractionScope(ControlRig, ControlsInteracting);
-		for (const FRigElementKey& ElementToReset : TransformElementsToReset)
+	FControlRigInteractionScope InteractionScope(ControlRig, ControlsInteracting);
+	for (const FRigElementKey& ElementToReset : TransformElementsToReset)
+	{
+		FRigControlElement* ControlElement = nullptr;
+		if (ElementToReset.Type == ERigElementType::Control)
 		{
-			FRigControlElement* ControlElement = nullptr;
-			if (ElementToReset.Type == ERigElementType::Control)
+			ControlElement = ControlRig->FindControl(ElementToReset.Name);
+			if (ControlElement->Settings.bIsTransientControl)
 			{
-				ControlElement = ControlRig->FindControl(ElementToReset.Name);
-				if (ControlElement->Settings.bIsTransientControl)
+				if(UControlRig::GetNodeNameFromTransientControl(ControlElement->GetKey()).IsEmpty())
 				{
-					if(UControlRig::GetNodeNameFromTransientControl(ControlElement->GetKey()).IsEmpty())
-					{
-						ControlElement = nullptr;
-					}
+					ControlElement = nullptr;
 				}
 			}
+		}
 			
-			const FTransform InitialLocalTransform = ControlRig->GetInitialLocalTransform(ElementToReset);
-			ControlRig->Modify();
-			if (bHasNonDefaultParent == true) //possibly not at default parent so switch to it
+		const FTransform InitialLocalTransform = ControlRig->GetInitialLocalTransform(ElementToReset);
+		ControlRig->Modify();
+		if (bHasNonDefaultParent == true) //possibly not at default parent so switch to it
+		{
+			ControlRig->GetHierarchy()->SwitchToDefaultParent(ElementToReset);
+		}
+		if (ControlElement)
+		{
+			const FVector InitialAngles = ControlRig->GetHierarchy()->GetControlPreferredEulerAngles(ControlElement, ControlElement->Settings.PreferredRotationOrder, true);
+			ControlRig->GetHierarchy()->SetControlPreferredEulerAngles(ControlElement, InitialAngles, ControlElement->Settings.PreferredRotationOrder);
+			ControlRig->SetControlLocalTransform(ElementToReset.Name, InitialLocalTransform, true, Context, true, true);
+
+			NotifyDrivenControls(ControlRig, ElementToReset, Context);
+
+			if (bHasNonDefaultParent == false)
 			{
-				ControlRig->GetHierarchy()->SwitchToDefaultParent(ElementToReset);
+				ControlRig->ControlModified().Broadcast(ControlRig, ControlElement, EControlRigSetKey::DoNotCare);
 			}
-			if (ControlElement)
-			{
-				const FVector InitialAngles = ControlRig->GetHierarchy()->GetControlPreferredEulerAngles(ControlElement, ControlElement->Settings.PreferredRotationOrder, true);
-				ControlRig->GetHierarchy()->SetControlPreferredEulerAngles(ControlElement, InitialAngles, ControlElement->Settings.PreferredRotationOrder);
-				ControlRig->SetControlLocalTransform(ElementToReset.Name, InitialLocalTransform, true, Context, true, true);
+		}
+		else
+		{
+			ControlRig->GetHierarchy()->SetLocalTransform(ElementToReset, InitialLocalTransform, false, true, true);
+		}
+
+		//@helge not sure what to do if the non-default parent
+		if (UControlRigBlueprint* Blueprint = Cast<UControlRigBlueprint>(ControlRig->GetClass()->ClassGeneratedBy))
+		{
+			Blueprint->Hierarchy->SetLocalTransform(ElementToReset, InitialLocalTransform);
+		}
+	}
 
-				NotifyDrivenControls(ControlRig, ElementToReset, Context);
+	if (bHasNonDefaultParent == true) //now we have the initial pose setup we need to get the global transforms as specified now then set them in the current parent space
+	{
+		EvaluateRig(ControlRig);
 
-				if (bHasNonDefaultParent == false)
+		//get global transforms
+		TMap<FRigElementKey, FTransform> GlobalTransforms;
+		for (const FRigElementKey& ElementToReset : TransformElementsToReset)
+		{
+			if (ElementToReset.IsTypeOf(ERigElementType::Control))
+			{
+				FRigControlElement* ControlElement = ControlRig->FindControl(ElementToReset.Name);
+				if (ControlElement && !ControlElement->Settings.bIsTransientControl)
 				{
-					ControlRig->ControlModified().Broadcast(ControlRig, ControlElement, EControlRigSetKey::DoNotCare);
+					FTransform GlobalTransform = ControlRig->GetHierarchy()->GetGlobalTransform(ElementToReset);
+					GlobalTransforms.Add(ElementToReset, GlobalTransform);
 				}
+				NotifyDrivenControls(ControlRig, ElementToReset,Context);
 			}
 			else
 			{
-				ControlRig->GetHierarchy()->SetLocalTransform(ElementToReset, InitialLocalTransform, false, true, true);
-			}
-
-			//@helge not sure what to do if the non-default parent
-			if (UControlRigBlueprint* Blueprint = Cast<UControlRigBlueprint>(ControlRig->GetClass()->ClassGeneratedBy))
-			{
-				Blueprint->Hierarchy->SetLocalTransform(ElementToReset, InitialLocalTransform);
+				FTransform GlobalTransform = ControlRig->GetHierarchy()->GetGlobalTransform(ElementToReset);
+				GlobalTransforms.Add(ElementToReset, GlobalTransform);
 			}
 		}
-
-		if (bHasNonDefaultParent == true) //now we have the initial pose setup we need to get the global transforms as specified now then set them in the current parent space
+		//switch back to original parent space
+		for (const FRigElementKey& ElementToReset : TransformElementsToReset)
 		{
-			EvaluateRig(ControlRig);
-
-			//get global transforms
-			TMap<FRigElementKey, FTransform> GlobalTransforms;
-			for (const FRigElementKey& ElementToReset : TransformElementsToReset)
+			if (const FRigElementKey* SpaceKey = Parents.Find(ElementToReset))
 			{
 				if (ElementToReset.IsTypeOf(ERigElementType::Control))
 				{
 					FRigControlElement* ControlElement = ControlRig->FindControl(ElementToReset.Name);
 					if (ControlElement && !ControlElement->Settings.bIsTransientControl)
 					{
-						FTransform GlobalTransform = ControlRig->GetHierarchy()->GetGlobalTransform(ElementToReset);
-						GlobalTransforms.Add(ElementToReset, GlobalTransform);
+						ControlRig->GetHierarchy()->SwitchToParent(ElementToReset, *SpaceKey);
 					}
-					NotifyDrivenControls(ControlRig, ElementToReset,Context);
 				}
 				else
 				{
-					FTransform GlobalTransform = ControlRig->GetHierarchy()->GetGlobalTransform(ElementToReset);
-					GlobalTransforms.Add(ElementToReset, GlobalTransform);
+					ControlRig->GetHierarchy()->SwitchToParent(ElementToReset, *SpaceKey);
 				}
 			}
-			//switch back to original parent space
+		}
+		//set global transforms in this space // do it twice since ControlsInOrder is not really always in order
+		for (int32 SetHack = 0; SetHack < 2; ++SetHack)
+		{
+			const TGuardValue<bool> GuardEvaluationType(ControlRig->bEvaluationTriggeredFromInteraction, true);
+			ControlRig->Evaluate_AnyThread();
 			for (const FRigElementKey& ElementToReset : TransformElementsToReset)
 			{
-				if (const FRigElementKey* SpaceKey = Parents.Find(ElementToReset))
+				if (const FTransform* GlobalTransform = GlobalTransforms.Find(ElementToReset))
 				{
 					if (ElementToReset.IsTypeOf(ERigElementType::Control))
 					{
 						FRigControlElement* ControlElement = ControlRig->FindControl(ElementToReset.Name);
 						if (ControlElement && !ControlElement->Settings.bIsTransientControl)
 						{
-							ControlRig->GetHierarchy()->SwitchToParent(ElementToReset, *SpaceKey);
+							ControlRig->SetControlGlobalTransform(ElementToReset.Name, *GlobalTransform, true);
+							ControlRig->Evaluate_AnyThread();
+							NotifyDrivenControls(ControlRig, ElementToReset, Context);
 						}
 					}
 					else
 					{
-						ControlRig->GetHierarchy()->SwitchToParent(ElementToReset, *SpaceKey);
+						ControlRig->GetHierarchy()->SetGlobalTransform(ElementToReset, *GlobalTransform, false, true, true);
 					}
 				}
 			}
-			//set global transforms in this space // do it twice since ControlsInOrder is not really always in order
-			for (int32 SetHack = 0; SetHack < 2; ++SetHack)
-			{
-				const TGuardValue<bool> GuardEvaluationType(ControlRig->bEvaluationTriggeredFromInteraction, true);
-				ControlRig->Evaluate_AnyThread();
-				for (const FRigElementKey& ElementToReset : TransformElementsToReset)
-				{
-					if (const FTransform* GlobalTransform = GlobalTransforms.Find(ElementToReset))
-					{
-						if (ElementToReset.IsTypeOf(ERigElementType::Control))
-						{
-							FRigControlElement* ControlElement = ControlRig->FindControl(ElementToReset.Name);
-							if (ControlElement && !ControlElement->Settings.bIsTransientControl)
-							{
-								ControlRig->SetControlGlobalTransform(ElementToReset.Name, *GlobalTransform, true);
-								ControlRig->Evaluate_AnyThread();
-								NotifyDrivenControls(ControlRig, ElementToReset, Context);
-							}
-						}
-						else
-						{
-							ControlRig->GetHierarchy()->SetGlobalTransform(ElementToReset, *GlobalTransform, false, true, true);
-						}
-					}
-				}
-			}
-			//send notifies
+		}
+		//send notifies
 
-			for (const FRigElementKey& ControlToReset : ControlsToReset)
+		for (const FRigElementKey& ControlToReset : ControlsToReset)
+		{
+			FRigControlElement* ControlElement = ControlRig->FindControl(ControlToReset.Name);
+			if (ControlElement && !ControlElement->Settings.bIsTransientControl)
 			{
-				FRigControlElement* ControlElement = ControlRig->FindControl(ControlToReset.Name);
-				if (ControlElement && !ControlElement->Settings.bIsTransientControl)
-				{
-					ControlRig->ControlModified().Broadcast(ControlRig, ControlElement, EControlRigSetKey::DoNotCare);
-				}
+				ControlRig->ControlModified().Broadcast(ControlRig, ControlElement, EControlRigSetKey::DoNotCare);
 			}
 		}
-		else
+	}
+	else
+	{
+		// we have to insert the interaction event before we run current events
+		TArray<FName> NewEventQueue = {FRigUnit_InteractionExecution::EventName};
+		NewEventQueue.Append(ControlRig->EventQueue);
+		TGuardValue<TArray<FName>> EventGuard(ControlRig->EventQueue, NewEventQueue);
+		const TGuardValue<bool> GuardEvaluationType(ControlRig->bEvaluationTriggeredFromInteraction, true);
+		ControlRig->Evaluate_AnyThread();
+		for (const FRigElementKey& ControlToReset : ControlsToReset)
 		{
-			// we have to insert the interaction event before we run current events
-			TArray<FName> NewEventQueue = {FRigUnit_InteractionExecution::EventName};
-			NewEventQueue.Append(ControlRig->EventQueue);
-			TGuardValue<TArray<FName>> EventGuard(ControlRig->EventQueue, NewEventQueue);
-			const TGuardValue<bool> GuardEvaluationType(ControlRig->bEvaluationTriggeredFromInteraction, true);
-			ControlRig->Evaluate_AnyThread();
-			for (const FRigElementKey& ControlToReset : ControlsToReset)
-			{
-				NotifyDrivenControls(ControlRig, ControlToReset, Context);
-			}
+			NotifyDrivenControls(ControlRig, ControlToReset, Context);
 		}
 	}
 }
 
-void FControlRigEditMode::InvertInputPose(bool bSelectionOnly)
+void FControlRigEditMode::InvertInputPose(bool bSelectionOnly, bool bIncludeChannels)
 {
 	// Gather up the control rigs for the selected controls
 	TArray<UControlRig*> ControlRigs;
@@ -3783,36 +3850,34 @@ void FControlRigEditMode::InvertInputPose(bool bSelectionOnly)
 	}
 
 	FScopedTransaction Transaction(LOCTEXT("HierarchyInvertTransformsToRestPose", "Invert Transforms to Rest Pose"));
-
+	FRigControlModifiedContext Context;
+	Context.SetKey = EControlRigSetKey::DoNotCare;
 	for (UControlRig* ControlRig : ControlRigs)
 	{
-		if (!ControlRig->IsAdditive())
-		{
-			ZeroTransforms(bSelectionOnly);
-			continue;
-		}
+		InvertInputPose(ControlRig, Context,  bSelectionOnly, bIncludeChannels);
+	}
+}
 
-		TArray<FRigElementKey> SelectedRigElements;
-		if (bSelectionOnly)
-		{
-			SelectedRigElements = GetSelectedRigElements(ControlRig);
-			SelectedRigElements = SelectedRigElements.FilterByPredicate([ControlRig](const FRigElementKey& Key)
-			{
-				if (FRigControlElement* Element = ControlRig->FindControl(Key.Name))
-				{
-					return Element->Settings.ControlType != ERigControlType::Bool;
-				}
-				return true;
-			});
-		}
+void FControlRigEditMode::InvertInputPose(UControlRig * ControlRig, const FRigControlModifiedContext & Context, bool bSelectionOnly, bool bIncludeChannels)
+{
+	if (!ControlRig->IsAdditive())
+	{
+		ZeroTransforms(ControlRig, Context, bSelectionOnly, bIncludeChannels);
+		return;
+	}
 
-		const TArray<FRigControlElement*> ModifiedElements = ControlRig->InvertInputPose(SelectedRigElements, EControlRigSetKey::Never);
-		EvaluateRig(ControlRig);
+	TArray<FRigElementKey> SelectedRigElements;
+	if(bSelectionOnly)
+	{
+		SelectedRigElements = FControlRigEditMode::GetRigElementsForSettingTransforms(ControlRig, bSelectionOnly, bIncludeChannels);
+	}
 
-		for (FRigControlElement* ControlElement : ModifiedElements)
-		{
-			ControlRig->ControlModified().Broadcast(ControlRig, ControlElement, EControlRigSetKey::DoNotCare);
-		}
+	const TArray<FRigControlElement*> ModifiedElements = ControlRig->InvertInputPose(SelectedRigElements, Context.SetKey);
+	EvaluateRig(ControlRig);
+
+	for (FRigControlElement* ControlElement : ModifiedElements)
+	{
+		ControlRig->ControlModified().Broadcast(ControlRig, ControlElement, Context.SetKey);
 	}
 }
 
@@ -6233,7 +6298,7 @@ static TArray<UControlRig*> GetControlRigsWithSelectedControls(const TArray<UCon
 /** Select Mirrored Controls on Current Selection*/
 void FControlRigEditMode::SelectMirroredControls()
 {
-	TArray<UControlRig*> ControlRigs = GetControlRigsArray(true);
+	TArray<UControlRig*> ControlRigs = GetControlRigsArray(false);
 	ControlRigs = GetControlRigsWithSelectedControls(ControlRigs);
 	if (ControlRigs.Num() == 0)
 	{
@@ -6258,7 +6323,7 @@ void FControlRigEditMode::SelectMirroredControls()
 /** Select Mirrored Controls on Current Selection, keeping current selection*/
 void FControlRigEditMode::AddMirroredControlsToSelection()
 {
-	TArray<UControlRig*> ControlRigs = GetControlRigsArray(true);
+	TArray<UControlRig*> ControlRigs = GetControlRigsArray(false);
 	ControlRigs = GetControlRigsWithSelectedControls(ControlRigs);
 	if (ControlRigs.Num() == 0)
 	{
@@ -6282,7 +6347,7 @@ void FControlRigEditMode::AddMirroredControlsToSelection()
 /** Put Selected Controls To Mirrored Pose*/
 void FControlRigEditMode::MirrorSelectedControls()
 {
-	TArray<UControlRig*> ControlRigs = GetControlRigsArray(true);
+	TArray<UControlRig*> ControlRigs = GetControlRigsArray(false);
 	ControlRigs = GetControlRigsWithSelectedControls(ControlRigs);
 	if (ControlRigs.Num() == 0)
 	{
@@ -6306,7 +6371,7 @@ void FControlRigEditMode::MirrorSelectedControls()
 /** Put Unselected Controls To Mirrored selcted controls*/
 void FControlRigEditMode::MirrorUnselectedControls()
 {
-	TArray<UControlRig*> ControlRigs = GetControlRigsArray(true);
+	TArray<UControlRig*> ControlRigs = GetControlRigsArray(false);
 	ControlRigs = GetControlRigsWithSelectedControls(ControlRigs);
 	if (ControlRigs.Num() == 0)
 	{
@@ -6332,7 +6397,7 @@ void FControlRigEditMode::MirrorUnselectedControls()
 /** Select All Controls*/
 void FControlRigEditMode::SelectAllControls()
 {
-	TArray<UControlRig*> ControlRigs = GetControlRigsArray(true);
+	TArray<UControlRig*> ControlRigs = GetControlRigsArray(false);
 	ControlRigs = GetControlRigsWithSelectedControls(ControlRigs);
 	if (ControlRigs.Num() == 0)
 	{
@@ -6357,7 +6422,7 @@ void FControlRigEditMode::SavePose(int32 PoseNum)
 {
 	if (StoredPose)
 	{
-		TArray<UControlRig*> ControlRigs = GetControlRigsArray(true);
+		TArray<UControlRig*> ControlRigs = GetControlRigsArray(false);
 		ControlRigs = GetControlRigsWithSelectedControls(ControlRigs);
 		for (UControlRig* ControlRig : ControlRigs)
 		{
@@ -6375,7 +6440,7 @@ void FControlRigEditMode::SelectPose(bool bMirror, int32 PoseNum)
 {
 	if (StoredPose)
 	{
-		TArray<UControlRig*> ControlRigs = GetControlRigsArray(true);
+		TArray<UControlRig*> ControlRigs = GetControlRigsArray(false);
 		if (ControlRigs.Num() < 0)
 		{
 			return;
@@ -6410,7 +6475,7 @@ void FControlRigEditMode::PastePose(bool bMirror, int32 PoseNum)
 {
 	if (StoredPose)
 	{
-		TArray<UControlRig*> ControlRigs = GetControlRigsArray(true);
+		TArray<UControlRig*> ControlRigs = GetControlRigsArray(false);
 		ControlRigs = GetControlRigsWithSelectedControls(ControlRigs);
 		if (ControlRigs.Num() == 0)
 		{
diff --git a/Engine/Plugins/Animation/ControlRig/Source/ControlRigEditor/Private/EditMode/ControlRigEditModeCommands.cpp b/Engine/Plugins/Animation/ControlRig/Source/ControlRigEditor/Private/EditMode/ControlRigEditModeCommands.cpp
--- a/Engine/Plugins/Animation/ControlRig/Source/ControlRigEditor/Private/EditMode/ControlRigEditModeCommands.cpp
+++ b/Engine/Plugins/Animation/ControlRig/Source/ControlRigEditor/Private/EditMode/ControlRigEditModeCommands.cpp
@@ -6,6 +6,8 @@
 
 void FControlRigEditModeCommands::RegisterCommands()
 {
+	UI_COMMAND(InvertTransformsAndChannels, "Invert Selected Controls And Channels to Rest Pose", "Invert the Selected Controls and Channels to Rest Pose", EUserInterfaceActionType::Button, FInputChord());
+	UI_COMMAND(InvertAllTransformsAndChannels, "Invert All Controls and Channels to Rest Pose", "Invert all of the Controls And Channels to Rest Pose", EUserInterfaceActionType::Button, FInputChord());
 	UI_COMMAND(InvertTransforms, "Invert to Rest Pose", "Invert the Controls Transforms to Rest Pose", EUserInterfaceActionType::Button, FInputChord(EKeys::G, EModifierKey::Control));
 	UI_COMMAND(InvertAllTransforms, "Invert All Transform to Rest Pose", "Invert all of the Controls Transforms to Rest Pose", EUserInterfaceActionType::Button, FInputChord(EKeys::G, EModifierKey::Control | EModifierKey::Shift));
 	UI_COMMAND(ZeroTransforms, "Zeros Transforms", "Resets Control Transforms to Zero", EUserInterfaceActionType::Button, FInputChord(EKeys::I, EModifierKey::Control));
@@ -13,6 +15,7 @@ void FControlRigEditModeCommands::RegisterCommands()
 	UI_COMMAND(ToggleManipulators, "Toggle Shapes", "Toggles visibility of active control rig shapes in the viewport", EUserInterfaceActionType::Button, FInputChord(EKeys::T));
 	UI_COMMAND(ToggleModuleManipulators, "Toggle Module Shapes", "Toggles visibility of active control rig shapes belonging to selected modules in the viewport", EUserInterfaceActionType::Button, FInputChord(EKeys::T, EModifierKey::Control));
 	UI_COMMAND(ToggleAllManipulators, "Toggle All Shapes", "Toggles visibility of all control rig shapes in the viewport", EUserInterfaceActionType::Button, FInputChord(EKeys::T, EModifierKey::Control | EModifierKey::Alt));
+	UI_COMMAND(ToggleControlsAsOverlay, "Toggle Contraols As Overlay", "Toggles showing controls as overlay", EUserInterfaceActionType::Button, FInputChord());
 	UI_COMMAND(FrameSelection, "Frame Selection", "Focus the viewport on the current selection", EUserInterfaceActionType::Button, FInputChord(EKeys::F));
 	UI_COMMAND(ClearSelection, "Clear Selection", "Clear Selection", EUserInterfaceActionType::Button, FInputChord(EKeys::Escape));
 
diff --git a/Engine/Plugins/Animation/ControlRig/Source/ControlRigEditor/Private/EditMode/ControlRigEditModeCommands.h b/Engine/Plugins/Animation/ControlRig/Source/ControlRigEditor/Private/EditMode/ControlRigEditModeCommands.h
--- a/Engine/Plugins/Animation/ControlRig/Source/ControlRigEditor/Private/EditMode/ControlRigEditModeCommands.h
+++ b/Engine/Plugins/Animation/ControlRig/Source/ControlRigEditor/Private/EditMode/ControlRigEditModeCommands.h
@@ -29,6 +29,15 @@ public:
 	/** Toggles hiding all manipulators on all control rigs in the viewport */
 	TSharedPtr< FUICommandInfo > ToggleAllManipulators;
 
+	/** Toggles show controls as an overlay in the viewport */
+	TSharedPtr< FUICommandInfo > ToggleControlsAsOverlay;
+
+	/** Invert Transforms and channels to Rest Pose */
+	TSharedPtr< FUICommandInfo > InvertTransformsAndChannels;
+
+	/** Invert All Transforms and channels  to Rest Pose */
+	TSharedPtr< FUICommandInfo > InvertAllTransformsAndChannels;
+
 	/** Invert Transforms to Rest Pose */
 	TSharedPtr< FUICommandInfo > InvertTransforms;
 
diff --git a/Engine/Plugins/Animation/ControlRig/Source/ControlRigEditor/Public/EditMode/ControlRigEditMode.h b/Engine/Plugins/Animation/ControlRig/Source/ControlRigEditor/Public/EditMode/ControlRigEditMode.h
--- a/Engine/Plugins/Animation/ControlRig/Source/ControlRigEditor/Public/EditMode/ControlRigEditMode.h
+++ b/Engine/Plugins/Animation/ControlRig/Source/ControlRigEditor/Public/EditMode/ControlRigEditMode.h
@@ -277,6 +277,9 @@ protected:
 
 public:
 
+	/** Toggle controls as overlay*/
+	void ToggleControlsAsOverlay();
+
 	/** Toggles visibility of acive control rig shapes inside the selected module in the viewport */
 	void ToggleModuleManipulators();
 	
@@ -322,10 +325,16 @@ public:
 	void OpenSpacePickerWidget();
 
 	/** Reset Transforms */
-	void ZeroTransforms(bool bSelectionOnly);
+	void ZeroTransforms(bool bSelectionOnly, bool bIncludeChannels = true);
 
 	/** Invert Input Pose */
-	void InvertInputPose(bool bSelectionOnly);
+	void InvertInputPose(bool bSelectionOnly, bool bIncludeChannels = true);
+
+	/** Reset Transforms for this Control Rig Controls based upon selection and channel states*/
+	static void ZeroTransforms(UControlRig* ControlRig, const FRigControlModifiedContext& Context,  bool bSelectionOnly, bool bIncludeChannels);
+
+	/** Invert Input Pose for this Control Rig Controls based upon selection and channel states*/
+	static void InvertInputPose(UControlRig* ControlRig, const FRigControlModifiedContext& Context, bool bSelectionOnly, bool bIncludeChannels);
 
 private:
 	
@@ -425,7 +434,11 @@ protected:
 	FControlRigSelected OnControlRigSelectedDelegate;
 
 	/** GetSelectedRigElements, if InControlRig is nullptr get the first one */
-	TArray<FRigElementKey> GetSelectedRigElements(UControlRig* InControlRig = nullptr) const;
+	TArray<FRigElementKey> GetSelectedRigElements() const;
+	static TArray<FRigElementKey> GetSelectedRigElements(UControlRig* InControlRig);
+
+	/** Get the rig elements, based upon the selection or if it's a channel, also will do internal additive filters(like no bools)  */
+	static TArray<FRigElementKey> GetRigElementsForSettingTransforms(UControlRig* InControlRig, bool bSelectionOnly, bool bIncludeChannels);
 
 	/* Flag to recreate shapes during tick */
 	ERecreateControlRigShape RecreateControlShapesRequired;
