CIS Fix

diff --git a/Engine/Plugins/Animation/ControlRig/Source/ControlRigEditor/Private/EditMode/ControlRigEditMode.cpp b/Engine/Plugins/Animation/ControlRig/Source/ControlRigEditor/Private/EditMode/ControlRigEditMode.cpp
--- a/Engine/Plugins/Animation/ControlRig/Source/ControlRigEditor/Private/EditMode/ControlRigEditMode.cpp
+++ b/Engine/Plugins/Animation/ControlRig/Source/ControlRigEditor/Private/EditMode/ControlRigEditMode.cpp
@@ -2436,7 +2436,7 @@ TArray<FRigElementKey> FControlRigEditMode::GetSelectedRigElements(UControlRig*
 
 bool FControlRigEditMode::AreRigElementsSelected(uint32 InTypes, UControlRig* InControlRig) const
 {
-	if (IsInLevelEditor())
+	if (IsInLevelEditor() && InControlRig)
 	{
 		// no need to look for transient controls when animating in the level editor
 		const URigHierarchy* Hierarchy = InControlRig->GetHierarchy();
