Control Rig dynamic hierarchy now supports using a skeleton for binding at import bones rig unit

diff --git a/Engine/Plugins/Animation/ControlRig/Source/ControlRig/Private/Units/Execution/RigUnit_DynamicHierarchy.cpp b/Engine/Plugins/Animation/ControlRig/Source/ControlRig/Private/Units/Execution/RigUnit_DynamicHierarchy.cpp
--- a/Engine/Plugins/Animation/ControlRig/Source/ControlRig/Private/Units/Execution/RigUnit_DynamicHierarchy.cpp
+++ b/Engine/Plugins/Animation/ControlRig/Source/ControlRig/Private/Units/Execution/RigUnit_DynamicHierarchy.cpp
@@ -338,59 +338,72 @@ FRigUnit_HierarchyImportFromSkeleton_Execute()
 	{
 		FRigHierarchyControllerInstructionBracket InstructionBracket(Controller, ExecuteContext.GetInstructionIndex());
 		
-		if(const USkeletalMeshComponent* SkelMeshComponent = ExecuteContext.UnitContext.DataSourceRegistry->RequestSource<USkeletalMeshComponent>(UControlRig::OwnerComponent))
+		USkeleton* Skeleton = nullptr;
+		const FReferenceSkeleton* ReferenceSkeleton = nullptr;
+
+		const USkeletalMeshComponent* SkelMeshComponent = ExecuteContext.UnitContext.DataSourceRegistry->RequestSource<USkeletalMeshComponent>(UControlRig::OwnerComponent);
+		if (SkelMeshComponent != nullptr)
 		{
-			if(SkelMeshComponent->GetSkeletalMeshAsset())
+			if (SkelMeshComponent->GetSkeletalMeshAsset())
 			{
-				const FReferenceSkeleton& ReferenceSkeleton = SkelMeshComponent->GetSkeletalMeshAsset()->GetRefSkeleton();
-				
-				if (bIncludeVirtualBones)
-				{
-					Items = Controller->ImportBones(ReferenceSkeleton, NameSpace, false, false, false, false);
-				}
-				else
-				{
-					const TArray<FMeshBoneInfo>& BoneInfos = ReferenceSkeleton.GetRefBoneInfo();
-					const TArray<FTransform>& BoneTransforms = ReferenceSkeleton.GetRefBonePose();
-
-					TArray<FMeshBoneInfo> FilteredBoneInfos;
-					TArray<FTransform> FilteredBoneTransforms;
-
-					const int32 NumBones = BoneInfos.Num();
-					for (int32 Index = 0; Index < NumBones; ++Index)
-					{
-						const FMeshBoneInfo& MeshBoneInfo = BoneInfos[Index];
-						const FTransform& BoneTransform = BoneTransforms[Index];
-
-						if (VirtualBoneNameHelpers::CheckVirtualBonePrefix(MeshBoneInfo.Name.ToString()))
-						{
-							continue;
-						}
-						FilteredBoneInfos.Add(MeshBoneInfo);
-						FilteredBoneTransforms.Add(BoneTransform);
-					}
+				Skeleton = SkelMeshComponent->GetSkeletalMeshAsset()->GetSkeleton();
+				ReferenceSkeleton = &SkelMeshComponent->GetSkeletalMeshAsset()->GetRefSkeleton();
+			}
+		}
+		else
+		{
+			Skeleton = ExecuteContext.UnitContext.DataSourceRegistry->RequestSource<USkeleton>(UControlRig::OwnerComponent);
+			if (Skeleton != nullptr)
+			{
+				ReferenceSkeleton = &Skeleton->GetReferenceSkeleton();
+			}
+		}
 
-					Items = Controller->ImportBones(FilteredBoneInfos, FilteredBoneTransforms, NameSpace, false, false, false, false);
-				}
+		if(ReferenceSkeleton != nullptr)
+		{
+			if (bIncludeVirtualBones)
+			{
+				Items = Controller->ImportBones(*ReferenceSkeleton, NameSpace, false, false, false, false);
+			}
+			else
+			{
+				const TArray<FMeshBoneInfo>& BoneInfos = ReferenceSkeleton->GetRefBoneInfo();
+				const TArray<FTransform>& BoneTransforms = ReferenceSkeleton->GetRefBonePose();
 
+				TArray<FMeshBoneInfo> FilteredBoneInfos;
+				TArray<FTransform> FilteredBoneTransforms;
 
-				if(bIncludeCurves)
+				const int32 NumBones = BoneInfos.Num();
+				for (int32 Index = 0; Index < NumBones; ++Index)
 				{
-					if(USkeleton* Skeleton = SkelMeshComponent->GetSkeletalMeshAsset()->GetSkeleton())
+					const FMeshBoneInfo& MeshBoneInfo = BoneInfos[Index];
+					const FTransform& BoneTransform = BoneTransforms[Index];
+
+					if (VirtualBoneNameHelpers::CheckVirtualBonePrefix(MeshBoneInfo.Name.ToString()))
 					{
-						Items.Append(Controller->ImportCurves(Skeleton, NameSpace, false, false, false));
+						continue;
 					}
+					FilteredBoneInfos.Add(MeshBoneInfo);
+					FilteredBoneTransforms.Add(BoneTransform);
 				}
-				if(bIncludeMeshSockets)
+
+				Items = Controller->ImportBones(FilteredBoneInfos, FilteredBoneTransforms, NameSpace, false, false, false, false);
+			}
+
+			if(bIncludeCurves && Skeleton != nullptr)
+			{
+				Items.Append(Controller->ImportCurves(Skeleton, NameSpace, false, false, false));
+			}
+			if(bIncludeMeshSockets && SkelMeshComponent != nullptr)
+			{
+				if(USkeletalMesh* SkeletalMesh = SkelMeshComponent->GetSkeletalMeshAsset())
 				{
-					if(USkeletalMesh* SkeletalMesh = SkelMeshComponent->GetSkeletalMeshAsset())
-					{
-						Items.Append(Controller->ImportSocketsFromSkeletalMesh(SkeletalMesh, NameSpace, false, false, false));
-					}
+					Items.Append(Controller->ImportSocketsFromSkeletalMesh(SkeletalMesh, NameSpace, false, false, false));
 				}
 			}
 		}
 	}
+	ensure(Items.Num() > 0);
 }
 
 FRigUnit_HierarchyRemoveElement_Execute()
