ControlRig Rotation Order: fixed preferred angle computation if the rotation order is not the default one (XYZ)
- URigHierarchy::SetControlPreferredEulerAngles takes a FTransform as input, suggesting that the incoming values come in XYZ order

#rnx

diff --git a/Engine/Plugins/Animation/ControlRig/Source/ControlRig/Private/Rigs/RigHierarchy.cpp b/Engine/Plugins/Animation/ControlRig/Source/ControlRig/Private/Rigs/RigHierarchy.cpp
--- a/Engine/Plugins/Animation/ControlRig/Source/ControlRig/Private/Rigs/RigHierarchy.cpp
+++ b/Engine/Plugins/Animation/ControlRig/Source/ControlRig/Private/Rigs/RigHierarchy.cpp
@@ -1,6 +1,8 @@
 // Copyright Epic Games, Inc. All Rights Reserved.
 
 #include "Rigs/RigHierarchy.h"
+
+#include "AnimationCoreLibrary.h"
 #include "Async/TaskGraphInterfaces.h"
 #include "ControlRig.h"
 
@@ -9822,6 +9824,32 @@ void URigHierarchy::PatchModularRigComponentKeys(const TMap<FRigHierarchyModuleP
 	}
 }
 
+void URigHierarchy::SetControlPreferredEulerAngles(FRigControlElement* InControlElement, const FTransform& InTransform, bool bIsInitial)
+{
+	if (InControlElement)
+	{
+		const FEulerTransform EulerTransform(InTransform);
+		const FVector EulerXYZAngle(EulerTransform.Rotation.Roll, EulerTransform.Rotation.Pitch, EulerTransform.Rotation.Yaw);
+
+		// switch from XYZ to the actual rotation order as FTransform doesn't have any notion of order
+		const EEulerRotationOrder RotationOrder = InControlElement->Settings.PreferredRotationOrder;
+		const FVector EulerAngle = GetUsePreferredRotationOrder(InControlElement) ?
+			AnimationCore::ChangeEulerRotationOrder(EulerXYZAngle, EEulerRotationOrder::XYZ, RotationOrder) : EulerXYZAngle;
+			
+		switch(InControlElement->Settings.ControlType)
+		{
+		case ERigControlType::Transform:
+		case ERigControlType::Rotator:
+		case ERigControlType::TransformNoScale:
+		case ERigControlType::EulerTransform:
+			SetControlSpecifiedEulerAngle(InControlElement, EulerAngle, bIsInitial);
+			break;
+		default:
+			break;
+		}
+	}
+}
+
 FRigHierarchyRedirectorGuard::FRigHierarchyRedirectorGuard(UControlRig* InControlRig)
 : Guard(InControlRig->GetHierarchy()->ElementKeyRedirector, &InControlRig->GetElementKeyRedirector())
 {
diff --git a/Engine/Plugins/Animation/ControlRig/Source/ControlRig/Public/Rigs/RigHierarchy.h b/Engine/Plugins/Animation/ControlRig/Source/ControlRig/Public/Rigs/RigHierarchy.h
--- a/Engine/Plugins/Animation/ControlRig/Source/ControlRig/Public/Rigs/RigHierarchy.h
+++ b/Engine/Plugins/Animation/ControlRig/Source/ControlRig/Public/Rigs/RigHierarchy.h
@@ -2921,37 +2921,7 @@ public:
 		}
 	}
 
-	void SetControlPreferredEulerAngles(FRigControlElement* InControlElement, const FTransform& InTransform, bool bIsInitial = false)
-	{
-		if (InControlElement)
-		{
-			FEulerTransform EulerTransform(InTransform);
-			switch(InControlElement->Settings.ControlType)
-			{
-				case ERigControlType::Transform:
-				case ERigControlType::Rotator:
-					{
-						FVector EulerAngle(EulerTransform.Rotation.Roll, EulerTransform.Rotation.Pitch, EulerTransform.Rotation.Yaw);
-						SetControlSpecifiedEulerAngle(InControlElement, EulerAngle, bIsInitial);
-						break;
-					}
-				case ERigControlType::TransformNoScale:
-					{
-						const FTransformNoScale NoScale = EulerTransform.ToFTransform();
-						FVector EulerAngle(EulerTransform.Rotation.Roll, EulerTransform.Rotation.Pitch, EulerTransform.Rotation.Yaw);
-						SetControlSpecifiedEulerAngle(InControlElement, EulerAngle, bIsInitial);
-						break;
-					}
-				case ERigControlType::EulerTransform:
-					{
-						FVector EulerAngle(EulerTransform.Rotation.Roll, EulerTransform.Rotation.Pitch, EulerTransform.Rotation.Yaw);
-						FQuat Quat = GetControlQuaternion(InControlElement, EulerAngle);
-						SetControlSpecifiedEulerAngle(InControlElement, EulerAngle, bIsInitial);
-						break;
-					}
-			}
-		}
-	}
+	void SetControlPreferredEulerAngles(FRigControlElement* InControlElement, const FTransform& InTransform, bool bIsInitial = false);
 
 	FQuat GetControlQuaternion(const FRigControlElement* InControlElement, const FVector& InEulerAngle) const
 	{
