Control Rig: Add interaction events (start, update, end) so that c++/python delegates can hook up to these events

diff --git a/Engine/Plugins/Animation/ControlRig/Source/ControlRig/Public/ControlRig.h b/Engine/Plugins/Animation/ControlRig/Source/ControlRig/Public/ControlRig.h
--- a/Engine/Plugins/Animation/ControlRig/Source/ControlRig/Public/ControlRig.h
+++ b/Engine/Plugins/Animation/ControlRig/Source/ControlRig/Public/ControlRig.h
@@ -30,6 +30,7 @@
 
 #if WITH_EDITOR
 #include "AnimPreviewInstance.h"
+#include "UnrealWidgetFwd.h"
 #endif 
 
 #include "ControlRig.generated.h"
@@ -41,6 +42,7 @@ class USkeletalMeshComponent;
 class AActor;
 class UTransformableControlHandle;
 class UControlRigReplay;
+class UControlRig;
 
 struct FReferenceSkeleton;
 struct FRigUnit;
@@ -49,6 +51,65 @@ struct FRigPhysicsSimulationBase;
 
 CONTROLRIG_API DECLARE_LOG_CATEGORY_EXTERN(LogControlRig, Log, All);
 
+UENUM()
+enum class EControlRigInteractionTransformSpace
+{
+	World,
+	Local,
+	Parent,
+	Explicit
+};
+
+/**
+ * FTransformContext provides a way of passing the various transform parameters to functions that need to know what the transform context is.
+ * Extend it if necessary, particularly to avoid overloading some functions signatures.
+ */
+USTRUCT(BlueprintType)
+struct FControlRigInteractionTransformContext
+{
+	GENERATED_BODY()
+	
+	FControlRigInteractionTransformContext() = default;
+
+#if WITH_EDITOR
+	FControlRigInteractionTransformContext(const UE::Widget::EWidgetMode& InWidgetMode)
+		: bTranslation(InWidgetMode == UE::Widget::WM_Translate || InWidgetMode == UE::Widget::WM_TranslateRotateZ)
+		, bRotation(InWidgetMode == UE::Widget::WM_Rotate || InWidgetMode == UE::Widget::WM_TranslateRotateZ)
+		, bScale(InWidgetMode == UE::Widget::WM_Scale)
+	{}
+#endif
+	
+	bool CanTransform() const
+	{
+		return bTranslation || bRotation || bScale;
+	}
+
+	bool bTranslation = false;
+	
+	UPROPERTY(BlueprintReadWrite, EditAnywhere, Category = "Interaction")
+	FVector Drag = FVector::ZeroVector;
+
+	bool bRotation = false;
+	
+	UPROPERTY(BlueprintReadWrite, EditAnywhere, Category = "Interaction")
+	FRotator Rot = FRotator::ZeroRotator;
+
+	bool bScale = false;
+
+	UPROPERTY(BlueprintReadWrite, EditAnywhere, Category = "Interaction")
+	FVector Scale = FVector::OneVector;
+
+	UPROPERTY(BlueprintReadWrite, EditAnywhere, Category = "Interaction")
+	EControlRigInteractionTransformSpace Space = EControlRigInteractionTransformSpace::World;
+};
+
+#if WITH_EDITORONLY_DATA
+// A delegate for subscribing / reacting to interaction started.
+DECLARE_DYNAMIC_MULTICAST_DELEGATE_ThreeParams(FControlRigInteractionStartedEvent, UControlRig*, ControlRig, TArray<FRigElementKey>, InteractionKeys, EControlRigInteractionType, InteractionType);
+DECLARE_DYNAMIC_MULTICAST_DELEGATE_TwoParams(FControlRigInteractionEndedEvent, UControlRig*, ControlRig, TArray<FRigElementKey>, InteractionKeys);
+DECLARE_DYNAMIC_MULTICAST_DELEGATE_ThreeParams(FControlRigInteractionUpdatedEvent, UControlRig*, ControlRig, TArray<FRigElementKey>, InteractionKeys, FControlRigInteractionTransformContext, TransformContext);
+#endif
+
 /** Runs logic for mapping input data to transforms (the "Rig") */
 UCLASS(Blueprintable, Abstract, editinlinenew)
 class CONTROLRIG_API UControlRig : public URigVMHost, public INodeMappingProviderInterface, public IRigHierarchyProvider, public IMeshDeformerProducer
@@ -686,6 +747,30 @@ protected:
 	UPROPERTY(BlueprintAssignable, Category = ControlRig, meta=(DisplayName="OnControlSelected"))
 	FOnControlSelectedBP OnControlSelected_BP;
 
+#if WITH_EDITORONLY_DATA
+
+	/** Broadcasts a notification when a gizmo manipulation has started */
+	UPROPERTY(BlueprintReadOnly, Category = ControlRig, meta = (ScriptName = "OnGizmoInteractionStarted", AllowPrivateAccess = "true"))
+	FControlRigInteractionStartedEvent OnGizmoInteractionStartedEvent;
+
+	/** Broadcasts a notification when a gizmo manipulation has ended */
+	UPROPERTY(BlueprintReadOnly, Category = ControlRig, meta = (ScriptName = "OnGizmoInteractionEnded", AllowPrivateAccess = "true"))
+	FControlRigInteractionEndedEvent OnGizmoInteractionEndedEvent;
+
+	/** Broadcasts a notification when a gizmo manipulation has been updated (before the update actually takes place) */
+	UPROPERTY(BlueprintReadOnly, Category = ControlRig, meta = (ScriptName = "OnGizmoInteractionUpdated", AllowPrivateAccess = "true"))
+	FControlRigInteractionUpdatedEvent OnGizmoInteractionUpdatedEvent;
+
+public:
+	/** Broadcasts a notification when a gizmo manipulation has started */
+	FControlRigInteractionStartedEvent& OnGizmoInteractionStarted() { return OnGizmoInteractionStartedEvent; }
+	/** Broadcasts a notification when a gizmo manipulation has ended */
+	FControlRigInteractionEndedEvent& OnGizmoInteractionEnded() { return OnGizmoInteractionEndedEvent; }
+	/** Broadcasts a notification when a gizmo manipulation has been updated (before the update actually takes place) */
+	FControlRigInteractionUpdatedEvent& OnGizmoInteractionUpdated() { return OnGizmoInteractionUpdatedEvent; }
+protected:
+#endif
+
 	TArray<FRigElementKey> QueuedModifiedControls;
 
 private:
@@ -988,6 +1073,12 @@ public:
 		, SyncBracketScope(InControlRig->InterRigSyncBracket)
 	{
 		InControlRig->GetHierarchy()->StartInteraction();
+#if WITH_EDITORONLY_DATA
+		if (InControlRig->OnGizmoInteractionStarted().IsBound())
+		{
+			InControlRig->OnGizmoInteractionStarted().Broadcast(InControlRig, InControlRig->ElementsBeingInteracted, EControlRigInteractionType::All);
+		}
+#endif
 	}
 
 	FControlRigInteractionScope(
@@ -1003,6 +1094,12 @@ public:
 		InControlRig->InteractionType = (uint8)InInteractionType;
 		InControlRig->bInteractionJustBegan = true;
 		InControlRig->GetHierarchy()->StartInteraction();
+#if WITH_EDITORONLY_DATA
+		if (InControlRig->OnGizmoInteractionStarted().IsBound())
+		{
+			InControlRig->OnGizmoInteractionStarted().Broadcast(InControlRig, InControlRig->ElementsBeingInteracted, InInteractionType);
+		}
+#endif
 	}
 
 	FControlRigInteractionScope(
@@ -1018,12 +1115,24 @@ public:
 		InControlRig->InteractionType = (uint8)InInteractionType;
 		InControlRig->bInteractionJustBegan = true;
 		InControlRig->GetHierarchy()->StartInteraction();
+#if WITH_EDITORONLY_DATA
+		if (InControlRig->OnGizmoInteractionStarted().IsBound())
+		{
+			InControlRig->OnGizmoInteractionStarted().Broadcast(InControlRig, InControlRig->ElementsBeingInteracted, InInteractionType);
+		}
+#endif
 	}
 
 	~FControlRigInteractionScope()
 	{
 		if(ensure(ControlRig.IsValid()))
 		{
+#if WITH_EDITORONLY_DATA
+			if (ControlRig->OnGizmoInteractionEnded().IsBound())
+			{
+				ControlRig->OnGizmoInteractionEnded().Broadcast(ControlRig.Get(), ControlRig->ElementsBeingInteracted);
+			}
+#endif
 			ControlRig->GetHierarchy()->EndInteraction();
 			ControlRig->InteractionType = (uint8)EControlRigInteractionType::None;
 			ControlRig->bInteractionJustBegan = false;
diff --git a/Engine/Plugins/Animation/ControlRig/Source/ControlRigEditor/Private/ControlRigElementDetails.cpp b/Engine/Plugins/Animation/ControlRig/Source/ControlRigEditor/Private/ControlRigElementDetails.cpp
--- a/Engine/Plugins/Animation/ControlRig/Source/ControlRigEditor/Private/ControlRigElementDetails.cpp
+++ b/Engine/Plugins/Animation/ControlRig/Source/ControlRigEditor/Private/ControlRigElementDetails.cpp
@@ -2959,7 +2959,7 @@ FDetailWidgetRow& FRigTransformElementDetails::CreateEulerTransformValueWidgetRo
 					DebuggedRig->ElementsBeingInteracted.AddUnique(ElementInfo.Element.GetKey());
 				}
 				
-				FControlRigInteractionScope* InteractionScope = new FControlRigInteractionScope(DebuggedRig);
+				FControlRigInteractionScope* InteractionScope = new FControlRigInteractionScope(DebuggedRig, DebuggedRig->ElementsBeingInteracted, Type);
 				StrongThisPtr->InteractionScopes.Add(InteractionScope);
 			}
 		});
diff --git a/Engine/Plugins/Animation/ControlRig/Source/ControlRigEditor/Private/EditMode/ControlRigEditMode.cpp b/Engine/Plugins/Animation/ControlRig/Source/ControlRigEditor/Private/EditMode/ControlRigEditMode.cpp
--- a/Engine/Plugins/Animation/ControlRig/Source/ControlRigEditor/Private/EditMode/ControlRigEditMode.cpp
+++ b/Engine/Plugins/Animation/ControlRig/Source/ControlRigEditor/Private/EditMode/ControlRigEditMode.cpp
@@ -1213,7 +1213,8 @@ bool FControlRigEditMode::HandleBeginTransform(const FEditorViewportClient* InVi
 				if (AreRigElementSelectedAndMovable(ControlRig))
 				{
 					//todo need to add multiple
-					FControlRigInteractionScope* InteractionScope = new FControlRigInteractionScope(ControlRig);
+					TArray<FRigElementKey> SelectedRigElements = GetSelectedRigElements(ControlRig);
+					FControlRigInteractionScope* InteractionScope = new FControlRigInteractionScope(ControlRig, SelectedRigElements, (EControlRigInteractionType)InteractionType);
 					InteractionScopes.Add(ControlRig,InteractionScope);
 					ControlRig->bInteractionJustBegan = true;
 				}
@@ -1227,7 +1228,8 @@ bool FControlRigEditMode::HandleBeginTransform(const FEditorViewportClient* InVi
 	else if(UControlRigEditorSettings::Get()->bEnableUndoForPoseInteraction)
 	{
 		UControlRig* ControlRig = RuntimeControlRigs[0].Get();
-		FControlRigInteractionScope* InteractionScope = new FControlRigInteractionScope(ControlRig);
+		TArray<FRigElementKey> SelectedRigElements = GetSelectedRigElements(ControlRig);
+		FControlRigInteractionScope* InteractionScope = new FControlRigInteractionScope(ControlRig, SelectedRigElements, (EControlRigInteractionType)InteractionType);
 		InteractionScopes.Add(ControlRig,InteractionScope);
 	}
 	else
@@ -1244,7 +1246,7 @@ bool FControlRigEditMode::HandleEndTransform(FEditorViewportClient* InViewportCl
 		return false;
 	}
 
-	const FTransformContext TransformContext(InViewportClient->GetWidgetMode());
+	const FControlRigInteractionTransformContext TransformContext(InViewportClient->GetWidgetMode());
 	const bool bWasInteracting = bManipulatorMadeChange && InteractionType != (uint8)EControlRigInteractionType::None;
 	
 	InteractionType = (uint8)EControlRigInteractionType::None;
@@ -1479,7 +1481,8 @@ bool FControlRigEditMode::HandleClick(FEditorViewportClient* InViewportClient, H
 
 					if (bInteractionScopePresent)
 					{
-						*InteractionScope = new FControlRigInteractionScope(ShapeActor->ControlRig.Get());
+						TArray<FRigElementKey> SelectedRigElements = GetSelectedRigElements(ShapeActor->ControlRig.Get());
+						*InteractionScope = new FControlRigInteractionScope(ShapeActor->ControlRig.Get(), SelectedRigElements, (EControlRigInteractionType) InteractionType);
 						InteractionScopes.Add(ShapeActor->ControlRig.Get(), *InteractionScope);
 
 					}
@@ -1984,9 +1987,9 @@ bool FControlRigEditMode::InputDelta(FEditorViewportClient* InViewportClient, FV
 
 	const UE::Widget::EWidgetMode WidgetMode = InViewportClient->GetWidgetMode();
 	const EAxisList::Type CurrentAxis = InViewportClient->GetCurrentWidgetAxis();
-	const ETransformSpace TransformSpace = GetTransformSpace();
+	const EControlRigInteractionTransformSpace TransformSpace = GetTransformSpace();
 	
-	FTransformContext TransformContext;
+	FControlRigInteractionTransformContext TransformContext;
 	TransformContext.bTranslation = !InDrag.IsZero() && (WidgetMode == UE::Widget::WM_Translate || WidgetMode == UE::Widget::WM_TranslateRotateZ);
 	TransformContext.Drag = InDrag;
 	TransformContext.bRotation = !InRot.IsZero() && (WidgetMode == UE::Widget::WM_Rotate || WidgetMode == UE::Widget::WM_TranslateRotateZ);
@@ -2001,12 +2004,12 @@ bool FControlRigEditMode::InputDelta(FEditorViewportClient* InViewportClient, FV
 		{
 			switch (TransformSpace)
 			{
-			case ETransformSpace::World:
+			case EControlRigInteractionTransformSpace::World:
 				return false;
-			case ETransformSpace::Local:
-			case ETransformSpace::Parent:
+			case EControlRigInteractionTransformSpace::Local:
+			case EControlRigInteractionTransformSpace::Parent:
 				return true;
-			case ETransformSpace::Explicit:
+			case EControlRigInteractionTransformSpace::Explicit:
 				return TransformContext.bTranslation;
 			default:
 				break;
@@ -2073,6 +2076,9 @@ bool FControlRigEditMode::InputDelta(FEditorViewportClient* InViewportClient, FV
 								GEditor->BeginTransaction(LOCTEXT("MoveControlTransaction", "Move Control"));
 							}
 
+							// The interaction update event must be broadcasted before setting any control value or executing the rig
+							ControlRig->OnGizmoInteractionUpdated().Broadcast(ControlRig, ControlRig->ElementsBeingInteracted, TransformContext);
+
 							// cache interaction dependencies + evaluate the rig at least once before manipulating anything
 							if (ControlRig->ElementsBeingInteracted.IsEmpty() && InteractionScopes.Contains(ControlRig))
 							{
@@ -2540,27 +2546,27 @@ bool FControlRigEditMode::ComputePivotFromEditedShape(UControlRig* InControlRig,
 	return true;
 }
 
-ETransformSpace FControlRigEditMode::GetTransformSpace() const
+EControlRigInteractionTransformSpace FControlRigEditMode::GetTransformSpace() const
 {
 	switch (GetCoordSystemSpace())
 	{
 		case COORD_World:
-			return ETransformSpace::World;
+			return EControlRigInteractionTransformSpace::World;
 		case COORD_Local:
-			return ETransformSpace::Local;
+			return EControlRigInteractionTransformSpace::Local;
 		case COORD_Parent:
-			return ETransformSpace::Parent;
+			return EControlRigInteractionTransformSpace::Parent;
 		case COORD_Explicit:
-			return ETransformSpace::Explicit;
+			return EControlRigInteractionTransformSpace::Explicit;
 		default:
 			break;
 	}
-	return ETransformSpace::Local;	
+	return EControlRigInteractionTransformSpace::Local;	
 };
 
 FTransform FControlRigEditMode::GetPivotOrientation(
 	const FRigElementKey& InControlKey, const UControlRig* InControlRig, URigHierarchy* InHierarchy,
-	const ETransformSpace InSpace, const FTransform& InComponentTransform) const
+	const EControlRigInteractionTransformSpace InSpace, const FTransform& InComponentTransform) const
 {
 	FRigControlElement* Control = InHierarchy->Find<FRigControlElement>(InControlKey);
 	if (!Control)
@@ -2570,10 +2576,10 @@ FTransform FControlRigEditMode::GetPivotOrientation(
 		
 	switch (InSpace)
 	{
-	case ETransformSpace::World:
-	case ETransformSpace::Local:
+	case EControlRigInteractionTransformSpace::World:
+	case EControlRigInteractionTransformSpace::Local:
 		return InHierarchy->GetTransform(Control, ERigTransformType::CurrentGlobal);
-	case ETransformSpace::Parent:
+	case EControlRigInteractionTransformSpace::Parent:
 		{	
 			if (TOptional<FTransform> ConstraintSpace = GetConstraintParentTransform(InControlRig, InControlKey.Name))
 			{
@@ -2583,7 +2589,7 @@ FTransform FControlRigEditMode::GetPivotOrientation(
 			const int32 NumParents = InHierarchy->GetNumberOfParents(Control);
 			return NumParents > 0 ? InHierarchy->GetParentTransform(Control, ERigTransformType::CurrentGlobal) : InHierarchy->GetTransform(Control, ERigTransformType::CurrentGlobal);
 		}
-	case ETransformSpace::Explicit:
+	case EControlRigInteractionTransformSpace::Explicit:
 		{
 			FRotationContext& RotationContext = GetRotationContext();
 					
@@ -2625,7 +2631,7 @@ FTransform FControlRigEditMode::GetPivotOrientation(
 	return FTransform::Identity;
 }
 
-bool FControlRigEditMode::ComputePivotFromShapeActors(UControlRig* InControlRig, const bool bEachLocalSpace, const ETransformSpace InSpace, FTransform& OutTransform) const
+bool FControlRigEditMode::ComputePivotFromShapeActors(UControlRig* InControlRig, const bool bEachLocalSpace, const EControlRigInteractionTransformSpace InSpace, FTransform& OutTransform) const
 {
 	if (!ensure(!bIsChangingControlShapeTransform))
 	{
@@ -4092,7 +4098,7 @@ void FControlRigEditMode::UpdateSequencerStatus()
     	{
     		if (FEditorModeTools* ModeTools = GetModeManager())
     		{
-    			const FTransformContext TransformContext(ModeTools->GetWidgetMode());
+    			const FControlRigInteractionTransformContext TransformContext(ModeTools->GetWidgetMode());
     			for (const auto& [ControlRig, InteractionScope] : InteractionScopes)
     			{
     				if (InteractionScope)
@@ -5143,7 +5149,7 @@ void FControlRigEditMode::SetControlShapeTransform(
 	const FTransform& InToWorldTransform,
 	const FRigControlModifiedContext& InContext,
 	const bool bPrintPython,
-	const FTransformContext& InTransformContext,
+	const FControlRigInteractionTransformContext& InTransformContext,
 	const bool bFixEulerFlips) const
 {
 	UControlRig* ControlRig = InShapeActor->ControlRig.Get();
@@ -5261,7 +5267,7 @@ FTransform FControlRigEditMode::GetControlShapeTransform(const AControlRigShapeA
 }
 
 bool FControlRigEditMode::MoveControlShapeLocally( AControlRigShapeActor* ShapeActor,
-	const FTransformContext& InTransformContext, const FTransform& ToWorldTransform,
+	const FControlRigInteractionTransformContext& InTransformContext, const FTransform& ToWorldTransform,
 	const FTransform& InLocal)
 {
 	if (!ensure(InTransformContext.CanTransform()))
@@ -5292,13 +5298,13 @@ bool FControlRigEditMode::MoveControlShapeLocally( AControlRigShapeActor* ShapeA
 		{
 			switch(InTransformContext.Space)
 			{
-			case ETransformSpace::World:
-			case ETransformSpace::Explicit:
+			case EControlRigInteractionTransformSpace::World:
+			case EControlRigInteractionTransformSpace::Explicit:
 				ensure(false);
 				break;
-			case ETransformSpace::Local:
+			case EControlRigInteractionTransformSpace::Local:
 				break;
-			case ETransformSpace::Parent:
+			case EControlRigInteractionTransformSpace::Parent:
 				{
 					const int32 NumParents = Hierarchy->GetNumberOfParents(Control);
 					const FTransform Global = Hierarchy->GetTransform(Control, ERigTransformType::CurrentGlobal);
@@ -5330,10 +5336,10 @@ bool FControlRigEditMode::MoveControlShapeLocally( AControlRigShapeActor* ShapeA
 		{
 			switch(InTransformContext.Space)
 			{
-			case ETransformSpace::World:
+			case EControlRigInteractionTransformSpace::World:
 				ensure(false);
 				break;
-			case ETransformSpace::Local:
+			case EControlRigInteractionTransformSpace::Local:
 			{
 				// in local mode, the incoming delta is expressed in the global space so it has to be switched back to the offset space
 				const FTransform Global = Hierarchy->GetTransform(Control, ERigTransformType::CurrentGlobal);
@@ -5342,7 +5348,7 @@ bool FControlRigEditMode::MoveControlShapeLocally( AControlRigShapeActor* ShapeA
 				Delta = Offset.InverseTransformVector(GlobalDelta);
 				break;
 			}
-			case ETransformSpace::Parent:
+			case EControlRigInteractionTransformSpace::Parent:
 			{
 				// in parent mode, the incoming delta is expressed in the parent space so it has to be switched back to the offset space
 				FTransform Parent = FTransform::Identity;  
@@ -5362,7 +5368,7 @@ bool FControlRigEditMode::MoveControlShapeLocally( AControlRigShapeActor* ShapeA
 				Delta = Offset.InverseTransformVector(GlobalDelta);
 				break;
 			}
-			case ETransformSpace::Explicit:
+			case EControlRigInteractionTransformSpace::Explicit:
 			{
 				// nothing to do as it has already been computed in the right space
 				break;
@@ -5412,7 +5418,7 @@ bool FControlRigEditMode::MoveControlShapeLocally( AControlRigShapeActor* ShapeA
 }
 
 void FControlRigEditMode::MoveControlShape( AControlRigShapeActor* ShapeActor,
-	const FTransformContext& InTransformContext, const FTransform& ToWorldTransform,
+	const FControlRigInteractionTransformContext& InTransformContext, const FTransform& ToWorldTransform,
 	const bool bUseLocal, const bool bCalcLocal, FTransform* InOutLocal,
 	TArray<TFunction<void()>>& OutTasks)
 {
@@ -5545,13 +5551,13 @@ void FControlRigEditMode::MoveControlShape( AControlRigShapeActor* ShapeActor,
 
 			switch (InTransformContext.Space)
 			{
-			case ETransformSpace::World:
+			case EControlRigInteractionTransformSpace::World:
 				ensure(false);
 				break;
-			case ETransformSpace::Local:
+			case EControlRigInteractionTransformSpace::Local:
 				TransformSpace = GlobalTransform;
 				break;
-			case ETransformSpace::Parent:
+			case EControlRigInteractionTransformSpace::Parent:
 				{
 					if (TOptional<FTransform> ConstraintSpace = GetConstraintParentTransform(ControlRig, ShapeActor->ControlName))
 					{
@@ -5571,7 +5577,7 @@ void FControlRigEditMode::MoveControlShape( AControlRigShapeActor* ShapeActor,
 					}
 				}
 				break;
-			case ETransformSpace::Explicit:
+			case EControlRigInteractionTransformSpace::Explicit:
 				if (ensure(InTransformContext.bTranslation))
 				{
 					if (TOptional<FTransform> ConstraintSpace = GetConstraintParentTransform(ControlRig, ShapeActor->ControlName))
@@ -5635,15 +5641,15 @@ void FControlRigEditMode::MoveControlShape( AControlRigShapeActor* ShapeActor,
 				
 				switch (InTransformContext.Space)
 				{
-					case ETransformSpace::World:
+					case EControlRigInteractionTransformSpace::World:
 						break;
-					case ETransformSpace::Local:
+					case EControlRigInteractionTransformSpace::Local:
 						if (InTransformContext.bTranslation)
 						{
 							InOutLocal->SetLocation(TransformSpace.InverseTransformVector(InTransformContext.Drag));
 						}
 						break;
-					case ETransformSpace::Parent:
+					case EControlRigInteractionTransformSpace::Parent:
 						if (InTransformContext.bTranslation)
 						{
 							InOutLocal->SetLocation(TransformSpace.InverseTransformVector(InTransformContext.Drag));
@@ -5655,7 +5661,7 @@ void FControlRigEditMode::MoveControlShape( AControlRigShapeActor* ShapeActor,
 							InOutLocal->SetRotation(TransformSpace.InverseTransformRotation(SpaceRotation));
 						}
 						break;
-					case ETransformSpace::Explicit:
+					case EControlRigInteractionTransformSpace::Explicit:
 						if (InTransformContext.bTranslation)
 						{
 							InOutLocal->SetLocation(TransformSpace.InverseTransformVector(InTransformContext.Drag));
@@ -5678,7 +5684,7 @@ void FControlRigEditMode::MoveControlShape( AControlRigShapeActor* ShapeActor,
 	}
 }
 
-void FControlRigEditMode::ChangeControlShapeTransform(AControlRigShapeActor* InShapeActor, const FTransformContext& InContext, const FTransform& ToWorldTransform)
+void FControlRigEditMode::ChangeControlShapeTransform(AControlRigShapeActor* InShapeActor, const FControlRigInteractionTransformContext& InContext, const FTransform& ToWorldTransform)
 {
 	if (!InContext.CanTransform())
 	{
diff --git a/Engine/Plugins/Animation/ControlRig/Source/ControlRigEditor/Private/EditMode/ControlRigEditModeUtil.cpp b/Engine/Plugins/Animation/ControlRig/Source/ControlRigEditor/Private/EditMode/ControlRigEditModeUtil.cpp
--- a/Engine/Plugins/Animation/ControlRig/Source/ControlRigEditor/Private/EditMode/ControlRigEditModeUtil.cpp
+++ b/Engine/Plugins/Animation/ControlRig/Source/ControlRigEditor/Private/EditMode/ControlRigEditModeUtil.cpp
@@ -16,7 +16,7 @@
 namespace UE::ControlRigEditMode
 {
 
-FExplicitRotationInteraction::FExplicitRotationInteraction(const FTransformContext& InContext,
+FExplicitRotationInteraction::FExplicitRotationInteraction(const FControlRigInteractionTransformContext& InContext,
 	UControlRig* InControlRig,
 	URigHierarchy* InHierarchy,
 	FRigControlElement* InControlElement,
@@ -30,7 +30,7 @@ FExplicitRotationInteraction::FExplicitRotationInteraction(const FTransformConte
 	
 bool FExplicitRotationInteraction::IsValid() const
 {
-	const bool bIsExplicitRotation = TransformContext.bRotation && TransformContext.Space == ETransformSpace::Explicit && !TransformContext.Rot.IsZero();
+	const bool bIsExplicitRotation = TransformContext.bRotation && TransformContext.Space == EControlRigInteractionTransformSpace::Explicit && !TransformContext.Rot.IsZero();
 	if (!bIsExplicitRotation)
 	{
 		return false;
diff --git a/Engine/Plugins/Animation/ControlRig/Source/ControlRigEditor/Private/Sequencer/EditModeAnimationUtil.cpp b/Engine/Plugins/Animation/ControlRig/Source/ControlRigEditor/Private/Sequencer/EditModeAnimationUtil.cpp
--- a/Engine/Plugins/Animation/ControlRig/Source/ControlRigEditor/Private/Sequencer/EditModeAnimationUtil.cpp
+++ b/Engine/Plugins/Animation/ControlRig/Source/ControlRigEditor/Private/Sequencer/EditModeAnimationUtil.cpp
@@ -132,7 +132,7 @@ void FControlRigKeyframer::Store(const uint32 InControlHash, FControlKeyframeDat
 	}
 }
 	
-void FControlRigKeyframer::Apply(const FControlRigInteractionScope& InInteractionScope, const FTransformContext& InTransformContext)
+void FControlRigKeyframer::Apply(const FControlRigInteractionScope& InInteractionScope, const FControlRigInteractionTransformContext& InTransformContext)
 {
 	if (!IsEnabled())
 	{
diff --git a/Engine/Plugins/Animation/ControlRig/Source/ControlRigEditor/Public/EditMode/ControlRigEditMode.h b/Engine/Plugins/Animation/ControlRig/Source/ControlRigEditor/Public/EditMode/ControlRigEditMode.h
--- a/Engine/Plugins/Animation/ControlRig/Source/ControlRigEditor/Public/EditMode/ControlRigEditMode.h
+++ b/Engine/Plugins/Animation/ControlRig/Source/ControlRigEditor/Public/EditMode/ControlRigEditMode.h
@@ -238,7 +238,7 @@ public:
 
 protected:
 
-	UE::ControlRigEditMode::ETransformSpace GetTransformSpace() const;
+	EControlRigInteractionTransformSpace GetTransformSpace() const;
 	
 	// shape related functions wrt enable/selection
 	/** Get the node name from the property path */
@@ -253,9 +253,9 @@ protected:
 	/** Updates cached pivot transforms */
 	void UpdatePivotTransforms();
 	bool ComputePivotFromEditedShape(UControlRig* InControlRig, FTransform& OutTransform) const;
-	bool ComputePivotFromShapeActors(UControlRig* InControlRig, const bool bEachLocalSpace, const UE::ControlRigEditMode::ETransformSpace InSpace, FTransform& OutTransform) const;
+	bool ComputePivotFromShapeActors(UControlRig* InControlRig, const bool bEachLocalSpace, const EControlRigInteractionTransformSpace InSpace, FTransform& OutTransform) const;
 	bool ComputePivotFromElements(UControlRig* InControlRig, FTransform& OutTransform) const;
-	FTransform GetPivotOrientation(const FRigElementKey& InControlKey, const UControlRig* InControlRig, URigHierarchy* InHierarchy, const UE::ControlRigEditMode::ETransformSpace InSpace, const FTransform& InComponentTransform) const;
+	FTransform GetPivotOrientation(const FRigElementKey& InControlKey, const UControlRig* InControlRig, URigHierarchy* InHierarchy, const EControlRigInteractionTransformSpace InSpace, const FTransform& InComponentTransform) const;
 	
 	/** Get the current coordinate system space */
 	ECoordSystem GetCoordSystemSpace() const;
@@ -557,10 +557,10 @@ private:
 	bool CheckMovieSceneSig();
 	void SetControlShapeTransform( const AControlRigShapeActor* InShapeActor, const FTransform& InGlobalTransform,
 		const FTransform& InToWorldTransform, const FRigControlModifiedContext& InContext, const bool bPrintPython,
-		const UE::ControlRigEditMode::FTransformContext& InTransformContext, const bool bFixEulerFlips = true) const;
+		const FControlRigInteractionTransformContext& InTransformContext, const bool bFixEulerFlips = true) const;
 	static FTransform GetControlShapeTransform(const AControlRigShapeActor* ShapeActor);
 
-	static void ChangeControlShapeTransform(AControlRigShapeActor* ShapeActor, const UE::ControlRigEditMode::FTransformContext& InContext, const FTransform& ToWorldTransform);
+	static void ChangeControlShapeTransform(AControlRigShapeActor* ShapeActor, const FControlRigInteractionTransformContext& InContext, const FTransform& ToWorldTransform);
 
 	bool ModeSupportedByShapeActor(const AControlRigShapeActor* ShapeActor, UE::Widget::EWidgetMode InMode) const;
 
@@ -571,11 +571,11 @@ public:
 protected:
 
 	bool MoveControlShapeLocally( AControlRigShapeActor* ShapeActor,
-		const UE::ControlRigEditMode::FTransformContext& InTransformContext, const FTransform& ToWorldTransform,
+		const FControlRigInteractionTransformContext& InTransformContext, const FTransform& ToWorldTransform,
 		const FTransform& InLocal);
 	
 	void MoveControlShape(AControlRigShapeActor* ShapeActor,
-		const UE::ControlRigEditMode::FTransformContext& InContext, const FTransform& ToWorldTransform,
+		const FControlRigInteractionTransformContext& InContext, const FTransform& ToWorldTransform,
 		const bool bUseLocal, const bool bCalcLocal, FTransform* InOutLocal, TArray<TFunction<void()>>& OutTasks);
 
 	/** Get bindings to a runtime object */
diff --git a/Engine/Plugins/Animation/ControlRig/Source/ControlRigEditor/Public/EditMode/ControlRigEditModeUtil.h b/Engine/Plugins/Animation/ControlRig/Source/ControlRigEditor/Public/EditMode/ControlRigEditModeUtil.h
--- a/Engine/Plugins/Animation/ControlRig/Source/ControlRigEditor/Public/EditMode/ControlRigEditModeUtil.h
+++ b/Engine/Plugins/Animation/ControlRig/Source/ControlRigEditor/Public/EditMode/ControlRigEditModeUtil.h
@@ -20,50 +20,11 @@ class UTickableConstraint;
 struct FConvexVolume;
 struct FRigControlModifiedContext;
 struct FRigControlElement;
+struct FControlRigInteractionTransformContext;
 
 namespace UE::ControlRigEditMode
 {
 
-enum class ETransformSpace
-{
-	World,
-	Local,
-	Parent,
-	Explicit
-};
-
-/**
- * FTransformContext provides a way of passing the various transform parameters to functions that need to know what the transform context is.
- * Extend it if necessary, particularly to avoid overloading some functions signatures.
- */
-
-struct FTransformContext
-{
-	FTransformContext() = default;
-	
-	FTransformContext(const UE::Widget::EWidgetMode& InWidgetMode)
-		: bTranslation(InWidgetMode == UE::Widget::WM_Translate || InWidgetMode == UE::Widget::WM_TranslateRotateZ)
-		, bRotation(InWidgetMode == UE::Widget::WM_Rotate || InWidgetMode == UE::Widget::WM_TranslateRotateZ)
-		, bScale(InWidgetMode == UE::Widget::WM_Scale)
-	{}
-	
-	bool CanTransform() const
-	{
-		return bTranslation || bRotation || bScale;
-	}
-
-	bool bTranslation = false;
-	FVector Drag = FVector::ZeroVector;
-
-	bool bRotation = false;
-	FRotator Rot = FRotator::ZeroRotator;
-
-	bool bScale = false;
-	FVector Scale = FVector::OneVector;
-
-	ETransformSpace Space = ETransformSpace::World;
-};
-
 /**
  * FInteractionDependencyCache provides a minimal "dependency graph" between the selected controls during interaction.
  * it stores information on who is a child and who is a parent in whatever is being manipulated, as well as
@@ -121,7 +82,7 @@ struct FInteractionDependencyCache
 	
 struct FExplicitRotationInteraction
 {
-	FExplicitRotationInteraction(const FTransformContext& InContext,
+	FExplicitRotationInteraction(const FControlRigInteractionTransformContext& InContext,
 		UControlRig* InControlRig,
 		URigHierarchy* InHierarchy,
 		FRigControlElement* InControlElement,
@@ -133,7 +94,7 @@ struct FExplicitRotationInteraction
 		const TArray< TWeakObjectPtr<UTickableConstraint> >& InConstraints = TArray< TWeakObjectPtr<UTickableConstraint> >()) const;
 
 private:
-	const FTransformContext& TransformContext;
+	const FControlRigInteractionTransformContext& TransformContext;
 	UControlRig* ControlRig = nullptr;
 	URigHierarchy* Hierarchy = nullptr;
 	FRigControlElement* ControlElement = nullptr;
diff --git a/Engine/Plugins/Animation/ControlRig/Source/ControlRigEditor/Public/Sequencer/EditModeAnimationUtil.h b/Engine/Plugins/Animation/ControlRig/Source/ControlRigEditor/Public/Sequencer/EditModeAnimationUtil.h
--- a/Engine/Plugins/Animation/ControlRig/Source/ControlRigEditor/Public/Sequencer/EditModeAnimationUtil.h
+++ b/Engine/Plugins/Animation/ControlRig/Source/ControlRigEditor/Public/Sequencer/EditModeAnimationUtil.h
@@ -14,11 +14,7 @@ class UAnimationAuthoringSettings;
 class UControlRig;
 
 struct FRigControlElement;
-
-namespace UE::ControlRigEditMode
-{
-	struct FTransformContext;
-}
+struct FControlRigInteractionTransformContext;
 
 namespace UE::AnimationEditMode
 {
@@ -103,7 +99,7 @@ struct FControlRigKeyframer
 	void Store(const uint32 InControlHash, FControlKeyframeData&& InData);
 
 	// Does the actual work off adding keyframes to the controls currently interacting.
-	void Apply(const FControlRigInteractionScope& InInteractionScope, const ControlRigEditMode::FTransformContext& InTransformContext);
+	void Apply(const FControlRigInteractionScope& InInteractionScope, const FControlRigInteractionTransformContext& InTransformContext);
 
 	// Updates whatever needs to once the keyframes have been added (updating constraints is one of them).
 	void Finalize(UWorld* InWorld);
