ControlRig constraint handle: make sure that ControlName is actually the control's unique name version instead of the namespaced one
- follow up to 40302964

#rnx

diff --git a/Engine/Plugins/Animation/ControlRig/Source/ControlRig/Private/ControlRig.cpp b/Engine/Plugins/Animation/ControlRig/Source/ControlRig/Private/ControlRig.cpp
--- a/Engine/Plugins/Animation/ControlRig/Source/ControlRig/Private/ControlRig.cpp
+++ b/Engine/Plugins/Animation/ControlRig/Source/ControlRig/Private/ControlRig.cpp
@@ -3645,9 +3645,9 @@ void UControlRig::SetDynamicHierarchy(TObjectPtr<URigHierarchy> InHierarchy)
 UTransformableControlHandle* UControlRig::CreateTransformableControlHandle(
 	const FName& InControlName) const
 {
-	auto IsConstrainable = [this](const FName& InControlName)
+	const FRigControlElement* ControlElement = FindControl(InControlName);
+	auto IsConstrainable = [ControlElement]()
 	{
-		const FRigControlElement* ControlElement = FindControl(InControlName);
 		if (!ControlElement)
 		{
 			return false;
@@ -3665,7 +3665,7 @@ UTransformableControlHandle* UControlRig::CreateTransformableControlHandle(
 		return true;
 	};
 
-	if (!IsConstrainable(InControlName))
+	if (!IsConstrainable())
 	{
 		return nullptr;
 	}
@@ -3674,6 +3674,27 @@ UTransformableControlHandle* UControlRig::CreateTransformableControlHandle(
 	check(CtrlHandle);
 	CtrlHandle->ControlRig = const_cast<UControlRig*>(this);
 	CtrlHandle->ControlName = InControlName;
+	
+	if (ControlElement->GetFName() != InControlName)
+	{
+		CtrlHandle->ControlName = ControlElement->GetFName();
+		
+#if WITH_EDITOR
+		const FRigElementKey ControlKey(InControlName, ERigElementType::Control);
+		if (DynamicHierarchy->ReceivedNameSpaceBasedKeys.Contains(ControlKey))
+		{
+			UE_LOG(
+				LogControlRig,
+				Warning,
+				TEXT("%s: control '%s' has been accessed using a namespace based key ('%s') when creating the control handle. Please consider updating your code."),
+				*GetName(),
+				*ControlElement->GetName(),
+				*InControlName.ToString()
+			);	
+		}
+#endif
+	}
+	
 	CtrlHandle->RegisterDelegates();
 	return CtrlHandle;
 }
