Anim In Engine: Action Set Key/Passthrough Key actions and  Add support for per frame anim layer merging and fix some issues when merging overide sections when controls aren't present in them, need to make sure they stay masked out.

diff --git a/Engine/Plugins/Animation/ControlRig/Source/ControlRig/Public/Sequencer/MovieSceneControlRigParameterTrack.h b/Engine/Plugins/Animation/ControlRig/Source/ControlRig/Public/Sequencer/MovieSceneControlRigParameterTrack.h
--- a/Engine/Plugins/Animation/ControlRig/Source/ControlRig/Public/Sequencer/MovieSceneControlRigParameterTrack.h
+++ b/Engine/Plugins/Animation/ControlRig/Source/ControlRig/Public/Sequencer/MovieSceneControlRigParameterTrack.h
@@ -239,7 +239,38 @@ private:
 	/** copy of the controlled control rig that we use in the game world so editor control rig doesn't conflict*/
 	UPROPERTY(transient)
 	TMap<TWeakObjectPtr<UWorld>,TObjectPtr<UControlRig>> GameWorldControlRigs;
-};
 
 
+	friend struct FControlRigParameterTrackSectionToKeyRestore;
+};
 
+//override the section to key, section to key per control, temporarilly
+struct FControlRigParameterTrackSectionToKeyRestore
+{
+	FControlRigParameterTrackSectionToKeyRestore(UMovieSceneControlRigParameterTrack* InTrack, UMovieSceneSection* NewSectionToKey, TMap<FName, TWeakObjectPtr<UMovieSceneSection>>& NewSectionToKeyPerControl)
+		: Track(InTrack)
+	{
+		if (Track)
+		{
+			OldSectionToKey = Track->SectionToKey;
+			OldSectionToKeyPerControl = Track->SectionToKeyPerControl;
+			Track->SectionToKey = NewSectionToKey;
+			Track->SectionToKeyPerControl = NewSectionToKeyPerControl;
+		}
+	}
+
+	~FControlRigParameterTrackSectionToKeyRestore()
+	{
+		if (Track)
+		{
+			Track->SectionToKey = OldSectionToKey;
+			Track->SectionToKeyPerControl = OldSectionToKeyPerControl;
+		}
+	}
+
+	FControlRigParameterTrackSectionToKeyRestore() = delete;
+	
+	UMovieSceneControlRigParameterTrack* Track;
+	UMovieSceneSection* OldSectionToKey;
+	TMap<FName, TWeakObjectPtr<UMovieSceneSection>> OldSectionToKeyPerControl;
+};
diff --git a/Engine/Plugins/Animation/ControlRig/Source/ControlRigEditor/Private/ControlRigSequencerEditorLibrary.cpp b/Engine/Plugins/Animation/ControlRig/Source/ControlRigEditor/Private/ControlRigSequencerEditorLibrary.cpp
--- a/Engine/Plugins/Animation/ControlRig/Source/ControlRigEditor/Private/ControlRigSequencerEditorLibrary.cpp
+++ b/Engine/Plugins/Animation/ControlRig/Source/ControlRigEditor/Private/ControlRigSequencerEditorLibrary.cpp
@@ -3440,6 +3440,7 @@ int32 UControlRigSequencerEditorLibrary::AddAnimLayerFromSelection()
 
 }
 
+//remove after 5.7
 bool UControlRigSequencerEditorLibrary::MergeAnimLayers(const TArray<int32>& Indices)
 {
 	TWeakPtr<ISequencer> WeakSequencer = GetSequencerFromAsset();
@@ -3451,7 +3452,27 @@ bool UControlRigSequencerEditorLibrary::MergeAnimLayers(const TArray<int32>& Ind
 	TSharedPtr<ISequencer> SequencerPtr = WeakSequencer.Pin();
 	if (UAnimLayers* AnimLayers = UAnimLayers::GetAnimLayers(SequencerPtr.Get()))
 	{
-		return AnimLayers->MergeAnimLayers(SequencerPtr.Get(), Indices, nullptr);
+		return AnimLayers->MergeAnimLayers(SequencerPtr, Indices, nullptr);
+	}
+	else
+	{
+		UE_LOG(LogControlRig, Error, TEXT("MergeAnimLayers: No Anim Layers on Level Sequence"));
+	}
+	return false;
+}
+
+bool UControlRigSequencerEditorLibrary::MergeAnimLayersWithSettings(const TArray<int32>& Indices, const FMergeAnimLayerSettings& MergeAnimLayerSettings)
+{
+	TWeakPtr<ISequencer> WeakSequencer = GetSequencerFromAsset();
+	if (WeakSequencer.IsValid() == false)
+	{
+		UE_LOG(LogControlRig, Error, TEXT("MergeAnimLayers: Need open Sequencer"));
+		return false;
+	}
+	TSharedPtr<ISequencer> SequencerPtr = WeakSequencer.Pin();
+	if (UAnimLayers* AnimLayers = UAnimLayers::GetAnimLayers(SequencerPtr.Get()))
+	{
+		return AnimLayers->MergeAnimLayers(SequencerPtr, Indices, &MergeAnimLayerSettings);
 	}
 	else
 	{
diff --git a/Engine/Plugins/Animation/ControlRig/Source/ControlRigEditor/Private/Sequencer/AnimLayers/AnimLayers.cpp b/Engine/Plugins/Animation/ControlRig/Source/ControlRigEditor/Private/Sequencer/AnimLayers/AnimLayers.cpp
--- a/Engine/Plugins/Animation/ControlRig/Source/ControlRigEditor/Private/Sequencer/AnimLayers/AnimLayers.cpp
+++ b/Engine/Plugins/Animation/ControlRig/Source/ControlRigEditor/Private/Sequencer/AnimLayers/AnimLayers.cpp
@@ -12,6 +12,7 @@
 #include "MVVM/ViewModels/ObjectBindingModel.h"
 #include "LevelSequence.h"
 #include "ControlRigSequencerEditorLibrary.h"
+#include "Sequencer/ControlRigParameterTrackEditor.h"
 #include "Sequencer/MovieSceneControlRigParameterTrack.h"
 #include "Sequencer/MovieSceneControlRigParameterSection.h"
 #include "Tracks/MovieScene3DTransformTrack.h"
@@ -35,14 +36,13 @@
 #include "Channels/MovieSceneDoubleChannel.h"
 #include "Channels/MovieSceneIntegerChannel.h"
 #include "Channels/MovieSceneByteChannel.h"
+#include "EditMode/ControlRigEditMode.h"
 #include "ISequencerPropertyKeyedStatus.h"
 #include "BakingAnimationKeySettings.h"
 #include "Algo/Accumulate.h"
 #include "MovieSceneToolHelpers.h"
-
 #define LOCTEXT_NAMESPACE "AnimLayers"
 
-
 static TArray<FGuid> GetSelectedOutlinerGuids(ISequencer* SequencerPtr)
 {
 	using namespace UE::Sequencer;
@@ -828,6 +828,13 @@ bool UAnimLayer::IsAccepableNonControlRigObject(UObject* InObject)
 }
 
 ECheckBoxState UAnimLayer::GetSelected() const
+{
+	TSet<UObject*> SelectedObjects; 
+	TMap<UControlRig*, TArray<FName>> SelectedControls;
+	return GetSelected(SelectedObjects, SelectedControls);
+}
+
+ECheckBoxState UAnimLayer::GetSelected(TSet<UObject*>& OutSelectedObjects, TMap<UControlRig*, TArray<FName>>& OutSelectedControls) const
 {
 	FAnimLayerObjects LayerObjects;
 	GetAnimLayerObjects(LayerObjects);
@@ -841,6 +848,15 @@ ECheckBoxState UAnimLayer::GetSelected() const
 			{
 				if (SelectedControls.Contains(ControlName))
 				{
+					if (TArray<FName>* ControlsArray = OutSelectedControls.Find(ControlRigObject.ControlRig.Get()))
+					{
+						ControlsArray->Add(ControlName);
+					}
+					else
+					{
+						TArray<FName>& ControlsRefArray = OutSelectedControls.Add(ControlRigObject.ControlRig.Get());
+						ControlsRefArray.Add(ControlName);
+					}
 					if (SelectionState.IsSet() == false)
 					{
 						SelectionState = ECheckBoxState::Checked;
@@ -892,6 +908,7 @@ ECheckBoxState UAnimLayer::GetSelected() const
 		{
 			if (SelectedActors.Contains(SceneObject.SceneObjectOrComponent))
 			{
+				OutSelectedObjects.Add(SceneObject.SceneObjectOrComponent.Get());
 				if (SelectionState.IsSet() == false)
 				{
 					SelectionState = ECheckBoxState::Checked;
@@ -918,6 +935,7 @@ ECheckBoxState UAnimLayer::GetSelected() const
 		{
 			if (SelectedComponents.Contains(SceneObject.SceneObjectOrComponent))
 			{
+				OutSelectedObjects.Add(SceneObject.SceneObjectOrComponent.Get());
 				if (SelectionState.IsSet() == false)
 				{
 					SelectionState = ECheckBoxState::Checked;
@@ -1745,8 +1763,9 @@ FAnimLayerSectionItem* FAnimLayerItem::FindMatchingSectionItem(UMovieSceneSectio
 	return nullptr;
 }
 
-template<typename ChannelType>
-void AddKeyToChannel(ChannelType* Channel, EMovieSceneKeyInterpolation DefaultInterpolation, const FFrameNumber& FrameNumber, double Value)
+
+template<typename ChannelType, typename ValueType>
+void AddKeyToChannel(ChannelType* Channel, EMovieSceneKeyInterpolation DefaultInterpolation, const FFrameNumber& FrameNumber, ValueType Value)
 {
 	switch (DefaultInterpolation)
 	{
@@ -1765,17 +1784,172 @@ void AddKeyToChannel(ChannelType* Channel, EMovieSceneKeyInterpolation DefaultIn
 		break;
 	}
 }
+static bool GetStartEndIndicesForControl(UMovieSceneControlRigParameterSection* BaseSection, const FRigControlElement* ControlElement, int& OutStartIndex, int& OutEndIndex)
+{
+	if (FChannelMapInfo* pChannelIndex = BaseSection->ControlChannelMap.Find(ControlElement->GetFName()))
+	{
+		const int32 ChannelIndex = pChannelIndex->ChannelIndex;
+		switch (ControlElement->Settings.ControlType)
+		{
+
+		case ERigControlType::Float:
+		case ERigControlType::ScaleFloat:
+		{
+			OutStartIndex = ChannelIndex;
+			OutEndIndex = ChannelIndex;
+			break;
+		}
+		case ERigControlType::Vector2D:
+		{
+			OutStartIndex = ChannelIndex;
+			OutEndIndex = ChannelIndex + 1;
+			break;
+		}
+		case ERigControlType::Position:
+		case ERigControlType::Scale:
+		case ERigControlType::Rotator:
+		{
+			OutStartIndex = ChannelIndex;
+			OutEndIndex = ChannelIndex + 2;
+			break;
+		}
+
+		case ERigControlType::Transform:
+		case ERigControlType::TransformNoScale:
+		case ERigControlType::EulerTransform:
+		{
+			if (ControlElement->Settings.ControlType == ERigControlType::TransformNoScale)
+			{
+				OutStartIndex = ChannelIndex;
+				OutEndIndex = ChannelIndex + 5;
+
+			}
+			else
+			{
+				OutStartIndex = ChannelIndex;
+				OutEndIndex = ChannelIndex + 8;
+			}
+			break;
+		}
+		default:
+			break;
+		}
+		return true;
+	}
+	return false;
+}
+template<typename ChannelType, typename ValueType>
+void AssigneOrSetValue(ChannelType* Channel, ValueType Value, const FFrameNumber& FrameNumber, EMovieSceneKeyInterpolation DefaultInterpolation)
+{
+	TArray<FKeyHandle> KeysAtCurrentTime;
+	Channel->GetKeys(TRange<FFrameNumber>(FrameNumber), nullptr, &KeysAtCurrentTime);
+	if (KeysAtCurrentTime.Num() > 0)
+	{
+		AssignValue(Channel, KeysAtCurrentTime[0], Value);
+	}
+	else
+	{
+		EMovieSceneKeyInterpolation Interpolation = GetInterpolationMode(Channel, FrameNumber, DefaultInterpolation);
+		AddKeyToChannel(Channel, FrameNumber, Value, Interpolation);
+	}
+}
+
+template<typename ChannelType,typename ValueType>
+void SetPassthroughKeys(TArrayView<ChannelType*> Channels,int32 StartIndex, int32 EndIndex, EAnimLayerType BlendType, const TArray<UMovieSceneSection*>& Sections, const TArray<UMovieSceneSection*>& AbsoluteSections,
+	EMovieSceneKeyInterpolation DefaultInterpolation, const FFrameNumber& FrameNumber)
+{
+	switch (BlendType)
+	{
+	case EAnimLayerType::Override:
+	{
+		TArray<ValueType> ChannelValues = MovieSceneToolHelpers::GetChannelValues<ChannelType,
+			ValueType>(StartIndex, EndIndex, Sections,AbsoluteSections,FrameNumber);
+		int32  ChannelValueIndex = 0;
+		for (int32 Index = StartIndex; Index <= EndIndex; ++Index)
+		{
+			ChannelType* Channel = Channels[Index];
+			ValueType Value = ChannelValues[ChannelValueIndex];
+			++ChannelValueIndex;
+			AssigneOrSetValue(Channel, Value, FrameNumber, DefaultInterpolation);
+		}
+	}
+	break;
+	case EAnimLayerType::Additive:
+	{
+		const ValueType Value = 0.0;
+		for (int32 Index = StartIndex; Index <= EndIndex; ++Index)
+		{
+			ChannelType* Channel = Channels[Index];
+			AssigneOrSetValue(Channel, Value, FrameNumber, DefaultInterpolation);
+		}
+	}
+	break;
+	}
+}
+
+template<typename ChannelType, typename ValueType>
+void SetCurrentKeys(TArrayView<ChannelType*> Channels, int32 StartIndex, int32 EndIndex ,EMovieSceneKeyInterpolation DefaultInterpolation, const FFrameNumber& FrameNumber)
+{
+	ValueType Value = 0.0;
+	FFrameTime FrameTime(FrameNumber);
+	for (int32 Index = StartIndex; Index <= EndIndex; ++Index)
+	{
+		ChannelType* Channel = Channels[Index];
+		if (Channel->Evaluate(FrameTime, Value))
+		{
+			AssigneOrSetValue(Channel, Value, FrameNumber, DefaultInterpolation);
+		}
+	}
+}
 
 bool UAnimLayers::SetPassthroughKey(ISequencer* InSequencer, int32 InIndex)
 {
-	if (InSequencer == nullptr || InIndex <= 0 || InIndex >= AnimLayers.Num())
+	return SetKeyValueOrPassthrough(InSequencer, InIndex, false /*IsValue*/);
+}
+
+bool UAnimLayers::SetKey(ISequencer* InSequencer, int32 InIndex)
+{
+	return SetKeyValueOrPassthrough(InSequencer, InIndex, true /*IsValue*/);
+}
+
+
+
+bool UAnimLayers::SetKeyValueOrPassthrough(ISequencer* InSequencer, int32 InIndex, bool bJustValue)
+{
+	auto GetSelectedRigElements = [](UControlRig* ControlRig, TArray<FRigElementKey>& OutSelectedKeys) 
+	{
+		if (ControlRig)
+		{
+			TArray<FRigElementKey> SelectedRigElements = ControlRig->GetHierarchy()->GetSelectedKeys();
+			if (ControlRig->IsAdditive())
+			{
+				// For additive rigs, ignore boolean controls
+				SelectedRigElements = SelectedRigElements.FilterByPredicate([ControlRig](const FRigElementKey& Key)
+					{
+						if (FRigControlElement* Element = ControlRig->FindControl(Key.Name))
+						{
+							return Element->CanTreatAsAdditive();
+						}
+						return true;
+					});
+			}
+		}
+
+	};
+	if (InSequencer == nullptr || InIndex <= -1 || InIndex >= AnimLayers.Num())
 	{
 		return false;
 	}
 	FFrameNumber FrameNumber = InSequencer->GetLocalTime().Time.GetFrame();
 	EMovieSceneKeyInterpolation DefaultInterpolation = InSequencer->GetKeyInterpolation();
+
+	FText TransactionText = bJustValue ? FText(LOCTEXT("SetKeyValue_Transaction", "Set Key")) : FText(LOCTEXT("SetPassthroughKey_Transaction", "Set Passthrough Key"));
 	if (UAnimLayer* AnimLayer = AnimLayers[InIndex].Get())
 	{
+		TSet<UObject*> SelectedObjects;
+		TMap<UControlRig*, TArray<FName>> SelectedControls;
+		bool bHasSelected = AnimLayer->GetSelected(SelectedObjects,SelectedControls) == ECheckBoxState::Unchecked ? false : true;
+
 		for (TPair<TWeakObjectPtr<UObject>, FAnimLayerItem>& Pair : AnimLayer->AnimLayerItems)
 		{
 			if (Pair.Key != nullptr)
@@ -1790,7 +1964,7 @@ bool UAnimLayers::SetPassthroughKey(ISequencer* InSequencer, int32 InIndex)
 							int32 SectionIndex = TrackSections.Find((SectionItem.Section.Get()));
 							if (SectionIndex != INDEX_NONE)
 							{
-								const FScopedTransaction Transaction(LOCTEXT("SetPassthroughKey_Transaction", "Set Passthrough Key"), !GIsTransacting);
+								const FScopedTransaction Transaction(TransactionText, !GIsTransacting);
 								SectionItem.Section->Modify();
 								TrackSections.SetNum(SectionIndex); //this will gives us up to the section 
 								TArray<UMovieSceneSection*> Sections;
@@ -1798,59 +1972,142 @@ bool UAnimLayers::SetPassthroughKey(ISequencer* InSequencer, int32 InIndex)
 								MovieSceneToolHelpers::SplitSectionsByBlendType(EMovieSceneBlendType::Absolute, TrackSections, Sections, AbsoluteSections);
 								TArrayView<FMovieSceneFloatChannel*> BaseFloatChannels = SectionItem.Section->GetChannelProxy().GetChannels<FMovieSceneFloatChannel>();
 								TArrayView<FMovieSceneDoubleChannel*> BaseDoubleChannels = SectionItem.Section->GetChannelProxy().GetChannels<FMovieSceneDoubleChannel>();
-								if (BaseDoubleChannels.Num() > 0)
+								if (bHasSelected == false)
 								{
-									if (SectionItem.Section->GetBlendType().IsValid() &&
-										SectionItem.Section->GetBlendType() == EMovieSceneBlendType::Additive)
+									if (UMovieSceneControlRigParameterSection* CRSection = Cast<UMovieSceneControlRigParameterSection>(SectionItem.Section.Get()))
 									{
-										const double Value = 0.0;
-										for (FMovieSceneDoubleChannel* DoubleChannel : BaseDoubleChannels)
+										if (BaseFloatChannels.Num() > 0 && CRSection->GetControlRig())
 										{
-											AddKeyToChannel(DoubleChannel, FrameNumber, Value, DefaultInterpolation);
+											//passthrough and base is seperate case we call edit mode function since it does lots of special case stuff
+											//need to make sure the base section is section to key and then set it back if not
+											if (bJustValue == false && AnimLayer->GetType() == EAnimLayerType::Base)
+											{
+												if (UMovieSceneControlRigParameterTrack* CRTrack = Cast<UMovieSceneControlRigParameterTrack>(OwnerTrack))
+												{
+													TMap<FName, TWeakObjectPtr<UMovieSceneSection>> Empty;
+													FControlRigParameterTrackSectionToKeyRestore Restore(CRTrack, CRSection, Empty);
+
+													FRigControlModifiedContext Context;
+													Context.SetKey = EControlRigSetKey::Always;
+
+													FControlRigEditMode::InvertInputPose(CRSection->GetControlRig(), Context, false/*selection only*/, false/* include channels*/);
+												}
+											}
+											else //not base and passthrough do each control 
+											{
+												for (const TPair<FName, FAnimLayerPropertyAndChannels>& Set : SectionItem.AnimLayerSet.Names)
+												{
+													FName ControlName = Set.Key;
+													if (FRigControlElement* Control = CRSection->GetControlRig()->FindControl(ControlName))
+													{
+														int32 StartIndex = 0, EndIndex = 0;
+														if (GetStartEndIndicesForControl(CRSection, Control, StartIndex, EndIndex))
+														{
+															if (bJustValue)
+															{
+																SetCurrentKeys<FMovieSceneFloatChannel, float>(BaseFloatChannels, StartIndex, EndIndex, DefaultInterpolation, FrameNumber);
+															}
+															else
+															{
+																SetPassthroughKeys<FMovieSceneFloatChannel, float>(BaseFloatChannels, StartIndex, EndIndex, AnimLayer->GetType(),
+																	Sections, AbsoluteSections, DefaultInterpolation, FrameNumber);
+															}
+														}
+													}
+												}
+											}	
 										}
 									}
-									else if (SectionItem.Section->GetBlendType().IsValid() &&
-										SectionItem.Section->GetBlendType() == EMovieSceneBlendType::Override)
+									else if (BaseDoubleChannels.Num() > 0)
 									{
-										int32 NumChannels = BaseDoubleChannels.Num();
+										const int32 NumChannels = BaseDoubleChannels.Num();
 										const int32 StartIndex = 0;
 										const int32 EndIndex = NumChannels - 1;
-										TArray<double> ChannelValues = MovieSceneToolHelpers::GetChannelValues<FMovieSceneDoubleChannel,
-											double>(StartIndex, EndIndex, Sections,AbsoluteSections,FrameNumber);
-										for (int32 Index = 0; Index < NumChannels; ++Index)
+
+										if (bJustValue)
 										{
-											FMovieSceneDoubleChannel* DoubleChannel = BaseDoubleChannels[Index];
-											const double Value = ChannelValues[Index];
-											AddKeyToChannel(DoubleChannel, FrameNumber, Value, DefaultInterpolation);
+											SetCurrentKeys<FMovieSceneDoubleChannel, double>(BaseDoubleChannels, StartIndex, EndIndex, DefaultInterpolation, FrameNumber);
+										}
+										else
+										{
+											SetPassthroughKeys<FMovieSceneDoubleChannel, double>(BaseDoubleChannels, StartIndex, EndIndex, AnimLayer->GetType(), Sections, AbsoluteSections,
+												DefaultInterpolation, FrameNumber);
 										}
 									}
 								}
-								else if (BaseFloatChannels.Num() > 0)
+								else //bHasSelected == true
 								{
-									int32 NumChannels = SectionItem.Section->IsA<UMovieSceneControlRigParameterSection>() ?
-										BaseFloatChannels.Num() - 1 : BaseFloatChannels.Num();// skip weight if CR section
-									const int32 StartIndex = 0;
-									const int32 EndIndex = NumChannels - 1;
-									if (SectionItem.Section->GetBlendType().IsValid() &&
-										SectionItem.Section->GetBlendType() == EMovieSceneBlendType::Additive)
+									if (UMovieSceneControlRigParameterSection* CRSection = Cast<UMovieSceneControlRigParameterSection>(SectionItem.Section.Get()))
 									{
-										const float Value = 0.0;
-										for (int32 Index = 0; Index < NumChannels; ++Index)
+										for (const TPair <UControlRig*, TArray<FName>>& ControlPair : SelectedControls)
 										{
-											FMovieSceneFloatChannel* FloatChannel = BaseFloatChannels[Index];
-											AddKeyToChannel(FloatChannel, FrameNumber, Value, DefaultInterpolation);
+											if (CRSection->GetControlRig() == ControlPair.Key)
+											{
+												//passthrough and base is seperate case we call edit mode function since it does lots of special case stuff
+												//need to make sure the base section is section to key and then set it back if not
+												if (bJustValue == false && AnimLayer->GetType() == EAnimLayerType::Base)
+												{
+													if (UMovieSceneControlRigParameterTrack* CRTrack = Cast<UMovieSceneControlRigParameterTrack>(OwnerTrack))
+													{
+														TMap<FName, TWeakObjectPtr<UMovieSceneSection>> Empty;
+														FControlRigParameterTrackSectionToKeyRestore Restore(CRTrack, CRSection, Empty);
+														FRigControlModifiedContext Context;
+														Context.SetKey = EControlRigSetKey::Always;
+
+														FControlRigEditMode::InvertInputPose(CRSection->GetControlRig(), Context, true/*selection only*/, false/* include channels*/);
+													}
+												}
+												else
+												{
+													for (const TPair<FName, FAnimLayerPropertyAndChannels>& Set : SectionItem.AnimLayerSet.Names)
+													{
+														FName ControlName = Set.Key;
+														if (ControlPair.Value.Contains(ControlName))
+														{
+															if (FRigControlElement* Control = ControlPair.Key->FindControl(ControlName))
+															{
+																int32 StartIndex = 0, EndIndex = 0;
+																if (GetStartEndIndicesForControl(CRSection, Control, StartIndex, EndIndex))
+																{
+																	if (bJustValue)
+																	{
+																		SetCurrentKeys<FMovieSceneFloatChannel, float>(BaseFloatChannels, StartIndex, EndIndex, DefaultInterpolation, FrameNumber);
+																	}
+																	else
+																	{
+																		SetPassthroughKeys<FMovieSceneFloatChannel, float>(BaseFloatChannels, StartIndex, EndIndex, AnimLayer->GetType(), Sections, AbsoluteSections,
+																			DefaultInterpolation, FrameNumber);
+																	}
+																}
+															}
+														}
+													}
+												}
+											}
 										}
 									}
-									else if (SectionItem.Section->GetBlendType().IsValid() &&
-										SectionItem.Section->GetBlendType() == EMovieSceneBlendType::Override)
+									else 
 									{
-										TArray<float> ChannelValues = MovieSceneToolHelpers::GetChannelValues<FMovieSceneFloatChannel,
-											float>(StartIndex,EndIndex, Sections,AbsoluteSections, FrameNumber);
-										for (int32 Index = 0; Index < NumChannels; ++Index)
+										for (UObject* SelectedObject : SelectedObjects)
 										{
-											FMovieSceneFloatChannel* FloatChannel = BaseFloatChannels[Index];
-											const float Value = ChannelValues[Index];
-											AddKeyToChannel(FloatChannel, FrameNumber, Value, DefaultInterpolation);
+											if (SectionItem.AnimLayerSet.BoundObject == SelectedObject) //okay selected
+											{
+												if (BaseDoubleChannels.Num() > 0)
+												{
+													const int32 NumChannels = BaseDoubleChannels.Num();
+													const int32 StartIndex = 0;
+													const int32 EndIndex = NumChannels - 1;
+													if (bJustValue)
+													{
+														SetCurrentKeys<FMovieSceneDoubleChannel, double>(BaseDoubleChannels, StartIndex, EndIndex, DefaultInterpolation, FrameNumber);
+													}
+													else
+													{
+														SetPassthroughKeys<FMovieSceneDoubleChannel, double>(BaseDoubleChannels, StartIndex, EndIndex, AnimLayer->GetType(), Sections, AbsoluteSections,
+															DefaultInterpolation, FrameNumber);
+													}
+												}
+											}
 										}
 									}
 								}
@@ -1873,16 +2130,88 @@ bool UAnimLayers::SetPassthroughKey(ISequencer* InSequencer, int32 InIndex)
 	return true;
 }
 
+//get the set of sections active, or with this control name in it
+static void GetValidCRSections(UMovieSceneSection* InSection, const FName& ControlName, TArray<UMovieSceneSection*>& OutSections)
+{
+	if (UMovieSceneTrack* OwnerTrack = InSection->GetTypedOuter<UMovieSceneTrack>())
+	{
+		TArray<UMovieSceneSection*> TrackSections = OwnerTrack->GetAllSections();
+		for (UMovieSceneSection* Section : TrackSections)
+		{
+			if (UMovieSceneControlRigParameterSection* CRSection = Cast<UMovieSceneControlRigParameterSection>(Section))
+			{
+				if (Section && (Section->IsActive() || Section == InSection) &&
+					CRSection->GetControlNameMask(ControlName))
+				{
+					OutSections.Add(Section);
+				}
+			}
+		}
+	}
+}
+
 
-static void MergeControlRigSections(UMovieSceneControlRigParameterSection* BaseSection, UMovieSceneControlRigParameterSection* Section, const TRange<FFrameNumber>& Range)
+static void GetValidSections(UMovieSceneSection* InSection, TArray<UMovieSceneSection*>& OutSections)
 {
-	//disable section if no control in the section and override, need to make sure we skip it's eval or it will incorrectly
-	//override when merging
-	auto ShouldDisableSection = [](UMovieSceneControlRigParameterSection* CRSection, FRigControlElement* ControlElement)
+	if (UMovieSceneTrack* OwnerTrack = InSection->GetTypedOuter<UMovieSceneTrack>())
 	{
-		return (CRSection->GetControlNameMask(ControlElement->GetFName()) == false &&
-				CRSection->GetBlendType().IsValid() && CRSection->GetBlendType() == EMovieSceneBlendType::Override);
-	};
+		TArray<UMovieSceneSection*> TrackSections = OwnerTrack->GetAllSections();
+		for (UMovieSceneSection* Section : TrackSections)
+		{
+			if (Section && (Section->IsActive() || Section == InSection))
+			{
+				OutSections.Add(Section);
+			}
+		}
+	}
+}
+static void RevertWeightChannelToOne(FMovieSceneFloatChannel* FloatChannel, const TRange<FFrameNumber>& FrameRange)
+{
+	if (FloatChannel)
+	{
+		//don't add key if there!
+		TArray<FFrameNumber> KeyTimes;
+		TArray<FKeyHandle> KeyHandles;
+		FloatChannel->GetKeys(FrameRange, &KeyTimes, &KeyHandles);
+		if (KeyTimes.Num() > 0)
+		{
+			FloatChannel->DeleteKeys(KeyHandles);
+			if (FloatChannel->GetNumKeys() > 0) //if it still has keys see if the weight is not one 
+			{
+				float StartValue = 0.0, EndValue = 0.0;
+				FFrameTime StartFrameTime(FrameRange.GetLowerBoundValue());
+				FFrameTime EndFrameTime(FrameRange.GetUpperBoundValue());
+				FloatChannel->Evaluate(StartFrameTime, StartValue);
+				FloatChannel->Evaluate(EndFrameTime, EndValue);
+				//if not zero at boundaries set one keys there
+				if (FMath::IsNearlyEqual(StartValue, 1.0f) == false ||
+					FMath::IsNearlyEqual(EndValue, 1.0f) == false)
+				{
+					FloatChannel->AddCubicKey(StartFrameTime.FrameNumber, 1.0f, ERichCurveTangentMode::RCTM_SmartAuto);
+					FloatChannel->AddCubicKey(EndFrameTime.FrameNumber, 1.0f, ERichCurveTangentMode::RCTM_SmartAuto);
+				}
+			}
+			else
+			{
+				FloatChannel->SetDefault(1.0f);
+			}
+		}
+	}
+}
+
+
+// disable section if no control in the section and override, need to make sure we skip it's eval or it will incorrectly
+//override when merging
+static bool ShouldDisableSection(UMovieSceneControlRigParameterSection* CRSection, FRigControlElement* ControlElement)
+{
+	return (CRSection->GetControlNameMask(ControlElement->GetFName()) == false &&
+		CRSection->GetBlendType().IsValid() && CRSection->GetBlendType() == EMovieSceneBlendType::Override);
+}
+
+static void MergeControlRigSections(UMovieSceneControlRigParameterSection* BaseSection, UMovieSceneControlRigParameterSection* Section, const TRange<FFrameNumber>& Range,
+	const int32* Increment)
+{
+
 	if (!BaseSection || !Section)
 	{
 		return;
@@ -1904,6 +2233,12 @@ static void MergeControlRigSections(UMovieSceneControlRigParameterSection* BaseS
 			return;
 		}
 
+		BaseSection->Modify();
+
+		const bool bIsOverride = (Section->GetBlendType().IsValid() && Section->GetBlendType() == EMovieSceneBlendType::Override) ||
+			(BaseSection->GetBlendType().IsValid() && BaseSection->GetBlendType() == EMovieSceneBlendType::Override);
+
+
 		for (int32 LocalControlIndex = 0; LocalControlIndex < Controls.Num(); ++LocalControlIndex)
 		{
 			FRigControlElement* ControlElement = Controls[LocalControlIndex];
@@ -1913,17 +2248,21 @@ static void MergeControlRigSections(UMovieSceneControlRigParameterSection* BaseS
 				continue;
 			}
 
-			if (FChannelMapInfo* pChannelIndex = BaseSection->ControlChannelMap.Find(ControlElement->GetFName()))
+			const FName& ControlName = ControlElement->GetFName();
+			if (FChannelMapInfo* pChannelIndex = BaseSection->ControlChannelMap.Find(ControlName))
 			{
 				const int32 ChannelIndex = pChannelIndex->ChannelIndex;			
-				//if section is additive an control not there, skip it. we need to keep overrides since we may be blending over
-				//an additive and we want the full override value (with prpoer scale)
-				const bool bMaskKeyOut = (Section->GetBlendType().IsValid() && Section->GetBlendType() == EMovieSceneBlendType::Additive 
-					&& Section->GetControlNameMask(ControlElement->GetFName()) == false);
+				//if override we mask out if both are masked out, if not override we mask out if the top section is masked out
+				const bool bMaskedOutOfBase = BaseSection->GetControlNameMask(ControlName) == false;
+				const bool bMaskedOutOfSection = Section->GetControlNameMask(ControlName) == false;
+				const bool bMaskKeyOut = bIsOverride ? (bMaskedOutOfBase && bMaskedOutOfSection)
+					: bMaskedOutOfSection;
 				if (bMaskKeyOut)
 				{
 					continue;
 				}
+				TArray<UMovieSceneSection*> TrackSections;
+				GetValidCRSections(BaseSection, ControlName, TrackSections);
 
 				TOptional<bool> bBaseSectionResetActive;
 				TOptional<bool> bSectionResetActive;
@@ -1939,8 +2278,7 @@ static void MergeControlRigSections(UMovieSceneControlRigParameterSection* BaseS
 					Section->SetIsActive(false);
 				}
 
-				
-
+			
 				switch (ControlElement->Settings.ControlType)
 				{
 
@@ -1950,7 +2288,7 @@ static void MergeControlRigSections(UMovieSceneControlRigParameterSection* BaseS
 						int32 StartIndex = ChannelIndex;
 						int32 EndIndex = ChannelIndex;
 						MovieSceneToolHelpers::MergeSections<FMovieSceneFloatChannel>(BaseSection,
-							Section, StartIndex, EndIndex, Range);
+							Section, StartIndex, EndIndex, Range, TrackSections,  Increment);
 						break;
 					}
 					case ERigControlType::Vector2D:
@@ -1958,7 +2296,7 @@ static void MergeControlRigSections(UMovieSceneControlRigParameterSection* BaseS
 						int32 StartIndex = ChannelIndex;
 						int32 EndIndex = ChannelIndex + 1;
 						MovieSceneToolHelpers::MergeSections<FMovieSceneFloatChannel>(BaseSection,
-							Section, StartIndex, EndIndex, Range);
+							Section, StartIndex, EndIndex, Range, TrackSections, Increment);
 						break;
 					}
 					case ERigControlType::Position:
@@ -1968,7 +2306,7 @@ static void MergeControlRigSections(UMovieSceneControlRigParameterSection* BaseS
 						int32 StartIndex = ChannelIndex;
 						int32 EndIndex = ChannelIndex + 2;
 						MovieSceneToolHelpers::MergeSections<FMovieSceneFloatChannel>(BaseSection,
-							Section, StartIndex, EndIndex, Range);
+							Section, StartIndex, EndIndex, Range, TrackSections,  Increment);
 						break;
 					}
 
@@ -1987,14 +2325,14 @@ static void MergeControlRigSections(UMovieSceneControlRigParameterSection* BaseS
 								int32 StartIndex = ChannelIndex;
 								int32 EndIndex = ChannelIndex + 5;
 								MovieSceneToolHelpers::MergeSections<FMovieSceneFloatChannel>(BaseSection,
-									Section, StartIndex, EndIndex, Range);
+									Section, StartIndex, EndIndex, Range, TrackSections, Increment);
 							}
 							else
 							{
 								int32 StartIndex = ChannelIndex;
 								int32 EndIndex = ChannelIndex + 8;
 								MovieSceneToolHelpers::MergeSections<FMovieSceneFloatChannel>(BaseSection,
-									Section, StartIndex, EndIndex, Range);
+									Section, StartIndex, EndIndex, Range, TrackSections, Increment);
 							}
 						}
 						else
@@ -2005,7 +2343,7 @@ static void MergeControlRigSections(UMovieSceneControlRigParameterSection* BaseS
 								int32 StartIndex = ChannelIndex;
 								int32 EndIndex = StartIndex;
 								MovieSceneToolHelpers::MergeSections<FMovieSceneFloatChannel>(BaseSection,
-									Section, StartIndex, EndIndex, Range);
+									Section, StartIndex, EndIndex, Range, TrackSections, Increment);
 							}
 							if (EnumHasAllFlags(BaseChannelMask, EMovieSceneTransformChannel::TranslationY)
 								&& EnumHasAllFlags(ChannelMask, EMovieSceneTransformChannel::TranslationY))
@@ -2013,7 +2351,7 @@ static void MergeControlRigSections(UMovieSceneControlRigParameterSection* BaseS
 								int32 StartIndex = ChannelIndex + 1;
 								int32 EndIndex = StartIndex;
 								MovieSceneToolHelpers::MergeSections<FMovieSceneFloatChannel>(BaseSection,
-									Section, StartIndex, EndIndex, Range);
+									Section, StartIndex, EndIndex, Range, TrackSections, Increment);
 							}
 							if (EnumHasAllFlags(BaseChannelMask, EMovieSceneTransformChannel::TranslationZ)
 								&& EnumHasAllFlags(ChannelMask, EMovieSceneTransformChannel::TranslationZ))
@@ -2021,7 +2359,7 @@ static void MergeControlRigSections(UMovieSceneControlRigParameterSection* BaseS
 								int32 StartIndex = ChannelIndex + 2;
 								int32 EndIndex = StartIndex;
 								MovieSceneToolHelpers::MergeSections<FMovieSceneFloatChannel>(BaseSection,
-									Section, StartIndex, EndIndex, Range);
+									Section, StartIndex, EndIndex, Range, TrackSections, Increment);
 							}
 							if (EnumHasAllFlags(BaseChannelMask, EMovieSceneTransformChannel::RotationX)
 								&& EnumHasAllFlags(ChannelMask, EMovieSceneTransformChannel::RotationX))
@@ -2029,7 +2367,7 @@ static void MergeControlRigSections(UMovieSceneControlRigParameterSection* BaseS
 								int32 StartIndex = ChannelIndex + 3;
 								int32 EndIndex = StartIndex;
 								MovieSceneToolHelpers::MergeSections<FMovieSceneFloatChannel>(BaseSection,
-									Section, StartIndex, EndIndex, Range);
+									Section, StartIndex, EndIndex, Range, TrackSections, Increment);
 							}
 							if (EnumHasAllFlags(BaseChannelMask, EMovieSceneTransformChannel::RotationY)
 								&& EnumHasAllFlags(ChannelMask, EMovieSceneTransformChannel::RotationY))
@@ -2037,7 +2375,7 @@ static void MergeControlRigSections(UMovieSceneControlRigParameterSection* BaseS
 								int32 StartIndex = ChannelIndex + 4;
 								int32 EndIndex = StartIndex;
 								MovieSceneToolHelpers::MergeSections<FMovieSceneFloatChannel>(BaseSection,
-									Section, StartIndex, EndIndex, Range);
+									Section, StartIndex, EndIndex, Range, TrackSections, Increment);
 							}
 							if (EnumHasAllFlags(BaseChannelMask, EMovieSceneTransformChannel::RotationZ)
 								&& EnumHasAllFlags(ChannelMask, EMovieSceneTransformChannel::RotationZ))
@@ -2045,7 +2383,7 @@ static void MergeControlRigSections(UMovieSceneControlRigParameterSection* BaseS
 								int32 StartIndex = ChannelIndex + 5;
 								int32 EndIndex = StartIndex;
 								MovieSceneToolHelpers::MergeSections<FMovieSceneFloatChannel>(BaseSection,
-									Section, StartIndex, EndIndex, Range);
+									Section, StartIndex, EndIndex, Range, TrackSections, Increment);
 							}
 							if (ControlElement->Settings.ControlType != ERigControlType::TransformNoScale)
 							{
@@ -2055,7 +2393,7 @@ static void MergeControlRigSections(UMovieSceneControlRigParameterSection* BaseS
 									int32 StartIndex = ChannelIndex + 6;
 									int32 EndIndex = StartIndex;
 									MovieSceneToolHelpers::MergeSections<FMovieSceneFloatChannel>(BaseSection,
-										Section, StartIndex, EndIndex, Range);
+										Section, StartIndex, EndIndex, Range, TrackSections, Increment);
 								}
 								if (EnumHasAllFlags(BaseChannelMask, EMovieSceneTransformChannel::ScaleY)
 									&& EnumHasAllFlags(ChannelMask, EMovieSceneTransformChannel::ScaleY))
@@ -2063,7 +2401,7 @@ static void MergeControlRigSections(UMovieSceneControlRigParameterSection* BaseS
 									int32 StartIndex = ChannelIndex + 7;
 									int32 EndIndex = StartIndex;
 									MovieSceneToolHelpers::MergeSections<FMovieSceneFloatChannel>(BaseSection,
-										Section, StartIndex, EndIndex, Range);
+										Section, StartIndex, EndIndex, Range, TrackSections, Increment);
 								}
 								if (EnumHasAllFlags(BaseChannelMask, EMovieSceneTransformChannel::ScaleZ)
 									&& EnumHasAllFlags(ChannelMask, EMovieSceneTransformChannel::ScaleZ))
@@ -2071,7 +2409,7 @@ static void MergeControlRigSections(UMovieSceneControlRigParameterSection* BaseS
 									int32 StartIndex = ChannelIndex + 8;
 									int32 EndIndex = StartIndex;
 									MovieSceneToolHelpers::MergeSections<FMovieSceneFloatChannel>(BaseSection,
-										Section, StartIndex, EndIndex, Range);
+										Section, StartIndex, EndIndex, Range, TrackSections, Increment);
 								}
 							}
 						}
@@ -2090,10 +2428,17 @@ static void MergeControlRigSections(UMovieSceneControlRigParameterSection* BaseS
 				} 
 			}
 		}
+		if (BaseSection && BaseSection->GetBlendType() == EMovieSceneBlendType::Override)
+		{
+			if (FMovieSceneFloatChannel * FloatChannel = &BaseSection->Weight)
+			{
+				RevertWeightChannelToOne(FloatChannel,Range);
+			}
+		}
 	}
 }
 
-static void MergeTransformSections(UMovieScene3DTransformSection* BaseSection, UMovieScene3DTransformSection* Section, const TRange<FFrameNumber>& Range)
+static void MergeTransformSections(UMovieScene3DTransformSection* BaseSection, UMovieScene3DTransformSection* Section, const TRange<FFrameNumber>& Range, const int32* Increment)
 {
 	if (!BaseSection || !Section)
 	{
@@ -2105,13 +2450,15 @@ static void MergeTransformSections(UMovieScene3DTransformSection* BaseSection, U
 		EMovieSceneTransformChannel BaseChannelMask = BaseSection->GetMask().GetChannels();
 		EMovieSceneTransformChannel ChannelMask = Section->GetMask().GetChannels();
 
+		TArray<UMovieSceneSection*> TrackSections;
+		GetValidSections(BaseSection, TrackSections);
 		const bool bDoAllTransform = EnumHasAllFlags(ChannelMask, EMovieSceneTransformChannel::AllTransform);
 		if (bDoAllTransform)
 		{
 			int32 StartIndex = 0;
 			int32 EndIndex = BaseDoubleChannels.Num() - 1;
 			MovieSceneToolHelpers::MergeSections<FMovieSceneDoubleChannel>(BaseSection,
-				Section, StartIndex, EndIndex, Range);
+				Section, StartIndex, EndIndex, Range, TrackSections, Increment);
 		}
 		else
 		{
@@ -2122,7 +2469,7 @@ static void MergeTransformSections(UMovieScene3DTransformSection* BaseSection, U
 				int32 StartIndex = ChannelIndex;
 				int32 EndIndex = StartIndex;
 				MovieSceneToolHelpers::MergeSections<FMovieSceneDoubleChannel>(BaseSection,
-					Section, StartIndex, EndIndex, Range);
+					Section, StartIndex, EndIndex, Range, TrackSections, Increment);
 			}
 			if (EnumHasAllFlags(BaseChannelMask, EMovieSceneTransformChannel::TranslationY)
 				&& EnumHasAllFlags(ChannelMask, EMovieSceneTransformChannel::TranslationY))
@@ -2130,7 +2477,7 @@ static void MergeTransformSections(UMovieScene3DTransformSection* BaseSection, U
 				int32 StartIndex = ChannelIndex + 1;
 				int32 EndIndex = StartIndex;
 				MovieSceneToolHelpers::MergeSections<FMovieSceneDoubleChannel>(BaseSection,
-					Section, StartIndex, EndIndex, Range);
+					Section, StartIndex, EndIndex, Range, TrackSections, Increment);
 			}
 			if (EnumHasAllFlags(BaseChannelMask, EMovieSceneTransformChannel::TranslationZ)
 				&& EnumHasAllFlags(ChannelMask, EMovieSceneTransformChannel::TranslationZ))
@@ -2138,7 +2485,7 @@ static void MergeTransformSections(UMovieScene3DTransformSection* BaseSection, U
 				int32 StartIndex = ChannelIndex + 2;
 				int32 EndIndex = StartIndex;
 				MovieSceneToolHelpers::MergeSections<FMovieSceneDoubleChannel>(BaseSection,
-					Section, StartIndex, EndIndex, Range);
+					Section, StartIndex, EndIndex, Range, TrackSections, Increment);
 			}
 			if (EnumHasAllFlags(BaseChannelMask, EMovieSceneTransformChannel::RotationX)
 				&& EnumHasAllFlags(ChannelMask, EMovieSceneTransformChannel::RotationX))
@@ -2146,7 +2493,7 @@ static void MergeTransformSections(UMovieScene3DTransformSection* BaseSection, U
 				int32 StartIndex = ChannelIndex + 3;
 				int32 EndIndex = StartIndex;
 				MovieSceneToolHelpers::MergeSections<FMovieSceneDoubleChannel>(BaseSection,
-					Section, StartIndex, EndIndex, Range);
+					Section, StartIndex, EndIndex, Range, TrackSections, Increment);
 			}
 			if (EnumHasAllFlags(BaseChannelMask, EMovieSceneTransformChannel::RotationY)
 				&& EnumHasAllFlags(ChannelMask, EMovieSceneTransformChannel::RotationY))
@@ -2154,7 +2501,7 @@ static void MergeTransformSections(UMovieScene3DTransformSection* BaseSection, U
 				int32 StartIndex = ChannelIndex + 4;
 				int32 EndIndex = StartIndex;
 				MovieSceneToolHelpers::MergeSections<FMovieSceneDoubleChannel>(BaseSection,
-					Section, StartIndex, EndIndex, Range);
+					Section, StartIndex, EndIndex, Range, TrackSections, Increment);
 			}
 			if (EnumHasAllFlags(BaseChannelMask, EMovieSceneTransformChannel::RotationZ)
 				&& EnumHasAllFlags(ChannelMask, EMovieSceneTransformChannel::RotationZ))
@@ -2162,7 +2509,7 @@ static void MergeTransformSections(UMovieScene3DTransformSection* BaseSection, U
 				int32 StartIndex = ChannelIndex +  5;
 				int32 EndIndex = StartIndex;
 				MovieSceneToolHelpers::MergeSections<FMovieSceneDoubleChannel>(BaseSection,
-					Section, StartIndex, EndIndex, Range);
+					Section, StartIndex, EndIndex, Range, TrackSections, Increment);
 			}
 			if (EnumHasAllFlags(BaseChannelMask, EMovieSceneTransformChannel::ScaleX)
 				&& EnumHasAllFlags(ChannelMask, EMovieSceneTransformChannel::ScaleX))
@@ -2170,7 +2517,7 @@ static void MergeTransformSections(UMovieScene3DTransformSection* BaseSection, U
 				int32 StartIndex = ChannelIndex + 6;
 				int32 EndIndex = StartIndex;
 				MovieSceneToolHelpers::MergeSections<FMovieSceneFloatChannel>(BaseSection,
-					Section, StartIndex, EndIndex, Range);
+					Section, StartIndex, EndIndex, Range, TrackSections, Increment);
 			}
 			if (EnumHasAllFlags(BaseChannelMask, EMovieSceneTransformChannel::ScaleY)
 				&& EnumHasAllFlags(ChannelMask, EMovieSceneTransformChannel::ScaleY))
@@ -2178,7 +2525,7 @@ static void MergeTransformSections(UMovieScene3DTransformSection* BaseSection, U
 				int32 StartIndex = ChannelIndex + 7;
 				int32 EndIndex = StartIndex;
 				MovieSceneToolHelpers::MergeSections<FMovieSceneDoubleChannel>(BaseSection,
-					Section, StartIndex, EndIndex, Range);
+					Section, StartIndex, EndIndex, Range, TrackSections, Increment);
 			}
 			if (EnumHasAllFlags(BaseChannelMask, EMovieSceneTransformChannel::ScaleZ)
 				&& EnumHasAllFlags(ChannelMask, EMovieSceneTransformChannel::ScaleZ))
@@ -2186,22 +2533,36 @@ static void MergeTransformSections(UMovieScene3DTransformSection* BaseSection, U
 				int32 StartIndex = ChannelIndex + 8;
 				int32 EndIndex = StartIndex;
 				MovieSceneToolHelpers::MergeSections<FMovieSceneFloatChannel>(BaseSection,
-					Section, StartIndex, EndIndex, Range);
+					Section, StartIndex, EndIndex, Range, TrackSections, Increment);
+			}
+		}
+		if (BaseSection && BaseSection->GetBlendType() == EMovieSceneBlendType::Override)
+		{
+			if (FMovieSceneFloatChannel* FloatChannel = BaseSection->GetWeightChannel())
+			{
+				RevertWeightChannelToOne(FloatChannel,Range);
 			}
 		}
 	}
 }
 
-bool UAnimLayers::MergeAnimLayers(ISequencer* InSequencer, const TArray<int32>& Indices, const FBakingAnimationKeySettings* InSettings)
+bool UAnimLayers::MergeAnimLayers(TSharedPtr<ISequencer>& InSequencerPtr, const TArray<int32>& Indices, const FMergeAnimLayerSettings* InSettings)
 {
-	if (InSequencer == nullptr)
+	if (InSequencerPtr.IsValid() == false)
 	{
 		return false;
 	}
+	ISequencer* InSequencer = InSequencerPtr.Get();
 	TArray<UAnimLayer*> LayersToMerge;
-	const FFrameNumber Min = InSettings ? InSettings->StartFrame : TNumericLimits<FFrameNumber>::Lowest();
-	const FFrameNumber Max = InSettings ? InSettings->EndFrame : TNumericLimits<FFrameNumber>::Max();
-	TRange<FFrameNumber> Range(Min, Max);
+	const FFrameNumber Min = TNumericLimits<FFrameNumber>::Lowest();
+	const FFrameNumber Max = TNumericLimits<FFrameNumber>::Max();
+
+	TRange<FFrameNumber> Range = InSequencerPtr->GetFocusedMovieSceneSequence()->GetMovieScene()->GetPlaybackRange();
+	TOptional<TRange<FFrameNumber>> OptionalRange = InSequencerPtr->GetSubSequenceRange();
+	if (OptionalRange.IsSet())
+	{
+		Range = TRange<FFrameNumber>(OptionalRange.GetValue().GetLowerBoundValue(), OptionalRange.GetValue().GetUpperBoundValue());
+	}
 	FScopedTransaction Transaction(LOCTEXT("Merge Anim Layers", "Merge Anim Layers"), !GIsTransacting);
 
 	TArray<int32> SortedIndices = Indices;
@@ -2225,7 +2586,8 @@ bool UAnimLayers::MergeAnimLayers(ISequencer* InSequencer, const TArray<int32>&
 		return false;
 	}
 	Modify();
-
+	//set up Increment if we are baking increments
+	const int32* Increment = (InSettings && InSettings->BakingKeySettings == EBakingKeySettings::AllFrames) ? &InSettings->FrameIncrement : nullptr;
 	for (int32 Index = 0; Index < LayersToMerge.Num() -1; ++Index)
 	{
 		UAnimLayer* BaseLayer = LayersToMerge[Index + 1];
@@ -2254,15 +2616,17 @@ bool UAnimLayers::MergeAnimLayers(ISequencer* InSequencer, const TArray<int32>&
 								if (UMovieSceneControlRigParameterSection* BaseCRSection = Cast< UMovieSceneControlRigParameterSection>(BaseSectionItem->Section.Get()))
 								{
 									UMovieSceneControlRigParameterSection* CRSection = Cast< UMovieSceneControlRigParameterSection>(SectionItem.Section.Get());
-									MergeControlRigSections(BaseCRSection, CRSection, Range);
+									MergeControlRigSections(BaseCRSection, CRSection, Range, Increment);
 								}
 								else if (UMovieScene3DTransformSection* BaseTRSection = Cast< UMovieScene3DTransformSection>(BaseSectionItem->Section.Get()))
 								{
 									UMovieScene3DTransformSection* TRSection = Cast< UMovieScene3DTransformSection>(SectionItem.Section.Get());
-									MergeTransformSections(BaseTRSection, TRSection, Range);
+									MergeTransformSections(BaseTRSection, TRSection, Range, Increment);
 								}
 								else
 								{
+									TArray<UMovieSceneSection*> TrackSections;
+									GetValidSections(BaseSectionItem->Section.Get(), TrackSections);
 									TArrayView<FMovieSceneFloatChannel*> BaseFloatChannels = BaseSectionItem->Section->GetChannelProxy().GetChannels<FMovieSceneFloatChannel>();
 									TArrayView<FMovieSceneDoubleChannel*> BaseDoubleChannels = BaseSectionItem->Section->GetChannelProxy().GetChannels<FMovieSceneDoubleChannel>();
 									if (BaseDoubleChannels.Num() > 0)
@@ -2270,7 +2634,7 @@ bool UAnimLayers::MergeAnimLayers(ISequencer* InSequencer, const TArray<int32>&
 										int32 StartIndex = 0;
 										int32 EndIndex = BaseDoubleChannels.Num() - 1;
 										MovieSceneToolHelpers::MergeSections<FMovieSceneDoubleChannel>(BaseSectionItem->Section.Get(),
-											SectionItem.Section.Get(), StartIndex, EndIndex, Range);
+											SectionItem.Section.Get(), StartIndex, EndIndex, Range, TrackSections, Increment);
 
 									}
 									else if (BaseFloatChannels.Num() > 0)
@@ -2278,7 +2642,7 @@ bool UAnimLayers::MergeAnimLayers(ISequencer* InSequencer, const TArray<int32>&
 										int32 StartIndex = 0;
 										int32 EndIndex = BaseFloatChannels.Num() - 1;
 										MovieSceneToolHelpers::MergeSections<FMovieSceneFloatChannel>(BaseSectionItem->Section.Get(),
-											SectionItem.Section.Get(), StartIndex, EndIndex, Range);
+											SectionItem.Section.Get(), StartIndex, EndIndex, Range, TrackSections, Increment);
 									}
 								}
 							}
@@ -2298,6 +2662,15 @@ bool UAnimLayers::MergeAnimLayers(ISequencer* InSequencer, const TArray<int32>&
 									}
 								}
 							}
+							//merging so do a key reduction possibly.
+							if (InSettings->bReduceKeys)
+							{
+								FSmartReduceParams SmartParams;
+								SmartParams.SampleRate = InSequencer->GetFocusedDisplayRate();
+								SmartParams.TolerancePercentage = InSettings->TolerancePercentage;
+								FControlRigParameterTrackEditor::SmartReduce(InSequencerPtr, SmartParams, BaseSectionItem->Section.Get());
+								return true;
+							}
 						}
 						else //okay this object doesn't exist in the first layer we are merging into so we need to move it to the other one 
 						{
@@ -2351,6 +2724,7 @@ bool UAnimLayers::MergeAnimLayers(ISequencer* InSequencer, const TArray<int32>&
 	{
 		SetUpBaseLayerSections(); //if it is the base reset it
 	}
+	
 	InSequencer->NotifyMovieSceneDataChanged(EMovieSceneDataChangeType::MovieSceneStructureItemAdded);
 	return true;
 }
@@ -2853,6 +3227,7 @@ void AdjustmentBlend(UMovieSceneSection* Section, TArrayView<ChannelType*>& Base
 				const double CurrentValue = (KeyInterval.EndValue > KeyInterval.StartValue) ?
 					PreviousValue + ValueDelta : PreviousValue - ValueDelta;
 				AddKeyToChannel(LayerChannel, DefaultInterpolation, TimeValue.Key, CurrentValue);
+				PreviousValue = CurrentValue;
 			}
 		}
 	}
@@ -2901,7 +3276,7 @@ bool UAnimLayers::AdjustmentBlendLayers(ISequencer* InSequencer, int32 LayerInde
 							TArrayView<FMovieSceneDoubleChannel*> BaseDoubleChannels = BaseSectionItem->Section->GetChannelProxy().GetChannels<FMovieSceneDoubleChannel>();
 							TArrayView<FMovieSceneFloatChannel*> LayerFloatChannels = SectionItem.Section->GetChannelProxy().GetChannels<FMovieSceneFloatChannel>();
 							TArrayView<FMovieSceneDoubleChannel*> LayerDoubleChannels = SectionItem.Section->GetChannelProxy().GetChannels<FMovieSceneDoubleChannel>();
-
+							SectionItem.Section->Modify();
 							
 							if (BaseDoubleChannels.Num() > 0)
 							{
@@ -2923,3 +3298,5 @@ bool UAnimLayers::AdjustmentBlendLayers(ISequencer* InSequencer, int32 LayerInde
 }
 
 #undef LOCTEXT_NAMESPACE
+
+
diff --git a/Engine/Plugins/Animation/ControlRig/Source/ControlRigEditor/Private/Sequencer/AnimLayers/AnimLayers.h b/Engine/Plugins/Animation/ControlRig/Source/ControlRigEditor/Private/Sequencer/AnimLayers/AnimLayers.h
--- a/Engine/Plugins/Animation/ControlRig/Source/ControlRigEditor/Private/Sequencer/AnimLayers/AnimLayers.h
+++ b/Engine/Plugins/Animation/ControlRig/Source/ControlRigEditor/Private/Sequencer/AnimLayers/AnimLayers.h
@@ -10,6 +10,7 @@
 #include "ControlRig.h"
 #include "Styling/SlateTypes.h"
 #include "ISequencerPropertyKeyedStatus.h"
+#include "BakingAnimationKeySettings.h"
 #include "AnimLayers.generated.h"
 
 class ULevelSequence;
@@ -19,10 +20,37 @@ class UMovieSceneSection;
 class ISequencer;
 class UMovieSceneControlRigParameterTrack;
 class UMovieSceneControlRigParameterSection;
-struct FBakingAnimationKeySettings;
 struct FMovieSceneFloatChannel;
 class IPropertyHandle;
 
+
+USTRUCT(BlueprintType)
+struct FMergeAnimLayerSettings
+{
+	GENERATED_BODY();
+
+	FMergeAnimLayerSettings()
+	{
+		BakingKeySettings = EBakingKeySettings::KeysOnly;
+		FrameIncrement = 1;
+		bReduceKeys = false;
+		TolerancePercentage = 5.0;
+	}
+	UPROPERTY(BlueprintReadWrite, EditAnywhere, Category = "Bake")
+	EBakingKeySettings BakingKeySettings;
+
+	UPROPERTY(BlueprintReadWrite, EditAnywhere, Category = "Bake", meta = (ClampMin = "1", UIMin = "1", EditCondition = "BakingKeySettings == EBakingKeySettings::AllFrames"))
+	int32 FrameIncrement;
+
+	UPROPERTY(BlueprintReadWrite, EditAnywhere, Category = "Bake", meta = (EditCondition = "BakingKeySettings == EBakingKeySettings::AllFrames"))
+	bool bReduceKeys;
+
+	UPROPERTY(BlueprintReadWrite, EditAnywhere, Category = "Bake", meta = (EditCondition = "BakingKeySettings == EBakingKeySettings::AllFrames && bReduceKeys"))
+	float TolerancePercentage;
+};
+
+
+
 USTRUCT(BlueprintType)
 // Name of a property and control and the specific channels that make up the layer
 struct FAnimLayerPropertyAndChannels
@@ -233,11 +261,13 @@ public:
 	bool RemoveAnimLayerItem(UObject* InObject);
 	void SetSelectedInList(bool bInValue); 
 	bool GetSelectedInList() const { return bIsSelectedInList; }
+	ECheckBoxState GetSelected(TSet<UObject*>& OutSelectedObjects, TMap<UControlRig*, TArray<FName>>& OutSelectedControls) const;
 
 	//for key property status for weights
 	void SetKey(TSharedPtr<ISequencer>& Sequencer, const IPropertyHandle& KeyedPropertyHandle);
 	EPropertyKeyedStatus GetPropertyKeyedStatus(TSharedPtr<ISequencer>& Sequencer, const IPropertyHandle& PropertyHandle);
 private:
+
 	UPROPERTY(VisibleAnywhere, Category = "Data")
 	TMap<TWeakObjectPtr<UObject>, FAnimLayerItem> AnimLayerItems;
 
@@ -278,10 +308,11 @@ public:
 	int32 DuplicateAnimLayer(ISequencer* InSequencer, int32 Index);
 	int32 AddAnimLayerFromSelection(ISequencer* InSequencer);
 	void GetAnimLayerStates(TArray<FAnimLayerState>& OutStates);
-	bool MergeAnimLayers(ISequencer* InSequencer, const TArray<int32>& Indices, const FBakingAnimationKeySettings* InSettings);
+	bool MergeAnimLayers(TSharedPtr<ISequencer>& InSequencerPtr, const TArray<int32>& Indices, const FMergeAnimLayerSettings* InSettings);
 	bool SetPassthroughKey(ISequencer* InSequencer, int32 Index);
+	bool SetKey(ISequencer* InSequencer, int32 Index);
 
-	//will always blend to base fo rnow
+	//will always blend to base for now
 	bool AdjustmentBlendLayers(ISequencer* InSequencer, int32 LayerIndex);
 
 	bool IsTrackOnSelectedLayer(const UMovieSceneTrack* InTrack)const; 
@@ -302,6 +333,7 @@ private:
 	void AnimLayerListChangedBroadcast();
 	FAnimLayerListChanged OnAnimLayerListChanged;
 	void AddBaseLayer();
+	bool SetKeyValueOrPassthrough(ISequencer* InSequencer, int32 Index, bool bJustValue);
 public:
 	void SetUpBaseLayerSections();
 
diff --git a/Engine/Plugins/Animation/ControlRig/Source/ControlRigEditor/Private/Sequencer/AnimLayers/SAnimLayers.cpp b/Engine/Plugins/Animation/ControlRig/Source/ControlRigEditor/Private/Sequencer/AnimLayers/SAnimLayers.cpp
--- a/Engine/Plugins/Animation/ControlRig/Source/ControlRigEditor/Private/Sequencer/AnimLayers/SAnimLayers.cpp
+++ b/Engine/Plugins/Animation/ControlRig/Source/ControlRigEditor/Private/Sequencer/AnimLayers/SAnimLayers.cpp
@@ -13,6 +13,7 @@
 #include <DetailsNameWidgetOverrideCustomization.h>
 #include "IDetailCustomization.h"
 #include "IDetailCustomNodeBuilder.h"
+#include "IStructureDetailsView.h"
 #include "Containers/Set.h"
 #include "PropertyPath.h"
 #include "Widgets/SWidget.h"
@@ -47,6 +48,7 @@ typedef TSharedPtr<FAnimLayerSourceUIEntry> FAnimLayerSourceUIEntryPtr;
 namespace AnimLayerSourceListUI
 {
 	static const FName LayerColumnName(TEXT("Layer"));
+	static const FName ActionColumnName(TEXT("Action"));
 	static const FName StatusColumnName(TEXT("Status"));
 	static const FName WeightColumnName(TEXT("Weight"));
 	static const FName TypeColumnName(TEXT("Type"));
@@ -68,6 +70,7 @@ public:
 	void Duplicate() const;
 	void DeleteAnimLayer() const;
 	void SetPassthroughKey() const;
+	void SetKey() const;
 
 	ECheckBoxState GetKeyed() const;
 	void SetKeyed();
@@ -208,6 +211,25 @@ void FAnimLayerSourceUIEntry::SetPassthroughKey() const
 	}
 }
 
+void FAnimLayerSourceUIEntry::SetKey() const
+{
+	if (AnimLayer)
+	{
+		if (TSharedPtr<ISequencer> SequencerPtr = UAnimLayers::GetSequencerFromAsset())
+		{
+			if (UAnimLayers* AnimLayers = UAnimLayers::GetAnimLayers(SequencerPtr.Get()))
+			{
+				int32 Index = AnimLayers->GetAnimLayerIndex(AnimLayer);
+				if (Index != INDEX_NONE)
+				{
+					AnimLayers->SetKey(SequencerPtr.Get(), Index);
+				}
+			}
+		}
+		ClearCaches();
+	}
+}
+
 
 ECheckBoxState FAnimLayerSourceUIEntry::GetKeyed() const
 {
@@ -643,6 +665,54 @@ public:
 							.OnTextCommitted(this, &SAnimLayerSourcesRow::OnLayerNameCommited)
 					]
 				];
+		}
+		else if (ColumnName == AnimLayerSourceListUI::ActionColumnName)
+		{
+			return SNew(SHorizontalBox)
+				+ SHorizontalBox::Slot()
+				.AutoWidth()
+				.Padding(1, 0)
+				.VAlign(VAlign_Center)
+				[
+					SNew(SButton)
+						.IsFocusable(false)
+						.ButtonStyle(FAppStyle::Get(), TEXT("SimpleButton"))
+						.OnClicked_Lambda([this]()
+							{
+								EntryPtr->SetKey();
+								return FReply::Handled();
+							})
+						.ContentPadding(1.f)
+						.ToolTipText(LOCTEXT("AnimLayerKeyTooltip", "Key selected controls or all controls if none selected"))
+						[
+							SNew(SImage)
+								.Image(FAppStyle::GetBrush("Sequencer.AddKey.Details"))
+								.ColorAndOpacity(FSlateColor::UseForeground())
+						]
+				]
+				+ SHorizontalBox::Slot()
+				.AutoWidth()
+				.Padding(1, 0)
+				.VAlign(VAlign_Center)
+				[
+					SNew(SButton)
+						.IsFocusable(false)
+						.ButtonStyle(FAppStyle::Get(), TEXT("SimpleButton"))
+						.OnClicked_Lambda([this]()
+							{
+								EntryPtr->SetPassthroughKey();
+								return FReply::Handled();
+							})
+						.ContentPadding(1.f)
+						.ToolTipText(LOCTEXT("AnimLayerPassthroughTooltip", "Set Default Pose(Base), Zero Key(Additive) or Passthrough(Override) key"))
+						[
+							SNew(SImage)
+								.Image(FAppStyle::GetBrush("Sequencer.AddKey.Details"))
+								.ColorAndOpacity(FSlateColor::UseForeground())
+						]
+				];
+
+
 		}
 		else if (ColumnName == AnimLayerSourceListUI::StatusColumnName)
 		{
@@ -863,6 +933,7 @@ private:
 	TSharedPtr<SButton> MuteButton;
 	TSharedPtr<SButton> LockButton;
 
+
 };
 
 FAnimLayerSourcesView::FAnimLayerSourcesView()
@@ -936,13 +1007,16 @@ void FAnimLayerSourcesView::CreateSourcesListView()
 		(
 			SNew(SHeaderRow)
 			+ SHeaderRow::Column(AnimLayerSourceListUI::LayerColumnName)
-			.FillWidth(120.f)
+			.FillWidth(100.f)
 			.DefaultLabel(LOCTEXT("LayerColumnName", "Name"))
+			+ SHeaderRow::Column(AnimLayerSourceListUI::ActionColumnName)
+			.FillWidth(30.f)
+			.DefaultLabel(LOCTEXT("LayerColumnName", "Action"))
 			+ SHeaderRow::Column(AnimLayerSourceListUI::StatusColumnName)
 			.FillWidth(40.f)
 			.DefaultLabel(LOCTEXT("StatusColumnName", "Status"))
 			+ SHeaderRow::Column(AnimLayerSourceListUI::WeightColumnName)
-			.FillWidth(80.f)
+			.FillWidth(60.f)
 			.DefaultLabel(LOCTEXT("WeightColumnName", "Weight"))
 			+ SHeaderRow::Column(AnimLayerSourceListUI::TypeColumnName)
 			.FillWidth(80.f)
@@ -1038,6 +1112,198 @@ void FAnimLayerSourcesView::AdjustmentBlend()
 		}
 	}
 }
+//////////////////////////////////////////////////////////////
+/// SMergeAnimLayersWidget
+///////////////////////////////////////////////////////////
+DECLARE_DELEGATE_TwoParams(FMergeAnimLayersCB, TSharedPtr<ISequencer>& InSequencer, const FMergeAnimLayerSettings& InSettings);
+
+/** Widget allowing collapsing of controls */
+class SMergeAnimLayersWidget : public SCompoundWidget
+{
+public:
+
+	SLATE_BEGIN_ARGS(SMergeAnimLayersWidget)
+		: _Sequencer(nullptr), _bSmartReduce(false), _TolerancePercentage(TNumericLimits<float>::Max())
+		{}
+		SLATE_ARGUMENT(TWeakPtr<ISequencer>, Sequencer)
+		SLATE_ARGUMENT(bool, bSmartReduce)
+		SLATE_ARGUMENT(float, TolerancePercentage)
+
+
+	SLATE_END_ARGS()
+
+	void Construct(const FArguments& InArgs);
+	virtual ~SMergeAnimLayersWidget() override {}
+
+	FReply OpenDialog(bool bModal = true);
+	void CloseDialog();
+
+	void SetMergeCB(FMergeAnimLayersCB& InCB) { MergeCB = InCB; }
+private:
+	void Merge();
+
+	TWeakPtr<ISequencer> Sequencer;
+	//static to be reused
+	static TOptional<FMergeAnimLayerSettings> MergeAnimLayersSettings;
+	//structonscope for details panel
+	TSharedPtr < TStructOnScope<FMergeAnimLayerSettings>> Settings;
+	TWeakPtr<SWindow> DialogWindow;
+	TSharedPtr<IStructureDetailsView> DetailsView;
+	FMergeAnimLayersCB MergeCB;
+};
+
+
+TOptional<FMergeAnimLayerSettings> SMergeAnimLayersWidget::MergeAnimLayersSettings;
+
+void SMergeAnimLayersWidget::Construct(const FArguments& InArgs)
+{
+	Sequencer = InArgs._Sequencer;
+
+	if (MergeAnimLayersSettings.IsSet() == false)
+	{
+		TSharedPtr<ISequencer> SequencerPtr = Sequencer.Pin();
+		MergeAnimLayersSettings = FMergeAnimLayerSettings();
+	}
+
+	if (InArgs._TolerancePercentage != TNumericLimits<float>::Max())
+	{
+		MergeAnimLayersSettings.GetValue().bReduceKeys = InArgs._bSmartReduce;
+		MergeAnimLayersSettings.GetValue().TolerancePercentage = InArgs._TolerancePercentage;
+	}
+
+	Settings = MakeShared<TStructOnScope<FMergeAnimLayerSettings>>();
+	Settings->InitializeAs<FMergeAnimLayerSettings>(MergeAnimLayersSettings.GetValue());
+
+	FStructureDetailsViewArgs StructureViewArgs;
+	StructureViewArgs.bShowObjects = true;
+	StructureViewArgs.bShowAssets = true;
+	StructureViewArgs.bShowClasses = true;
+	StructureViewArgs.bShowInterfaces = true;
+
+	FDetailsViewArgs ViewArgs;
+	ViewArgs.bAllowSearch = false;
+	ViewArgs.bHideSelectionTip = false;
+	ViewArgs.bShowObjectLabel = false;
+
+	FPropertyEditorModule& PropertyEditor = FModuleManager::Get().LoadModuleChecked<FPropertyEditorModule>(TEXT("PropertyEditor"));
+
+	DetailsView = PropertyEditor.CreateStructureDetailView(ViewArgs, StructureViewArgs, TSharedPtr<FStructOnScope>());
+	DetailsView->GetDetailsView()->RegisterInstancedCustomPropertyTypeLayout("FrameNumber",
+		FOnGetPropertyTypeCustomizationInstance::CreateSP(Sequencer.Pin().ToSharedRef(), &ISequencer::MakeFrameNumberDetailsCustomization));
+	DetailsView->SetStructureData(Settings);
+
+	ChildSlot
+		[
+			SNew(SBorder)
+				.Visibility(EVisibility::Visible)
+				[
+					SNew(SVerticalBox)
+						+ SVerticalBox::Slot()
+						.AutoHeight()
+						.Padding(0.f)
+						[
+							DetailsView->GetWidget().ToSharedRef()
+						]
+						+ SVerticalBox::Slot()
+						.AutoHeight()
+						.Padding(16.f)
+						[
+
+							SNew(SHorizontalBox)
+								+ SHorizontalBox::Slot()
+								.FillWidth(1.f)
+								.HAlign(HAlign_Fill)
+								[
+									SNew(SSpacer)
+								]
+
+								+ SHorizontalBox::Slot()
+								.AutoWidth()
+								.VAlign(VAlign_Center)
+								.HAlign(HAlign_Right)
+								.Padding(0.f)
+								[
+									SNew(SButton)
+										.HAlign(HAlign_Center)
+										.ContentPadding(FAppStyle::GetMargin("StandardDialog.ContentPadding"))
+										.Text(LOCTEXT("OK", "OK"))
+										.OnClicked_Lambda([this, InArgs]()
+											{
+												Merge();
+												CloseDialog();
+												return FReply::Handled();
+
+											})
+										.IsEnabled_Lambda([this]()
+											{
+												return (Settings.IsValid());
+											})
+								]
+						]
+				]
+		];
+}
+
+void  SMergeAnimLayersWidget::Merge()
+{
+	FMergeAnimLayerSettings* BakeSettings = Settings->Get();
+	TSharedPtr<ISequencer> SequencerPtr = Sequencer.Pin();
+	MergeCB.ExecuteIfBound(SequencerPtr, *BakeSettings);
+	MergeAnimLayersSettings = *BakeSettings;
+}
+
+class SMergeAnimLayersWidgetWindow : public SWindow
+{
+};
+
+FReply SMergeAnimLayersWidget::OpenDialog(bool bModal)
+{
+	check(!DialogWindow.IsValid());
+
+	const FVector2D CursorPos = FSlateApplication::Get().GetCursorPos();
+
+	TSharedRef<SMergeAnimLayersWidgetWindow> Window = SNew(SMergeAnimLayersWidgetWindow)
+		.Title(LOCTEXT("MergeAnimLayers", "Merge Anim Layer"))
+		.CreateTitleBar(true)
+		.Type(EWindowType::Normal)
+		.SizingRule(ESizingRule::Autosized)
+		.ScreenPosition(CursorPos)
+		.FocusWhenFirstShown(true)
+		.ActivationPolicy(EWindowActivationPolicy::FirstShown)
+		[
+			AsShared()
+		];
+
+	Window->SetWidgetToFocusOnActivate(AsShared());
+
+	DialogWindow = Window;
+
+	Window->MoveWindowTo(CursorPos);
+
+	if (bModal)
+	{
+		GEditor->EditorAddModalWindow(Window);
+	}
+	else
+	{
+		FSlateApplication::Get().AddWindow(Window);
+	}
+
+	return FReply::Handled();
+}
+
+void SMergeAnimLayersWidget::CloseDialog()
+{
+	if (DialogWindow.IsValid())
+	{
+		DialogWindow.Pin()->RequestDestroyWindow();
+		DialogWindow.Reset();
+	}
+}
+///////////////
+//end of SMergeAnimLayersWidget
+/////////////////////
+
 
 void FAnimLayerSourcesView::MergeLayers()
 {
@@ -1062,22 +1328,22 @@ void FAnimLayerSourcesView::MergeLayers()
 			}
 			if (LayersToMerge.Num() > 1)
 			{
-				AnimLayers->MergeAnimLayers(SequencerPtr.Get(), LayersToMerge, nullptr);
-				//for now we don't use the dialog since it's not doing proper per key frame reductions
-				/*
-				FCollapseControlsCB CollapseCB = FCollapseControlsCB::CreateLambda([AnimLayers,LayersToMerge](TSharedPtr<ISequencer>& InSequencer,const FBakingAnimationKeySettings& InSettings)
+			
+				FMergeAnimLayersCB MergeCB = FMergeAnimLayersCB::CreateLambda([AnimLayers,LayersToMerge](TSharedPtr<ISequencer>& InSequencer,const FMergeAnimLayerSettings& InSettings)
 				{
-					AnimLayers->MergeAnimLayers(InSequencer.Get(), LayersToMerge, InSettings);
+					AnimLayers->MergeAnimLayers(InSequencer, LayersToMerge, &InSettings);
 
 				});
 
-				TSharedRef<SCollapseControlsWidget> BakeWidget =
-					SNew(SCollapseControlsWidget)
-					.Sequencer(SequencerPtr);
 
-				BakeWidget->SetCollapseCB(CollapseCB);
+				TSharedRef<SMergeAnimLayersWidget> BakeWidget =
+					SNew(SMergeAnimLayersWidget)
+					.Sequencer(SequencerPtr)
+					.bSmartReduce(false)
+					.TolerancePercentage(5.0);
+				
+				BakeWidget->SetMergeCB(MergeCB);
 				BakeWidget->OpenDialog(false);
-				*/
 			}
 		}
 	}
@@ -1464,7 +1730,7 @@ void SAnimLayers::Construct(const FArguments& InArgs, FControlRigEditMode& InEdi
 								.ButtonStyle(FAppStyle::Get(), "NoBorder")
 								.ButtonColorAndOpacity_Lambda([this]()
 									{
-										return FLinearColor::White;
+										return FSlateColor(FLinearColor(1.f, 1.f, 1.f, 1));
 									})
 								.OnClicked_Lambda([this]()
 									{
@@ -1478,7 +1744,7 @@ void SAnimLayers::Construct(const FArguments& InArgs, FControlRigEditMode& InEdi
 										.ColorAndOpacity_Lambda([this]()
 											{
 
-												const FLinearColor Selected = FLinearColor(FColor::FromHex(TEXT("#EBA30A")));
+												const FLinearColor Selected = FLinearColor::White;
 												const FColor NotSelected(56, 56, 56);
 												FSlateColor SlateColor;
 												if (AnimLayerController->IsSelectionFilterActive() == true)
diff --git a/Engine/Plugins/Animation/ControlRig/Source/ControlRigEditor/Private/Sequencer/ControlRigParameterTrackEditor.cpp b/Engine/Plugins/Animation/ControlRig/Source/ControlRigEditor/Private/Sequencer/ControlRigParameterTrackEditor.cpp
--- a/Engine/Plugins/Animation/ControlRig/Source/ControlRigEditor/Private/Sequencer/ControlRigParameterTrackEditor.cpp
+++ b/Engine/Plugins/Animation/ControlRig/Source/ControlRigEditor/Private/Sequencer/ControlRigParameterTrackEditor.cpp
@@ -1135,18 +1135,18 @@ FReply SBakeToAnimAndControlRigOptionsWindow::OnResetToDefaultClick() const
 	return FReply::Handled();
 }
 
-void FControlRigParameterTrackEditor::SmartReduce(const TSharedPtr<ISequencer>& InSequencer, const FSmartReduceParams& InParams, UMovieSceneControlRigParameterSection * ParamSection)
+void FControlRigParameterTrackEditor::SmartReduce(const TSharedPtr<ISequencer>& InSequencer, const FSmartReduceParams& InParams, UMovieSceneSection* InSection)
 {
 	using namespace UE::Sequencer;
 
-	if (ParamSection)
+	if (InSection)
 	{
 		ISequencerModule& SequencerModule = FModuleManager::LoadModuleChecked<ISequencerModule>("Sequencer");
 		constexpr bool bNeedToTestExisting = false;
 		TOptional<FKeyHandleSet> KeyHandleSet;
 
-		FMovieSceneChannelProxy& ChannelProxy = ParamSection->GetChannelProxy();
-		for (const FMovieSceneChannelEntry& Entry : ParamSection->GetChannelProxy().GetAllEntries())
+		FMovieSceneChannelProxy& ChannelProxy = InSection->GetChannelProxy();
+		for (const FMovieSceneChannelEntry& Entry : InSection->GetChannelProxy().GetAllEntries())
 		{
 			const FName ChannelTypeName = Entry.GetChannelTypeName();
 			TArrayView<FMovieSceneChannel* const> Channels = Entry.GetChannels();
@@ -1155,7 +1155,7 @@ void FControlRigParameterTrackEditor::SmartReduce(const TSharedPtr<ISequencer>&
 				FMovieSceneChannelHandle ChannelHandle = ChannelProxy.MakeHandle(ChannelTypeName, Index);
 				ISequencerChannelInterface* EditorInterface = SequencerModule.FindChannelEditorInterface(ChannelHandle.GetChannelTypeName());
 
-				FCreateCurveEditorModelParams CurveModelParams{ ParamSection, ParamSection, InSequencer.ToSharedRef() };
+				FCreateCurveEditorModelParams CurveModelParams{ InSection, InSection, InSequencer.ToSharedRef() };
 				if (TUniquePtr<FCurveModel> CurveModel = EditorInterface->CreateCurveEditorModel_Raw(ChannelHandle, CurveModelParams))
 				{
 					FKeyHandleSet OutHandleSet;
diff --git a/Engine/Plugins/Animation/ControlRig/Source/ControlRigEditor/Private/Sequencer/ControlRigParameterTrackEditor.h b/Engine/Plugins/Animation/ControlRig/Source/ControlRigEditor/Private/Sequencer/ControlRigParameterTrackEditor.h
--- a/Engine/Plugins/Animation/ControlRig/Source/ControlRigEditor/Private/Sequencer/ControlRigParameterTrackEditor.h
+++ b/Engine/Plugins/Animation/ControlRig/Source/ControlRigEditor/Private/Sequencer/ControlRigParameterTrackEditor.h
@@ -316,7 +316,7 @@ public:
 		bool bOntoSelectedControls, UMovieSceneControlRigParameterSection* ParamSection);
 
 	//** Function to smart reduce all keys on all constraints based on the parameters
-	static void SmartReduce(const TSharedPtr<ISequencer>& InSequencer, const FSmartReduceParams& InParams, UMovieSceneControlRigParameterSection* InSection);
+	static void SmartReduce(const TSharedPtr<ISequencer>& InSequencer, const FSmartReduceParams& InParams, UMovieSceneSection* InSection);
 
 	/** Adds the constraint to sequencer and create the bindings if needed. Note that ControlRig Track auto-generation is disabled in this function. */
 	static void AddConstraintToSequencer(const TSharedPtr<ISequencer>& InSequencer, UTickableTransformConstraint* InConstraint);
diff --git a/Engine/Plugins/Animation/ControlRig/Source/ControlRigEditor/Public/ControlRigSequencerEditorLibrary.h b/Engine/Plugins/Animation/ControlRig/Source/ControlRigEditor/Public/ControlRigSequencerEditorLibrary.h
--- a/Engine/Plugins/Animation/ControlRig/Source/ControlRigEditor/Public/ControlRigSequencerEditorLibrary.h
+++ b/Engine/Plugins/Animation/ControlRig/Source/ControlRigEditor/Public/ControlRigSequencerEditorLibrary.h
@@ -25,6 +25,7 @@ class UTickableConstraint;
 class UTickableTransformConstraint;
 class UTransformableHandle;
 struct FBakingAnimationKeySettings;
+struct FMergeAnimLayerSettings;
 class UControlRig;
 class UAnimLayer;
 
@@ -1175,13 +1176,20 @@ public:
 	UFUNCTION(BlueprintCallable, Category = "Editor Scripting | Sequencer Tools | Control Rig | Animation Layers")
 	static int32 AddAnimLayerFromSelection();
 
+
+	UE_DEPRECATED(5.6, "Use MergeAnimLayersWithSettings that takes FMergeAnimLayerSettings")
+	UFUNCTION(BlueprintCallable, Category = "Editor Scripting | Sequencer Tools | Control Rig | AnimationLayers", meta = (DeprecatedFunction, DeprecationMessage = "Use MergeAnimLayersWithSettings that takes FMergeAnimLayerSettings"))
+	static bool MergeAnimLayers(const TArray<int32>& Indices);
+
 	/**
 	* Merge specified anim layers into one layer. Will merge onto the anim layer with the lowest index
 	* @param Indices The indices to merge
+	* @param MergeAnimLayerSettings Settings with per frame or key options and key reduction
 	* @return Returns true if successful, false otherwise
 	*/
+
 	UFUNCTION(BlueprintCallable, Category = "Editor Scripting | Sequencer Tools | Control Rig | Animation Layers")
-	static bool MergeAnimLayers(const TArray<int32>& Indices);
+	static bool MergeAnimLayersWithSettings(const TArray<int32>& Indices, const FMergeAnimLayerSettings& MergeAnimLayerSettings);
 
 	/**
 	* Get the animation layer objects
diff --git a/Engine/Source/Editor/MovieSceneTools/Private/MovieSceneToolHelpers.cpp b/Engine/Source/Editor/MovieSceneTools/Private/MovieSceneToolHelpers.cpp
--- a/Engine/Source/Editor/MovieSceneTools/Private/MovieSceneToolHelpers.cpp
+++ b/Engine/Source/Editor/MovieSceneTools/Private/MovieSceneToolHelpers.cpp
@@ -5471,7 +5471,7 @@ bool MovieSceneToolHelpers::CollapseSection(TSharedPtr<ISequencer>& SequencerPtr
 			if (BaseSection && Section && BaseSection->GetBlendType().IsValid() && Section->GetBlendType().IsValid())
 			{
 				OwnerTrack->Modify();
-
+				TArray<UMovieSceneSection*> TrackSections = OwnerTrack->GetAllSections();
 				if (Section->IsActive())//active sections merge them
 				{
 					TArrayView<FMovieSceneFloatChannel*> BaseFloatChannels = BaseSection->GetChannelProxy().GetChannels<FMovieSceneFloatChannel>();
@@ -5481,7 +5481,7 @@ bool MovieSceneToolHelpers::CollapseSection(TSharedPtr<ISequencer>& SequencerPtr
 						const int StartIndex = 0;
 						const int EndIndex = BaseDoubleChannels.Num() - 1;
 						MovieSceneToolHelpers::MergeSections<FMovieSceneDoubleChannel>(BaseSection,
-							Section, StartIndex, EndIndex, Range);
+							Section, StartIndex, EndIndex, Range, TrackSections);
 
 					}
 					else if (BaseFloatChannels.Num() > 0)
@@ -5489,7 +5489,7 @@ bool MovieSceneToolHelpers::CollapseSection(TSharedPtr<ISequencer>& SequencerPtr
 						const int StartIndex = 0;
 						const int EndIndex = BaseFloatChannels.Num() - 1;
 						MovieSceneToolHelpers::MergeSections<FMovieSceneFloatChannel>(BaseSection,
-							Section,StartIndex, EndIndex, Range);
+							Section,StartIndex, EndIndex, Range, TrackSections);
 					}
 				}
 
diff --git a/Engine/Source/Editor/MovieSceneTools/Public/MovieSceneToolHelpers.h b/Engine/Source/Editor/MovieSceneTools/Public/MovieSceneToolHelpers.h
--- a/Engine/Source/Editor/MovieSceneTools/Public/MovieSceneToolHelpers.h
+++ b/Engine/Source/Editor/MovieSceneTools/Public/MovieSceneToolHelpers.h
@@ -120,7 +120,6 @@ enum class FChannelMergeAlgorithm : uint8
 class MOVIESCENETOOLS_API MovieSceneToolHelpers
 {
 public:
-
 	/**
 	 * Trim section at the given time
 	 *
@@ -769,6 +768,87 @@ public:
 		return Values;
 	}
 
+
+	/*
+	* Merge the value based upon the specified merge algorithm using cached parameters
+	* For the algorithm to be relaible it should be done per section blending type.
+	* @param InOutValue  Current Value you want to merge with the other sections.
+	* @param Channels  Set of channels of this type that we will merge together onto the first specified one.
+	* @param Frame  Current Frame to merge
+	* @param SectionChannelIndex  The index of the channel of this type in the sections.
+	* @param Sections Set of corresponding sections for each channel specified.
+	* @param AbsoluteSections Set of absoluate sections for each channel specified.
+	* @param OtherSections Set of non-absolute sections for each channel specified.
+	* @param OverrideChanneIndex ChannelIndex for an override merge.
+	* @param MergeAlgorithm The algorithm to use to blend the channels together.
+	* @return Returns true if successful.
+	*/
+	template<typename ChannelType, typename CurveValueType>
+	static bool MergeValue(CurveValueType& InOutValue, TArray<ChannelType*>& Channels, const FFrameNumber& Frame, int32 SectionChannelIndex, const  TArray<UMovieSceneSection*>& Sections, 
+		const TArray<UMovieSceneSection*>& AbsoluteSections, const TArray<UMovieSceneSection*>& OtherSections, int32 OverrideChannelIndex, FChannelMergeAlgorithm MergeAlgorithm)
+	{
+		const FFrameTime FrameTime(Frame);
+		if (MergeAlgorithm == FChannelMergeAlgorithm::Average)
+		{
+			double DNumChannels = 0.0;
+			for (int32 WeightIndex = 0; WeightIndex < Sections.Num(); ++WeightIndex)
+			{
+				if (Sections[WeightIndex]->GetRange().Contains(Frame))
+				{
+					float Weight = Sections[WeightIndex]->GetTotalWeightValue(FrameTime);
+					CurveValueType WeightedValue = 0.0;
+					ChannelType* EachChannel = Channels[WeightIndex];
+					EachChannel->Evaluate(FrameTime, WeightedValue);
+					WeightedValue *= ((double)Weight);
+					InOutValue += WeightedValue;
+					DNumChannels += 1.0;
+				}
+			}
+			if (DNumChannels > 0.0) //should always happen since base(0) at least be here
+			{
+				InOutValue /= DNumChannels;
+			}
+		}
+		else if (MergeAlgorithm == FChannelMergeAlgorithm::Add)
+		{
+			for (int32 WeightIndex = 0; WeightIndex < Sections.Num(); ++WeightIndex)
+			{
+				if (Sections[WeightIndex]->GetRange().Contains(Frame))
+				{
+					float Weight = Sections[WeightIndex]->GetTotalWeightValue(FrameTime);
+					if (Sections[WeightIndex]->GetBlendType().IsValid() == false ||
+						Sections[WeightIndex]->GetBlendType().Get() != EMovieSceneBlendType::Additive)
+					{
+						Weight = 1.0;
+					}
+					CurveValueType WeightedValue = 0.0;
+					ChannelType* EachChannel = Channels[WeightIndex];
+					EachChannel->Evaluate(FrameTime, WeightedValue);
+					WeightedValue *= Weight;
+					InOutValue += WeightedValue;
+				}
+			}
+		}
+		else if (MergeAlgorithm == FChannelMergeAlgorithm::Override)
+		{
+			//when doing an override merge we need to just get the full value since the new layer will also be
+			//an override layer
+			if (Sections[1]->GetRange().Contains(Frame))
+			{
+				TArray<CurveValueType> ChannelValues = MovieSceneToolHelpers::GetChannelValues<ChannelType,
+					CurveValueType>(OverrideChannelIndex, OverrideChannelIndex, OtherSections, AbsoluteSections, Frame);
+				if (ChannelValues.Num() == 1)
+				{
+					InOutValue = ChannelValues[0];
+				}
+			}
+		}
+		else
+		{
+			return false;
+		}
+		return true;
+	}
 	/*
 	* Merge the set of passed in channels from each section at the specified section channel index.
 	* For the algorithm to be relaible it should be done per section blending type.
@@ -777,10 +857,13 @@ public:
 	* @param Sections Set of corresponding sections for each channel specified.
 	* @param Range Range of time over which we will merge
 	* @param MergeAlgorithm The algorithm to use to blend the channels together
+	* @param TrackSections Sections in this track
+	* @param Increment Optional increment index, if specfied we bake over the range using the increment value, otherwise by default we just merge onto the existing keys
 	* @return Returns true if successful
-	*/	template<typename ChannelType>
+	*/	
+	template<typename ChannelType>
 	static bool MergeChannels(int32 SectionChannelIndex, TArray<ChannelType*>& Channels, const  TArray<UMovieSceneSection*>& Sections, const TRange<FFrameNumber>& Range,
-		FChannelMergeAlgorithm MergeAlgorithm)
+		FChannelMergeAlgorithm MergeAlgorithm, const TArray<UMovieSceneSection*>& TrackSections, const int32* Increment = nullptr)
 	{
 		if (Channels.Num() < 2 || Channels.Num() != Sections.Num())
 		{
@@ -801,108 +884,83 @@ public:
 		int32 OverrideChannelIndex = SectionChannelIndex;
 		int32 NumChannelsInSection = INDEX_NONE;
 		if (MergeAlgorithm == FChannelMergeAlgorithm::Override) //if we are overriding we need to get the absolute value at the frame!
-		{
-			if (UMovieSceneTrack* OwnerTrack = Sections[1]->GetTypedOuter<UMovieSceneTrack>())
-			{
-				TArray<UMovieSceneSection*> TrackSections = OwnerTrack->GetAllSections();
-				int32 SectionIndex = TrackSections.Find(Sections[1]);
-				if (SectionIndex != INDEX_NONE)
-				{
-					TrackSections.SetNum(SectionIndex + 1); //this wil make sure we include the current sections
-					MovieSceneToolHelpers::SplitSectionsByBlendType(EMovieSceneBlendType::Absolute,TrackSections, OtherSections, AbsoluteSections);
-					TArrayView<ChannelType*> OverrideChannels = Sections[1]->GetChannelProxy().GetChannels<ChannelType>();
-					NumChannelsInSection = OverrideChannels.Num();
-				}
-			}
+	{
+			MovieSceneToolHelpers::SplitSectionsByBlendType(EMovieSceneBlendType::Absolute,TrackSections, OtherSections, AbsoluteSections);
+			TArrayView<ChannelType*> OverrideChannels = Sections[1]->GetChannelProxy().GetChannels<ChannelType>();
+			NumChannelsInSection = OverrideChannels.Num();
 		}
-
 		//base channel we set values on
 		ChannelType* BaseChannel = Channels[0];
 		TMovieSceneChannelData<ChannelValueType> BaseChannelData = BaseChannel->GetData();
-		//iterate over each key
-		TArray<FFrameNumber> KeyTimes;
-		TArray<FKeyHandle> Handles;
-		//cached set that we set at the end
-		TArray<TPair< FFrameNumber, ChannelValueType>> KeysToSet;
-		for (int32 ChannelIndex = 0; ChannelIndex < Channels.Num(); ++ChannelIndex)
+		//similar but slightly different implementations based upon whether or not doing per key or by frame.
+		//for performance we do this outside of the loop
+		if (Increment != nullptr && *Increment > 0)
 		{
-			KeyTimes.Reset();
-			Handles.Reset();
-			ChannelType* Channel = Channels[ChannelIndex];
-			Channel->GetKeys(Range, &KeyTimes, &Handles);
-			for (int32 FrameIndex = 0; FrameIndex < KeyTimes.Num(); ++FrameIndex)
+			//iterate over each frame
+			if (UMovieScene* MovieScene = Sections[0]->GetTypedOuter<UMovieScene>())
 			{
-				const FFrameNumber& Frame = KeyTimes[FrameIndex];
-				if (Sections[0]->GetRange().Contains(Frame) == false)  //frame is outside base range so skip
-				{
-					continue;
-				}
-				const FFrameTime FrameTime(Frame);
-				int32 KeyIndex = Channel->GetData().GetIndex(Handles[FrameIndex]);
-				ChannelValueType Value = Channel->GetData().GetValues()[KeyIndex];
-				//got value with tangents and times, now we perform the operation
-				Value.Value = 0.0; //zero out the value we calculate it 
-				if (MergeAlgorithm == FChannelMergeAlgorithm::Average)
+				TArray<FFrameNumber> Frames;
+				CalculateFramesBetween(MovieScene, Range.GetLowerBoundValue(), Range.GetUpperBoundValue(), *Increment, Frames);
+				TArray<TPair< FFrameNumber, CurveValueType>> ValuesToSet;
+				for (int32 ChannelIndex = 0; ChannelIndex < Channels.Num(); ++ChannelIndex)
 				{
-					double DNumChannels = 0.0;
-					for (int32 WeightIndex = 0; WeightIndex < Sections.Num(); ++WeightIndex)
+					ChannelType* Channel = Channels[ChannelIndex];
+					for (const FFrameNumber& Frame : Frames)
 					{
-						if (Sections[WeightIndex]->GetRange().Contains(Frame))
+						if (Sections[0]->GetRange().Contains(Frame) == false)  //frame is outside base range so skip
 						{
-							float Weight = Sections[WeightIndex]->GetTotalWeightValue(FrameTime);
-							CurveValueType WeightedValue = 0.0;
-							ChannelType* EachChannel = Channels[WeightIndex];
-							EachChannel->Evaluate(FrameTime, WeightedValue);
-							WeightedValue *= ((double)Weight);
-							Value.Value += WeightedValue;
-							DNumChannels += 1.0;
+							continue;
 						}
-					}
-					if (DNumChannels > 0.0) //should always happen since base(0) at least be here
-					{
-						Value.Value /= DNumChannels;
+						CurveValueType Value = 0.0;
+						const FFrameTime FrameTime(Frame);
+						Channel->Evaluate(FrameTime, Value);
+
+						MergeValue(Value, Channels, Frame, SectionChannelIndex, Sections, AbsoluteSections,
+							OtherSections, OverrideChannelIndex, MergeAlgorithm);
+						ValuesToSet.Add(TPair<FFrameNumber, CurveValueType>(Frame, Value));
 					}
 				}
-				else if (MergeAlgorithm == FChannelMergeAlgorithm::Add)
+				for (TPair<FFrameNumber, CurveValueType>& ValueToSet : ValuesToSet)
 				{
-					for (int32 WeightIndex = 0; WeightIndex < Sections.Num(); ++WeightIndex)
-					{
-						if (Sections[WeightIndex]->GetRange().Contains(Frame))
-						{
-							float Weight = Sections[WeightIndex]->GetTotalWeightValue(FrameTime);
-							if (Sections[WeightIndex]->GetBlendType().IsValid() == false ||
-								Sections[WeightIndex]->GetBlendType().Get() != EMovieSceneBlendType::Additive)
-							{
-								Weight = 1.0;
-							}
-							CurveValueType WeightedValue = 0.0;
-							ChannelType* EachChannel = Channels[WeightIndex];
-							EachChannel->Evaluate(FrameTime, WeightedValue);
-							WeightedValue *= Weight;
-							Value.Value += WeightedValue;
-						}
-					}
+					MovieSceneToolHelpers::SetOrAddKey(BaseChannelData, ValueToSet.Key, ValueToSet.Value);
 				}
-				else if (MergeAlgorithm == FChannelMergeAlgorithm::Override)
+			}
+		}
+		else
+		{
+			//iterate over each key
+			TArray<FFrameNumber> KeyTimes;
+			TArray<FKeyHandle> Handles;
+			//cached set that we set at the end
+			TArray<TPair< FFrameNumber, ChannelValueType>> KeysToSet;
+			for (int32 ChannelIndex = 0; ChannelIndex < Channels.Num(); ++ChannelIndex)
+			{
+				KeyTimes.Reset();
+				Handles.Reset();
+				ChannelType* Channel = Channels[ChannelIndex];
+				Channel->GetKeys(Range, &KeyTimes, &Handles);
+				for (int32 FrameIndex = 0; FrameIndex < KeyTimes.Num(); ++FrameIndex)
 				{
-					//when doing an override merge we need to just get the full value since the new layer will also be
-					//an override layer
-					if (Sections[1]->GetRange().Contains(Frame))
+					const FFrameNumber& Frame = KeyTimes[FrameIndex];
+					if (Sections[0]->GetRange().Contains(Frame) == false)  //frame is outside base range so skip
 					{
-						TArray<CurveValueType> ChannelValues = MovieSceneToolHelpers::GetChannelValues<ChannelType,
-							CurveValueType>(OverrideChannelIndex, OverrideChannelIndex, OtherSections, AbsoluteSections, Frame);
-						if (ChannelValues.Num() == 1)
-						{
-							Value.Value = ChannelValues[0];
-						}
+						continue;
 					}
+					const FFrameTime FrameTime(Frame);
+					int32 KeyIndex = Channel->GetData().GetIndex(Handles[FrameIndex]);
+					ChannelValueType Value = Channel->GetData().GetValues()[KeyIndex];
+					//got value with tangents and times, now we perform the operation
+					Value.Value = 0.0; //zero out the value we calculate it 
+
+					MergeValue(Value.Value, Channels, Frame, SectionChannelIndex, Sections, AbsoluteSections,
+						OtherSections, OverrideChannelIndex, MergeAlgorithm);
+					KeysToSet.Add(TPair<FFrameNumber, ChannelValueType>(Frame, Value));
 				}
-				KeysToSet.Add(TPair<FFrameNumber, ChannelValueType>(Frame, Value));
 			}
-		}
-		for (TPair<FFrameNumber, ChannelValueType>& KeyToSet : KeysToSet)
-		{
-			MovieSceneToolHelpers::SetOrAddKey(BaseChannelData, KeyToSet.Key, KeyToSet.Value);
+			for (TPair<FFrameNumber, ChannelValueType>& KeyToSet : KeysToSet)
+			{
+				MovieSceneToolHelpers::SetOrAddKey(BaseChannelData, KeyToSet.Key, KeyToSet.Value);
+			}
 		}
 		return true;
 	}
@@ -914,13 +972,14 @@ public:
 	* @param TopSection  The section we will merge from
 	* @param StartIndex Start Index of the channels we are merging
 	* @param EndIndex End Index of the channels we are merging, should be no greater than the number of channels -1
-	* @param Range Range of time over which we will merge
-	* Control Rig sections and don't want to merge the last weight float channel
+	* @param Range Range of time over which we will merge  Control Rig sections and don't want to merge the last weight float channel
+	* @param TrackSections Sections in this track
+	* @param Increment Optional increment index, if specfied we bake over the range using the increment value, otherwise by default we just merge onto the existing keys
 	* @return Returns true if successful
 	*/	
 	template<typename ChannelType>
 	static bool MergeSections(UMovieSceneSection* BaseSection, UMovieSceneSection* TopSection,
-		int32 StartIndex, int32 EndIndex, const TRange<FFrameNumber>& Range)
+		int32 StartIndex, int32 EndIndex, const TRange<FFrameNumber>& Range, const TArray<UMovieSceneSection*>& TrackSections, const int32* Increment = nullptr)
 	{
 		TArrayView<ChannelType*> BaseChannels = BaseSection->GetChannelProxy().GetChannels<ChannelType>();
 		TArrayView<ChannelType*> TopChannels = TopSection->GetChannelProxy().GetChannels<ChannelType>();
@@ -935,8 +994,6 @@ public:
 			return false;
 		}
 
-		BaseSection->Modify();
-
 		TArray<ChannelType*> Channels;
 		int32 ChannelIndex = 0;
 		TArray<UMovieSceneSection*> Sections;
@@ -962,7 +1019,7 @@ public:
 			Channels.Add(TopChannels[ChannelIndex]);
 
 			if (MergeChannels(ChannelIndex, Channels, Sections, Range,
-				MergeAlgorithm) == false)
+				MergeAlgorithm, TrackSections, Increment) == false)
 			{
 				UE_LOG(LogMovieScene, Warning, TEXT("MergeSections:: Could not merge channels"));
 				return false;
@@ -985,13 +1042,13 @@ public:
 	* @param AddditveSections  Set of Additive functions to merge
 	* @param StartIndex Start Index of the channels we are merging
 	* @param EndIndex End Index of the channels we are merging, should be no greater than the number of channels -1
-	* @param Range Range of time over which we will merge
-	* Control Rig sections and don't want to merge the last weight float channel
+	* @param Range Range of time over which we will merge Control Rig sections and don't want to merge the last weight float channel
+	* @param Increment Optional increment index, if specfied we bake over the range using the increment value, otherwise by default we just merge onto the existing keys
 	* @return Returns true if successful
 	*/	
 	template<typename ChannelType>
 	static bool MergeSections(UMovieSceneSection* BaseSection, TArray<UMovieSceneSection*>& AbsoluteSections, TArray<UMovieSceneSection*>& AdditiveSections,
-		int32 StartIndex, int32 EndIndex, const TRange<FFrameNumber>& Range)
+		int32 StartIndex, int32 EndIndex, const TRange<FFrameNumber>& Range, const int32* Increment = nullptr)
 	{
 		TArrayView<ChannelType*> BaseChannels = BaseSection->GetChannelProxy().GetChannels<ChannelType>();
 		if (StartIndex < 0 || EndIndex >= BaseChannels.Num() || EndIndex < StartIndex)
@@ -1021,9 +1078,20 @@ public:
 					return false;
 				}
 			}
+			TArray<UMovieSceneSection*> TrackSections;
+			if (UMovieSceneTrack* OwnerTrack = BaseSection->GetTypedOuter<UMovieSceneTrack>())
+			{
+				TrackSections = OwnerTrack->GetAllSections();
+			}
+			else
+			{
+				UE_LOG(LogMovieScene, Warning, TEXT("MergeSections:: No Owner Track"));
+				return false;
+			}
+			TArray<ChannelType*> Channels;
+			
 			BaseSection->Modify();
 
-			TArray<ChannelType*> Channels;
 			for (int32 ChannelIndex = StartIndex; ChannelIndex <= EndIndex; ++ChannelIndex)
 			{
 				if (AbsoluteSections.Num() > 0)
@@ -1036,7 +1104,7 @@ public:
 					}
 					//now blend them
 					if (MergeChannels(ChannelIndex, Channels, AbsoluteSections, Range,
-						FChannelMergeAlgorithm::Average) == false)
+						FChannelMergeAlgorithm::Average, TrackSections, Increment) == false)
 					{
 						UE_LOG(LogMovieScene, Warning, TEXT("MergeSections:: Could not merge channels"));
 						return false;
@@ -1053,7 +1121,7 @@ public:
 					}
 					//now blend them
 					if (MergeChannels(ChannelIndex,Channels, AdditiveSections, Range,
-						FChannelMergeAlgorithm::Add) == false)
+						FChannelMergeAlgorithm::Add, TrackSections, Increment) == false)
 					{
 						UE_LOG(LogMovieScene, Warning, TEXT("MergeSections:: Could not merge channels"));
 						return false;
