Fix Typo

diff --git a/Engine/Plugins/Animation/ControlRig/Source/ControlRigEditor/Private/EditMode/ControlRigEditModeCommands.cpp b/Engine/Plugins/Animation/ControlRig/Source/ControlRigEditor/Private/EditMode/ControlRigEditModeCommands.cpp
--- a/Engine/Plugins/Animation/ControlRig/Source/ControlRigEditor/Private/EditMode/ControlRigEditModeCommands.cpp
+++ b/Engine/Plugins/Animation/ControlRig/Source/ControlRigEditor/Private/EditMode/ControlRigEditModeCommands.cpp
@@ -15,7 +15,7 @@ void FControlRigEditModeCommands::RegisterCommands()
 	UI_COMMAND(ToggleManipulators, "Toggle Shapes", "Toggles visibility of active control rig shapes in the viewport", EUserInterfaceActionType::Button, FInputChord(EKeys::T));
 	UI_COMMAND(ToggleModuleManipulators, "Toggle Module Shapes", "Toggles visibility of active control rig shapes belonging to selected modules in the viewport", EUserInterfaceActionType::Button, FInputChord(EKeys::T, EModifierKey::Control));
 	UI_COMMAND(ToggleAllManipulators, "Toggle All Shapes", "Toggles visibility of all control rig shapes in the viewport", EUserInterfaceActionType::Button, FInputChord(EKeys::T, EModifierKey::Control | EModifierKey::Alt));
-	UI_COMMAND(ToggleControlsAsOverlay, "Toggle Contraols As Overlay", "Toggles showing controls as overlay", EUserInterfaceActionType::Button, FInputChord());
+	UI_COMMAND(ToggleControlsAsOverlay, "Toggle Controls As Overlay", "Toggles showing controls as overlay", EUserInterfaceActionType::Button, FInputChord());
 	UI_COMMAND(FrameSelection, "Frame Selection", "Focus the viewport on the current selection", EUserInterfaceActionType::Button, FInputChord(EKeys::F));
 	UI_COMMAND(ClearSelection, "Clear Selection", "Clear Selection", EUserInterfaceActionType::Button, FInputChord(EKeys::Escape));
 
