Anim Layers: Fix LOCTEXT clashes.

#rnx

diff --git a/Engine/Plugins/Animation/ControlRig/Source/ControlRigEditor/Private/Sequencer/AnimLayers/SAnimLayers.cpp b/Engine/Plugins/Animation/ControlRig/Source/ControlRigEditor/Private/Sequencer/AnimLayers/SAnimLayers.cpp
--- a/Engine/Plugins/Animation/ControlRig/Source/ControlRigEditor/Private/Sequencer/AnimLayers/SAnimLayers.cpp
+++ b/Engine/Plugins/Animation/ControlRig/Source/ControlRigEditor/Private/Sequencer/AnimLayers/SAnimLayers.cpp
@@ -1011,7 +1011,7 @@ void FAnimLayerSourcesView::CreateSourcesListView()
 			.DefaultLabel(LOCTEXT("LayerColumnName", "Name"))
 			+ SHeaderRow::Column(AnimLayerSourceListUI::ActionColumnName)
 			.FillWidth(30.f)
-			.DefaultLabel(LOCTEXT("LayerColumnName", "Action"))
+			.DefaultLabel(LOCTEXT("ActionColumnName", "Action"))
 			+ SHeaderRow::Column(AnimLayerSourceListUI::StatusColumnName)
 			.FillWidth(40.f)
 			.DefaultLabel(LOCTEXT("StatusColumnName", "Status"))
