Control Rig: Add tag filter to the module browser tab

diff --git a/Engine/Plugins/Animation/ControlRig/Source/ControlRigEditor/Private/Editor/SRigModuleAssetBrowser.cpp b/Engine/Plugins/Animation/ControlRig/Source/ControlRigEditor/Private/Editor/SRigModuleAssetBrowser.cpp
--- a/Engine/Plugins/Animation/ControlRig/Source/ControlRigEditor/Private/Editor/SRigModuleAssetBrowser.cpp
+++ b/Engine/Plugins/Animation/ControlRig/Source/ControlRigEditor/Private/Editor/SRigModuleAssetBrowser.cpp
@@ -9,11 +9,15 @@
 #include "UObject/AssetRegistryTagsContext.h"
 
 #include "ControlRigEditor.h"
+#include "FrontendFilterBase.h"
+#include "Editor/RigVMEditorTools.h"
 #include "Widgets/Layout/SUniformGridPanel.h"
 #include "Widgets/SRigVMVariantWidget.h"
 
 #define LOCTEXT_NAMESPACE "RigModuleAssetBrowser"
 
+typedef UE::RigVM::Editor::Tools::FFilterByAssetTag FRigVMFilterTag; 
+
 void SRigModuleAssetBrowser::Construct(
 	const FArguments& InArgs,
 	TSharedRef<IControlRigBaseEditor> InEditor)
@@ -71,8 +75,27 @@ void SRigModuleAssetBrowser::RefreshView()
 	// allow to open the rigs directly on double click
 	AssetPickerConfig.OnAssetDoubleClicked = FOnAssetDoubleClicked::CreateSP(this, &SRigModuleAssetBrowser::OnAssetDoubleClicked);
 
+	TSharedRef<FFrontendFilterCategory>	ControlRigFilterCategory = MakeShared<FFrontendFilterCategory>(LOCTEXT("ControlRigFilterCategoryName", "Control Rig Tags"), LOCTEXT("ControlRigFilterCategoryToolTip", "Filter ControlRigs by variant tags specified in ControlRig Blueprint class settings"));
+	const URigVMProjectSettings* Settings = GetDefault<URigVMProjectSettings>(URigVMProjectSettings::StaticClass());
+	TArray<FRigVMTag> AvailableTags = Settings->VariantTags;
+
+	TArray<TSharedRef<FRigVMFilterTag>> Filters;
+	for (const FRigVMTag& CurTag : AvailableTags)
+	{
+		if (CurTag.bShowInUserInterface)
+		{
+			Filters.Add(MakeShared<FRigVMFilterTag>(ControlRigFilterCategory, CurTag));
+			AssetPickerConfig.ExtraFrontendFilters.Add(Filters.Last());
+		}
+	}
+
 	const FContentBrowserModule& ContentBrowserModule = FModuleManager::Get().LoadModuleChecked<FContentBrowserModule>(TEXT("ContentBrowser"));
 	AssetBrowserBox->SetContent(ContentBrowserModule.Get().CreateAssetPicker(AssetPickerConfig));
+
+	for (TSharedRef<FRigVMFilterTag> Filter : Filters)
+	{
+		Filter->SetActive(Filter->ShouldBeMarkedAsInvalid());
+	}
 }
 
 TSharedPtr<SWidget> SRigModuleAssetBrowser::OnGetAssetContextMenu(const TArray<FAssetData>& SelectedAssets) const
@@ -136,23 +159,6 @@ bool SRigModuleAssetBrowser::OnShouldFilterAsset(const struct FAssetData& AssetD
 		return true;
 	}
 
-	static const FName AssetVariantPropertyName = GET_MEMBER_NAME_CHECKED(URigVMBlueprint, AssetVariant);
-	const FProperty* AssetVariantProperty = CastField<FProperty>(URigVMBlueprint::StaticClass()->FindPropertyByName(AssetVariantPropertyName));
-	const FString VariantStr = AssetData.GetTagValueRef<FString>(AssetVariantPropertyName);
-	if(!VariantStr.IsEmpty())
-	{
-		FRigVMVariant AssetVariant;
-		AssetVariantProperty->ImportText_Direct(*VariantStr, &AssetVariant, nullptr, EPropertyPortFlags::PPF_None);
-
-		for(const FRigVMTag& VariantTag : AssetVariant.Tags)
-		{
-			if(VariantTag.bMarksSubjectAsInvalid)
-			{
-				return true;
-			}
-		}
-	}
-
 	return false;
 }
 
diff --git a/Engine/Plugins/Animation/ControlRig/Source/ControlRigEditor/Private/Sequencer/ControlRigParameterTrackEditor.cpp b/Engine/Plugins/Animation/ControlRig/Source/ControlRigEditor/Private/Sequencer/ControlRigParameterTrackEditor.cpp
--- a/Engine/Plugins/Animation/ControlRig/Source/ControlRigEditor/Private/Sequencer/ControlRigParameterTrackEditor.cpp
+++ b/Engine/Plugins/Animation/ControlRig/Source/ControlRigEditor/Private/Sequencer/ControlRigParameterTrackEditor.cpp
@@ -113,6 +113,7 @@
 #include "Sequencer/SLoadAnimToControlRig.h"
 #include "FrontendFilterBase.h"
 #include "InstancedPropertyBagStructureDataProvider.h"
+#include "Editor/RigVMEditorTools.h"
 
 #define LOCTEXT_NAMESPACE "FControlRigParameterTrackEditor"
 
@@ -257,63 +258,6 @@ static bool DoesControlRigAllowMultipleInstances(const FTopLevelAssetPath& InGen
 	return false;
 }
 
-
-class FFrontendFilter_ControlRigFilterByAssetTag: public FFrontendFilter
-{
-public:
-	FFrontendFilter_ControlRigFilterByAssetTag(TSharedPtr<FFrontendFilterCategory> InCategory, const FRigVMTag& InTag) :
-		FFrontendFilter(InCategory),
-		Tag(InTag)
-	{}
-
-	FString GetName() const override { return Tag.bMarksSubjectAsInvalid ? TEXT("Exclude ") + Tag.Name.ToString() : Tag.Name.ToString(); }
-	FText GetDisplayName() const override { return FText::FromString(Tag.bMarksSubjectAsInvalid ? TEXT("Exclude ") + Tag.GetLabel() : Tag.GetLabel()); }
-	FText GetToolTipText() const override { return Tag.ToolTip; }
-	FLinearColor GetColor() const override { return Tag.Color; };
-	
-	bool PassesFilter(const FContentBrowserItem& InItem) const override
-	{
-		FAssetData AssetData;
-		if (InItem.Legacy_TryGetAssetData(AssetData))
-		{
-			static const FName AssetVariantPropertyName = GET_MEMBER_NAME_CHECKED(URigVMBlueprint, AssetVariant);
-			const FProperty* AssetVariantProperty = CastField<FProperty>(URigVMBlueprint::StaticClass()->FindPropertyByName(AssetVariantPropertyName));
-			const FString VariantStr = AssetData.GetTagValueRef<FString>(AssetVariantPropertyName);
-			if(!VariantStr.IsEmpty())
-			{
-				FRigVMVariant AssetVariant;
-				AssetVariantProperty->ImportText_Direct(*VariantStr, &AssetVariant, nullptr, EPropertyPortFlags::PPF_None);
-
-				if (Tag.bMarksSubjectAsInvalid)
-				{
-					if (!AssetVariant.Tags.Contains(Tag))
-					{
-						return true;
-					}
-				}
-				else
-				{
-					if (AssetVariant.Tags.Contains(Tag))
-					{
-						return true;
-					}	
-				}
-			}
-			else
-			{
-				if (Tag.bMarksSubjectAsInvalid)
-				{
-					return true;
-				}
-			}
-		}
-		return false;
-	}
-
-private:
-	FRigVMTag Tag;
-};
-
 bool FControlRigParameterTrackEditor::bAutoGenerateControlRigTrack = true;
 FCriticalSection FControlRigParameterTrackEditor::ControlUndoTransactionMutex;
 
@@ -1879,7 +1823,7 @@ void FControlRigParameterTrackEditor::HandleAddTrackSubMenu(FMenuBuilder& MenuBu
 
 			for (const FRigVMTag& Tag : AvailableTags)
 			{
-				AssetPickerConfig.ExtraFrontendFilters.Add(MakeShared<FFrontendFilter_ControlRigFilterByAssetTag>(ControlRigFilterCategory, Tag));
+				AssetPickerConfig.ExtraFrontendFilters.Add(MakeShared<UE::RigVM::Editor::Tools::FFilterByAssetTag>(ControlRigFilterCategory, Tag));
 			}
 
 			// This is so that we can remove the "Other Filters" section easily
diff --git a/Engine/Plugins/Runtime/RigVM/Source/RigVMDeveloper/Public/RigVMBlueprint.h b/Engine/Plugins/Runtime/RigVM/Source/RigVMDeveloper/Public/RigVMBlueprint.h
--- a/Engine/Plugins/Runtime/RigVM/Source/RigVMDeveloper/Public/RigVMBlueprint.h
+++ b/Engine/Plugins/Runtime/RigVM/Source/RigVMDeveloper/Public/RigVMBlueprint.h
@@ -24,6 +24,10 @@ class URigVMBlueprintGeneratedClass;
 
 #if WITH_EDITOR
 class IRigVMEditorModule;
+namespace UE::RigVM::Editor::Tools
+{
+	class FFilterByAssetTag;
+}
 #endif
 struct FEndLoadPackageContext;
 struct FRigVMMemoryStorageStruct;
@@ -864,7 +868,7 @@ private:
 	friend class SRigVMGraphNode;
 	friend struct FGuardSkipDirtyBlueprintStatus;
 	friend class SRigModuleAssetBrowser;
-	friend class FFrontendFilter_ControlRigFilterByAssetTag;
+	friend class UE::RigVM::Editor::Tools::FFilterByAssetTag;
 };
 
 class RIGVMDEVELOPER_API FRigVMBlueprintCompileScope
diff --git a/Engine/Plugins/Runtime/RigVM/Source/RigVMEditor/Private/Editor/RigVMEditorTools.cpp b/Engine/Plugins/Runtime/RigVM/Source/RigVMEditor/Private/Editor/RigVMEditorTools.cpp
--- a/Engine/Plugins/Runtime/RigVM/Source/RigVMEditor/Private/Editor/RigVMEditorTools.cpp
+++ b/Engine/Plugins/Runtime/RigVM/Source/RigVMEditor/Private/Editor/RigVMEditorTools.cpp
@@ -158,5 +158,48 @@ RIGVMEDITOR_API FAssetData FindAssetFromAnyPath(const FString& InPartialOrFullPa
 	return EditorAssetSubsystem->FindAssetData(InPartialOrFullPath);
 }
 
+FFilterByAssetTag::FFilterByAssetTag(TSharedPtr<FFrontendFilterCategory> InCategory, const FRigVMTag& InTag):
+	FFrontendFilter(InCategory), Tag(InTag)
+{}
+
+bool FFilterByAssetTag::PassesFilter(const FContentBrowserItem& InItem) const
+{
+	FAssetData AssetData;
+	if (InItem.Legacy_TryGetAssetData(AssetData))
+	{
+		static const FName AssetVariantPropertyName = GET_MEMBER_NAME_CHECKED(URigVMBlueprint, AssetVariant);
+		const FProperty* AssetVariantProperty = CastField<FProperty>(URigVMBlueprint::StaticClass()->FindPropertyByName(AssetVariantPropertyName));
+		const FString VariantStr = AssetData.GetTagValueRef<FString>(AssetVariantPropertyName);
+		if(!VariantStr.IsEmpty())
+		{
+			FRigVMVariant AssetVariant;
+			AssetVariantProperty->ImportText_Direct(*VariantStr, &AssetVariant, nullptr, EPropertyPortFlags::PPF_None);
+
+			if (Tag.bMarksSubjectAsInvalid)
+			{
+				if (!AssetVariant.Tags.Contains(Tag))
+				{
+					return true;
+				}
+			}
+			else
+			{
+				if (AssetVariant.Tags.Contains(Tag))
+				{
+					return true;
+				}
+			}
+		}
+		else
+		{
+			if (Tag.bMarksSubjectAsInvalid)
+			{
+				return true;
+			}
+		}
+	}
+	return false;
+}
+
 
 } // end namespace UE::RigVM::Editor::Tools
diff --git a/Engine/Plugins/Runtime/RigVM/Source/RigVMEditor/Public/Editor/RigVMEditorTools.h b/Engine/Plugins/Runtime/RigVM/Source/RigVMEditor/Public/Editor/RigVMEditorTools.h
--- a/Engine/Plugins/Runtime/RigVM/Source/RigVMEditor/Public/Editor/RigVMEditorTools.h
+++ b/Engine/Plugins/Runtime/RigVM/Source/RigVMEditor/Public/Editor/RigVMEditorTools.h
@@ -4,6 +4,8 @@
 
 #include "CoreMinimal.h"
 #include "AssetRegistry/AssetData.h"
+#include "FrontendFilterBase.h"
+#include "RigVMCore/RigVMVariant.h"
 
 class URigVMController;
 class URigVMGraph;
@@ -39,4 +41,21 @@ RIGVMEDITOR_API void OnRequestLocalizeFunctionDialog(FRigVMGraphFunctionIdentifi
 	
 RIGVMEDITOR_API FAssetData FindAssetFromAnyPath(const FString& InPartialOrFullPath, bool bConvertToRootPath);
 
+class RIGVMEDITOR_API FFilterByAssetTag : public FFrontendFilter
+{
+public:
+	FFilterByAssetTag(TSharedPtr<FFrontendFilterCategory> InCategory, const FRigVMTag& InTag);
+
+	FString GetName() const override { return Tag.bMarksSubjectAsInvalid ? TEXT("Exclude ") + Tag.Name.ToString() : Tag.Name.ToString(); }
+	FText GetDisplayName() const override { return FText::FromString(Tag.bMarksSubjectAsInvalid ? TEXT("Exclude ") + Tag.GetLabel() : Tag.GetLabel()); }
+	FText GetToolTipText() const override { return Tag.ToolTip; }
+	FLinearColor GetColor() const override { return Tag.Color; };
+
+	bool PassesFilter(const FContentBrowserItem& InItem) const override;
+	bool ShouldBeMarkedAsInvalid() const { return Tag.bMarksSubjectAsInvalid; }
+
+private:
+	FRigVMTag Tag;
+};
+
 } // end namespace UE::RigVM::Editor::Tools
diff --git a/Engine/Plugins/Runtime/RigVM/Source/RigVMEditor/RigVMEditor.Build.cs b/Engine/Plugins/Runtime/RigVM/Source/RigVMEditor/RigVMEditor.Build.cs
--- a/Engine/Plugins/Runtime/RigVM/Source/RigVMEditor/RigVMEditor.Build.cs
+++ b/Engine/Plugins/Runtime/RigVM/Source/RigVMEditor/RigVMEditor.Build.cs
@@ -33,6 +33,7 @@ public class RigVMEditor : ModuleRules
                 "MessageLog",
 				"StructUtilsEditor",
 				"ContentBrowser",
+				"ContentBrowserData",
 				"EditorFramework"
 			}
 		);
