Control Rig: Add gizmo interaction events to ControlRigEditMode, and remove them from ControlRig

diff --git a/Engine/Plugins/Animation/ControlRig/Source/ControlRig/Public/ControlRig.h b/Engine/Plugins/Animation/ControlRig/Source/ControlRig/Public/ControlRig.h
--- a/Engine/Plugins/Animation/ControlRig/Source/ControlRig/Public/ControlRig.h
+++ b/Engine/Plugins/Animation/ControlRig/Source/ControlRig/Public/ControlRig.h
@@ -747,30 +747,6 @@ protected:
 	UPROPERTY(BlueprintAssignable, Category = ControlRig, meta=(DisplayName="OnControlSelected"))
 	FOnControlSelectedBP OnControlSelected_BP;
 
-#if WITH_EDITORONLY_DATA
-
-	/** Broadcasts a notification when a gizmo manipulation has started */
-	UPROPERTY(BlueprintReadOnly, Category = ControlRig, meta = (ScriptName = "OnGizmoInteractionStarted", AllowPrivateAccess = "true"))
-	FControlRigInteractionStartedEvent OnGizmoInteractionStartedEvent;
-
-	/** Broadcasts a notification when a gizmo manipulation has ended */
-	UPROPERTY(BlueprintReadOnly, Category = ControlRig, meta = (ScriptName = "OnGizmoInteractionEnded", AllowPrivateAccess = "true"))
-	FControlRigInteractionEndedEvent OnGizmoInteractionEndedEvent;
-
-	/** Broadcasts a notification when a gizmo manipulation has been updated (before the update actually takes place) */
-	UPROPERTY(BlueprintReadOnly, Category = ControlRig, meta = (ScriptName = "OnGizmoInteractionUpdated", AllowPrivateAccess = "true"))
-	FControlRigInteractionUpdatedEvent OnGizmoInteractionUpdatedEvent;
-
-public:
-	/** Broadcasts a notification when a gizmo manipulation has started */
-	FControlRigInteractionStartedEvent& OnGizmoInteractionStarted() { return OnGizmoInteractionStartedEvent; }
-	/** Broadcasts a notification when a gizmo manipulation has ended */
-	FControlRigInteractionEndedEvent& OnGizmoInteractionEnded() { return OnGizmoInteractionEndedEvent; }
-	/** Broadcasts a notification when a gizmo manipulation has been updated (before the update actually takes place) */
-	FControlRigInteractionUpdatedEvent& OnGizmoInteractionUpdated() { return OnGizmoInteractionUpdatedEvent; }
-protected:
-#endif
-
 	TArray<FRigElementKey> QueuedModifiedControls;
 
 private:
@@ -1073,12 +1049,6 @@ public:
 		, SyncBracketScope(InControlRig->InterRigSyncBracket)
 	{
 		InControlRig->GetHierarchy()->StartInteraction();
-#if WITH_EDITORONLY_DATA
-		if (InControlRig->OnGizmoInteractionStarted().IsBound())
-		{
-			InControlRig->OnGizmoInteractionStarted().Broadcast(InControlRig, InControlRig->ElementsBeingInteracted, EControlRigInteractionType::All);
-		}
-#endif
 	}
 
 	FControlRigInteractionScope(
@@ -1094,12 +1064,6 @@ public:
 		InControlRig->InteractionType = (uint8)InInteractionType;
 		InControlRig->bInteractionJustBegan = true;
 		InControlRig->GetHierarchy()->StartInteraction();
-#if WITH_EDITORONLY_DATA
-		if (InControlRig->OnGizmoInteractionStarted().IsBound())
-		{
-			InControlRig->OnGizmoInteractionStarted().Broadcast(InControlRig, InControlRig->ElementsBeingInteracted, InInteractionType);
-		}
-#endif
 	}
 
 	FControlRigInteractionScope(
@@ -1115,24 +1079,12 @@ public:
 		InControlRig->InteractionType = (uint8)InInteractionType;
 		InControlRig->bInteractionJustBegan = true;
 		InControlRig->GetHierarchy()->StartInteraction();
-#if WITH_EDITORONLY_DATA
-		if (InControlRig->OnGizmoInteractionStarted().IsBound())
-		{
-			InControlRig->OnGizmoInteractionStarted().Broadcast(InControlRig, InControlRig->ElementsBeingInteracted, InInteractionType);
-		}
-#endif
 	}
 
 	~FControlRigInteractionScope()
 	{
 		if(ensure(ControlRig.IsValid()))
 		{
-#if WITH_EDITORONLY_DATA
-			if (ControlRig->OnGizmoInteractionEnded().IsBound())
-			{
-				ControlRig->OnGizmoInteractionEnded().Broadcast(ControlRig.Get(), ControlRig->ElementsBeingInteracted);
-			}
-#endif
 			ControlRig->GetHierarchy()->EndInteraction();
 			ControlRig->InteractionType = (uint8)EControlRigInteractionType::None;
 			ControlRig->bInteractionJustBegan = false;
diff --git a/Engine/Plugins/Animation/ControlRig/Source/ControlRigEditor/Private/EditMode/ControlRigEditMode.cpp b/Engine/Plugins/Animation/ControlRig/Source/ControlRigEditor/Private/EditMode/ControlRigEditMode.cpp
--- a/Engine/Plugins/Animation/ControlRig/Source/ControlRigEditor/Private/EditMode/ControlRigEditMode.cpp
+++ b/Engine/Plugins/Animation/ControlRig/Source/ControlRigEditor/Private/EditMode/ControlRigEditMode.cpp
@@ -1206,6 +1206,22 @@ bool FControlRigEditMode::HandleBeginTransform(const FEditorViewportClient* InVi
 
 	if (!AreEditingControlRigDirectly())
 	{
+		if (OnGizmoInteractionStartedDelegate.IsBound())
+		{
+			FMultiControlRigElementSelection Selection;
+			Selection.Rigs = RuntimeControlRigs;
+			for (TWeakObjectPtr<UControlRig>& RuntimeRigPtr : RuntimeControlRigs)
+			{
+				if (UControlRig* ControlRig = RuntimeRigPtr.Get())
+				{
+					if (AreRigElementSelectedAndMovable(ControlRig))
+					{
+						Selection.KeysPerRig.Add(GetSelectedRigElements(ControlRig));
+					}
+				}
+			}
+			OnGizmoInteractionStartedDelegate.Broadcast(Selection, (EControlRigInteractionType)InteractionType);
+		}
 		for (TWeakObjectPtr<UControlRig>& RuntimeRigPtr : RuntimeControlRigs)
 		{
 			if (UControlRig* ControlRig = RuntimeRigPtr.Get())
@@ -1228,6 +1244,14 @@ bool FControlRigEditMode::HandleBeginTransform(const FEditorViewportClient* InVi
 	else if(UControlRigEditorSettings::Get()->bEnableUndoForPoseInteraction)
 	{
 		UControlRig* ControlRig = RuntimeControlRigs[0].Get();
+		if (OnGizmoInteractionStartedDelegate.IsBound())
+		{
+			FMultiControlRigElementSelection Selection;
+			Selection.Rigs = {ControlRig};
+			Selection.KeysPerRig.Add(GetSelectedRigElements(ControlRig));
+			
+			OnGizmoInteractionStartedDelegate.Broadcast(Selection, (EControlRigInteractionType)InteractionType);
+		}
 		TArray<FRigElementKey> SelectedRigElements = GetSelectedRigElements(ControlRig);
 		FControlRigInteractionScope* InteractionScope = new FControlRigInteractionScope(ControlRig, SelectedRigElements, (EControlRigInteractionType)InteractionType);
 		InteractionScopes.Add(ControlRig,InteractionScope);
@@ -1264,6 +1288,22 @@ bool FControlRigEditMode::HandleEndTransform(FEditorViewportClient* InViewportCl
 		TArray<TWeakObjectPtr<UControlRig>> RigsToTick;
 		RigsToTick.Reserve(InteractionScopes.Num());
 
+		if (OnGizmoInteractionEndedDelegate.IsBound())
+		{
+			FMultiControlRigElementSelection Selection;
+			TArray<UControlRig*> Rigs;
+			InteractionScopes.GetKeys(Rigs);
+			Selection.Rigs.Append(Rigs);
+			for (UControlRig* ControlRig : Rigs)
+			{
+				if (AreRigElementSelectedAndMovable(ControlRig))
+				{
+					Selection.KeysPerRig.Add(GetSelectedRigElements(ControlRig));
+				}
+			}
+			OnGizmoInteractionEndedDelegate.Broadcast(Selection);
+		}
+		
 		for (TPair<UControlRig*, FControlRigInteractionScope*>& RigInteractionScope : InteractionScopes)
 		{
 			if (UControlRig* ControlRig = RigInteractionScope.Key)
@@ -2026,6 +2066,28 @@ bool FControlRigEditMode::InputDelta(FEditorViewportClient* InViewportClient, FV
 			// reset the dependency cache as the hierarchy might have changed since the previous frame
 			InteractionDependencies.Reset();
 		}
+
+		// The interaction update event must be broadcasted before setting any control value or executing the rig
+		if (OnGizmoInteractionUpdatedDelegate.IsBound())
+		{
+			FMultiControlRigElementSelection MultiRigSelection;
+			for (TPair<UControlRig*, TArray<TObjectPtr<AControlRigShapeActor>>>& Pairs : ControlRigShapeActors)
+			{
+				MultiRigSelection.Rigs.Add(Pairs.Key);
+				FRigElementKeyCollection& Keys = MultiRigSelection.KeysPerRig.Add_GetRef({});
+				if (AreRigElementsSelected(ValidControlTypeMask(), Pairs.Key))
+				{
+					for (TObjectPtr<AControlRigShapeActor> ShapeActor : Pairs.Value)
+					{
+						if (ShapeActor->IsEnabled() && ShapeActor->IsSelected())
+						{
+							Keys.Add(ShapeActor->GetElementKey());
+						}
+					}
+				}
+			}
+			OnGizmoInteractionUpdatedDelegate.Broadcast(MultiRigSelection, TransformContext);
+		}
 		
 		for (auto& Pairs : ControlRigShapeActors)
 		{
@@ -2076,9 +2138,6 @@ bool FControlRigEditMode::InputDelta(FEditorViewportClient* InViewportClient, FV
 								GEditor->BeginTransaction(LOCTEXT("MoveControlTransaction", "Move Control"));
 							}
 
-							// The interaction update event must be broadcasted before setting any control value or executing the rig
-							ControlRig->OnGizmoInteractionUpdated().Broadcast(ControlRig, ControlRig->ElementsBeingInteracted, TransformContext);
-
 							// cache interaction dependencies + evaluate the rig at least once before manipulating anything
 							if (ControlRig->ElementsBeingInteracted.IsEmpty() && InteractionScopes.Contains(ControlRig))
 							{
diff --git a/Engine/Plugins/Animation/ControlRig/Source/ControlRigEditor/Public/EditMode/ControlRigEditMode.h b/Engine/Plugins/Animation/ControlRig/Source/ControlRigEditor/Public/EditMode/ControlRigEditMode.h
--- a/Engine/Plugins/Animation/ControlRig/Source/ControlRigEditor/Public/EditMode/ControlRigEditMode.h
+++ b/Engine/Plugins/Animation/ControlRig/Source/ControlRigEditor/Public/EditMode/ControlRigEditMode.h
@@ -129,6 +129,22 @@ private:
 
 };
 
+USTRUCT()
+struct CONTROLRIGEDITOR_API FMultiControlRigElementSelection
+{
+	GENERATED_BODY()
+
+	UPROPERTY()
+	TArray<TWeakObjectPtr<UControlRig>> Rigs;
+
+	UPROPERTY()
+	TArray<FRigElementKeyCollection> KeysPerRig;
+};
+
+DECLARE_DYNAMIC_MULTICAST_DELEGATE_TwoParams(FControlRigEditModeInteractionStartedEvent, FMultiControlRigElementSelection, InteractionKeys, EControlRigInteractionType, InteractionType);
+DECLARE_DYNAMIC_MULTICAST_DELEGATE_OneParam(FControlRigEditModeInteractionEndedEvent, FMultiControlRigElementSelection, InteractionKeys);
+DECLARE_DYNAMIC_MULTICAST_DELEGATE_TwoParams(FControlRigEditModeInteractionUpdatedEvent, FMultiControlRigElementSelection, InteractionKeys, FControlRigInteractionTransformContext, TransformContext);
+
 class CONTROLRIGEDITOR_API FControlRigEditMode : public IPersonaEditMode
 {
 public:
@@ -214,6 +230,15 @@ public:
 	/* Control Rig Selected Delegate*/
 	FControlRigSelected& OnControlRigSelected() { return OnControlRigSelectedDelegate; }
 
+	/** Broadcasts a notification when a gizmo manipulation has started */
+	FControlRigEditModeInteractionStartedEvent& OnGizmoInteractionStarted() { return OnGizmoInteractionStartedDelegate; }
+
+	/** Broadcasts a notification when a gizmo manipulation has ended */
+	FControlRigEditModeInteractionEndedEvent& OnGizmoInteractionEnded() { return OnGizmoInteractionEndedDelegate; }
+
+	/** Broadcasts a notification when a gizmo manipulation has been updated (before the update actually takes place) */
+	FControlRigEditModeInteractionUpdatedEvent& OnGizmoInteractionUpdated() { return OnGizmoInteractionUpdatedDelegate; }
+
 	// callback that gets called when rig element is selected in other view
 	void OnHierarchyModified(ERigHierarchyNotification InNotif, URigHierarchy* InHierarchy, const FRigNotificationSubject& InSubject);
 	void OnHierarchyModified_AnyThread(ERigHierarchyNotification InNotif, URigHierarchy* InHierarchy, const FRigNotificationSubject& InSubject);
@@ -433,6 +458,15 @@ protected:
 	FControlRigAddedOrRemoved OnControlRigAddedOrRemovedDelegate;
 	FControlRigSelected OnControlRigSelectedDelegate;
 
+	/** Broadcasts a notification when a gizmo manipulation has started */
+	FControlRigEditModeInteractionStartedEvent OnGizmoInteractionStartedDelegate;
+
+	/** Broadcasts a notification when a gizmo manipulation has ended */
+	FControlRigEditModeInteractionEndedEvent OnGizmoInteractionEndedDelegate;
+
+	/** Broadcasts a notification when a gizmo manipulation has been updated (before the update actually takes place) */
+	FControlRigEditModeInteractionUpdatedEvent OnGizmoInteractionUpdatedDelegate;
+
 	/** GetSelectedRigElements, if InControlRig is nullptr get the first one */
 	TArray<FRigElementKey> GetSelectedRigElements() const;
 	static TArray<FRigElementKey> GetSelectedRigElements(UControlRig* InControlRig);
