Control Rig: Sync controls visibility between ControlRigEditMode and SControlRigOutliner
#rnx

diff --git a/Engine/Plugins/Animation/ControlRig/Source/ControlRigEditor/Private/EditMode/ControlRigBaseDockableView.cpp b/Engine/Plugins/Animation/ControlRig/Source/ControlRigEditor/Private/EditMode/ControlRigBaseDockableView.cpp
--- a/Engine/Plugins/Animation/ControlRig/Source/ControlRigEditor/Private/EditMode/ControlRigBaseDockableView.cpp
+++ b/Engine/Plugins/Animation/ControlRig/Source/ControlRigEditor/Private/EditMode/ControlRigBaseDockableView.cpp
@@ -22,6 +22,7 @@ FControlRigBaseDockableView::~FControlRigBaseDockableView()
 	{
 		EditMode->OnControlRigAddedOrRemoved().RemoveAll(this);
 		EditMode->OnControlRigSelected().RemoveAll(this);
+		EditMode->OnControlRigVisibilityChanged().RemoveAll(this);
 	}
 }
 
@@ -43,6 +44,7 @@ void FControlRigBaseDockableView::SetEditMode(FControlRigEditMode& InEditMode)
 	{
 		EditMode->OnControlRigAddedOrRemoved().RemoveAll(this);
 		EditMode->OnControlRigAddedOrRemoved().AddRaw(this, &FControlRigBaseDockableView::HandleControlAdded);
+		EditMode->OnControlRigVisibilityChanged().AddRaw(this, &FControlRigBaseDockableView::HandleRigVisibilityChanged);
 		EditMode->OnControlRigSelected().RemoveAll(this);
 		EditMode->OnControlRigSelected().AddRaw(this, &FControlRigBaseDockableView::HandlElementSelected);
 	}
@@ -52,6 +54,10 @@ void FControlRigBaseDockableView::HandleControlAdded(UControlRig* ControlRig, bo
 {
 }
 
+void FControlRigBaseDockableView::HandleRigVisibilityChanged(UControlRig* control_rig)
+{
+}
+
 void FControlRigBaseDockableView::HandleControlSelected(UControlRig* ControlRig, FRigControlElement* InControl, bool bSelected)
 {
 }
diff --git a/Engine/Plugins/Animation/ControlRig/Source/ControlRigEditor/Private/EditMode/ControlRigBaseDockableView.h b/Engine/Plugins/Animation/ControlRig/Source/ControlRigEditor/Private/EditMode/ControlRigBaseDockableView.h
--- a/Engine/Plugins/Animation/ControlRig/Source/ControlRigEditor/Private/EditMode/ControlRigBaseDockableView.h
+++ b/Engine/Plugins/Animation/ControlRig/Source/ControlRigEditor/Private/EditMode/ControlRigBaseDockableView.h
@@ -27,6 +27,7 @@ public:
 protected:
 	virtual void HandleControlSelected(UControlRig* Subject, FRigControlElement* InControl, bool bSelected);
 	virtual void HandleControlAdded(UControlRig* ControlRig, bool bIsAdded);
+	virtual void HandleRigVisibilityChanged(UControlRig* control_rig);
 
 	void HandlElementSelected(UControlRig* Subject, const FRigElementKey& Key, bool bSelected);
 
diff --git a/Engine/Plugins/Animation/ControlRig/Source/ControlRigEditor/Private/EditMode/ControlRigEditMode.cpp b/Engine/Plugins/Animation/ControlRig/Source/ControlRigEditor/Private/EditMode/ControlRigEditMode.cpp
--- a/Engine/Plugins/Animation/ControlRig/Source/ControlRigEditor/Private/EditMode/ControlRigEditMode.cpp
+++ b/Engine/Plugins/Animation/ControlRig/Source/ControlRigEditor/Private/EditMode/ControlRigEditMode.cpp
@@ -216,6 +216,7 @@ FControlRigEditMode::~FControlRigEditMode()
 	DestroyShapesActors(nullptr);
 	OnControlRigAddedOrRemovedDelegate.Clear();
 	OnControlRigSelectedDelegate.Clear();
+	OnControlRigVisibilityChangedDelegate.Clear();
 
 	TArray<TWeakObjectPtr<UControlRig>> PreviousRuntimeRigs = RuntimeControlRigs;
 	for (int32 PreviousRuntimeRigIndex = 0; PreviousRuntimeRigIndex < PreviousRuntimeRigs.Num(); PreviousRuntimeRigIndex++)
@@ -547,6 +548,7 @@ void FControlRigEditMode::Exit()
 	ClearOutAnyActiveTools();
 	OnControlRigAddedOrRemovedDelegate.Clear();
 	OnControlRigSelectedDelegate.Clear();
+	OnControlRigVisibilityChangedDelegate.Clear();
 	for (TWeakObjectPtr<UControlRig>& RuntimeRigPtr : RuntimeControlRigs)
 	{
 		if (UControlRig* ControlRig = RuntimeRigPtr.Get())
@@ -3517,6 +3519,10 @@ void FControlRigEditMode::ToggleManipulators()
 				FScopedTransaction ScopedTransaction(LOCTEXT("ToggleControlsVisibility", "Toggle Controls Visibility"),!GIsTransacting);
 				ControlRig->Modify();
 				ControlRig->ToggleControlsVisible();
+				if (OnControlRigVisibilityChangedDelegate.IsBound())
+				{
+					OnControlRigVisibilityChangedDelegate.Broadcast(ControlRig);
+				}
 			}
 		}
 	}
@@ -3558,6 +3564,11 @@ void FControlRigEditMode::ToggleModuleManipulators()
 								}
 							}
 						}
+
+						if (OnControlRigVisibilityChangedDelegate.IsBound())
+						{
+							OnControlRigVisibilityChangedDelegate.Broadcast(ControlRig);
+						}
 					}
 				}
 			}
diff --git a/Engine/Plugins/Animation/ControlRig/Source/ControlRigEditor/Private/EditMode/SControlRigOutliner.cpp b/Engine/Plugins/Animation/ControlRig/Source/ControlRigEditor/Private/EditMode/SControlRigOutliner.cpp
--- a/Engine/Plugins/Animation/ControlRig/Source/ControlRigEditor/Private/EditMode/SControlRigOutliner.cpp
+++ b/Engine/Plugins/Animation/ControlRig/Source/ControlRigEditor/Private/EditMode/SControlRigOutliner.cpp
@@ -267,6 +267,34 @@ void FMultiRigTreeElement::RefreshDisplaySettings(const URigHierarchy* InHierarc
 	}
 
 }
+
+bool FMultiRigTreeElement::AreControlsVisible() const
+{
+	switch (Data.Type)
+	{
+		case FMultiRigData::EMultiRigDataType_ControlRig:
+		{
+			if (TStrongObjectPtr<UControlRig> Rig = Data.WeakControlRig.Pin())
+			{
+				return Rig->GetControlsVisible();
+			}
+			break;
+		}
+		case FMultiRigData::EMultiRigDataType_Module:
+		{
+			if (FRigModuleInstance* Module = Data.GetModuleInstance())
+			{
+				if (UControlRig* ModuleRig = Module->GetRig())
+				{
+					return ModuleRig->GetControlsVisible();
+				}
+			}
+			break;
+		}
+	}
+	return false;
+}
+
 //////////////////////////////////////////////////////////////
 /// SMultiRigHierarchyItem
 ///////////////////////////////////////////////////////////
@@ -363,7 +391,6 @@ FReply SMultiRigHierarchyItem::OnToggleVisibilityClicked()
 						bSetControlsVisible = !Rig->GetControlsVisible();
 					}
 
-					ElementToToggle->bAreControlsVisible = bSetControlsVisible.GetValue();
 					RigsToToggle.AddUnique(Rig);
 				}
 			}
@@ -506,7 +533,7 @@ TSharedRef<SWidget> SMultiRigHierarchyItem::GenerateWidgetForColumn(const FName&
 				{
 					if (TSharedPtr<FMultiRigTreeElement> Element = WeakRigTreeElement.Pin())
 					{
-						if (Element->bAreControlsVisible && !Element->bIsRowHovered && !Element->bIsEyeballIconHovered)
+						if (Element->AreControlsVisible() && !Element->bIsRowHovered && !Element->bIsEyeballIconHovered)
 						{
 							return FSlateColor(FLinearColor::Transparent);
 						}
@@ -1906,6 +1933,12 @@ void SControlRigOutliner::HandleControlSelected(UControlRig* Subject, FRigContro
 	}
 }
 
+void SControlRigOutliner::HandleRigVisibilityChanged(UControlRig* control_rig)
+{
+	FControlRigBaseDockableView::HandleRigVisibilityChanged(control_rig);
+	HierarchyTreeView->GetTreeView()->RequestTreeViewRefresh();
+}
+
 void SControlRigOutliner::HandleSelectionChanged(TSharedPtr<FMultiRigTreeElement> Selection, ESelectInfo::Type SelectInfo)
 {
 	if (bIsChangingRigHierarchy)
diff --git a/Engine/Plugins/Animation/ControlRig/Source/ControlRigEditor/Private/EditMode/SControlRigOutliner.h b/Engine/Plugins/Animation/ControlRig/Source/ControlRigEditor/Private/EditMode/SControlRigOutliner.h
--- a/Engine/Plugins/Animation/ControlRig/Source/ControlRigEditor/Private/EditMode/SControlRigOutliner.h
+++ b/Engine/Plugins/Animation/ControlRig/Source/ControlRigEditor/Private/EditMode/SControlRigOutliner.h
@@ -147,6 +147,7 @@ public:
 	TSharedRef<ITableRow> MakeTreeRowWidget(const TSharedRef<STableViewBase>& InOwnerTable, TSharedRef<FMultiRigTreeElement> InRigTreeElement, TSharedPtr<SMultiRigHierarchyTreeView> InTreeView, const FRigTreeDisplaySettings& InSettings, bool bPinned);
 
 	void RefreshDisplaySettings(const URigHierarchy* InHierarchy, const FRigTreeDisplaySettings& InSettings);
+	bool AreControlsVisible() const;
 
 	/** The current filter result */
 	ERigTreeFilterResult FilterResult;
@@ -160,9 +161,6 @@ public:
 	/** The color to use when rendering the label text */
 	FSlateColor TextColor;
 
-	/** Whether or not the controls of this element are visible */
-	bool bAreControlsVisible = true;
-
 	/** Whether or not this row is being hovered */
 	bool bIsRowHovered = false;
 
@@ -347,6 +345,7 @@ class SControlRigOutliner : public FControlRigBaseDockableView, public SCompound
 private:
 	virtual void HandleControlAdded(UControlRig* ControlRig, bool bIsAdded) override;
 	virtual void HandleControlSelected(UControlRig* Subject, FRigControlElement* InControl, bool bSelected) override;
+	virtual void HandleRigVisibilityChanged(UControlRig* control_rig) override;
 
 	//control rig delegates
 	void HandleOnControlRigBound(UControlRig* InControlRig);
diff --git a/Engine/Plugins/Animation/ControlRig/Source/ControlRigEditor/Public/EditMode/ControlRigEditMode.h b/Engine/Plugins/Animation/ControlRig/Source/ControlRigEditor/Public/EditMode/ControlRigEditMode.h
--- a/Engine/Plugins/Animation/ControlRig/Source/ControlRigEditor/Public/EditMode/ControlRigEditMode.h
+++ b/Engine/Plugins/Animation/ControlRig/Source/ControlRigEditor/Public/EditMode/ControlRigEditMode.h
@@ -62,6 +62,7 @@ DECLARE_DELEGATE_ThreeParams(FOnSetRigElementTransform, const FRigElementKey& /*
 DECLARE_DELEGATE_RetVal(TSharedPtr<FUICommandList>, FNewMenuCommandsDelegate);
 DECLARE_MULTICAST_DELEGATE_TwoParams(FControlRigAddedOrRemoved, UControlRig*, bool /*true if added, false if removed*/);
 DECLARE_MULTICAST_DELEGATE_ThreeParams(FControlRigSelected, UControlRig*, const FRigElementKey& /*RigElementKey*/,const bool /*bIsSelected*/);
+DECLARE_MULTICAST_DELEGATE_OneParam(FOnControlRigVisibilityChanged, UControlRig*);
 DECLARE_DELEGATE_RetVal(UToolMenu*, FOnGetContextMenu);
 
 class FControlRigEditMode;
@@ -230,6 +231,9 @@ public:
 	/* Control Rig Selected Delegate*/
 	FControlRigSelected& OnControlRigSelected() { return OnControlRigSelectedDelegate; }
 
+	/* Control Rig Visibility Delegate*/
+	FOnControlRigVisibilityChanged& OnControlRigVisibilityChanged() { return OnControlRigVisibilityChangedDelegate; }
+
 	/** Broadcasts a notification when a gizmo manipulation has started */
 	FControlRigEditModeInteractionStartedEvent& OnGizmoInteractionStarted() { return OnGizmoInteractionStartedDelegate; }
 
@@ -457,6 +461,7 @@ protected:
 	FSimpleMulticastDelegate OnAnimSystemInitializedDelegate;
 	FControlRigAddedOrRemoved OnControlRigAddedOrRemovedDelegate;
 	FControlRigSelected OnControlRigSelectedDelegate;
+	FOnControlRigVisibilityChanged OnControlRigVisibilityChangedDelegate;
 
 	/** Broadcasts a notification when a gizmo manipulation has started */
 	FControlRigEditModeInteractionStartedEvent OnGizmoInteractionStartedDelegate;
