Control Rig: Fix Control Rig viewport submenu not visible in the new control rig editor
#rnx

diff --git a/Engine/Plugins/Animation/ControlRig/Source/ControlRigEditor/Private/Editor/ControlRigEditor.cpp b/Engine/Plugins/Animation/ControlRig/Source/ControlRigEditor/Private/Editor/ControlRigEditor.cpp
--- a/Engine/Plugins/Animation/ControlRig/Source/ControlRigEditor/Private/Editor/ControlRigEditor.cpp
+++ b/Engine/Plugins/Animation/ControlRig/Source/ControlRigEditor/Private/Editor/ControlRigEditor.cpp
@@ -162,6 +162,18 @@ UControlRigBlueprint* GetControlRigBlueprintFromContext(const UUnrealEdViewportT
 					StaticCastSharedPtr<FAnimationViewportClient>(Viewport->GetViewportClient()))
 			{
 #if WITH_RIGVMLEGACYEDITOR
+				// Check if we are editing UControlRigBlueprints, otherwise this is not a ControlRigEditor
+				if (const TArray<UObject*>* EditingObjects = AnimViewportClient->GetAssetEditorToolkit()->GetObjectsCurrentlyBeingEdited())
+				{
+					for (UObject* Object : *EditingObjects)
+					{
+						if (!Object->IsA<UControlRigBlueprint>())
+						{
+							return nullptr;
+						}
+					}
+				}
+
 				if (AnimViewportClient->GetAssetEditorToolkit()->IsBlueprintEditor())
 				{
 					const TSharedRef<FBlueprintEditor>& BlueprintEditor =
@@ -169,14 +181,11 @@ UControlRigBlueprint* GetControlRigBlueprintFromContext(const UUnrealEdViewportT
 
 					return Cast<UControlRigBlueprint>(BlueprintEditor->GetBlueprintObj());
 				}
-#else
-				// TODO sara-s Figure out another way to make sure this is a ControlRigEditor
-				// Possibly getting the EditingObjects from the AssetEditorToolkit and checking if its a ControlRigBlueprint
+#endif
 				const TSharedRef<FControlRigEditor>& Editor =
 					StaticCastSharedRef<FControlRigEditor>(AnimViewportClient->GetAssetEditorToolkit());
 
 				return Cast<UControlRigBlueprint>(Editor->GetRigVMBlueprint());
-#endif
 			}
 		}
 	}
