Control Rig dynamic hierarchy: remove useless ensure
- follow up to 41311744

#rnx

diff --git a/Engine/Plugins/Animation/ControlRig/Source/ControlRig/Private/Units/Execution/RigUnit_DynamicHierarchy.cpp b/Engine/Plugins/Animation/ControlRig/Source/ControlRig/Private/Units/Execution/RigUnit_DynamicHierarchy.cpp
--- a/Engine/Plugins/Animation/ControlRig/Source/ControlRig/Private/Units/Execution/RigUnit_DynamicHierarchy.cpp
+++ b/Engine/Plugins/Animation/ControlRig/Source/ControlRig/Private/Units/Execution/RigUnit_DynamicHierarchy.cpp
@@ -403,7 +403,6 @@ FRigUnit_HierarchyImportFromSkeleton_Execute()
 			}
 		}
 	}
-	ensure(Items.Num() > 0);
 }
 
 FRigUnit_HierarchyRemoveElement_Execute()
