ControlRig Rotation Order: fixed rotator's preferred angle computation if the rotation order is not the default one (XYZ)

- using FRigPreferredEulerAngles::SetRotator was suggesting that the incoming value was already ordered, which was not the case

#rnx

diff --git a/Engine/Plugins/Animation/ControlRig/Source/ControlRig/Private/Rigs/RigHierarchy.cpp b/Engine/Plugins/Animation/ControlRig/Source/ControlRig/Private/Rigs/RigHierarchy.cpp
--- a/Engine/Plugins/Animation/ControlRig/Source/ControlRig/Private/Rigs/RigHierarchy.cpp
+++ b/Engine/Plugins/Animation/ControlRig/Source/ControlRig/Private/Rigs/RigHierarchy.cpp
@@ -5827,14 +5827,24 @@ FRigControlValue URigHierarchy::GetControlValue(FRigControlElement* InControlEle
 	return Value;
 }
 
-void URigHierarchy::SetPreferredEulerAnglesFromValue(FRigControlElement* InControlElement, const FRigControlValue& InValue, const ERigControlValueType& InValueType, const bool bFixEulerFlips)
+void URigHierarchy::SetPreferredEulerAnglesFromValue(FRigControlElement* InControlElement,
+	const FRigControlValue& InValue,
+	const ERigControlValueType& InValueType,
+	const bool bFixEulerFlips) const
 {
 	const bool bInitial = InValueType == ERigControlValueType::Initial;
+
+	FRigPreferredEulerAngles& PreferredEulerAngles = InControlElement->PreferredEulerAngles;
+	const EEulerRotationOrder RotationOrder = PreferredEulerAngles.RotationOrder;
+	
 	switch(InControlElement->Settings.ControlType)
 	{
 	case ERigControlType::Rotator:
 		{
-			InControlElement->PreferredEulerAngles.SetRotator(GetRotatorFromControlValue(InValue), bInitial, bFixEulerFlips);
+			const FVector3f EulerXYZf = InValue.Get<FVector3f>();
+			const FVector EulerXYZ(EulerXYZf.X, EulerXYZf.Y, EulerXYZf.Z);
+			const FVector Angle = AnimationCore::ChangeEulerRotationOrder(EulerXYZ, EEulerRotationOrder::XYZ, RotationOrder);
+			PreferredEulerAngles.SetAngles(Angle, bInitial, RotationOrder, bFixEulerFlips);
 			break;
 		}
 	case ERigControlType::EulerTransform:
@@ -5842,7 +5852,7 @@ void URigHierarchy::SetPreferredEulerAnglesFromValue(FRigControlElement* InContr
 			FEulerTransform EulerTransform = GetEulerTransformFromControlValue(InValue);
 			FQuat Quat = EulerTransform.GetRotation();
 			const FVector Angle = GetControlAnglesFromQuat(InControlElement, Quat, bFixEulerFlips);
-			InControlElement->PreferredEulerAngles.SetAngles(Angle, bInitial, InControlElement->PreferredEulerAngles.RotationOrder, bFixEulerFlips);
+			PreferredEulerAngles.SetAngles(Angle, bInitial, RotationOrder, bFixEulerFlips);
 			break;
 		}
 	case ERigControlType::Transform:
@@ -5850,7 +5860,7 @@ void URigHierarchy::SetPreferredEulerAnglesFromValue(FRigControlElement* InContr
 			FTransform Transform = GetTransformFromControlValue(InValue);
 			FQuat Quat = Transform.GetRotation();
 			const FVector Angle = GetControlAnglesFromQuat(InControlElement, Quat, bFixEulerFlips);
-			InControlElement->PreferredEulerAngles.SetAngles(Angle, bInitial, InControlElement->PreferredEulerAngles.RotationOrder, bFixEulerFlips);
+			PreferredEulerAngles.SetAngles(Angle, bInitial, RotationOrder, bFixEulerFlips);
 			break;
 		}
 	case ERigControlType::TransformNoScale:
@@ -5858,7 +5868,7 @@ void URigHierarchy::SetPreferredEulerAnglesFromValue(FRigControlElement* InContr
 			FTransform Transform = GetTransformNoScaleFromControlValue(InValue);
 			FQuat Quat = Transform.GetRotation();
 			const FVector Angle = GetControlAnglesFromQuat(InControlElement, Quat, bFixEulerFlips);
-			InControlElement->PreferredEulerAngles.SetAngles(Angle, bInitial, InControlElement->PreferredEulerAngles.RotationOrder, bFixEulerFlips);
+			PreferredEulerAngles.SetAngles(Angle, bInitial, RotationOrder, bFixEulerFlips);
 			break;
 		}
 	default:
diff --git a/Engine/Plugins/Animation/ControlRig/Source/ControlRig/Public/Rigs/RigHierarchy.h b/Engine/Plugins/Animation/ControlRig/Source/ControlRig/Public/Rigs/RigHierarchy.h
--- a/Engine/Plugins/Animation/ControlRig/Source/ControlRig/Public/Rigs/RigHierarchy.h
+++ b/Engine/Plugins/Animation/ControlRig/Source/ControlRig/Public/Rigs/RigHierarchy.h
@@ -4392,8 +4392,17 @@ public:
 	 */
 	FRigControlValue GetControlValue(FRigControlElement* InControlElement, ERigControlValueType InValueType, bool bUsePreferredAngles = true) const;
 
-	void SetPreferredEulerAnglesFromValue(FRigControlElement* InControlElement, const FRigControlValue& InValue,
-	                                      const ERigControlValueType& InValueType, bool bFixEulerFlips);
+	/**
+	 * Sets a control's preferred euler angle.
+	 * @param InControlElement The element to set the current value for
+	 * @param InValue The value to ge the angle from to set on the control
+	 * @param InValueType The type of value to set
+	 * @param bFixEulerFlips If true the angle will use the shortest path
+	 */
+	void SetPreferredEulerAnglesFromValue(FRigControlElement* InControlElement,
+		const FRigControlValue& InValue,
+		const ERigControlValueType& InValueType,
+		bool bFixEulerFlips) const;
 
 	template<typename T>
 	T GetControlValue(FRigControlElement* InControlElement, ERigControlValueType InValueType) const
