Control Rig: Fix saving sequencer while control tracks selected freezes editor if the anim outliner is opened
#rnx

diff --git a/Engine/Plugins/Animation/ControlRig/Source/ControlRigEditor/Private/EditMode/SControlRigOutliner.cpp b/Engine/Plugins/Animation/ControlRig/Source/ControlRigEditor/Private/EditMode/SControlRigOutliner.cpp
--- a/Engine/Plugins/Animation/ControlRig/Source/ControlRigEditor/Private/EditMode/SControlRigOutliner.cpp
+++ b/Engine/Plugins/Animation/ControlRig/Source/ControlRigEditor/Private/EditMode/SControlRigOutliner.cpp
@@ -1356,7 +1356,8 @@ void SMultiRigHierarchyTreeView::RefreshTreeView()
 			}
 		}
 	}
-	
+
+	TGuardValue<bool> Guard(Delegates.bIsChangingRigHierarchy, true);
 	ClearSelection();
 
 	for (const TPair<UControlRig*, TArray<FRigElementKey>>& RigAndKeys : RigAndSelection)
