Control Rig: Fix toggle module visibility not working when only selected in the anim outliner
#rnx

diff --git a/Engine/Plugins/Animation/ControlRig/Source/ControlRigEditor/Private/EditMode/ControlRigBaseDockableView.cpp b/Engine/Plugins/Animation/ControlRig/Source/ControlRigEditor/Private/EditMode/ControlRigBaseDockableView.cpp
--- a/Engine/Plugins/Animation/ControlRig/Source/ControlRigEditor/Private/EditMode/ControlRigBaseDockableView.cpp
+++ b/Engine/Plugins/Animation/ControlRig/Source/ControlRigEditor/Private/EditMode/ControlRigBaseDockableView.cpp
@@ -54,7 +54,7 @@ void FControlRigBaseDockableView::HandleControlAdded(UControlRig* ControlRig, bo
 {
 }
 
-void FControlRigBaseDockableView::HandleRigVisibilityChanged(UControlRig* control_rig)
+void FControlRigBaseDockableView::HandleRigVisibilityChanged(TArray<UControlRig*> InControlRigs)
 {
 }
 
diff --git a/Engine/Plugins/Animation/ControlRig/Source/ControlRigEditor/Private/EditMode/ControlRigBaseDockableView.h b/Engine/Plugins/Animation/ControlRig/Source/ControlRigEditor/Private/EditMode/ControlRigBaseDockableView.h
--- a/Engine/Plugins/Animation/ControlRig/Source/ControlRigEditor/Private/EditMode/ControlRigBaseDockableView.h
+++ b/Engine/Plugins/Animation/ControlRig/Source/ControlRigEditor/Private/EditMode/ControlRigBaseDockableView.h
@@ -27,7 +27,7 @@ public:
 protected:
 	virtual void HandleControlSelected(UControlRig* Subject, FRigControlElement* InControl, bool bSelected);
 	virtual void HandleControlAdded(UControlRig* ControlRig, bool bIsAdded);
-	virtual void HandleRigVisibilityChanged(UControlRig* control_rig);
+	virtual void HandleRigVisibilityChanged(TArray<UControlRig*> InControlRigs);
 
 	void HandlElementSelected(UControlRig* Subject, const FRigElementKey& Key, bool bSelected);
 
diff --git a/Engine/Plugins/Animation/ControlRig/Source/ControlRigEditor/Private/EditMode/ControlRigEditMode.cpp b/Engine/Plugins/Animation/ControlRig/Source/ControlRigEditor/Private/EditMode/ControlRigEditMode.cpp
--- a/Engine/Plugins/Animation/ControlRig/Source/ControlRigEditor/Private/EditMode/ControlRigEditMode.cpp
+++ b/Engine/Plugins/Animation/ControlRig/Source/ControlRigEditor/Private/EditMode/ControlRigEditMode.cpp
@@ -3521,7 +3521,7 @@ void FControlRigEditMode::ToggleManipulators()
 				ControlRig->ToggleControlsVisible();
 				if (OnControlRigVisibilityChangedDelegate.IsBound())
 				{
-					OnControlRigVisibilityChangedDelegate.Broadcast(ControlRig);
+					OnControlRigVisibilityChangedDelegate.Broadcast({ControlRig});
 				}
 			}
 		}
@@ -3548,6 +3548,7 @@ void FControlRigEditMode::ToggleModuleManipulators()
 					if (UModularRig* ModularRig = Cast<UModularRig>(ControlRig))
 					{
 						TArray<FString> ModuleNames;
+						TArray<UControlRig*> ChangedRigs;
 						TArray<FRigElementKey> Selected = GetSelectedRigElements(ControlRig);
 						for (const FRigElementKey& Key : Selected)
 						{
@@ -3561,13 +3562,14 @@ void FControlRigEditMode::ToggleModuleManipulators()
 								if (UControlRig* Rig = Module->GetRig())
 								{
 									Rig->ToggleControlsVisible();
+									ChangedRigs.Add(Rig);
 								}
 							}
 						}
 
 						if (OnControlRigVisibilityChangedDelegate.IsBound())
 						{
-							OnControlRigVisibilityChangedDelegate.Broadcast(ControlRig);
+							OnControlRigVisibilityChangedDelegate.Broadcast(ChangedRigs);
 						}
 					}
 				}
diff --git a/Engine/Plugins/Animation/ControlRig/Source/ControlRigEditor/Private/EditMode/SControlRigOutliner.cpp b/Engine/Plugins/Animation/ControlRig/Source/ControlRigEditor/Private/EditMode/SControlRigOutliner.cpp
--- a/Engine/Plugins/Animation/ControlRig/Source/ControlRigEditor/Private/EditMode/SControlRigOutliner.cpp
+++ b/Engine/Plugins/Animation/ControlRig/Source/ControlRigEditor/Private/EditMode/SControlRigOutliner.cpp
@@ -1934,10 +1934,26 @@ void SControlRigOutliner::HandleControlSelected(UControlRig* Subject, FRigContro
 	}
 }
 
-void SControlRigOutliner::HandleRigVisibilityChanged(UControlRig* control_rig)
+void SControlRigOutliner::HandleRigVisibilityChanged(TArray<UControlRig*> InControlRigs)
 {
-	FControlRigBaseDockableView::HandleRigVisibilityChanged(control_rig);
-	HierarchyTreeView->GetTreeView()->RequestTreeViewRefresh();
+	FControlRigBaseDockableView::HandleRigVisibilityChanged(InControlRigs);
+
+	// If only modules or super items are selected, we might have to toggle the visibility ourselves (instead of relying on the edit mode)
+	TArray<FMultiRigData> Selected = HierarchyTreeView->GetTreeView()->GetSelectedData();
+	
+	for (const FMultiRigData& Data : Selected)
+	{
+		if (Data.IsModule())
+		{
+			if (FRigModuleInstance* Module = Data.GetModuleInstance())
+			{
+				if (!InControlRigs.Contains(Module->GetRig()))
+				{
+					Module->GetRig()->ToggleControlsVisible();
+				}
+			}
+		}
+	}
 }
 
 void SControlRigOutliner::HandleSelectionChanged(TSharedPtr<FMultiRigTreeElement> Selection, ESelectInfo::Type SelectInfo)
diff --git a/Engine/Plugins/Animation/ControlRig/Source/ControlRigEditor/Private/EditMode/SControlRigOutliner.h b/Engine/Plugins/Animation/ControlRig/Source/ControlRigEditor/Private/EditMode/SControlRigOutliner.h
--- a/Engine/Plugins/Animation/ControlRig/Source/ControlRigEditor/Private/EditMode/SControlRigOutliner.h
+++ b/Engine/Plugins/Animation/ControlRig/Source/ControlRigEditor/Private/EditMode/SControlRigOutliner.h
@@ -345,7 +345,7 @@ class SControlRigOutliner : public FControlRigBaseDockableView, public SCompound
 private:
 	virtual void HandleControlAdded(UControlRig* ControlRig, bool bIsAdded) override;
 	virtual void HandleControlSelected(UControlRig* Subject, FRigControlElement* InControl, bool bSelected) override;
-	virtual void HandleRigVisibilityChanged(UControlRig* control_rig) override;
+	virtual void HandleRigVisibilityChanged(TArray<UControlRig*> InControlRigs) override;
 
 	//control rig delegates
 	void HandleOnControlRigBound(UControlRig* InControlRig);
diff --git a/Engine/Plugins/Animation/ControlRig/Source/ControlRigEditor/Public/EditMode/ControlRigEditMode.h b/Engine/Plugins/Animation/ControlRig/Source/ControlRigEditor/Public/EditMode/ControlRigEditMode.h
--- a/Engine/Plugins/Animation/ControlRig/Source/ControlRigEditor/Public/EditMode/ControlRigEditMode.h
+++ b/Engine/Plugins/Animation/ControlRig/Source/ControlRigEditor/Public/EditMode/ControlRigEditMode.h
@@ -62,7 +62,7 @@ DECLARE_DELEGATE_ThreeParams(FOnSetRigElementTransform, const FRigElementKey& /*
 DECLARE_DELEGATE_RetVal(TSharedPtr<FUICommandList>, FNewMenuCommandsDelegate);
 DECLARE_MULTICAST_DELEGATE_TwoParams(FControlRigAddedOrRemoved, UControlRig*, bool /*true if added, false if removed*/);
 DECLARE_MULTICAST_DELEGATE_ThreeParams(FControlRigSelected, UControlRig*, const FRigElementKey& /*RigElementKey*/,const bool /*bIsSelected*/);
-DECLARE_MULTICAST_DELEGATE_OneParam(FOnControlRigVisibilityChanged, UControlRig*);
+DECLARE_MULTICAST_DELEGATE_OneParam(FOnControlRigVisibilityChanged, TArray<UControlRig*>);
 DECLARE_DELEGATE_RetVal(UToolMenu*, FOnGetContextMenu);
 
 class FControlRigEditMode;
