ControlRig track: avoid useless calls to Modify that dirties the level sequence unnecessarily

#rnx

diff --git a/Engine/Plugins/Animation/ControlRig/Source/ControlRig/Private/Sequencer/MovieSceneControlRigParameterTrack.cpp b/Engine/Plugins/Animation/ControlRig/Source/ControlRig/Private/Sequencer/MovieSceneControlRigParameterTrack.cpp
--- a/Engine/Plugins/Animation/ControlRig/Source/ControlRig/Private/Sequencer/MovieSceneControlRigParameterTrack.cpp
+++ b/Engine/Plugins/Animation/ControlRig/Source/ControlRig/Private/Sequencer/MovieSceneControlRigParameterTrack.cpp
@@ -452,35 +452,62 @@ void UMovieSceneControlRigParameterTrack::SetSectionToKey(UMovieSceneSection* In
 
 void UMovieSceneControlRigParameterTrack::SetSectionToKey(UMovieSceneSection* InSection)
 {
-	if (Sections.Num() < 1 || InSection == nullptr)
+	if (Sections.IsEmpty())
 	{
 		return;
 	}
-	if (UMovieSceneControlRigParameterSection* CRSection = Cast<UMovieSceneControlRigParameterSection>(InSection))
+
+	UMovieSceneControlRigParameterSection* ControlRigSection = Cast<UMovieSceneControlRigParameterSection>(InSection);
+	if (!ControlRigSection)
 	{
-		Modify();
-		if (SectionToKeyPerControl.Num() > 0) //we have sections that are in layers so need to respect them
+		return;
+	}
+
+	if (!SectionToKeyPerControl.IsEmpty()) //we have sections that are in layers so need to respect them
+	{
+		bool bAlreadyModified = false;
+		auto ModifyThis = [&bAlreadyModified](UMovieSceneControlRigParameterTrack* InTrack)
+		{
+			if (!bAlreadyModified)
+			{
+				InTrack->Modify();
+				bAlreadyModified = true;
+			}
+		};
+		
+		if (bSetSectionToKeyPerControl)
 		{
-			if (bSetSectionToKeyPerControl)
+			for (TPair<FName, TWeakObjectPtr<UMovieSceneSection>>& SectionToKeyItem : SectionToKeyPerControl)
 			{
-				for (TPair < FName, TWeakObjectPtr<UMovieSceneSection>>& SectionToKeyItem : SectionToKeyPerControl)
+				const FName& ControlName = SectionToKeyItem.Key;
+				
+				//only set it as the section to key if it's in that section, otherwise leave it alone
+				if (ControlRigSection->ControlChannelMap.Contains(ControlName))
 				{
-					//only set it as the section to key if it's in that section, otherwise leave it alone
-					FChannelMapInfo* pChannelIndex = CRSection->ControlChannelMap.Find(SectionToKeyItem.Key);
-					if (pChannelIndex)
+					if (ControlRigSection->GetControlNameMask(ControlName))
 					{
-						if (CRSection->GetControlNameMask(SectionToKeyItem.Key))
+						if (InSection != SectionToKeyItem.Value)
 						{
-							SectionToKeyItem.Value = InSection;
+							ModifyThis(this);
+							SectionToKeyItem.Value = ControlRigSection;
 						}
 					}
 				}
 			}
+		}
+
+		if (Sections[0] != SectionToKey)
+		{
+			ModifyThis(this);
 			SectionToKey = Sections[0];
 		}
-		else
+	}
+	else
+	{
+		if (ControlRigSection != SectionToKey)
 		{
-			SectionToKey = InSection;
+			Modify();
+			SectionToKey = ControlRigSection;
 		}
 	}
 }
