[AutoRTFM][ControlRig] Update "AutoRTFM/AutoRTFM.h" include to "AutoRTFM.h"

Now that AutoRTFM has been moved out of Core, the header is no longer in an AutoRTFM directory.

diff --git a/Engine/Plugins/Animation/ControlRig/Source/ControlRig/Private/Rigs/FKControlRig.cpp b/Engine/Plugins/Animation/ControlRig/Source/ControlRig/Private/Rigs/FKControlRig.cpp
--- a/Engine/Plugins/Animation/ControlRig/Source/ControlRig/Private/Rigs/FKControlRig.cpp
+++ b/Engine/Plugins/Animation/ControlRig/Source/ControlRig/Private/Rigs/FKControlRig.cpp
@@ -2,7 +2,7 @@
 
 #include "Rigs/FKControlRig.h"
 #include "Animation/SmartName.h"
-#include "AutoRTFM/AutoRTFM.h"
+#include "AutoRTFM.h"
 #include "Containers/StaticArray.h"
 #include "Engine/SkeletalMesh.h"
 #include "IControlRigObjectBinding.h"
