Modular Control Rig support for Overrides

#rnx

diff --git a/Engine/Plugins/Animation/ControlRig/Source/ControlRig/Private/ControlRigOverride.cpp b/Engine/Plugins/Animation/ControlRig/Source/ControlRig/Private/ControlRigOverride.cpp
new file mode 100644
--- /dev/null
+++ b/Engine/Plugins/Animation/ControlRig/Source/ControlRig/Private/ControlRigOverride.cpp
@@ -0,0 +1,1101 @@
+// Copyright Epic Games, Inc. All Rights Reserved.
+
+#include "ControlRigOverride.h"
+#include "RigVMCore/RigVMMemoryCommon.h"
+
+#if WITH_EDITOR
+#include "PropertyPath.h"
+#endif
+
+#include UE_INLINE_GENERATED_CPP_BY_NAME(ControlRigOverride)
+
+////////////////////////////////////////////////////////////////////////////////
+// FControlRigOverrideValue
+////////////////////////////////////////////////////////////////////////////////
+
+FControlRigOverrideValue::FControlRigOverrideValue(const FString& InPath, const UObject* InSubject)
+{
+	if(InSubject)
+	{
+		*this = FControlRigOverrideValue(InPath, InSubject->GetClass(), InSubject, InSubject->GetFName());
+	}
+}
+
+FControlRigOverrideValue::FControlRigOverrideValue(const FString& InPath, const UStruct* InOwnerStruct, const void* InSubjectPtr, const FName& InSubjectName)
+{
+	ensure(InOwnerStruct);
+	ensureMsgf(InPath.Contains(TEXT(".")) == false, TEXT("Paths for overrides should use the -> separator."));
+	
+	Path = InPath;
+	SubjectName = InSubjectName;
+
+	FString PropertyNameString;
+	FString ArrayIndexString;
+	bool bIsParsingPropertyName = true;
+	const UStruct* OwnerStruct = InOwnerStruct;
+
+	auto ProcessPropertyName = [this, &PropertyNameString, &OwnerStruct]() -> bool
+	{
+		FString StringToProcess;
+		Swap(StringToProcess, PropertyNameString);
+		if(!StringToProcess.IsEmpty())
+		{
+			if(FProperty* Property = FindProperty(OwnerStruct, StringToProcess))
+			{
+				FProperty* PropertyForPath = Property;
+				
+				if(const FArrayProperty* ArrayProperty = CastField<FArrayProperty>(Property))
+				{
+					Property = ArrayProperty->Inner;
+
+					if(Properties.IsEmpty())
+					{
+						PropertyForPath = Property;
+					}
+				}
+				if(const FStructProperty* StructProperty = CastField<FStructProperty>(Property))
+				{
+					OwnerStruct = StructProperty->Struct;
+				}
+
+				Properties.Add({PropertyForPath, INDEX_NONE});
+
+				return true;
+			}
+		}
+		return false;
+	};
+	
+	auto ProcessArrayIndex = [this, &ArrayIndexString]() -> bool
+	{
+		FString StringToProcess;
+		Swap(StringToProcess, ArrayIndexString);
+		if(!StringToProcess.IsEmpty())
+		{
+			Properties.Last().ArrayIndex = FCString::Atoi(*StringToProcess);
+			return true;
+		}
+		return false;
+	};
+
+	// Parse the provided Property path. Path may look like
+	// Settings[2]->Color->R resulting in three segments, the
+	// first one with an array index != INDEX_NONE
+	int32 CharIndex = 0;
+	while(CharIndex < InPath.Len())
+	{
+		// potentially start parsing an array index
+		if(InPath[CharIndex] == TEXT('['))
+		{
+			if(!ProcessPropertyName())
+			{
+				Reset();
+				return;
+			}
+			bIsParsingPropertyName = false;
+			CharIndex++;
+			continue;;
+		}
+		
+		// potentially end parsing an array index
+		if(InPath[CharIndex] == TEXT(']'))
+		{
+			if(!ProcessArrayIndex())
+			{
+				Reset();
+				return;
+			}
+			CharIndex++;
+			continue;;
+		}
+
+		// given a path separator parse the property name and store it.
+		// if bIsParsingPropertyName == false we already processed the property name
+		// and have been parsing an array index
+		if(InPath.Mid(CharIndex, PathSeparatorLength).Equals(PathSeparator, ESearchCase::CaseSensitive))
+		{
+			if(bIsParsingPropertyName)
+			{
+				if(!ProcessPropertyName())
+				{
+					Reset();
+					return;
+				}
+			}
+			bIsParsingPropertyName = true;
+			CharIndex += PathSeparatorLength;
+			continue;
+		}
+
+		if(bIsParsingPropertyName)
+		{
+			PropertyNameString += InPath[CharIndex];
+		}
+		else
+		{
+			ArrayIndexString += InPath[CharIndex];
+		}
+		CharIndex++;
+	}
+
+	if(!PropertyNameString.IsEmpty() && !ProcessPropertyName())
+	{
+		Reset();
+		return;
+	}
+	if(!ArrayIndexString.IsEmpty() && !ProcessArrayIndex())
+	{
+		Reset();
+		return;
+	}
+
+	if(InSubjectPtr)
+	{
+		SetFromSubject(InSubjectPtr);
+	}
+}
+
+FControlRigOverrideValue::FControlRigOverrideValue(const FString& InPath, const UStruct* InOwnerStruct, const FString& InValueAsString, const FName& InSubjectName)
+{
+	*this = FControlRigOverrideValue(InPath, InOwnerStruct, nullptr, InSubjectName);
+	SetFromString(InValueAsString);
+}
+
+FControlRigOverrideValue::FControlRigOverrideValue(const FControlRigOverrideValue& InOther)
+{
+	*this = InOther;
+}
+
+FControlRigOverrideValue::~FControlRigOverrideValue()
+{
+	Reset();
+}
+
+FControlRigOverrideValue& FControlRigOverrideValue::operator=(const FControlRigOverrideValue& InOther)
+{
+	Reset();
+	Path = InOther.Path;
+	SubjectName = InOther.SubjectName;
+	CachedStringValue = InOther.CachedStringValue;
+	Properties = InOther.Properties;
+
+	if(const FProperty* LeafProperty = GetLeafProperty())
+	{
+		DataArray.SetNumUninitialized(LeafProperty->GetElementSize());
+		LeafProperty->InitializeValue(DataArray.GetData());
+
+		if(!InOther.DataArray.IsEmpty())
+		{
+			LeafProperty->CopyCompleteValue(DataArray.GetData(), InOther.DataArray.GetData());
+		}
+	}
+
+	return *this;
+}
+
+bool FControlRigOverrideValue::Serialize(FArchive& Ar)
+{
+	if(Ar.IsLoading())
+	{
+		Reset();
+	}
+
+	Ar << SubjectName;
+	
+	int32 NumProperties = 0;
+
+	if(Ar.IsLoading())
+	{
+		Ar << NumProperties;
+		Properties.Reserve(NumProperties);
+
+		for(int32 Index = 0; Index < NumProperties; Index++)
+		{
+			FPropertyInfo Info;
+			Ar << Info.Property;
+			Ar << Info.ArrayIndex;
+			Properties.Add(Info);
+		}
+
+		// rebuild path from properties
+		FStringBuilderBase PathBuilder;
+		bool bFirstProperty = true;
+
+		for(int32 Index = 0; Index < NumProperties; Index++)
+		{
+			if(const FProperty* Property = Properties[Index].Property.Get())
+			{
+				if(!bFirstProperty)
+				{
+					PathBuilder.Append(FControlRigOverrideValue::PathSeparator);
+				}
+				else
+				{
+					bFirstProperty = false;
+				}
+				PathBuilder.Append(Property->GetName());
+				Path = Property->GetName();
+				if(Properties[Index].ArrayIndex != INDEX_NONE)
+				{
+					Path->Appendf(TEXT("[%d]"),  Properties[Index].ArrayIndex);
+				}
+			}
+			else
+			{
+				PathBuilder.Reset();
+				break;
+			}
+		}
+
+		if(PathBuilder.Len() > 0)
+		{
+			Path = PathBuilder.ToString();
+		}
+	}
+	else
+	{
+		NumProperties = Properties.Num();
+		Ar << NumProperties;
+
+		for(int32 Index = 0; Index < NumProperties; Index++)
+		{
+			Ar << Properties[Index].Property;
+			Ar << Properties[Index].ArrayIndex;
+		}
+	}
+
+	const int64 ArchivePositionBeforeData = Ar.GetArchiveState().Tell();
+
+	int64 OffsetForData = 0;
+	Ar << OffsetForData;
+
+	if(Ar.IsLoading())
+	{
+		if(const FProperty* LeafProperty = GetLeafProperty())
+		{
+			DataArray.SetNumUninitialized(LeafProperty->GetElementSize());
+			LeafProperty->InitializeValue(DataArray.GetData());
+			LeafProperty->SerializeItem(FStructuredArchiveFromArchive(Ar).GetSlot(), DataArray.GetData());
+
+			FString ValueAsString;
+			LeafProperty->ExportTextItem_Direct(ValueAsString, DataArray.GetData(), DataArray.GetData(), nullptr, PPF_None, nullptr);
+			CachedStringValue = ValueAsString;
+		}
+		else
+		{
+			// we can't deserialize the data - skip it
+			Ar.Seek(ArchivePositionBeforeData + OffsetForData);
+		}
+	}
+	else
+	{
+		if(const FProperty* LeafProperty = GetLeafProperty())
+		{
+			LeafProperty->SerializeItem(FStructuredArchiveFromArchive(Ar).GetSlot(), DataArray.GetData());
+
+			const int64 ArchivePositionAfterData = Ar.GetArchiveState().Tell();
+			int32 DataToSkip = ArchivePositionAfterData - ArchivePositionBeforeData;
+			
+			Ar.Seek(ArchivePositionBeforeData);
+			Ar << DataToSkip;
+			Ar.Seek(ArchivePositionAfterData);
+		}
+	}
+
+	return true;
+}
+
+bool FControlRigOverrideValue::IsValid() const
+{
+	const bool bIsSet = Path.IsSet() && CachedStringValue.IsSet() && (GetData() != nullptr);
+	return bIsSet && (GetLeafProperty() != nullptr);
+}
+
+void FControlRigOverrideValue::Reset()
+{
+	if(!Properties.IsEmpty() && !DataArray.IsEmpty())
+	{
+		if(const FProperty* LeafProperty = GetLeafProperty())
+		{
+			check(DataArray.Num() == LeafProperty->GetElementSize());
+			LeafProperty->DestroyValue(DataArray.GetData());
+		}
+	}
+
+	CachedStringValue.Reset();
+	Properties.Reset();
+	DataArray.Reset();
+}
+
+const FName& FControlRigOverrideValue::GetSubjectName() const
+{
+	static const FName NoneName(NAME_None);
+	return SubjectName.Get(NoneName);
+}
+
+const FString& FControlRigOverrideValue::GetPath() const
+{
+	static const FString EmptyPath;
+	return Path.Get(EmptyPath);
+}
+
+uint8* FControlRigOverrideValue::SubjectPtrToValuePtr(const void* InSubjectPtr, bool bResizeArrays) const
+{
+	if(Properties.IsEmpty())
+	{
+		return nullptr;
+	}
+
+	if(InSubjectPtr == nullptr)
+	{
+		return nullptr;
+	}
+
+	void* ValuePtr = const_cast<void*>(InSubjectPtr);
+	
+	for (const FPropertyInfo& Info : Properties)
+	{
+		const FProperty* Property = Info.Property.Get();
+		check(Property);
+
+		if(Info.ArrayIndex != INDEX_NONE)
+		{
+			const FArrayProperty* ArrayProperty = CastField<FArrayProperty>(Property);
+			if(ArrayProperty == nullptr)
+			{
+				ArrayProperty = CastField<FArrayProperty>(Property->GetOwnerProperty());
+			}
+			if(ArrayProperty == nullptr)
+			{
+				return nullptr;
+			}
+
+			ValuePtr = ArrayProperty->ContainerPtrToValuePtr<void*>(ValuePtr);
+
+			FScriptArrayHelper ArrayHelper(ArrayProperty, ValuePtr);
+			if(ArrayHelper.Num() <= Info.ArrayIndex)
+			{
+				if(bResizeArrays)
+				{
+					ArrayHelper.Resize(Info.ArrayIndex + 1);
+				}
+				else
+				{
+					return nullptr;
+				}
+			}
+			
+			ValuePtr = ArrayHelper.GetRawPtr(Info.ArrayIndex);
+			continue;
+		}
+
+		ValuePtr = Property->ContainerPtrToValuePtr<void*>(ValuePtr);
+	}
+
+	return static_cast<uint8*>(ValuePtr);
+}
+
+const FProperty* FControlRigOverrideValue::GetRootProperty() const
+{
+	if(Properties.IsEmpty())
+	{
+		return nullptr;
+	}
+	return Properties[0].Property.Get();
+}
+
+FProperty* FControlRigOverrideValue::FindProperty(const UStruct* InStruct, const FString& InNameOrDisplayName)
+{
+	if(InStruct == nullptr)
+	{
+		return  nullptr;
+	}
+
+	if(InNameOrDisplayName.IsEmpty())
+	{
+		return nullptr;
+	}
+	
+	if(FProperty* Property = InStruct->FindPropertyByName(*InNameOrDisplayName))
+	{
+		return Property;
+	}
+
+#if WITH_EDITORONLY_DATA
+	const FText DisplayName = FText::FromString(InNameOrDisplayName);
+	for (TFieldIterator<FProperty> PropertyIt(InStruct); PropertyIt; ++PropertyIt)
+	{
+		if(FProperty* PropertyOnStruct = *PropertyIt)
+		{
+			if(PropertyOnStruct->GetDisplayNameText().EqualToCaseIgnored(DisplayName))
+			{
+				return PropertyOnStruct;
+			}
+		}
+	}
+#endif
+
+	return nullptr;	
+}
+
+uint32 GetTypeHash(const FControlRigOverrideValue& InOverride)
+{
+	uint32 Hash = GetTypeHash(InOverride.SubjectName);
+	if(InOverride.IsValid())
+	{
+		for(const FControlRigOverrideValue::FPropertyInfo& PropertyInfo : InOverride.Properties)
+		{
+			Hash = HashCombine(Hash, GetTypeHash(PropertyInfo.Property));
+			Hash = HashCombine(Hash, GetTypeHash(PropertyInfo.ArrayIndex));
+		}
+		Hash = HashCombine(Hash, GetTypeHash(InOverride.ToString()));
+	}
+	return Hash;
+}
+
+const FProperty* FControlRigOverrideValue::GetLeafProperty() const
+{
+	if(Properties.IsEmpty())
+	{
+		return nullptr;
+	}
+	
+	if(const FProperty* LeafProperty = Properties.Last().Property.Get())
+	{
+		const int32& ArrayIndex = Properties.Last().ArrayIndex;
+		if(ArrayIndex != INDEX_NONE)
+		{
+			if(const FArrayProperty* ArrayProperty = CastField<FArrayProperty>(LeafProperty))
+			{
+				LeafProperty = ArrayProperty->Inner;
+			}
+		}
+		else
+		{
+			if(const FArrayProperty* ArrayProperty = CastField<FArrayProperty>(LeafProperty->GetOwnerProperty()))
+			{
+				LeafProperty = ArrayProperty;
+			}
+		}
+
+		return LeafProperty;
+	}
+
+	return nullptr;
+}
+
+void* FControlRigOverrideValue::GetData()
+{
+	if(DataArray.IsEmpty())
+	{
+		return nullptr;
+	}
+	return DataArray.GetData();
+}
+
+const void* FControlRigOverrideValue::GetData() const
+{
+	if(DataArray.IsEmpty())
+	{
+		return nullptr;
+	}
+	return DataArray.GetData();
+}
+
+const FString& FControlRigOverrideValue::ToString() const
+{
+	static const FString EmptyPath;
+	return CachedStringValue.Get(EmptyPath);
+}
+
+bool FControlRigOverrideValue::SetFromString(const FString& InValue)
+{
+	if(InValue != CachedStringValue.Get(FString()))
+	{
+		if(const FProperty* LeafProperty = GetLeafProperty())
+		{
+			if(DataArray.IsEmpty())
+			{
+				DataArray.SetNumUninitialized(LeafProperty->GetElementSize());
+				LeafProperty->InitializeValue(DataArray.GetData());
+			}
+			
+			FRigVMMemoryStorageImportErrorContext ErrorPipe(false);
+			LeafProperty->ImportText_Direct(*InValue, DataArray.GetData(), nullptr, PPF_None, &ErrorPipe);
+			if(ErrorPipe.NumErrors == 0)
+			{
+				CachedStringValue = InValue;
+				return true;
+			}
+		}
+	}
+	return false;
+}
+
+bool FControlRigOverrideValue::operator==(const FControlRigOverrideValue& InOtherValue) const
+{
+	if(!IsValid() || !InOtherValue.IsValid())
+	{
+		return false;
+	}
+	if(GetSubjectName() != InOtherValue.GetSubjectName())
+	{
+		return false;
+	}
+	if(Properties.Num() != InOtherValue.Properties.Num())
+	{
+		return false;
+	}
+	for(int32 PropertyIndex = 0; PropertyIndex < Properties.Num(); PropertyIndex++)
+	{
+		if(Properties[PropertyIndex].Property != InOtherValue.Properties[PropertyIndex].Property)
+		{
+			return false;
+		}
+		if(Properties[PropertyIndex].ArrayIndex != InOtherValue.Properties[PropertyIndex].ArrayIndex)
+		{
+			return false;
+		}
+	}
+	return IdenticalValue(InOtherValue.GetData());
+}
+
+bool FControlRigOverrideValue::Identical(const FControlRigOverrideValue& InOtherValue) const
+{
+	return *this == InOtherValue;
+}
+
+bool FControlRigOverrideValue::IdenticalValue(const void* InValuePtr) const
+{
+	if(!IsValid() || (InValuePtr == nullptr))
+	{
+		return false;
+	}
+
+	// IsValid checks the validity of the data pointer and leaf property
+	const FProperty* LeafProperty = GetLeafProperty();
+	check(LeafProperty);
+	return LeafProperty->Identical(GetData(), InValuePtr);
+}
+
+bool FControlRigOverrideValue::IdenticalValueInSubject(const void* InSubjectPtr) const
+{
+	return IdenticalValue(SubjectPtrToValuePtr(InSubjectPtr, false));
+}
+
+void FControlRigOverrideValue::CopyToSubject(void* InSubjectPtr, const UStruct* InSubjectStruct) const
+{
+	if(InSubjectStruct)
+	{
+		if(const FProperty* RootProperty = GetRootProperty())
+		{
+			if(RootProperty->GetOwnerStruct() != InSubjectStruct)
+			{
+				return;
+			}
+		}
+	}
+
+	if(const FProperty* LeafProperty = GetLeafProperty())
+	{
+		if(const void* SourcePtr = GetData())
+		{
+			if(void* DestPtr = SubjectPtrToValuePtr(InSubjectPtr, true))
+			{
+				LeafProperty->CopyCompleteValue(DestPtr, SourcePtr);
+			}
+		}
+	}
+}
+
+void FControlRigOverrideValue::SetFromSubject(const void* InSubjectPtr, const UStruct* InSubjectStruct)
+{
+	if(InSubjectStruct)
+	{
+		if(const FProperty* RootProperty = GetRootProperty())
+		{
+			if(RootProperty->GetOwnerStruct() != InSubjectStruct)
+			{
+				return;
+			}
+		}
+	}
+	
+	CachedStringValue.Reset();
+
+	if(const FProperty* LeafProperty = GetLeafProperty())
+	{
+		if(DataArray.IsEmpty())
+		{
+			DataArray.SetNumUninitialized(LeafProperty->GetElementSize());
+			LeafProperty->InitializeValue(DataArray.GetData());
+		}
+
+		if(InSubjectPtr)
+		{
+			if(const void* SourcePtr = SubjectPtrToValuePtr(InSubjectPtr, false))
+			{
+				void* DestPtr = DataArray.GetData();
+				FString ValueAsString;
+				LeafProperty->CopyCompleteValue(DestPtr, SourcePtr);
+				LeafProperty->ExportTextItem_Direct(ValueAsString, DestPtr, DestPtr, nullptr, PPF_None, nullptr);
+				CachedStringValue = ValueAsString;
+			}
+		}
+	}
+}
+
+void FControlRigOverrideValue::CopyToUObject(UObject* InSubjectPtr) const
+{
+	if(InSubjectPtr)
+	{
+		CopyToSubject(InSubjectPtr, InSubjectPtr->GetClass());
+	}
+}
+
+void FControlRigOverrideValue::SetFromUObject(const UObject* InSubjectPtr)
+{
+	if(InSubjectPtr)
+	{
+		SetFromSubject(InSubjectPtr, InSubjectPtr->GetClass());
+	}
+}
+
+////////////////////////////////////////////////////////////////////////////////
+// FControlRigOverrideContainer
+////////////////////////////////////////////////////////////////////////////////
+
+void FControlRigOverrideContainer::Reset()
+{
+	Values.Reset();
+	HashIndexLookup.Reset();
+	SubjectIndexLookup.Reset();
+	InvalidateCache();
+}
+
+void FControlRigOverrideContainer::Empty()
+{
+	Values.Empty();
+	HashIndexLookup.Empty();
+	SubjectIndexLookup.Reset();
+	InvalidateCache();
+}
+
+void FControlRigOverrideContainer::Reserve(int32 InNum)
+{
+	Values.Reserve(InNum);
+	HashIndexLookup.Reserve(InNum);
+}
+
+bool FControlRigOverrideContainer::UsesSubjectName() const
+{
+	return bUsesSubjectName;
+}
+
+void FControlRigOverrideContainer::SetUsesSubjectName(bool InUsesSubjectName)
+{
+	bUsesSubjectName = InUsesSubjectName;
+	RebuildLookup();
+	InvalidateCache();
+}
+
+int32 FControlRigOverrideContainer::Add(const FControlRigOverrideValue& InValue)
+{
+	if(!InValue.IsValid())
+	{
+		return INDEX_NONE;
+	}
+	
+	const FString& Path = InValue.GetPath();
+	if(!Path.IsEmpty())
+	{
+		if(!Contains(Path, InValue.GetSubjectName()))
+		{
+			// don't allow to add overrides for child paths of
+			// existing parent paths
+			if(ContainsParentPathOf(Path, InValue.GetSubjectName()))
+			{
+				return INDEX_NONE;
+			}
+			
+			// remove child paths when adding parent overrides
+			if(ContainsChildPathOf(Path, InValue.GetSubjectName()))
+			{
+				TArray<int32> IndicesToRemove;
+				for(int32 Index = 0; Index < Values.Num(); Index++)
+				{
+					const FString& ChildPath = Values[Index].GetPath();
+					if(IsChildPathOf(ChildPath, Path))
+					{
+						IndicesToRemove.Add(Index);
+					}
+				}
+
+				if(!IndicesToRemove.IsEmpty())
+				{
+					// remove in reverse
+					for(int32 Index = IndicesToRemove.Num() - 1; Index >= 0; Index--)
+					{
+						Values.RemoveAt(IndicesToRemove[Index]);
+					}
+					RebuildLookup();
+				}
+			}
+			
+			const int32 Index = Values.Add(InValue);
+			const uint32 Hash = GetLookupHash(Path, InValue.GetSubjectName());
+			HashIndexLookup.Add(Hash, Index);
+
+			if(bUsesSubjectName)
+			{
+				SubjectIndexLookup.FindOrAdd(InValue.GetSubjectName()).Add(Index);
+			}
+			InvalidateCache();
+			return Index;
+		}
+	}
+	return INDEX_NONE;
+}
+
+const FControlRigOverrideValue* FControlRigOverrideContainer::FindOrAdd(const FControlRigOverrideValue& InValue)
+{
+	const FString& Path = InValue.GetPath();
+	if(!Path.IsEmpty())
+	{
+		if(FControlRigOverrideValue* ExistingValue = Find(Path, InValue.GetSubjectName()))
+		{
+			*ExistingValue = InValue;
+			return ExistingValue;
+		}
+
+		const int32 NewIndex = Add(InValue);
+		if(IsValidIndex(NewIndex))
+		{
+			return &Values[NewIndex];
+		}
+	}
+	return nullptr;
+}
+
+bool FControlRigOverrideContainer::Remove(const FControlRigOverrideValue& InValue)
+{
+	return Remove(InValue.GetPath(), InValue.GetSubjectName());
+}
+
+bool FControlRigOverrideContainer::Remove(const FString& InPath, const FName& InSubjectName)
+{
+	const int32 Index = GetIndex(InPath, InSubjectName);
+	if(IsValidIndex(Index))
+	{
+		Values.RemoveAt(Index);
+		RebuildLookup();
+		InvalidateCache();
+		return true;
+	}
+	return false;
+}
+
+bool FControlRigOverrideContainer::RemoveAll(const FName& InSubjectName)
+{
+	if(IsEmpty())
+	{
+		return false;
+	}
+	if(bUsesSubjectName)
+	{
+		if(const TArray<int32>* IndicesPtr = GetIndicesForSubject(InSubjectName))
+		{
+			for(int32 Index = IndicesPtr->Num() - 1; Index >= 0; Index--)
+			{
+				Values.RemoveAt(IndicesPtr->operator[](Index));
+			}
+			RebuildLookup();
+			InvalidateCache();
+			return true;
+		}
+		return false;
+	}
+	else
+	{
+		Reset();
+	}
+	return true;
+}
+
+int32 FControlRigOverrideContainer::GetIndex(const FString& InPath, const FName& InSubjectName) const
+{
+	const uint32 Hash = GetLookupHash(InPath, InSubjectName);
+	if(const int32* Index = HashIndexLookup.Find(Hash))
+	{
+		return *Index;
+	}
+	return INDEX_NONE;
+}
+
+const TArray<int32>* FControlRigOverrideContainer::GetIndicesForSubject(const FName& InSubjectName) const
+{
+	return SubjectIndexLookup.Find(InSubjectName);
+}
+
+const FControlRigOverrideValue* FControlRigOverrideContainer::Find(const FString& InPath, const FName& InSubjectName) const
+{
+	const int32 Index = GetIndex(InPath, InSubjectName);
+	if(IsValidIndex(Index))
+	{
+		return &Values[Index];
+	}
+	return nullptr;
+}
+
+FControlRigOverrideValue* FControlRigOverrideContainer::Find(const FString& InPath, const FName& InSubjectName)
+{
+	const int32 Index = GetIndex(InPath, InSubjectName);
+	if(IsValidIndex(Index))
+	{
+		return &Values[Index];
+	}
+	return nullptr;
+}
+
+const FControlRigOverrideValue& FControlRigOverrideContainer::FindChecked(const FString& InPath, const FName& InSubjectName) const
+{
+	const int32 Index = GetIndex(InPath, InSubjectName);
+	check(IsValidIndex(Index));
+	return Values[Index];
+}
+
+FControlRigOverrideValue& FControlRigOverrideContainer::FindChecked(const FString& InPath, const FName& InSubjectName)
+{
+	const int32 Index = GetIndex(InPath, InSubjectName);
+	check(IsValidIndex(Index));
+	return Values[Index];
+}
+
+bool FControlRigOverrideContainer::Contains(const FString& InPath, const FName& InSubjectName) const
+{
+	return IsValidIndex(GetIndex(InPath, InSubjectName));
+}
+
+bool FControlRigOverrideContainer::ContainsParentPathOf(const FString& InChildPath, const FName& InSubjectName) const
+{
+	if(InChildPath.IsEmpty())
+	{
+		return false;
+	}
+
+	const uint32 Hash = GetLookupHash(InChildPath, InSubjectName);
+	if(const bool* CachedResult = ContainsParentPathCache.Find(Hash))
+	{
+		return *CachedResult;
+	}
+
+	const FString ParentPath = GetParentPath(InChildPath);
+	const bool bResult = !ParentPath.IsEmpty() && (Contains(ParentPath, InSubjectName) || ContainsParentPathOf(ParentPath, InSubjectName));
+	ContainsParentPathCache.Add(Hash, bResult);
+	return bResult;
+}
+
+bool FControlRigOverrideContainer::ContainsChildPathOf(const FString& InParentPath, const FName& InSubjectName) const
+{
+	if(InParentPath.IsEmpty())
+	{
+		return false;
+	}
+	
+	const uint32 Hash = GetLookupHash(InParentPath, InSubjectName);
+	if(const bool* CachedResult = ContainsChildPathCache.Find(Hash))
+	{
+		return *CachedResult;
+	}
+
+	for(const FControlRigOverrideValue& Value : Values)
+	{
+		if(bUsesSubjectName)
+		{
+			if(Value.GetSubjectName() != InSubjectName)
+			{
+				continue;
+			}
+		}
+		const FString& Path = Value.GetPath();
+		if(IsChildPathOf(Path, InParentPath))
+		{
+			ContainsChildPathCache.Add(Hash, true);
+			return true;
+		}
+	}
+
+	ContainsChildPathCache.Add(Hash, false);
+	return false;
+}
+
+bool FControlRigOverrideContainer::Contains(const FControlRigOverrideValue& InOverrideValue) const
+{
+	return Contains(InOverrideValue.GetPath(), InOverrideValue.GetSubjectName());
+}
+
+bool FControlRigOverrideContainer::ContainsParentPathOf(const FControlRigOverrideValue& InOverrideValue) const
+{
+	return ContainsParentPathOf(InOverrideValue.GetPath(), InOverrideValue.GetSubjectName());
+}
+
+bool FControlRigOverrideContainer::ContainsChildPathOf(const FControlRigOverrideValue& InOverrideValue) const
+{
+	return ContainsChildPathOf(InOverrideValue.GetPath(), InOverrideValue.GetSubjectName());
+}
+
+void FControlRigOverrideContainer::CopyToSubject(void* InSubjectPtr, const UStruct* InSubjectStruct, const FName& InSubjectName) const
+{
+	for(const FControlRigOverrideValue& Override : Values)
+	{
+		if(bUsesSubjectName)
+		{
+			if(Override.GetSubjectName() != InSubjectName)
+			{
+				continue;
+			}
+		}
+		Override.CopyToSubject(InSubjectPtr, InSubjectStruct);
+	}
+}
+
+void FControlRigOverrideContainer::SetFromSubject(const void* InSubjectPtr, const UStruct* InSubjectStruct, const FName& InSubjectName)
+{
+	for(FControlRigOverrideValue& Override : Values)
+	{
+		if(bUsesSubjectName)
+		{
+			if(Override.GetSubjectName() != InSubjectName)
+			{
+				continue;
+			}
+		}
+		Override.SetFromSubject(InSubjectPtr, InSubjectStruct);
+	}
+}
+
+void FControlRigOverrideContainer::CopyToUObject(UObject* InSubjectPtr) const
+{
+	CopyToSubject(InSubjectPtr, InSubjectPtr->GetClass(), InSubjectPtr->GetFName());
+}
+
+void FControlRigOverrideContainer::SetFromUObject(const UObject* InSubjectPtr)
+{
+	SetFromSubject(InSubjectPtr, InSubjectPtr->GetClass(), InSubjectPtr->GetFName());
+}
+
+bool FControlRigOverrideContainer::Serialize(FArchive& Ar)
+{
+	Ar << bUsesSubjectName;
+	Ar << Values;
+	
+	if(Ar.IsLoading())
+	{
+		RebuildLookup();
+		InvalidateCache();
+	}
+
+	return true;
+}
+
+void FControlRigOverrideContainer::RebuildLookup()
+{
+	HashIndexLookup.Reset();
+	HashIndexLookup.Reserve(Values.Num());
+	SubjectIndexLookup.Reset();
+
+	for(int32 Index = 0; Index < Values.Num(); Index++)
+	{
+		if(Values[Index].IsValid())
+		{
+			const FString& Path = Values[Index].GetPath();
+			if(!Path.IsEmpty())
+			{
+				const uint32 Hash = GetLookupHash(Path, Values[Index].GetSubjectName());
+				HashIndexLookup.Add(Hash, Index);
+
+				if(bUsesSubjectName)
+				{
+					SubjectIndexLookup.FindOrAdd(Values[Index].GetSubjectName()).Add(Index);
+				}
+			}
+		}
+	}
+}
+
+void FControlRigOverrideContainer::InvalidateCache() const
+{
+	ContainsParentPathCache.Reset();
+	ContainsChildPathCache.Reset();
+}
+
+FString FControlRigOverrideContainer::GetParentPath(const FString& InChildPath)
+{
+	int32 LastSeparatorIndex = InChildPath.Find(FControlRigOverrideValue::PathSeparator, ESearchCase::CaseSensitive, ESearchDir::FromEnd);
+
+	int32 LastArrayIndex = INDEX_NONE;
+	if(InChildPath.FindLastChar(FControlRigOverrideValue::ArraySeparator, LastArrayIndex))
+	{
+		LastSeparatorIndex = FMath::Max(LastSeparatorIndex, LastArrayIndex);
+	}
+
+	if(LastSeparatorIndex != INDEX_NONE)
+	{
+		return InChildPath.Left(LastSeparatorIndex);
+	}
+
+	return FString();
+}
+
+bool FControlRigOverrideContainer::IsChildPathOf(const FString& InChildPath, const FString& InParentPath)
+{
+	if(!InChildPath.IsEmpty())
+	{
+		// this is true for both cases of child paths - one being separated by -> and the other being separated by [
+		if(InChildPath.Len() > InParentPath.Len() + FControlRigOverrideValue::PathSeparatorLength)
+		{
+			const FStringView ChildPathView(InChildPath);
+			if(ChildPathView.StartsWith(InParentPath))
+			{
+				if(ChildPathView[InParentPath.Len()] == FControlRigOverrideValue::ArraySeparator)
+				{
+					return true;
+				}
+				if(ChildPathView.Mid(InParentPath.Len(), FControlRigOverrideValue::PathSeparatorLength) == FControlRigOverrideValue::PathSeparator)
+				{
+					return true;
+				}
+			}
+		}
+	}
+	return false;
+}
+
+#if WITH_EDITOR
+TSharedPtr<FPropertyPath> FControlRigOverrideValue::ToPropertyPath() const
+{
+	if(Properties.IsEmpty())
+	{
+		return nullptr;
+	}
+	
+	TSharedPtr<FPropertyPath> PropertyPath = FPropertyPath::CreateEmpty();
+
+	for(const FPropertyInfo& PropertyInfo : Properties)
+	{
+		PropertyPath->AddProperty({PropertyInfo.Property, PropertyInfo.ArrayIndex});
+	}
+
+	return PropertyPath;
+}
+#endif
+
+uint32 GetTypeHash(const FControlRigOverrideContainer& InContainer)
+{
+	uint32 Hash = 0;
+	for(const FControlRigOverrideValue& Override : InContainer.Values)
+	{
+		Hash = HashCombine(Hash, GetTypeHash(Override));
+	}
+	return Hash;
+}
diff --git a/Engine/Plugins/Animation/ControlRig/Source/ControlRig/Private/ModularRig.cpp b/Engine/Plugins/Animation/ControlRig/Source/ControlRig/Private/ModularRig.cpp
--- a/Engine/Plugins/Animation/ControlRig/Source/ControlRig/Private/ModularRig.cpp
+++ b/Engine/Plugins/Animation/ControlRig/Source/ControlRig/Private/ModularRig.cpp
@@ -351,7 +351,7 @@ void UModularRig::UpdateModuleHierarchyFromCDO()
 					InModuleReference->Class.Get(),
 					FindModule(InModuleReference->ParentModuleName),
 					Model.Connections.GetModuleConnectionMap(InModuleReference->Name),
-					InModuleReference->ConfigValues);
+					InModuleReference->ConfigOverrides);
 			}
 
 			// continue to the next module
@@ -548,10 +548,7 @@ void UModularRig::ExecuteQueue()
 					// put the variable defaults back
 					if(const FRigModuleReference* ModuleReference = GetModularRigModel().FindModule(ExecutionElement.ModuleName))
 					{
-						for (const TPair<FName, FString>& Variable : ModuleReference->ConfigValues)
-						{
-							ModuleRig->SetVariableFromString(Variable.Key, Variable.Value);
-						}
+						ModuleReference->ConfigOverrides.CopyToSubject(ModuleRig);
 					}
 				}
 
@@ -786,7 +783,7 @@ TArray<FName> UModularRig::GetModuleNames() const
 }
 
 FRigModuleInstance* UModularRig::AddModuleInstance(const FName& InModuleName, TSubclassOf<UControlRig> InModuleClass, const FRigModuleInstance* InParent,
-                                                   const FRigElementKeyRedirector::FKeyMap& InConnectionMap, const TMap<FName, FString>& InVariableDefaultValues) 
+                                                   const FRigElementKeyRedirector::FKeyMap& InConnectionMap, const FControlRigOverrideContainer& InConfigValues) 
 {
 	// Make sure there are no name clashes
 	FName ParentName = NAME_None;
@@ -848,6 +845,20 @@ FRigModuleInstance* UModularRig::AddModuleInstance(const FName& InModuleName, TS
 	{
 		NewModuleRig = NewObject<UControlRig>(this, InModuleClass, InModuleName);
 	}
+	else
+	{
+		// make sure to reset its public variables back to the value from the CDO
+		const UControlRig* CDO = NewModuleRig->GetClass()->GetDefaultObject<UControlRig>();
+		for (TFieldIterator<FProperty> PropertyIt( NewModuleRig->GetClass()); PropertyIt; ++PropertyIt)
+		{
+			const FProperty* Property = *PropertyIt;
+			if(Property->IsNative())
+			{
+				continue;
+			}
+			Property->CopyCompleteValue_InContainer(NewModuleRig, CDO);
+		}
+	}
 	
 	NewModule.SetRig(NewModuleRig);
 
@@ -871,10 +882,8 @@ FRigModuleInstance* UModularRig::AddModuleInstance(const FName& InModuleName, TS
 		ModulePublicContext.RigModulePrefixHash = GetTypeHash(ModulePublicContext.RigModulePrefix);
 		NewModuleRig->SetElementKeyRedirector(FRigElementKeyRedirector(InConnectionMap, Hierarchy));
 
-		for (const TPair<FName, FString>& Variable : InVariableDefaultValues )
-		{
-			NewModule.GetRig()->SetVariableFromString(Variable.Key, Variable.Value);
-		}
+		UControlRig* ModuleRig = NewModule.GetRig();
+		InConfigValues.CopyToSubject(ModuleRig);
 	}
 
 	return &NewModule;
diff --git a/Engine/Plugins/Animation/ControlRig/Source/ControlRig/Private/ModularRigController.cpp b/Engine/Plugins/Animation/ControlRig/Source/ControlRig/Private/ModularRigController.cpp
--- a/Engine/Plugins/Animation/ControlRig/Source/ControlRig/Private/ModularRigController.cpp
+++ b/Engine/Plugins/Animation/ControlRig/Source/ControlRig/Private/ModularRigController.cpp
@@ -6,6 +6,9 @@
 #include "ModularRig.h"
 #include "ModularRigModel.h"
 #include "ModularRigRuleManager.h"
+#if WITH_EDITOR
+#include "PropertyPath.h"
+#endif
 #include "Misc/DefaultValueHelper.h"
 #include "Rigs/RigHierarchyController.h"
 #include "RigVMFunctions/Math/RigVMMathLibrary.h"
@@ -850,31 +853,54 @@ bool UModularRigController::SetConfigValueInModule(const FName& InModuleName, co
 		return false;
 	}
 
-	const FProperty* Property = Module->Class->FindPropertyByName(InVariableName);
-	if (!Property)
+	const FControlRigOverrideValue OverrideValue(InVariableName.ToString(), Module->Class.Get(), InValue, Module->Name);
+	return SetConfigValueInModule(InModuleName, OverrideValue, bSetupUndo);
+}
+
+bool UModularRigController::SetConfigValueInModule(const FName& InModuleName, const FControlRigOverrideValue& InValue, bool bSetupUndo)
+{
+	FRigModuleReference* Module = FindModule(InModuleName);
+	if (!Module)
 	{
-		UE_LOG(LogControlRig, Error, TEXT("Could not find variable %s in module %s"), *InVariableName.ToString(), *InModuleName.ToString());
+		UE_LOG(LogControlRig, Error, TEXT("Could not find module %s"), *InModuleName.ToString());
 		return false;
 	}
 
-	if (Property->HasAllPropertyFlags(CPF_BlueprintReadOnly))
+	if (!Module->Class.IsValid())
 	{
-		UE_LOG(LogControlRig, Error, TEXT("The target variable %s in module %s is read only"), *InVariableName.ToString(), *InModuleName.ToString());
+		UE_LOG(LogControlRig, Error, TEXT("Class defined in module %s is not valid"), *InModuleName.ToString());
 		return false;
 	}
 
-#if WITH_EDITOR
-	TArray<uint8, TAlignedHeapAllocator<16>> TempStorage;
-	TempStorage.AddZeroed(Property->GetSize());
-	uint8* TempMemory = TempStorage.GetData();
-	Property->InitializeValue(TempMemory);
+	if(!InValue.IsValid())
+	{
+		UE_LOG(LogControlRig, Error, TEXT("Provided value is not valid."));
+		return false;
+	}
 
-	if (!FBlueprintEditorUtils::PropertyValueFromString_Direct(Property, InValue, TempMemory))
+	const FString& Path = InValue.GetPath();
+
+	// we cannot set a config value if we already have a value set on a parent path.
+	// so for example you can't override Color.R if Color itself has been overridden already.
+	if(Module->ConfigOverrides.ContainsParentPathOf(InValue))
 	{
-		UE_LOG(LogControlRig, Error, TEXT("Value %s for variable %s in module %s is not valid"), *InValue, *InVariableName.ToString(), *InModuleName.ToString());
+		UE_LOG(LogControlRig, Error, TEXT("Cannot set a config value for '%s' for the module instance class %s since there's a value on a parent already."), *Path, *Module->Class.ToString());
 		return false;
 	}
 
+	if(InValue.GetRootProperty()->GetOwnerUObject() != Module->Class.Get())
+	{
+		UE_LOG(LogControlRig, Error, TEXT("Provided path '%s' does not belong to the module instance class %s."), *Path, *Module->Class.ToString());
+		return false;
+	}
+
+	if (InValue.GetRootProperty()->HasAllPropertyFlags(CPF_BlueprintReadOnly))
+	{
+		UE_LOG(LogControlRig, Error, TEXT("The target property %s in module %s is read only"), *InValue.GetRootProperty()->GetName(), *InModuleName.ToString());
+		return false;
+	}
+
+#if WITH_EDITOR
 	TSharedPtr<FScopedTransaction> TransactionPtr;
 	if (bSetupUndo)
 	{
@@ -884,9 +910,26 @@ bool UModularRigController::SetConfigValueInModule(const FName& InModuleName, co
 			Blueprint->Modify();
 		}
 	}
-#endif 
+#endif
+	
+	// remove all nested config values (child properties under the provided ones)
+	{
+		TArray<FString> KeysToRemove;
+		for(const FControlRigOverrideValue& Override : Module->ConfigOverrides)
+		{
+			const FString& ChildPath = Override.GetPath();
+			if(FControlRigOverrideContainer::IsChildPathOf(ChildPath, Path))
+			{
+				KeysToRemove.Add(ChildPath);
+			}
+		}
+		for(const FString& KeyToRemove : KeysToRemove)
+		{
+			Module->ConfigOverrides.Remove(KeyToRemove, Module->Name);
+		}
+	}
 
-	Module->ConfigValues.FindOrAdd(InVariableName) = InValue;
+	Module->ConfigOverrides.FindOrAdd(InValue);
 
 	Notify(EModularRigNotification::ModuleConfigValueChanged, Module);
 
@@ -897,6 +940,80 @@ bool UModularRigController::SetConfigValueInModule(const FName& InModuleName, co
 	return true;
 }
 
+bool UModularRigController::ResetConfigValueInModule(const FName& InModuleName, const FString& InPath, bool bClearOverride, bool bSetupUndo)
+{
+	FRigModuleReference* Module = FindModule(InModuleName);
+	if (!Module)
+	{
+		UE_LOG(LogControlRig, Error, TEXT("Could not find module %s"), *InModuleName.ToString());
+		return false;
+	}
+
+	if (!Module->Class.IsValid())
+	{
+		UE_LOG(LogControlRig, Error, TEXT("Class defined in module %s is not valid"), *InModuleName.ToString());
+		return false;
+	}
+
+	if(InPath.IsEmpty())
+	{
+		UE_LOG(LogControlRig, Error, TEXT("Provided path is not valid."));
+		return false;
+	}
+
+	const FControlRigOverrideValue DefaultValue(InPath, Module->Class.Get()->GetDefaultObject());
+	if(!DefaultValue.IsValid())
+	{
+		UE_LOG(LogControlRig, Error, TEXT("Provided path '%s' does not belong to the module instance class %s."), *InPath, *Module->Class.ToString());
+		return false;
+	}
+
+	if(!bClearOverride)
+	{
+		SetConfigValueInModule(InModuleName, DefaultValue, bSetupUndo);
+		return true;
+	}
+	
+	const FString PathPrefix = InPath + TEXT("->");
+
+#if WITH_EDITOR
+	TSharedPtr<FScopedTransaction> TransactionPtr;
+	if (bSetupUndo)
+	{
+		TransactionPtr = MakeShared<FScopedTransaction>(NSLOCTEXT("ModularRigController", "ConfigureModuleValueTransaction", "Configure Module Value"), !GIsTransacting);
+		if(UBlueprint* Blueprint = Cast<UBlueprint>(GetOuter()))
+		{
+			Blueprint->Modify();
+		}
+	}
+#endif
+	
+	// remove all nested config values (child properties under the provided ones)
+	{
+		TArray<FString> KeysToRemove;
+		for(const FControlRigOverrideValue& Override : Module->ConfigOverrides)
+		{
+			const FString& ChildPath = Override.GetPath();
+			if(InPath == ChildPath || FControlRigOverrideContainer::IsChildPathOf(ChildPath, InPath))
+			{
+				KeysToRemove.Add(ChildPath);
+			}
+		}
+		for(const FString& KeyToRemove : KeysToRemove)
+		{
+			Module->ConfigOverrides.Remove(KeyToRemove, Module->Name);
+		}
+	}
+
+	Notify(EModularRigNotification::ModuleConfigValueChanged, Module);
+
+#if WITH_EDITOR
+	TransactionPtr.Reset();
+#endif
+	
+	return true;
+}
+
 TArray<FString> UModularRigController::GetPossibleBindings(const FName& InModuleName, const FName& InVariableName)
 {
 	TArray<FString> PossibleBindings;
@@ -1445,7 +1562,7 @@ FName UModularRigController::MirrorModule(const FName& InModuleName, const FRigV
 	const TMap<FName, FString> OriginalBindings = OriginalModule->Bindings;
 	const TSubclassOf<UControlRig> OriginalClass = OriginalModule->Class.Get();
 	const FName OriginalParentName = OriginalModule->ParentModuleName;
-	const TMap<FName, FString> OriginalConfigValues = OriginalModule->ConfigValues;
+	const FControlRigOverrideContainer OriginalConfigValues = OriginalModule->ConfigOverrides;
 
 	FModularRigControllerCompileBracketScope CompileBracketScope(this);
 
@@ -1476,7 +1593,7 @@ FName UModularRigController::MirrorModule(const FName& InModuleName, const FRigV
 		BindModuleVariable(NewModuleName, Pair.Key, NewSourcePath, bSetupUndo);
 	}
 
-	TSet<FName> ConfigValueSet;
+	TSet<FString> ConfigValueSet;
 #if WITH_EDITOR
 	for (TFieldIterator<FProperty> PropertyIt(OriginalClass); PropertyIt; ++PropertyIt)
 	{
@@ -1512,19 +1629,19 @@ FName UModularRigController::MirrorModule(const FName& InModuleName, const FRigV
 		}
 
 		FString NewValueStr;
-		if (const FString* OriginalValue = OriginalConfigValues.Find(Property->GetFName()))
+		if (const FControlRigOverrideValue* OriginalValue = OriginalConfigValues.Find(Property->GetName(), OriginalModule->Name))
 		{
 			if (bIsVector)
 			{
 				FVector Value;
-				FBlueprintEditorUtils::PropertyValueFromString_Direct(Property, *OriginalValue, (uint8*)&Value);
+				FBlueprintEditorUtils::PropertyValueFromString_Direct(Property, *OriginalValue->ToString(), (uint8*)&Value);
 				Value = InSettings.MirrorVector(Value);
 				FBlueprintEditorUtils::PropertyValueToString_Direct(Property, (uint8*)&Value, NewValueStr, nullptr);
 			}
 			else
 			{
 				FTransform Value;
-				FBlueprintEditorUtils::PropertyValueFromString_Direct(Property, *OriginalValue, (uint8*)&Value);
+				FBlueprintEditorUtils::PropertyValueFromString_Direct(Property, *OriginalValue->ToString(), (uint8*)&Value);
 				Value = InSettings.MirrorTransform(Value);
 				FBlueprintEditorUtils::PropertyValueToString_Direct(Property, (uint8*)&Value, NewValueStr, nullptr);
 			}
@@ -1548,17 +1665,19 @@ FName UModularRigController::MirrorModule(const FName& InModuleName, const FRigV
 			}
 		}
 
-		ConfigValueSet.Add(Property->GetFName());
-		SetConfigValueInModule(NewModuleName, Property->GetFName(), NewValueStr, bSetupUndo);
+		ConfigValueSet.Add(Property->GetName());
+
+		const FControlRigOverrideValue NewValue(Property->GetName(), OriginalClass, NewValueStr, NewModuleName);
+		SetConfigValueInModule(NewModuleName, NewValue, bSetupUndo);
 	}
 #endif
 
 	// Add any other config value that was set in the original module, but was not mirrored
-	for (const TPair<FName, FString>& Pair : OriginalConfigValues)
+	for (const FControlRigOverrideValue& OriginalOverride : OriginalConfigValues)
 	{
-		if (!ConfigValueSet.Contains(Pair.Key))
+		if (!ConfigValueSet.Contains(OriginalOverride.GetPath()))
 		{
-			SetConfigValueInModule(NewModuleName, Pair.Key, Pair.Value, bSetupUndo);
+			SetConfigValueInModule(NewModuleName, OriginalOverride, bSetupUndo);
 		}
 	}
 	
@@ -1819,19 +1938,47 @@ void UModularRigController::RefreshModuleVariables(const FRigModuleReference* In
 		const bool bIsInstanceEditable = !Property->HasAnyPropertyFlags(CPF_DisableEditOnInstance);
 		if (bIsAdvanced || !bIsPublic || !bIsInstanceEditable)
 		{
-			Module->ConfigValues.Remove(Property->GetFName());
-			Module->Bindings.Remove(Property->GetFName());
+			FString PropertyName = Property->GetName();
+			FString PropertyPrefix = PropertyName + TEXT("->");
+
+			TArray<FString> ConfigValuesToRemove;
+			TArray<FName> BindingsToRemove;
+			for (const FControlRigOverrideValue& Override : Module->ConfigOverrides)
+			{
+				const FString& KeyString = Override.GetPath();
+				if(KeyString == PropertyName || KeyString.StartsWith(PropertyPrefix))
+				{
+					ConfigValuesToRemove.Add(KeyString);
+				}
+			}
+			for (const TPair<FName, FString>& Pair : Module->Bindings)
+			{
+				const FString KeyString = Pair.Key.ToString();
+				if(KeyString == PropertyName || KeyString.StartsWith(PropertyPrefix))
+				{
+					BindingsToRemove.Add(Pair.Key);
+				}
+			}
+
+			for(const FString& ConfigValueToRemove : ConfigValuesToRemove)
+			{
+				Module->ConfigOverrides.Remove(ConfigValueToRemove, ModuleName);
+			}
+			for(const FName& BindingToRemove : BindingsToRemove)
+			{
+				Module->Bindings.Remove(BindingToRemove);
+			}
 		}
 	}
 
 	// Make sure all the types are valid
-	const TMap<FName, FString> ConfigValues = Module->ConfigValues;
+	const FControlRigOverrideContainer ConfigOverrides = Module->ConfigOverrides;
 	const TMap<FName, FString> Bindings = Module->Bindings;
-	Module->ConfigValues.Reset();
+	Module->ConfigOverrides.Reset();
 	Module->Bindings.Reset();
-	for (const TPair<FName, FString>& Pair : ConfigValues)
+	for (const FControlRigOverrideValue& Override : ConfigOverrides)
 	{
-		SetConfigValueInModule(ModuleName, Pair.Key, Pair.Value, false);
+		SetConfigValueInModule(ModuleName, Override, false);
 	}
 	for (const TPair<FName, FString>& Pair : Bindings)
 	{
diff --git a/Engine/Plugins/Animation/ControlRig/Source/ControlRig/Private/ModularRigModel.cpp b/Engine/Plugins/Animation/ControlRig/Source/ControlRig/Private/ModularRigModel.cpp
--- a/Engine/Plugins/Animation/ControlRig/Source/ControlRig/Private/ModularRigModel.cpp
+++ b/Engine/Plugins/Animation/ControlRig/Source/ControlRig/Private/ModularRigModel.cpp
@@ -102,6 +102,24 @@ TArray<const FRigConnectorElement*> FRigModuleReference::FindConnectors(const UR
 	return Connectors;
 }
 
+void FRigModuleReference::PatchModelsOnLoad()
+{
+	if(const UClass* ClassPtr = Class.Get())
+	{
+		if(!ConfigValues_DEPRECATED.IsEmpty())
+		{
+			ConfigOverrides.Reset();
+			ConfigOverrides.SetUsesSubjectName(false);
+			ConfigOverrides.Reserve(ConfigValues_DEPRECATED.Num());
+			for(const TPair<FName, FString>& Pair : ConfigValues_DEPRECATED)
+			{
+				ConfigOverrides.Add(FControlRigOverrideValue(Pair.Key.ToString(), ClassPtr, Pair.Value, Name));
+			}
+			ConfigValues_DEPRECATED.Reset();
+		}
+	}
+}
+
 FRigElementKeyRedirector::FKeyArray FModularRigSingleConnection::GetTargetArray() const
 {
 	FRigElementKeyRedirector::FKeyArray Result;
@@ -229,6 +247,12 @@ void FModularRigModel::PatchModelsOnLoad()
 
 	Connections.PatchOnLoad(&PreviousModulePaths);
 	Connections.UpdateFromConnectionList();
+
+	for(int32 ModuleIndex = 0; ModuleIndex < Modules.Num(); ModuleIndex++)
+	{
+		FRigModuleReference& Module = Modules[ModuleIndex];
+		Module.PatchModelsOnLoad();
+	}
 }
 
 UModularRigController* FModularRigModel::GetController(bool bCreateIfNeeded)
diff --git a/Engine/Plugins/Animation/ControlRig/Source/ControlRig/Public/ControlRigOverride.h b/Engine/Plugins/Animation/ControlRig/Source/ControlRig/Public/ControlRigOverride.h
new file mode 100644
--- /dev/null
+++ b/Engine/Plugins/Animation/ControlRig/Source/ControlRig/Public/ControlRigOverride.h
@@ -0,0 +1,291 @@
+// Copyright Epic Games, Inc. All Rights Reserved.
+
+#pragma once
+
+#include "CoreMinimal.h"
+#include "Hash/CityHash.h"
+#include "UObject/WeakFieldPtr.h"
+#include "ControlRigOverride.generated.h"
+
+#if WITH_EDITOR
+class FPropertyPath;
+#endif
+
+/**
+ * A single value used to represent an override on a subject.
+ * The value serialized the data using binary serialization based on the last
+ * property in the property chain.
+ * The value can be copied on to the subject or from the subject
+ * as well as copied from and to string.
+ * This data-structure is not thread-safe for writing.
+ */
+USTRUCT(BlueprintType)
+struct CONTROLRIG_API FControlRigOverrideValue
+{
+	GENERATED_BODY()
+
+public:
+
+	static inline const FString PathSeparator = TEXT("->");
+	static inline constexpr int32 PathSeparatorLength = 2;
+	static inline constexpr TCHAR ArraySeparator = TEXT('[');
+
+	FControlRigOverrideValue()
+	{}
+
+	// constructor given the path and the subject (the instance representing the memory)
+	FControlRigOverrideValue(const FString& InPath, const UObject* InSubject);
+
+	// constructor given the path, the owning structure as well as the container
+	FControlRigOverrideValue(const FString& InPath, const UStruct* InOwnerStruct, const void* InSubjectPtr, const FName& InSubjectName = NAME_None);
+
+	// constructor given the path, the owning structure as well as the value as string
+	FControlRigOverrideValue(const FString& InPath, const UStruct* InOwnerStruct, const FString& InValueAsString, const FName& InSubjectName = NAME_None);
+
+	FControlRigOverrideValue(const FControlRigOverrideValue& InOther);
+
+	~FControlRigOverrideValue();
+
+	FControlRigOverrideValue& operator=(const FControlRigOverrideValue& InOther);
+
+	bool Serialize(FArchive& Ar);
+
+	friend FArchive& operator<<(FArchive& Ar, FControlRigOverrideValue& P)
+	{
+		P.Serialize(Ar);
+		return Ar;
+	}
+
+	// returns true if this value is valid - meaning if the property chain is valid
+	// and the value has been set
+	bool IsValid() const;
+
+	// resets the contents of this value
+	void Reset();
+
+	// returns the name of the subject
+	const FName& GetSubjectName() const;
+
+	// returns the path this value represents
+	const FString& GetPath() const;
+
+	// returns the raw data of this value
+	void* GetData();
+	
+	// returns the raw data of this value
+	const void* GetData() const;
+
+	// returns the value of this override as a string (or an empty string)  
+	const FString& ToString() const;
+
+	// sets the contents of this value from string
+	bool SetFromString(const FString& InValue);
+
+	// copies the override value onto a container
+	void CopyToSubject(void* InSubjectPtr, const UStruct* InSubjectStruct = nullptr) const;
+
+	// copies the value from a container into this override value
+	void SetFromSubject(const void* InSubjectPtr, const UStruct* InSubjectStruct = nullptr);
+
+	// copies the override value onto a container
+	void CopyToUObject(UObject* InSubjectPtr) const;
+
+	// copies the value from a container into this override value
+	void SetFromUObject(const UObject* InSubjectPtr);
+
+	bool operator==(const FControlRigOverrideValue& InOtherValue) const;
+
+	// returns true if the stored value matches the provided value
+	bool Identical(const FControlRigOverrideValue& InOtherValue) const;
+
+	// returns true if the stored value matches the provided value memory
+	bool IdenticalValue(const void* InValuePtr) const;
+
+	// returns true if the stored value matches the provided value memory
+	bool IdenticalValueInSubject(const void* InSubjectPtr) const;
+
+#if WITH_EDITOR
+	// helper conversion used mainly for user interface code
+	TSharedPtr<FPropertyPath> ToPropertyPath() const;
+#endif
+
+	// returns the root property represented by this value
+	const FProperty* GetRootProperty() const;
+	
+	// returns the leaf property represented by this value (potentially leaf == root property)
+	const FProperty* GetLeafProperty() const;
+
+	/**
+	 * converts a memory pointer from the subject to the leaf property.
+	 * @param bResizeArrays If set to true the subject memory will be resized to fit the property chain
+	 */
+	uint8* SubjectPtrToValuePtr(const void* InSubjectPtr, bool bResizeArrays) const;
+
+	// helper method to look up a property under a structure
+	static FProperty* FindProperty(const UStruct* InStruct, const FString& InNameOrDisplayName);
+
+	CONTROLRIG_API friend uint32 GetTypeHash(const FControlRigOverrideValue& InOverride);
+
+private:
+
+	TOptional<FString> Path;
+	TOptional<FString> CachedStringValue;
+	TOptional<FName> SubjectName;
+
+	struct FPropertyInfo
+	{
+		TWeakFieldPtr<FProperty> Property;
+		int32 ArrayIndex = INDEX_NONE;
+	};
+
+	TArray<FPropertyInfo> Properties;
+	TArray<uint8, TAlignedHeapAllocator<16>> DataArray;
+};
+
+template<>
+struct TStructOpsTypeTraits<FControlRigOverrideValue> : public TStructOpsTypeTraitsBase2<FControlRigOverrideValue>
+{
+	enum 
+	{
+		WithSerializer = true,
+	};
+	static constexpr EPropertyObjectReferenceType WithSerializerObjectReferences = EPropertyObjectReferenceType::None;
+};
+
+/**
+ * A container to represent a set of override values for one or more subjects.
+ * The bUsesSubjectName setting will control if the subject name is respected when
+ * adding / looking up / applying overrides.
+ * This data-structure is not thread-safe for writing.
+ */
+USTRUCT(BlueprintType)
+struct CONTROLRIG_API FControlRigOverrideContainer
+{
+	GENERATED_BODY()
+
+public:
+
+	FControlRigOverrideContainer()
+		: bUsesSubjectName(true)
+	{}
+
+	bool IsEmpty() const { return Values.IsEmpty(); }
+	int32 Num() const { return Values.Num(); }
+	bool IsValidIndex(int32 InIndex) const { return Values.IsValidIndex(InIndex); }
+	FControlRigOverrideValue& operator[](int32 InIndex) { return Values[InIndex]; }
+	const FControlRigOverrideValue& operator[](int32 InIndex) const { return Values[InIndex]; }
+	TArray<FControlRigOverrideValue>::RangedForIteratorType begin() { return Values.begin(); }
+	TArray<FControlRigOverrideValue>::RangedForIteratorType end() { return Values.end(); }
+	TArray<FControlRigOverrideValue>::RangedForConstIteratorType begin() const { return Values.begin(); }
+	TArray<FControlRigOverrideValue>::RangedForConstIteratorType end() const { return Values.end(); }
+
+	void Reset();
+	void Empty();
+	void Reserve(int32 InNum);
+	bool UsesSubjectName() const;
+	void SetUsesSubjectName(bool InUsesSubjectName);
+
+	int32 Add(const FControlRigOverrideValue& InValue);
+	const FControlRigOverrideValue* FindOrAdd(const FControlRigOverrideValue& InValue);
+	bool Remove(const FControlRigOverrideValue& InValue);
+	bool Remove(const FString& InPath, const FName& InSubjectName = NAME_None);
+	bool RemoveAll(const FName& InSubjectName = NAME_None);
+
+	int32 Emplace(const FString& InPath, const UStruct* InOwnerStruct, const void* InSubjectPtr, const FName& InSubjectName = NAME_None)
+	{
+		const FControlRigOverrideValue Value(InPath, InOwnerStruct, InSubjectPtr, InSubjectName);
+		return Add(Value);
+	}
+
+	int32 Emplace(const FString& InPath, const UObject* InSubject)
+	{
+		const FControlRigOverrideValue Value(InPath, InSubject->GetClass(), InSubject, InSubject->GetFName());
+		return Add(Value);
+	}
+	
+	int32 Emplace(const FString& InPath, const UStruct* InOwnerStruct, const FString& InValueAsString, const FName& InSubjectName = NAME_None)
+	{
+		const FControlRigOverrideValue Value(InPath, InOwnerStruct, InValueAsString, InSubjectName);
+		return Add(Value);
+	}
+
+	int32 GetIndex(const FString& InPath, const FName& InSubjectName = NAME_None) const;
+	const TArray<int32>* GetIndicesForSubject(const FName& InSubjectName) const;
+	const FControlRigOverrideValue* Find(const FString& InPath, const FName& InSubjectName = NAME_None) const;
+	FControlRigOverrideValue* Find(const FString& InPath, const FName& InSubjectName = NAME_None);
+	const FControlRigOverrideValue& FindChecked(const FString& InPath, const FName& InSubjectName = NAME_None) const;
+	FControlRigOverrideValue& FindChecked(const FString& InPath, const FName& InSubjectName = NAME_None);
+
+	bool Contains(const FString& InPath, const FName& InSubjectName = NAME_None) const;
+	bool ContainsParentPathOf(const FString& InChildPath, const FName& InSubjectName = NAME_None) const;
+	bool ContainsChildPathOf(const FString& InParentPath, const FName& InSubjectName = NAME_None) const;
+
+	bool Contains(const FControlRigOverrideValue& InOverrideValue) const;
+	bool ContainsParentPathOf(const FControlRigOverrideValue& InOverrideValue) const;
+	bool ContainsChildPathOf(const FControlRigOverrideValue& InOverrideValue) const;
+
+	// copies the override values onto a subject.
+	void CopyToSubject(void* InSubjectPtr, const UStruct* InSubjectStruct = nullptr, const FName& InSubjectName = NAME_None) const;
+
+	// copies the value from the subject into this container
+	void SetFromSubject(const void* InSubjectPtr, const UStruct* InSubjectStruct = nullptr, const FName& InSubjectName = NAME_None);
+
+	// copies the override values onto a subject.
+	void CopyToUObject(UObject* InSubjectPtr) const;
+
+	// copies the value from the subject into this container
+	void SetFromUObject(const UObject* InSubjectPtr);
+
+	bool Serialize(FArchive& Ar);
+	friend FArchive& operator<<(FArchive& Ar, FControlRigOverrideContainer& P)
+	{
+		P.Serialize(Ar);
+		return Ar;
+	}
+
+	// returns a parent path of a given child path, separators used are -> and [,
+	// so parent of 'Values[1]->Flag' is 'Values[1]',
+	// and parent of 'Values[1]' is 'Values'.
+	static FString GetParentPath(const FString& InChildPath);
+
+	// returns true if a given child path is a child of a given parent path 
+	static bool IsChildPathOf(const FString& InChildPath, const FString& InParentPath);
+
+	CONTROLRIG_API friend uint32 GetTypeHash(const FControlRigOverrideContainer& InContainer);
+
+private:
+
+	void RebuildLookup();
+	void InvalidateCache() const;
+	
+	FORCEINLINE uint32 GetLookupHash(const FString& InPath, const FName& InSubjectName) const
+	{
+		const uint32 PathHash = CityHash32(
+				(const char*)InPath.GetCharArray().GetData(),
+				InPath.GetCharArray().GetTypeSize() *
+				InPath.GetCharArray().Num());
+		if(!bUsesSubjectName)
+		{
+			return PathHash;
+		}
+		return HashCombineFast(GetTypeHash(InSubjectName), PathHash);
+	}
+
+	TArray<FControlRigOverrideValue> Values;
+
+	bool bUsesSubjectName;
+	TMap<uint32, int32> HashIndexLookup;
+	TMap<FName, TArray<int32>> SubjectIndexLookup;
+	mutable TMap<uint32, bool> ContainsParentPathCache; 
+	mutable TMap<uint32, bool> ContainsChildPathCache; 
+};
+
+template<>
+struct TStructOpsTypeTraits<FControlRigOverrideContainer> : public TStructOpsTypeTraitsBase2<FControlRigOverrideContainer>
+{
+	enum 
+	{
+		WithSerializer = true,
+	};
+	static constexpr EPropertyObjectReferenceType WithSerializerObjectReferences = EPropertyObjectReferenceType::None;
+};
diff --git a/Engine/Plugins/Animation/ControlRig/Source/ControlRig/Public/ModularRig.h b/Engine/Plugins/Animation/ControlRig/Source/ControlRig/Public/ModularRig.h
--- a/Engine/Plugins/Animation/ControlRig/Source/ControlRig/Public/ModularRig.h
+++ b/Engine/Plugins/Animation/ControlRig/Source/ControlRig/Public/ModularRig.h
@@ -306,7 +306,7 @@ public:
 private:
 
 	/** Adds a module to the rig*/
-	FRigModuleInstance* AddModuleInstance(const FName& InModuleName, TSubclassOf<UControlRig> InModuleClass, const FRigModuleInstance* InParent, const FRigElementKeyRedirector::FKeyMap& InConnectionMap, const TMap<FName, FString>& InVariableDefaultValues);
+	FRigModuleInstance* AddModuleInstance(const FName& InModuleName, TSubclassOf<UControlRig> InModuleClass, const FRigModuleInstance* InParent, const FRigElementKeyRedirector::FKeyMap& InConnectionMap, const FControlRigOverrideContainer& InConfigValues);
 
 	/** Updates the module's variable bindings */
 	bool SetModuleVariableBindings(const FName& InModuleName, const TMap<FName, FString>& InVariableBindings);
diff --git a/Engine/Plugins/Animation/ControlRig/Source/ControlRig/Public/ModularRigController.h b/Engine/Plugins/Animation/ControlRig/Source/ControlRig/Public/ModularRigController.h
--- a/Engine/Plugins/Animation/ControlRig/Source/ControlRig/Public/ModularRigController.h
+++ b/Engine/Plugins/Animation/ControlRig/Source/ControlRig/Public/ModularRigController.h
@@ -57,8 +57,13 @@ class CONTROLRIG_API UModularRigController : public UObject
 	UFUNCTION(BlueprintCallable, Category = "Control Rig | Modules")
 	bool AutoConnectModules(const TArray<FName>& InModuleNames, bool bReplaceExistingConnections, bool bSetupUndo = true);
 
+	UE_DEPRECATED(5.6, "Please rely on the overload taking an FControlRigOverrideValue")
 	UFUNCTION(BlueprintCallable, Category = "Control Rig | Modules")
 	bool SetConfigValueInModule(const FName& InModuleName, const FName& InVariableName, const FString& InValue, bool bSetupUndo = true);
+	bool SetConfigValueInModule(const FName& InModuleName, const FControlRigOverrideValue& InValue, bool bSetupUndo = true);
+
+	UFUNCTION(BlueprintCallable, Category = "Control Rig | Modules")
+	bool ResetConfigValueInModule(const FName& InModuleName, const FString& InPath, bool bClearOverride = true, bool bSetupUndo = true);
 
 	UFUNCTION(BlueprintCallable, Category = "ControlRig | Modules")
 	bool BindModuleVariable(const FName& InModuleName, const FName& InVariableName, const FString& InSourcePath, bool bSetupUndo = true);
diff --git a/Engine/Plugins/Animation/ControlRig/Source/ControlRig/Public/ModularRigModel.h b/Engine/Plugins/Animation/ControlRig/Source/ControlRig/Public/ModularRigModel.h
--- a/Engine/Plugins/Animation/ControlRig/Source/ControlRig/Public/ModularRigModel.h
+++ b/Engine/Plugins/Animation/ControlRig/Source/ControlRig/Public/ModularRigModel.h
@@ -3,11 +3,16 @@
 #pragma once
 
 #include "ControlRig.h"
+#include "ControlRigOverride.h"
 #include "ModularRigModel.generated.h"
 
 class UModularRigController;
 struct FModularRigModel;
 
+#if WITH_EDITOR
+class FPropertyPath;
+#endif
+
 UENUM()
 enum class EModularRigNotification : uint8
 {
@@ -54,7 +59,9 @@ public:
 		, Class(nullptr)
 		, PreviousName(NAME_None)
 		, PreviousParentName(NAME_None)
-	{}
+	{
+		ConfigOverrides.SetUsesSubjectName(false);
+	}
 	
 	FRigModuleReference(const FName& InName, TSubclassOf<UControlRig> InClass, const FName& InParentModuleName, const FModularRigModel* InModel)
 		: Name(InName)
@@ -64,7 +71,9 @@ public:
 		, PreviousName(NAME_None)
 		, PreviousParentName(NAME_None)
 		, Model(InModel)
-	{}
+	{
+		ConfigOverrides.SetUsesSubjectName(false);
+	}
 
 	UPROPERTY()
 	FName Name;
@@ -88,7 +97,10 @@ public:
 	TMap<FRigElementKey, FRigElementKey> Connections_DEPRECATED; // Connectors to Connection element
 
 	UPROPERTY()
-	TMap<FName, FString> ConfigValues;
+	TMap<FName, FString> ConfigValues_DEPRECATED;
+
+	UPROPERTY()
+	FControlRigOverrideContainer ConfigOverrides;
 
 	UPROPERTY()
 	TMap<FName, FString> Bindings; // ExternalVariableName (current module) -> SourceExternalVariableNamespacedPath (root rig or other module)
@@ -126,7 +138,7 @@ public:
 
 	const FRigModuleReference* GetParentModule() const;
 	const FRigModuleReference* GetRootModule() const;
-	
+
 	friend bool operator==(const FRigModuleReference& A, const FRigModuleReference& B)
 	{
 		return A.ParentModuleName == B.ParentModuleName &&
@@ -136,6 +148,8 @@ public:
 	const FRigConnectorElement* FindPrimaryConnector(const URigHierarchy* InHierarchy) const;
 	TArray<const FRigConnectorElement*> FindConnectors(const URigHierarchy* InHierarchy) const;
 
+	void PatchModelsOnLoad();
+
 private:
 	const FModularRigModel* Model = nullptr;
 
diff --git a/Engine/Plugins/Animation/ControlRig/Source/ControlRigDeveloper/ControlRigDeveloper.Build.cs b/Engine/Plugins/Animation/ControlRig/Source/ControlRigDeveloper/ControlRigDeveloper.Build.cs
--- a/Engine/Plugins/Animation/ControlRig/Source/ControlRigDeveloper/ControlRigDeveloper.Build.cs
+++ b/Engine/Plugins/Animation/ControlRig/Source/ControlRigDeveloper/ControlRigDeveloper.Build.cs
@@ -57,6 +57,13 @@ namespace UnrealBuildTool.Rules
 					}
                 );
 
+                PublicDependencyModuleNames.AddRange(
+	                new string[]
+	                {
+		                "AnimationEditorWidgets",
+	                }
+                );
+
                 PrivateIncludePathModuleNames.Add("ControlRigEditor");
                 DynamicallyLoadedModuleNames.Add("ControlRigEditor");
             }
diff --git a/Engine/Plugins/Animation/ControlRig/Source/ControlRigDeveloper/Private/ControlRigBlueprint.cpp b/Engine/Plugins/Animation/ControlRig/Source/ControlRigDeveloper/Private/ControlRigBlueprint.cpp
--- a/Engine/Plugins/Animation/ControlRig/Source/ControlRigDeveloper/Private/ControlRigBlueprint.cpp
+++ b/Engine/Plugins/Animation/ControlRig/Source/ControlRigDeveloper/Private/ControlRigBlueprint.cpp
@@ -1155,6 +1155,49 @@ void UControlRigBlueprint::UpdateExposedModuleConnectors() const
 	PropagateHierarchyFromBPToInstances();
 }
 
+#if WITH_EDITOR
+
+TArray<FOverrideStatusSubject> UControlRigBlueprint::GetOverrideSubjects() const
+{
+	TArray<FOverrideStatusSubject> Subjects;
+
+	if(const UModularRig* DebuggedRig = Cast<UModularRig>(GetObjectBeingDebugged()))
+	{
+		ModularRigModel.ForEachModule([&Subjects, DebuggedRig](const FRigModuleReference* ModuleReference) -> bool
+		{
+			if(const FRigModuleInstance* ModuleInstance = DebuggedRig->FindModule(ModuleReference->Name))
+			{
+				if(const UControlRig* ModuleRig = ModuleInstance->GetRig())
+				{
+					for(const FControlRigOverrideValue& Override : ModuleReference->ConfigOverrides)
+					{
+						Subjects.Add({ModuleRig, Override.ToPropertyPath()});
+					}
+				}
+			}
+			return true;
+		});
+	}
+
+	return Subjects;
+}
+
+uint32 UControlRigBlueprint::GetOverrideSubjectsHash() const
+{
+	uint32 Hash = 0;
+	
+	ModularRigModel.ForEachModule([&Hash](const FRigModuleReference* ModuleReference) -> bool
+	{
+		Hash = HashCombine(Hash, GetTypeHash(ModuleReference->Name));
+		Hash = HashCombine(Hash, GetTypeHash(ModuleReference->ConfigOverrides));
+		return true;
+	});
+
+	return Hash;
+}
+
+#endif
+
 bool UControlRigBlueprint::ResolveConnector(const FRigElementKey& DraggedKey, const FRigElementKey& TargetKey, bool bSetupUndoRedo)
 {
 	return ResolveConnectorToArray(DraggedKey, {TargetKey}, bSetupUndoRedo);
@@ -1477,6 +1520,7 @@ void UControlRigBlueprint::HandlePackageDone()
 		ModularRigModel.ForEachModule([](const FRigModuleReference* Element) -> bool
 		{
 			(void)Element->Class.LoadSynchronous();
+			const_cast<FRigModuleReference*>(Element)->PatchModelsOnLoad();
 			return true;
 		});
 
diff --git a/Engine/Plugins/Animation/ControlRig/Source/ControlRigDeveloper/Public/ControlRigBlueprint.h b/Engine/Plugins/Animation/ControlRig/Source/ControlRigDeveloper/Public/ControlRigBlueprint.h
--- a/Engine/Plugins/Animation/ControlRig/Source/ControlRigDeveloper/Public/ControlRigBlueprint.h
+++ b/Engine/Plugins/Animation/ControlRig/Source/ControlRigDeveloper/Public/ControlRigBlueprint.h
@@ -22,6 +22,7 @@
 
 #if WITH_EDITOR
 #include "Kismet2/CompilerResultsLog.h"
+#include "Overrides/SOverrideListWidget.h"
 #endif
 
 #include "ControlRigBlueprint.generated.h"
@@ -216,6 +217,11 @@ public:
 
 	UFUNCTION(BlueprintCallable, Category = "Control Rig Blueprint")
 	void UpdateExposedModuleConnectors() const;
+
+#if WITH_EDITOR
+	TArray<FOverrideStatusSubject> GetOverrideSubjects() const;
+	uint32 GetOverrideSubjectsHash() const;
+#endif
 	
 protected:
 
diff --git a/Engine/Plugins/Animation/ControlRig/Source/ControlRigEditor/Private/ControlRigModuleDetails.cpp b/Engine/Plugins/Animation/ControlRig/Source/ControlRigEditor/Private/ControlRigModuleDetails.cpp
--- a/Engine/Plugins/Animation/ControlRig/Source/ControlRigEditor/Private/ControlRigModuleDetails.cpp
+++ b/Engine/Plugins/Animation/ControlRig/Source/ControlRigEditor/Private/ControlRigModuleDetails.cpp
@@ -25,6 +25,7 @@
 #include "Features/IModularFeatures.h"
 #include "IPropertyAccessEditor.h"
 #include "ModularRigRuleManager.h"
+#include "PropertyPath.h"
 #include "ScopedTransaction.h"
 #include "Editor/SRigHierarchyTreeView.h"
 #include "Widgets/SRigVMVariantTagWidget.h"
@@ -642,9 +643,9 @@ void FRigModuleInstanceDetails::CustomizeDetails(IDetailLayoutBuilder& DetailBui
 					continue;
 				}
 
-				const FSimpleDelegate OnValueChangedDelegate = FSimpleDelegate::CreateSP(this, &FRigModuleInstanceDetails::OnConfigValueChanged, Property->GetFName());
-				DefaultProperty->SetOnPropertyValueChanged(OnValueChangedDelegate);
-				DefaultProperty->SetOnChildPropertyValueChanged(OnValueChangedDelegate);
+				TDelegate<void(const FPropertyChangedEvent&)> OnValueChangedDelegate = TDelegate<void(const FPropertyChangedEvent&)>::CreateSP(this, &FRigModuleInstanceDetails::OnConfigValueChanged);
+				DefaultProperty->SetOnPropertyValueChangedWithData(OnValueChangedDelegate);
+				DefaultProperty->SetOnChildPropertyValueChangedWithData(OnValueChangedDelegate);
 
 				FPropertyBindingWidgetArgs BindingArgs;
 				BindingArgs.Property = (FProperty*)Property;
@@ -881,29 +882,134 @@ FRigElementKeyRedirector FRigModuleInstanceDetails::GetConnections() const
 	return FRigElementKeyRedirector();
 }
 
-void FRigModuleInstanceDetails::OnConfigValueChanged(const FName InVariableName)
+void FRigModuleInstanceDetails::OnConfigValueChanged(const FPropertyChangedEvent& InPropertyChangedEvent)
 {
-	TMap<FName, FString> ModuleValues;
-	for(const FPerModuleInfo& Info : PerModuleInfos)
+	UControlRigBlueprint* Blueprint = PerModuleInfos[0].GetBlueprint();
+	if(Blueprint == nullptr)
 	{
-		if (const FRigModuleInstance* ModuleInstance = Info.GetModule())
+		return;
+	}
+
+	FString PropertyPathString;
+	{
+		TSharedRef< FPropertyPath > PropertyPath = FPropertyPath::CreateEmpty();
+		FProperty* Property = InPropertyChangedEvent.Property;
+		if(!ensure(Property))
 		{
-			if (const UControlRig* ModuleRig = ModuleInstance->GetRig())
+			return;
+		}
+
+		TMap<FString, int32> PropertyNameStack;
+		InPropertyChangedEvent.GetArrayIndicesPerObject(0, PropertyNameStack);
+
+		if (PropertyNameStack.Num() > 0)
+		{
+			TArray<FString> PropertyNames;
+			PropertyNameStack.GetKeys(PropertyNames);
+
+			// the property names in the array are provided in reverse order
+			for(int32 Index = PropertyNames.Num() - 1; Index >= 0; Index--)
 			{
-				FString ValueStr = ModuleRig->GetVariableAsString(InVariableName);
-				ModuleValues.Add(ModuleInstance->Name, ValueStr);
+				if(!PropertyPath->IsValid())
+				{
+					ensure(InPropertyChangedEvent.GetNumObjectsBeingEdited() > 0);
+					Property = FControlRigOverrideValue::FindProperty(InPropertyChangedEvent.GetObjectBeingEdited(0)->GetClass(), PropertyNames[Index]);
+					if(!ensure(Property != nullptr))
+					{
+						return;
+					}
+				}
+				else
+				{
+					if(const FArrayProperty* ArrayProperty = CastField<FArrayProperty>(Property))
+					{
+						Property = ArrayProperty->Inner;
+					}
+				
+					const FStructProperty* StructProperty = CastField<FStructProperty>(Property);
+					if(!ensure(StructProperty != nullptr))
+					{
+						return;
+					}
+
+					Property = FControlRigOverrideValue::FindProperty(StructProperty->Struct, PropertyNames[Index]);
+					if(!ensure(Property))
+					{
+						return;
+					}
+				}
+				
+				PropertyPath->AddProperty(
+					{
+						TWeakFieldPtr<FProperty>(Property),
+						InPropertyChangedEvent.GetArrayIndex(PropertyNames[Index])
+					}
+				);
+			}
+		}
+		else
+		{
+			PropertyPath->AddProperty(
+				{
+					TWeakFieldPtr<FProperty>(InPropertyChangedEvent.MemberProperty),
+					InPropertyChangedEvent.GetArrayIndex(InPropertyChangedEvent.GetMemberPropertyName().ToString())
+				}
+			);
+		}
+
+		PropertyPathString = PropertyPath->ToString();
+	}
+	
+	// we need to shorten the path if we already have data on something above this.
+	bool bFoundValueWithShorterPath = true;
+	while(bFoundValueWithShorterPath)
+	{
+		bFoundValueWithShorterPath = false;
+
+		for(const FPerModuleInfo& Info : PerModuleInfos)
+		{
+			if (const FRigModuleReference* ModuleReference = Blueprint->ModularRigModel.FindModule(Info.ModuleName))
+			{
+				for(const FControlRigOverrideValue& Override : ModuleReference->ConfigOverrides)
+				{
+					const FString& ParentPath = Override.GetPath();
+					if(FControlRigOverrideContainer::IsChildPathOf(PropertyPathString, ParentPath))
+					{
+						bFoundValueWithShorterPath = true;
+						PropertyPathString = ParentPath;
+						break;
+					}
+				}
+			}
+			
+			if(bFoundValueWithShorterPath)
+			{
+				break;
 			}
 		}
 	}
+
+	TMap<FName, FControlRigOverrideValue> ModuleValues;
+	ModuleValues.Reserve(PerModuleInfos.Num());
 	
-	for (const TPair<FName, FString>& Value : ModuleValues)
+	for(const FPerModuleInfo& Info : PerModuleInfos)
 	{
-		if (UControlRigBlueprint* Blueprint = PerModuleInfos[0].GetBlueprint())
+		if (const FRigModuleInstance* ModuleInstance = Info.GetModule())
 		{
-			UModularRigController* Controller = Blueprint->GetModularRigController();
-			Controller->SetConfigValueInModule(Value.Key, InVariableName, Value.Value);
+			if (const UControlRig* ModuleRig = ModuleInstance->GetRig())
+			{
+				FControlRigOverrideValue ConfigValue = FControlRigOverrideValue(PropertyPathString, ModuleRig);
+				ModuleValues.Add(ModuleInstance->Name, ConfigValue);
+			}
 		}
 	}
+
+	FScopedTransaction Transaction(LOCTEXT("ModuleHierarchySetConfigValue", "Set Module Config Value"));
+	for (const TPair<FName, FControlRigOverrideValue>& Value : ModuleValues)
+	{
+		UModularRigController* Controller = Blueprint->GetModularRigController();
+		Controller->SetConfigValueInModule(Value.Key, Value.Value);
+	}
 }
 
 bool FRigModuleInstanceDetails::OnConnectorTargetChanged(TArray<FRigElementKey> InTargets, FRigModuleConnector InConnector)
diff --git a/Engine/Plugins/Animation/ControlRig/Source/ControlRigEditor/Private/ControlRigModuleDetails.h b/Engine/Plugins/Animation/ControlRig/Source/ControlRigEditor/Private/ControlRigModuleDetails.h
--- a/Engine/Plugins/Animation/ControlRig/Source/ControlRigEditor/Private/ControlRigModuleDetails.h
+++ b/Engine/Plugins/Animation/ControlRig/Source/ControlRigEditor/Private/ControlRigModuleDetails.h
@@ -47,7 +47,7 @@ public:
 	TArray<FRigModuleConnector> GetConnectors() const;
 	FRigElementKeyRedirector GetConnections() const;
 
-	void OnConfigValueChanged(const FName InVariableName);
+	void OnConfigValueChanged(const FPropertyChangedEvent& InPropertyChangedEvent);
 	bool OnConnectorTargetChanged(TArray<FRigElementKey> InTargets, FRigModuleConnector InConnector);
 	
 	struct FPerModuleInfo
diff --git a/Engine/Plugins/Animation/ControlRig/Source/ControlRigEditor/Private/EditMode/SControlRigOutliner.cpp b/Engine/Plugins/Animation/ControlRig/Source/ControlRigEditor/Private/EditMode/SControlRigOutliner.cpp
--- a/Engine/Plugins/Animation/ControlRig/Source/ControlRigEditor/Private/EditMode/SControlRigOutliner.cpp
+++ b/Engine/Plugins/Animation/ControlRig/Source/ControlRigEditor/Private/EditMode/SControlRigOutliner.cpp
@@ -41,6 +41,7 @@
 #include "Rigs/FKControlRig.h"
 #include "SEnumCombo.h"
 #include "Components/SkeletalMeshComponent.h"
+#include "Overrides/SOverrideStatusWidget.h"
 
 #define LOCTEXT_NAMESPACE "ControlRigOutliner"
 
@@ -1635,15 +1636,56 @@ TSharedRef<SWidget> SSearchableMultiRigHierarchyTreeView::OnGetOptionsMenu()
 			.AutoHeight()
 			.Padding(3)
 			[
-				SNew(SEnumComboBox, StaticEnum<EElementNameDisplayMode>())
-				.CurrentValue_Lambda([this]() -> int32
-				{
-					return static_cast<int32>(GetElementNameDisplayMode());
-				})
-				.OnEnumSelectionChanged_Lambda([this](int32 InEnumValue, ESelectInfo::Type)
-				{
-					SetElementNameDisplayMode(static_cast<EElementNameDisplayMode>(InEnumValue));
-				})
+				SNew(SHorizontalBox)
+
+				+ SHorizontalBox::Slot()
+				.AutoWidth()
+				.Padding(0)
+				[
+					SNew(SEnumComboBox, StaticEnum<EElementNameDisplayMode>())
+					.CurrentValue_Lambda([this]() -> int32
+					{
+						return static_cast<int32>(GetElementNameDisplayMode());
+					})
+					.OnEnumSelectionChanged_Lambda([this](int32 InEnumValue, ESelectInfo::Type)
+					{
+						SetElementNameDisplayMode(static_cast<EElementNameDisplayMode>(InEnumValue));
+					})
+				]
+				
+				+ SHorizontalBox::Slot()
+				.AutoWidth()
+				.Padding(4, 0, 0, 0)
+				.VAlign(VAlign_Center)
+				.HAlign(HAlign_Left)
+				[
+					SNew(SOverrideStatusWidget)
+					.Status_Lambda([this]() -> EOverrideWidgetStatus::Type
+					{
+						return GetElementNameDisplayMode() == EElementNameDisplayMode::AssetDefault ?
+							EOverrideWidgetStatus::None : EOverrideWidgetStatus::ChangedHere;
+					})
+					.MenuContent_Lambda([this]() -> TSharedRef<SWidget>
+					{
+						if(GetElementNameDisplayMode() != EElementNameDisplayMode::AssetDefault)
+						{
+							FMenuBuilder MenuBuilder(true, nullptr);
+							MenuBuilder.AddMenuEntry(
+								LOCTEXT("RemoveOverride", "Remove Override"),
+								LOCTEXT("RemoveElementNameOverrideTooltip", "Removes the override from the element name mode and uses the AssetDefault option."),
+								FSlateIcon(),
+								FUIAction(
+									FExecuteAction::CreateLambda([this]()
+									{
+										SetElementNameDisplayMode(EElementNameDisplayMode::AssetDefault);
+									})
+								)
+							);
+							return MenuBuilder.MakeWidget();
+						}
+						return SNullWidget::NullWidget;
+					})
+				]
 			]
 			,
 			LOCTEXT("ElementNameDisplayMode", "Name Mode")
diff --git a/Engine/Plugins/Animation/ControlRig/Source/ControlRigEditor/Private/Editor/ControlRigEditor.cpp b/Engine/Plugins/Animation/ControlRig/Source/ControlRigEditor/Private/Editor/ControlRigEditor.cpp
--- a/Engine/Plugins/Animation/ControlRig/Source/ControlRigEditor/Private/Editor/ControlRigEditor.cpp
+++ b/Engine/Plugins/Animation/ControlRig/Source/ControlRigEditor/Private/Editor/ControlRigEditor.cpp
@@ -124,6 +124,7 @@
 #include "ITransportControl.h"
 #include "PropertyPath.h"
 #include "Editor/SRigVMDetailsInspector.h"
+#include "Overrides/OverrideStatusDetailsObjectFilter.h"
 
 #if WITH_RIGVMLEGACYEDITOR
 #include "SKismetInspector.h"
@@ -1568,7 +1569,228 @@ void FControlRigBaseEditor::SetDetailViewForRigModules(const TArray<FName> InMod
 			Objects.Add(ModuleInstance);
 		}
 	}
-	
+
+	if(!Objects.IsEmpty() && CVarRigVMEnablePinOverrides.GetValueOnAnyThread())
+	{
+		TSharedPtr<FOverrideStatusDetailsViewObjectFilter> ObjectFilter = FOverrideStatusDetailsViewObjectFilter::Create();
+
+		ObjectFilter->OnCanMergeObjects().BindLambda([](const UObject* InObjectA, const UObject* InObjectB) -> bool
+		{
+			if(InObjectA && InObjectB)
+			{
+				return InObjectA->IsA<UControlRig>() && InObjectB->IsA<UControlRig>(); 
+			}
+			return false;
+		});
+
+		ObjectFilter->OnCanCreateWidget().BindLambda([](const FOverrideStatusSubject& InSubject) -> bool
+		{
+			static const TSet<FName> CategoriesToIgnore = {
+				TEXT("General"),
+				TEXT("Connections"),
+			};
+			return InSubject.Contains<UControlRig>() && !CategoriesToIgnore.Contains(InSubject.GetCategory());
+		});
+
+		ObjectFilter->OnGetStatus().BindLambda([RigBlueprint](const FOverrideStatusSubject& InSubject)
+		{
+			const FString PropertyPath = InSubject.GetPropertyPathString();
+			return InSubject.GetStatus<UControlRig>(
+				[PropertyPath, RigBlueprint](const UControlRig* InModuleRig) -> TOptional<EOverrideWidgetStatus::Type>
+				{
+					if(const FRigModuleReference* ModuleReference = RigBlueprint->ModularRigModel.FindModule(InModuleRig->GetFName()))
+					{
+						if(PropertyPath.IsEmpty())
+						{
+							if(!ModuleReference->ConfigOverrides.IsEmpty())
+							{
+								return EOverrideWidgetStatus::ChangedInside; 
+							}
+						}
+						else
+						{
+							if(ModuleReference->ConfigOverrides.Contains(PropertyPath, ModuleReference->Name))
+							{
+								return EOverrideWidgetStatus::ChangedHere;
+							}
+
+							if(ModuleReference->ConfigOverrides.ContainsChildPathOf(PropertyPath, ModuleReference->Name))
+							{
+								return EOverrideWidgetStatus::ChangedInside;
+							}
+							if(ModuleReference->ConfigOverrides.ContainsParentPathOf(PropertyPath, ModuleReference->Name))
+							{
+								return EOverrideWidgetStatus::ChangedOutside;
+							}
+						}
+						return EOverrideWidgetStatus::None;
+					}
+					return {};
+				}
+			).Get(EOverrideWidgetStatus::Mixed);
+		});
+
+		ObjectFilter->OnAddOverride().BindLambda([RigBlueprint](const FOverrideStatusSubject& InSubject)
+			{
+				FScopedTransaction Transaction(LOCTEXT("AddOverride", "Add Override"));
+				RigBlueprint->Modify();
+
+				InSubject.ForEach<UControlRig>(
+					[RigBlueprint, InSubject](const UControlRig* InModuleRig)
+					{
+						if(InSubject.HasPropertyPath())
+						{
+							const FControlRigOverrideValue Value(InSubject.GetPropertyPathString(), InModuleRig);
+							if(Value.IsValid())
+							{
+								RigBlueprint->GetModularRigController()->SetConfigValueInModule(InModuleRig->GetFName(), Value);
+							}
+						}
+						else
+						{
+							const TArray<FRigVMExternalVariable> Variables = InModuleRig->GetPublicVariables();
+							for(const FRigVMExternalVariable& Variable : Variables)
+							{
+								const FString PropertyPath = Variable.Name.ToString();
+								const FControlRigOverrideValue Value(PropertyPath, InModuleRig);
+								RigBlueprint->GetModularRigController()->SetConfigValueInModule(InModuleRig->GetFName(), Value);
+							}
+						}
+					}
+				);
+
+				return FReply::Handled();
+			});
+
+		ObjectFilter->OnClearOverride().BindLambda(
+			[this, RigBlueprint](const FOverrideStatusSubject& InSubject)
+			{
+				FScopedTransaction Transaction(LOCTEXT("ClearOverride", "Clear Override"));
+				RigBlueprint->Modify();
+
+				InSubject.ForEach<UControlRig>(
+					[RigBlueprint, InSubject](const UControlRig* InModuleRig)
+					{
+						if(InSubject.HasPropertyPath())
+						{
+							RigBlueprint->GetModularRigController()->ResetConfigValueInModule(
+								InModuleRig->GetFName(), InSubject.GetPropertyPathString(), true);
+						}
+						else
+						{
+							if(const FRigModuleReference* ModuleReference = RigBlueprint->ModularRigModel.FindModule(InModuleRig->GetFName()))
+							{
+								TArray<FString> PathsToClear;
+								for(const FControlRigOverrideValue& Override : ModuleReference->ConfigOverrides)
+								{
+									if(Override.IsValid())
+									{
+										PathsToClear.Add(Override.GetPath());
+									}
+								}
+								for(const FString& Path : PathsToClear)
+								{
+									RigBlueprint->GetModularRigController()->ResetConfigValueInModule(ModuleReference->GetFName(), Path, true);
+								}
+							}
+						}
+					}
+				);
+
+				return FReply::Handled();
+			});
+
+		ObjectFilter->OnResetToDefault().BindLambda(
+			[this, RigBlueprint](const FOverrideStatusSubject& InSubject)
+			{
+				FScopedTransaction Transaction(LOCTEXT("ResetConfigValue", "Reset Config Value"));
+				RigBlueprint->Modify();
+
+				InSubject.ForEach<UControlRig>(
+					[RigBlueprint, InSubject](const UControlRig* InModuleRig)
+					{
+						if(InSubject.HasPropertyPath())
+						{
+							RigBlueprint->GetModularRigController()->ResetConfigValueInModule(
+								InModuleRig->GetFName(), InSubject.GetPropertyPathString(), false);
+						}
+						else
+						{
+							if(const FRigModuleReference* ModuleReference = RigBlueprint->ModularRigModel.FindModule(InModuleRig->GetFName()))
+							{
+								TArray<FString> PathsToClear;
+								for(const FControlRigOverrideValue& Override : ModuleReference->ConfigOverrides)
+								{
+									if(Override.IsValid())
+									{
+										PathsToClear.Add(Override.GetPath());
+									}
+								}
+								for(const FString& Path : PathsToClear)
+								{
+									RigBlueprint->GetModularRigController()->ResetConfigValueInModule(ModuleReference->GetFName(), Path, false);
+								}
+							}
+						}
+					}
+				);
+
+				return FReply::Handled();
+			});
+
+		ObjectFilter->OnValueDiffersFromDefault().BindLambda(
+			[this, RigBlueprint](const FOverrideStatusSubject& InSubject)
+			{
+				return InSubject.GetCommonValue<bool, UControlRig>(
+					[RigBlueprint, InSubject](const UControlRig* InModuleRig)
+					{
+						if(const FRigModuleReference* ModuleReference = RigBlueprint->ModularRigModel.FindModule(InModuleRig->GetFName()))
+						{
+							TArray<TSharedPtr<const FPropertyPath>> PropertyPathsToCheck;
+							if(InSubject.HasPropertyPath())
+							{
+								PropertyPathsToCheck.Add(InSubject.GetPropertyPath());
+							}
+							else
+							{
+								const TArray<FRigVMExternalVariable> Variables = InModuleRig->GetPublicVariables();
+								for(const FRigVMExternalVariable& Variable : Variables)
+								{
+									if(FProperty* Property = InModuleRig->GetClass()->FindPropertyByName(Variable.Name))
+									{
+										PropertyPathsToCheck.Add(FPropertyPath::Create(TWeakFieldPtr<FProperty>(Property)));
+									}
+								}
+							}
+
+							const UClass* Class = ModuleReference->Class.LoadSynchronous();
+							const UObject* CDO = Class->GetDefaultObject();
+							
+							for(const TSharedPtr<const FPropertyPath>& PropertyPath : PropertyPathsToCheck)
+							{
+								const FString PropertyPathString = PropertyPath->ToString();
+								const FString PropertyPathPrefix = PropertyPathString + TEXT("->");
+								for(const FControlRigOverrideValue& Override : ModuleReference->ConfigOverrides)
+								{
+									if(!Override.IdenticalValueInSubject(CDO))
+									{
+										const FString& ValuePathString = Override.GetPath();
+										if(ValuePathString == PropertyPathString || ValuePathString.StartsWith(PropertyPathPrefix))
+										{
+											return true;
+										}
+									}
+								}
+							}
+						}
+						return false;
+					}
+				).Get(false);
+			});
+
+		SetDetailObjectFilter(ObjectFilter);
+	}
+
 	SetDetailObjects(Objects);
 
 	// In case the modules selected are still not available, lets set them again
diff --git a/Engine/Plugins/Animation/ControlRig/Source/ControlRigEditor/Private/Editor/ControlRigEditor.h b/Engine/Plugins/Animation/ControlRig/Source/ControlRigEditor/Private/Editor/ControlRigEditor.h
--- a/Engine/Plugins/Animation/ControlRig/Source/ControlRigEditor/Private/Editor/ControlRigEditor.h
+++ b/Engine/Plugins/Animation/ControlRig/Source/ControlRigEditor/Private/Editor/ControlRigEditor.h
@@ -133,6 +133,7 @@ protected:
 	virtual void OnWrappedPropertyChangedChainEvent(URigVMDetailsViewWrapperObject* InWrapperObject, const FString& InPropertyPath, FPropertyChangedChainEvent& InPropertyChangedChainEvent) = 0;
 	virtual void SetDetailObjects(const TArray<UObject*>& InObjects) = 0;
 	virtual void SetDetailObjectsSuper(const TArray<UObject*>& InObjects) = 0;
+	virtual void SetDetailObjectFilter(TSharedPtr<FDetailsViewObjectFilter> InObjectFilter) = 0;
 	virtual bool DetailViewShowsStruct(UScriptStruct* InStruct) const = 0;
 	virtual TSharedPtr<class SWidget> GetInspector() const = 0;
 	virtual TArray<FName> GetEventQueue() const = 0;
diff --git a/Engine/Plugins/Animation/ControlRig/Source/ControlRigEditor/Private/Editor/ControlRigEditorMode.cpp b/Engine/Plugins/Animation/ControlRig/Source/ControlRigEditor/Private/Editor/ControlRigEditorMode.cpp
--- a/Engine/Plugins/Animation/ControlRig/Source/ControlRigEditor/Private/Editor/ControlRigEditorMode.cpp
+++ b/Engine/Plugins/Animation/ControlRig/Source/ControlRigEditor/Private/Editor/ControlRigEditorMode.cpp
@@ -20,6 +20,7 @@
 #include "Editor/RigVMCompilerResultsTabSummoner.h"
 #include "Editor/RigVMDetailsInspectorTabSummoner.h"
 #include "Editor/RigVMFindReferencesTabSummoner.h"
+#include "Editor/RigOverrideListTabSummoner.h"
 
 #if WITH_RIGVMLEGACYEDITOR
 FControlRigLegacyEditorMode::FControlRigLegacyEditorMode(const TSharedRef<FControlRigLegacyEditor>& InControlRigEditor, bool bCreateDefaultLayout)
@@ -159,7 +160,11 @@ FModularRigLegacyEditorMode::FModularRigLegacyEditorMode(const TSharedRef<FContr
 
 	TabFactories.RegisterFactory(MakeShared<FModularRigModelTabSummoner>(InControlRigEditor));
 	TabFactories.RegisterFactory(MakeShared<FRigModuleAssetBrowserTabSummoner>(InControlRigEditor));
-	
+	if(CVarRigVMEnablePinOverrides.GetValueOnAnyThread())
+	{
+		TabFactories.RegisterFactory(MakeShared<FRigOverrideListTabSummoner>(InControlRigEditor));
+	}
+
 	TabLayout = FTabManager::NewLayout("Standalone_ModularRigEditMode_Layout_v1.3")
 		->AddArea
 		(
@@ -375,7 +380,11 @@ FModularRigEditorMode::FModularRigEditorMode(const TSharedRef<FControlRigEditor>
 
 	TabFactories.RegisterFactory(MakeShared<FModularRigModelTabSummoner>(InControlRigEditor));
 	TabFactories.RegisterFactory(MakeShared<FRigModuleAssetBrowserTabSummoner>(InControlRigEditor));
-	
+	if(CVarRigVMEnablePinOverrides.GetValueOnAnyThread())
+	{
+		TabFactories.RegisterFactory(MakeShared<FRigOverrideListTabSummoner>(InControlRigEditor));
+	}
+
 	TabLayout = FTabManager::NewLayout("Standalone_ModularRigEditMode_Layout_v1.3")
 		->AddArea
 		(
diff --git a/Engine/Plugins/Animation/ControlRig/Source/ControlRigEditor/Private/Editor/ControlRigLegacyEditor.h b/Engine/Plugins/Animation/ControlRig/Source/ControlRigEditor/Private/Editor/ControlRigLegacyEditor.h
--- a/Engine/Plugins/Animation/ControlRig/Source/ControlRigEditor/Private/Editor/ControlRigLegacyEditor.h
+++ b/Engine/Plugins/Animation/ControlRig/Source/ControlRigEditor/Private/Editor/ControlRigLegacyEditor.h
@@ -99,6 +99,7 @@ public:
 
 	virtual void SetDetailObjects(const TArray<UObject*>& InObjects) override { return FControlRigBaseEditor::SetDetailObjectsImpl(InObjects); }
 	virtual void SetDetailObjectsSuper(const TArray<UObject*>& InObjects) override { return IControlRigLegacyEditor::SetDetailObjects(InObjects); }
+	virtual void SetDetailObjectFilter(TSharedPtr<FDetailsViewObjectFilter> InObjectFilter) override { return IControlRigLegacyEditor::SetDetailObjectFilter(InObjectFilter); };
 	virtual void RefreshDetailView() override { return FControlRigBaseEditor::RefreshDetailViewImpl(); }
 	virtual void RefreshDetailViewSuper() override { return IControlRigLegacyEditor::RefreshDetailView(); }
 
diff --git a/Engine/Plugins/Animation/ControlRig/Source/ControlRigEditor/Private/Editor/ControlRigNewEditor.h b/Engine/Plugins/Animation/ControlRig/Source/ControlRigEditor/Private/Editor/ControlRigNewEditor.h
--- a/Engine/Plugins/Animation/ControlRig/Source/ControlRigEditor/Private/Editor/ControlRigNewEditor.h
+++ b/Engine/Plugins/Animation/ControlRig/Source/ControlRigEditor/Private/Editor/ControlRigNewEditor.h
@@ -98,6 +98,7 @@ public:
 
 	virtual void SetDetailObjects(const TArray<UObject*>& InObjects) override { return FControlRigBaseEditor::SetDetailObjectsImpl(InObjects); }
 	virtual void SetDetailObjectsSuper(const TArray<UObject*>& InObjects) override { return IControlRigNewEditor::SetDetailObjects(InObjects); }
+	virtual void SetDetailObjectFilter(TSharedPtr<FDetailsViewObjectFilter> InObjectFilter) override { return IControlRigNewEditor::SetDetailObjectFilter(InObjectFilter); };
 	virtual void RefreshDetailView() override { return FControlRigBaseEditor::RefreshDetailViewImpl(); }
 	virtual void RefreshDetailViewSuper() override { return IControlRigNewEditor::RefreshDetailView(); }
 
diff --git a/Engine/Plugins/Animation/ControlRig/Source/ControlRigEditor/Private/Editor/RigOverrideListTabSummoner.cpp b/Engine/Plugins/Animation/ControlRig/Source/ControlRigEditor/Private/Editor/RigOverrideListTabSummoner.cpp
new file mode 100644
--- /dev/null
+++ b/Engine/Plugins/Animation/ControlRig/Source/ControlRigEditor/Private/Editor/RigOverrideListTabSummoner.cpp
@@ -0,0 +1,43 @@
+// Copyright Epic Games, Inc. All Rights Reserved.
+
+#include "Editor/RigOverrideListTabSummoner.h"
+#include "Editor/RigVMEditorStyle.h"
+#include "Editor/ControlRigEditor.h"
+#include "ControlRigBlueprint.h"
+#include "Overrides/SOverrideListWidget.h"
+
+#define LOCTEXT_NAMESPACE "RigOverrideListTabSummoner"
+
+const FName FRigOverrideListTabSummoner::TabID(TEXT("OverrideList"));
+
+FRigOverrideListTabSummoner::FRigOverrideListTabSummoner(const TSharedRef<IControlRigBaseEditor>& InControlRigEditor)
+	: FWorkflowTabFactory(TabID, InControlRigEditor->GetHostingApp())
+	, ControlRigEditor(InControlRigEditor)
+{
+	TabLabel = LOCTEXT("OverrideListTabLabel", "Active Overrides");
+	TabIcon = FSlateIcon(FAppStyle::Get().GetStyleSetName(), "DetailsView.OverrideInside.Hovered");
+	
+	bIsSingleton = true;
+
+	ViewMenuDescription = LOCTEXT("OverrideList_ViewMenu_Desc", "Active Overrides");
+	ViewMenuTooltip = LOCTEXT("OverrideList_ViewMenu_ToolTip", "Show the Active Overrides tab");
+}
+
+TSharedRef<SWidget> FRigOverrideListTabSummoner::CreateTabBody(const FWorkflowTabSpawnInfo& Info) const
+{
+	if (TSharedPtr<IControlRigBaseEditor> Editor = ControlRigEditor.Pin())
+	{
+		if (UControlRigBlueprint* RigBlueprint = Cast<UControlRigBlueprint>(Editor->GetControlRigBlueprint()))
+		{
+			check(RigBlueprint->IsModularRig());
+
+			return SNew(SOverrideListWidget)
+				.SubjectsHash_UObject(RigBlueprint, &UControlRigBlueprint::GetOverrideSubjectsHash)
+				.Subjects_UObject(RigBlueprint, &UControlRigBlueprint::GetOverrideSubjects);
+		}
+	}
+
+	return SNullWidget::NullWidget;
+}
+
+#undef LOCTEXT_NAMESPACE 
diff --git a/Engine/Plugins/Animation/ControlRig/Source/ControlRigEditor/Private/Editor/RigOverrideListTabSummoner.h b/Engine/Plugins/Animation/ControlRig/Source/ControlRigEditor/Private/Editor/RigOverrideListTabSummoner.h
new file mode 100644
--- /dev/null
+++ b/Engine/Plugins/Animation/ControlRig/Source/ControlRigEditor/Private/Editor/RigOverrideListTabSummoner.h
@@ -0,0 +1,22 @@
+// Copyright Epic Games, Inc. All Rights Reserved.
+
+#pragma once
+
+#include "CoreMinimal.h"
+#include "WorkflowOrientedApp/WorkflowTabFactory.h"
+
+class IControlRigBaseEditor;
+
+struct FRigOverrideListTabSummoner : public FWorkflowTabFactory
+{
+public:
+	static const FName TabID;
+	
+public:
+	FRigOverrideListTabSummoner(const TSharedRef<IControlRigBaseEditor>& InControlRigEditor);
+	
+	virtual TSharedRef<SWidget> CreateTabBody(const FWorkflowTabSpawnInfo& Info) const override;
+	
+protected:
+	TWeakPtr<IControlRigBaseEditor> ControlRigEditor;
+};
diff --git a/Engine/Plugins/Animation/ControlRig/Source/ControlRigEditor/Private/Editor/SRigHierarchy.cpp b/Engine/Plugins/Animation/ControlRig/Source/ControlRigEditor/Private/Editor/SRigHierarchy.cpp
--- a/Engine/Plugins/Animation/ControlRig/Source/ControlRigEditor/Private/Editor/SRigHierarchy.cpp
+++ b/Engine/Plugins/Animation/ControlRig/Source/ControlRigEditor/Private/Editor/SRigHierarchy.cpp
@@ -1191,6 +1191,18 @@ TSharedRef< SWidget > SRigHierarchy::CreateFilterMenu()
 		MenuBuilder.AddMenuEntry(Actions.FilteringFlattensHierarchy);
 		MenuBuilder.AddMenuEntry(Actions.HideParentsWhenFiltering);
 
+		const UEnum* NameModeEnum = StaticEnum<EElementNameDisplayMode>();
+		TArray<int32> EnumValueSubset;
+		EnumValueSubset.Reserve(NameModeEnum->NumEnums() - 1);
+		for (int32 i = 0; i < NameModeEnum->NumEnums() - 1; i++)
+		{
+			const int32 Value = static_cast<int32>(NameModeEnum->GetValueByIndex(i));
+			if(Value != static_cast<int32>(EElementNameDisplayMode::AssetDefault))
+			{
+				EnumValueSubset.Add(Value);
+			}
+		}
+
 		MenuBuilder.AddWidget(
 			SNew(SVerticalBox)
 
@@ -1198,7 +1210,8 @@ TSharedRef< SWidget > SRigHierarchy::CreateFilterMenu()
 			.AutoHeight()
 			.Padding(3)
 			[
-				SNew(SEnumComboBox, StaticEnum<EElementNameDisplayMode>())
+				SNew(SEnumComboBox, NameModeEnum)
+				.EnumValueSubset(EnumValueSubset)
 				.CurrentValue_Lambda([this]() -> int32
 				{
 					const IControlRigBaseEditor* StrongEditor = ControlRigEditor.Pin().Get();
