Control Rig: Follow fix to 41601390

#rnx

diff --git a/Engine/Plugins/Animation/ControlRig/Source/ControlRig/Private/ControlRigOverride.cpp b/Engine/Plugins/Animation/ControlRig/Source/ControlRig/Private/ControlRigOverride.cpp
--- a/Engine/Plugins/Animation/ControlRig/Source/ControlRig/Private/ControlRigOverride.cpp
+++ b/Engine/Plugins/Animation/ControlRig/Source/ControlRig/Private/ControlRigOverride.cpp
@@ -329,7 +329,7 @@ void FControlRigOverrideValue::Reset()
 
 const FName& FControlRigOverrideValue::GetSubjectName() const
 {
-	static const FName NoneName(NAME_None);
+	static const FName NoneName = FName(NAME_None);
 	return SubjectName.Get(NoneName);
 }
 
