[Backout] - CL41601390
#rnx
Original CL Desc
-----------------------------------------------------------------
Modular Control Rig support for Overrides

#rnx

diff --git a/Engine/Plugins/Animation/ControlRig/Source/ControlRig/Private/ControlRigOverride.cpp b/Engine/Plugins/Animation/ControlRig/Source/ControlRig/Private/ControlRigOverride.cpp
deleted file mode 100644
diff --git a/Engine/Plugins/Animation/ControlRig/Source/ControlRig/Private/ModularRig.cpp b/Engine/Plugins/Animation/ControlRig/Source/ControlRig/Private/ModularRig.cpp
--- a/Engine/Plugins/Animation/ControlRig/Source/ControlRig/Private/ModularRig.cpp
+++ b/Engine/Plugins/Animation/ControlRig/Source/ControlRig/Private/ModularRig.cpp
@@ -351,7 +351,7 @@ void UModularRig::UpdateModuleHierarchyFromCDO()
 					InModuleReference->Class.Get(),
 					FindModule(InModuleReference->ParentModuleName),
 					Model.Connections.GetModuleConnectionMap(InModuleReference->Name),
-					InModuleReference->ConfigOverrides);
+					InModuleReference->ConfigValues);
 			}
 
 			// continue to the next module
@@ -548,7 +548,10 @@ void UModularRig::ExecuteQueue()
 					// put the variable defaults back
 					if(const FRigModuleReference* ModuleReference = GetModularRigModel().FindModule(ExecutionElement.ModuleName))
 					{
-						ModuleReference->ConfigOverrides.CopyToSubject(ModuleRig);
+						for (const TPair<FName, FString>& Variable : ModuleReference->ConfigValues)
+						{
+							ModuleRig->SetVariableFromString(Variable.Key, Variable.Value);
+						}
 					}
 				}
 
@@ -783,7 +786,7 @@ TArray<FName> UModularRig::GetModuleNames() const
 }
 
 FRigModuleInstance* UModularRig::AddModuleInstance(const FName& InModuleName, TSubclassOf<UControlRig> InModuleClass, const FRigModuleInstance* InParent,
-                                                   const FRigElementKeyRedirector::FKeyMap& InConnectionMap, const FControlRigOverrideContainer& InConfigValues) 
+                                                   const FRigElementKeyRedirector::FKeyMap& InConnectionMap, const TMap<FName, FString>& InVariableDefaultValues) 
 {
 	// Make sure there are no name clashes
 	FName ParentName = NAME_None;
@@ -845,20 +848,6 @@ FRigModuleInstance* UModularRig::AddModuleInstance(const FName& InModuleName, TS
 	{
 		NewModuleRig = NewObject<UControlRig>(this, InModuleClass, InModuleName);
 	}
-	else
-	{
-		// make sure to reset its public variables back to the value from the CDO
-		const UControlRig* CDO = NewModuleRig->GetClass()->GetDefaultObject<UControlRig>();
-		for (TFieldIterator<FProperty> PropertyIt( NewModuleRig->GetClass()); PropertyIt; ++PropertyIt)
-		{
-			const FProperty* Property = *PropertyIt;
-			if(Property->IsNative())
-			{
-				continue;
-			}
-			Property->CopyCompleteValue_InContainer(NewModuleRig, CDO);
-		}
-	}
 	
 	NewModule.SetRig(NewModuleRig);
 
@@ -882,8 +871,10 @@ FRigModuleInstance* UModularRig::AddModuleInstance(const FName& InModuleName, TS
 		ModulePublicContext.RigModulePrefixHash = GetTypeHash(ModulePublicContext.RigModulePrefix);
 		NewModuleRig->SetElementKeyRedirector(FRigElementKeyRedirector(InConnectionMap, Hierarchy));
 
-		UControlRig* ModuleRig = NewModule.GetRig();
-		InConfigValues.CopyToSubject(ModuleRig);
+		for (const TPair<FName, FString>& Variable : InVariableDefaultValues )
+		{
+			NewModule.GetRig()->SetVariableFromString(Variable.Key, Variable.Value);
+		}
 	}
 
 	return &NewModule;
diff --git a/Engine/Plugins/Animation/ControlRig/Source/ControlRig/Private/ModularRigController.cpp b/Engine/Plugins/Animation/ControlRig/Source/ControlRig/Private/ModularRigController.cpp
--- a/Engine/Plugins/Animation/ControlRig/Source/ControlRig/Private/ModularRigController.cpp
+++ b/Engine/Plugins/Animation/ControlRig/Source/ControlRig/Private/ModularRigController.cpp
@@ -6,9 +6,6 @@
 #include "ModularRig.h"
 #include "ModularRigModel.h"
 #include "ModularRigRuleManager.h"
-#if WITH_EDITOR
-#include "PropertyPath.h"
-#endif
 #include "Misc/DefaultValueHelper.h"
 #include "Rigs/RigHierarchyController.h"
 #include "RigVMFunctions/Math/RigVMMathLibrary.h"
@@ -853,54 +850,31 @@ bool UModularRigController::SetConfigValueInModule(const FName& InModuleName, co
 		return false;
 	}
 
-	const FControlRigOverrideValue OverrideValue(InVariableName.ToString(), Module->Class.Get(), InValue, Module->Name);
-	return SetConfigValueInModule(InModuleName, OverrideValue, bSetupUndo);
-}
-
-bool UModularRigController::SetConfigValueInModule(const FName& InModuleName, const FControlRigOverrideValue& InValue, bool bSetupUndo)
-{
-	FRigModuleReference* Module = FindModule(InModuleName);
-	if (!Module)
-	{
-		UE_LOG(LogControlRig, Error, TEXT("Could not find module %s"), *InModuleName.ToString());
-		return false;
-	}
-
-	if (!Module->Class.IsValid())
+	const FProperty* Property = Module->Class->FindPropertyByName(InVariableName);
+	if (!Property)
 	{
-		UE_LOG(LogControlRig, Error, TEXT("Class defined in module %s is not valid"), *InModuleName.ToString());
+		UE_LOG(LogControlRig, Error, TEXT("Could not find variable %s in module %s"), *InVariableName.ToString(), *InModuleName.ToString());
 		return false;
 	}
 
-	if(!InValue.IsValid())
+	if (Property->HasAllPropertyFlags(CPF_BlueprintReadOnly))
 	{
-		UE_LOG(LogControlRig, Error, TEXT("Provided value is not valid."));
+		UE_LOG(LogControlRig, Error, TEXT("The target variable %s in module %s is read only"), *InVariableName.ToString(), *InModuleName.ToString());
 		return false;
 	}
 
-	const FString& Path = InValue.GetPath();
-
-	// we cannot set a config value if we already have a value set on a parent path.
-	// so for example you can't override Color.R if Color itself has been overridden already.
-	if(Module->ConfigOverrides.ContainsParentPathOf(InValue))
-	{
-		UE_LOG(LogControlRig, Error, TEXT("Cannot set a config value for '%s' for the module instance class %s since there's a value on a parent already."), *Path, *Module->Class.ToString());
-		return false;
-	}
-
-	if(InValue.GetRootProperty()->GetOwnerUObject() != Module->Class.Get())
-	{
-		UE_LOG(LogControlRig, Error, TEXT("Provided path '%s' does not belong to the module instance class %s."), *Path, *Module->Class.ToString());
-		return false;
-	}
+#if WITH_EDITOR
+	TArray<uint8, TAlignedHeapAllocator<16>> TempStorage;
+	TempStorage.AddZeroed(Property->GetSize());
+	uint8* TempMemory = TempStorage.GetData();
+	Property->InitializeValue(TempMemory);
 
-	if (InValue.GetRootProperty()->HasAllPropertyFlags(CPF_BlueprintReadOnly))
+	if (!FBlueprintEditorUtils::PropertyValueFromString_Direct(Property, InValue, TempMemory))
 	{
-		UE_LOG(LogControlRig, Error, TEXT("The target property %s in module %s is read only"), *InValue.GetRootProperty()->GetName(), *InModuleName.ToString());
+		UE_LOG(LogControlRig, Error, TEXT("Value %s for variable %s in module %s is not valid"), *InValue, *InVariableName.ToString(), *InModuleName.ToString());
 		return false;
 	}
 
-#if WITH_EDITOR
 	TSharedPtr<FScopedTransaction> TransactionPtr;
 	if (bSetupUndo)
 	{
@@ -910,26 +884,9 @@ bool UModularRigController::SetConfigValueInModule(const FName& InModuleName, co
 			Blueprint->Modify();
 		}
 	}
-#endif
-	
-	// remove all nested config values (child properties under the provided ones)
-	{
-		TArray<FString> KeysToRemove;
-		for(const FControlRigOverrideValue& Override : Module->ConfigOverrides)
-		{
-			const FString& ChildPath = Override.GetPath();
-			if(FControlRigOverrideContainer::IsChildPathOf(ChildPath, Path))
-			{
-				KeysToRemove.Add(ChildPath);
-			}
-		}
-		for(const FString& KeyToRemove : KeysToRemove)
-		{
-			Module->ConfigOverrides.Remove(KeyToRemove, Module->Name);
-		}
-	}
+#endif 
 
-	Module->ConfigOverrides.FindOrAdd(InValue);
+	Module->ConfigValues.FindOrAdd(InVariableName) = InValue;
 
 	Notify(EModularRigNotification::ModuleConfigValueChanged, Module);
 
@@ -940,80 +897,6 @@ bool UModularRigController::SetConfigValueInModule(const FName& InModuleName, co
 	return true;
 }
 
-bool UModularRigController::ResetConfigValueInModule(const FName& InModuleName, const FString& InPath, bool bClearOverride, bool bSetupUndo)
-{
-	FRigModuleReference* Module = FindModule(InModuleName);
-	if (!Module)
-	{
-		UE_LOG(LogControlRig, Error, TEXT("Could not find module %s"), *InModuleName.ToString());
-		return false;
-	}
-
-	if (!Module->Class.IsValid())
-	{
-		UE_LOG(LogControlRig, Error, TEXT("Class defined in module %s is not valid"), *InModuleName.ToString());
-		return false;
-	}
-
-	if(InPath.IsEmpty())
-	{
-		UE_LOG(LogControlRig, Error, TEXT("Provided path is not valid."));
-		return false;
-	}
-
-	const FControlRigOverrideValue DefaultValue(InPath, Module->Class.Get()->GetDefaultObject());
-	if(!DefaultValue.IsValid())
-	{
-		UE_LOG(LogControlRig, Error, TEXT("Provided path '%s' does not belong to the module instance class %s."), *InPath, *Module->Class.ToString());
-		return false;
-	}
-
-	if(!bClearOverride)
-	{
-		SetConfigValueInModule(InModuleName, DefaultValue, bSetupUndo);
-		return true;
-	}
-	
-	const FString PathPrefix = InPath + TEXT("->");
-
-#if WITH_EDITOR
-	TSharedPtr<FScopedTransaction> TransactionPtr;
-	if (bSetupUndo)
-	{
-		TransactionPtr = MakeShared<FScopedTransaction>(NSLOCTEXT("ModularRigController", "ConfigureModuleValueTransaction", "Configure Module Value"), !GIsTransacting);
-		if(UBlueprint* Blueprint = Cast<UBlueprint>(GetOuter()))
-		{
-			Blueprint->Modify();
-		}
-	}
-#endif
-	
-	// remove all nested config values (child properties under the provided ones)
-	{
-		TArray<FString> KeysToRemove;
-		for(const FControlRigOverrideValue& Override : Module->ConfigOverrides)
-		{
-			const FString& ChildPath = Override.GetPath();
-			if(InPath == ChildPath || FControlRigOverrideContainer::IsChildPathOf(ChildPath, InPath))
-			{
-				KeysToRemove.Add(ChildPath);
-			}
-		}
-		for(const FString& KeyToRemove : KeysToRemove)
-		{
-			Module->ConfigOverrides.Remove(KeyToRemove, Module->Name);
-		}
-	}
-
-	Notify(EModularRigNotification::ModuleConfigValueChanged, Module);
-
-#if WITH_EDITOR
-	TransactionPtr.Reset();
-#endif
-	
-	return true;
-}
-
 TArray<FString> UModularRigController::GetPossibleBindings(const FName& InModuleName, const FName& InVariableName)
 {
 	TArray<FString> PossibleBindings;
@@ -1562,7 +1445,7 @@ FName UModularRigController::MirrorModule(const FName& InModuleName, const FRigV
 	const TMap<FName, FString> OriginalBindings = OriginalModule->Bindings;
 	const TSubclassOf<UControlRig> OriginalClass = OriginalModule->Class.Get();
 	const FName OriginalParentName = OriginalModule->ParentModuleName;
-	const FControlRigOverrideContainer OriginalConfigValues = OriginalModule->ConfigOverrides;
+	const TMap<FName, FString> OriginalConfigValues = OriginalModule->ConfigValues;
 
 	FModularRigControllerCompileBracketScope CompileBracketScope(this);
 
@@ -1593,7 +1476,7 @@ FName UModularRigController::MirrorModule(const FName& InModuleName, const FRigV
 		BindModuleVariable(NewModuleName, Pair.Key, NewSourcePath, bSetupUndo);
 	}
 
-	TSet<FString> ConfigValueSet;
+	TSet<FName> ConfigValueSet;
 #if WITH_EDITOR
 	for (TFieldIterator<FProperty> PropertyIt(OriginalClass); PropertyIt; ++PropertyIt)
 	{
@@ -1629,19 +1512,19 @@ FName UModularRigController::MirrorModule(const FName& InModuleName, const FRigV
 		}
 
 		FString NewValueStr;
-		if (const FControlRigOverrideValue* OriginalValue = OriginalConfigValues.Find(Property->GetName(), OriginalModule->Name))
+		if (const FString* OriginalValue = OriginalConfigValues.Find(Property->GetFName()))
 		{
 			if (bIsVector)
 			{
 				FVector Value;
-				FBlueprintEditorUtils::PropertyValueFromString_Direct(Property, *OriginalValue->ToString(), (uint8*)&Value);
+				FBlueprintEditorUtils::PropertyValueFromString_Direct(Property, *OriginalValue, (uint8*)&Value);
 				Value = InSettings.MirrorVector(Value);
 				FBlueprintEditorUtils::PropertyValueToString_Direct(Property, (uint8*)&Value, NewValueStr, nullptr);
 			}
 			else
 			{
 				FTransform Value;
-				FBlueprintEditorUtils::PropertyValueFromString_Direct(Property, *OriginalValue->ToString(), (uint8*)&Value);
+				FBlueprintEditorUtils::PropertyValueFromString_Direct(Property, *OriginalValue, (uint8*)&Value);
 				Value = InSettings.MirrorTransform(Value);
 				FBlueprintEditorUtils::PropertyValueToString_Direct(Property, (uint8*)&Value, NewValueStr, nullptr);
 			}
@@ -1665,19 +1548,17 @@ FName UModularRigController::MirrorModule(const FName& InModuleName, const FRigV
 			}
 		}
 
-		ConfigValueSet.Add(Property->GetName());
-
-		const FControlRigOverrideValue NewValue(Property->GetName(), OriginalClass, NewValueStr, NewModuleName);
-		SetConfigValueInModule(NewModuleName, NewValue, bSetupUndo);
+		ConfigValueSet.Add(Property->GetFName());
+		SetConfigValueInModule(NewModuleName, Property->GetFName(), NewValueStr, bSetupUndo);
 	}
 #endif
 
 	// Add any other config value that was set in the original module, but was not mirrored
-	for (const FControlRigOverrideValue& OriginalOverride : OriginalConfigValues)
+	for (const TPair<FName, FString>& Pair : OriginalConfigValues)
 	{
-		if (!ConfigValueSet.Contains(OriginalOverride.GetPath()))
+		if (!ConfigValueSet.Contains(Pair.Key))
 		{
-			SetConfigValueInModule(NewModuleName, OriginalOverride, bSetupUndo);
+			SetConfigValueInModule(NewModuleName, Pair.Key, Pair.Value, bSetupUndo);
 		}
 	}
 	
@@ -1938,47 +1819,19 @@ void UModularRigController::RefreshModuleVariables(const FRigModuleReference* In
 		const bool bIsInstanceEditable = !Property->HasAnyPropertyFlags(CPF_DisableEditOnInstance);
 		if (bIsAdvanced || !bIsPublic || !bIsInstanceEditable)
 		{
-			FString PropertyName = Property->GetName();
-			FString PropertyPrefix = PropertyName + TEXT("->");
-
-			TArray<FString> ConfigValuesToRemove;
-			TArray<FName> BindingsToRemove;
-			for (const FControlRigOverrideValue& Override : Module->ConfigOverrides)
-			{
-				const FString& KeyString = Override.GetPath();
-				if(KeyString == PropertyName || KeyString.StartsWith(PropertyPrefix))
-				{
-					ConfigValuesToRemove.Add(KeyString);
-				}
-			}
-			for (const TPair<FName, FString>& Pair : Module->Bindings)
-			{
-				const FString KeyString = Pair.Key.ToString();
-				if(KeyString == PropertyName || KeyString.StartsWith(PropertyPrefix))
-				{
-					BindingsToRemove.Add(Pair.Key);
-				}
-			}
-
-			for(const FString& ConfigValueToRemove : ConfigValuesToRemove)
-			{
-				Module->ConfigOverrides.Remove(ConfigValueToRemove, ModuleName);
-			}
-			for(const FName& BindingToRemove : BindingsToRemove)
-			{
-				Module->Bindings.Remove(BindingToRemove);
-			}
+			Module->ConfigValues.Remove(Property->GetFName());
+			Module->Bindings.Remove(Property->GetFName());
 		}
 	}
 
 	// Make sure all the types are valid
-	const FControlRigOverrideContainer ConfigOverrides = Module->ConfigOverrides;
+	const TMap<FName, FString> ConfigValues = Module->ConfigValues;
 	const TMap<FName, FString> Bindings = Module->Bindings;
-	Module->ConfigOverrides.Reset();
+	Module->ConfigValues.Reset();
 	Module->Bindings.Reset();
-	for (const FControlRigOverrideValue& Override : ConfigOverrides)
+	for (const TPair<FName, FString>& Pair : ConfigValues)
 	{
-		SetConfigValueInModule(ModuleName, Override, false);
+		SetConfigValueInModule(ModuleName, Pair.Key, Pair.Value, false);
 	}
 	for (const TPair<FName, FString>& Pair : Bindings)
 	{
diff --git a/Engine/Plugins/Animation/ControlRig/Source/ControlRig/Private/ModularRigModel.cpp b/Engine/Plugins/Animation/ControlRig/Source/ControlRig/Private/ModularRigModel.cpp
--- a/Engine/Plugins/Animation/ControlRig/Source/ControlRig/Private/ModularRigModel.cpp
+++ b/Engine/Plugins/Animation/ControlRig/Source/ControlRig/Private/ModularRigModel.cpp
@@ -102,24 +102,6 @@ TArray<const FRigConnectorElement*> FRigModuleReference::FindConnectors(const UR
 	return Connectors;
 }
 
-void FRigModuleReference::PatchModelsOnLoad()
-{
-	if(const UClass* ClassPtr = Class.Get())
-	{
-		if(!ConfigValues_DEPRECATED.IsEmpty())
-		{
-			ConfigOverrides.Reset();
-			ConfigOverrides.SetUsesSubjectName(false);
-			ConfigOverrides.Reserve(ConfigValues_DEPRECATED.Num());
-			for(const TPair<FName, FString>& Pair : ConfigValues_DEPRECATED)
-			{
-				ConfigOverrides.Add(FControlRigOverrideValue(Pair.Key.ToString(), ClassPtr, Pair.Value, Name));
-			}
-			ConfigValues_DEPRECATED.Reset();
-		}
-	}
-}
-
 FRigElementKeyRedirector::FKeyArray FModularRigSingleConnection::GetTargetArray() const
 {
 	FRigElementKeyRedirector::FKeyArray Result;
@@ -247,12 +229,6 @@ void FModularRigModel::PatchModelsOnLoad()
 
 	Connections.PatchOnLoad(&PreviousModulePaths);
 	Connections.UpdateFromConnectionList();
-
-	for(int32 ModuleIndex = 0; ModuleIndex < Modules.Num(); ModuleIndex++)
-	{
-		FRigModuleReference& Module = Modules[ModuleIndex];
-		Module.PatchModelsOnLoad();
-	}
 }
 
 UModularRigController* FModularRigModel::GetController(bool bCreateIfNeeded)
diff --git a/Engine/Plugins/Animation/ControlRig/Source/ControlRig/Public/ControlRigOverride.h b/Engine/Plugins/Animation/ControlRig/Source/ControlRig/Public/ControlRigOverride.h
deleted file mode 100644
diff --git a/Engine/Plugins/Animation/ControlRig/Source/ControlRig/Public/ModularRig.h b/Engine/Plugins/Animation/ControlRig/Source/ControlRig/Public/ModularRig.h
--- a/Engine/Plugins/Animation/ControlRig/Source/ControlRig/Public/ModularRig.h
+++ b/Engine/Plugins/Animation/ControlRig/Source/ControlRig/Public/ModularRig.h
@@ -306,7 +306,7 @@ public:
 private:
 
 	/** Adds a module to the rig*/
-	FRigModuleInstance* AddModuleInstance(const FName& InModuleName, TSubclassOf<UControlRig> InModuleClass, const FRigModuleInstance* InParent, const FRigElementKeyRedirector::FKeyMap& InConnectionMap, const FControlRigOverrideContainer& InConfigValues);
+	FRigModuleInstance* AddModuleInstance(const FName& InModuleName, TSubclassOf<UControlRig> InModuleClass, const FRigModuleInstance* InParent, const FRigElementKeyRedirector::FKeyMap& InConnectionMap, const TMap<FName, FString>& InVariableDefaultValues);
 
 	/** Updates the module's variable bindings */
 	bool SetModuleVariableBindings(const FName& InModuleName, const TMap<FName, FString>& InVariableBindings);
diff --git a/Engine/Plugins/Animation/ControlRig/Source/ControlRig/Public/ModularRigController.h b/Engine/Plugins/Animation/ControlRig/Source/ControlRig/Public/ModularRigController.h
--- a/Engine/Plugins/Animation/ControlRig/Source/ControlRig/Public/ModularRigController.h
+++ b/Engine/Plugins/Animation/ControlRig/Source/ControlRig/Public/ModularRigController.h
@@ -57,13 +57,8 @@ class CONTROLRIG_API UModularRigController : public UObject
 	UFUNCTION(BlueprintCallable, Category = "Control Rig | Modules")
 	bool AutoConnectModules(const TArray<FName>& InModuleNames, bool bReplaceExistingConnections, bool bSetupUndo = true);
 
-	UE_DEPRECATED(5.6, "Please rely on the overload taking an FControlRigOverrideValue")
 	UFUNCTION(BlueprintCallable, Category = "Control Rig | Modules")
 	bool SetConfigValueInModule(const FName& InModuleName, const FName& InVariableName, const FString& InValue, bool bSetupUndo = true);
-	bool SetConfigValueInModule(const FName& InModuleName, const FControlRigOverrideValue& InValue, bool bSetupUndo = true);
-
-	UFUNCTION(BlueprintCallable, Category = "Control Rig | Modules")
-	bool ResetConfigValueInModule(const FName& InModuleName, const FString& InPath, bool bClearOverride = true, bool bSetupUndo = true);
 
 	UFUNCTION(BlueprintCallable, Category = "ControlRig | Modules")
 	bool BindModuleVariable(const FName& InModuleName, const FName& InVariableName, const FString& InSourcePath, bool bSetupUndo = true);
diff --git a/Engine/Plugins/Animation/ControlRig/Source/ControlRig/Public/ModularRigModel.h b/Engine/Plugins/Animation/ControlRig/Source/ControlRig/Public/ModularRigModel.h
--- a/Engine/Plugins/Animation/ControlRig/Source/ControlRig/Public/ModularRigModel.h
+++ b/Engine/Plugins/Animation/ControlRig/Source/ControlRig/Public/ModularRigModel.h
@@ -3,16 +3,11 @@
 #pragma once
 
 #include "ControlRig.h"
-#include "ControlRigOverride.h"
 #include "ModularRigModel.generated.h"
 
 class UModularRigController;
 struct FModularRigModel;
 
-#if WITH_EDITOR
-class FPropertyPath;
-#endif
-
 UENUM()
 enum class EModularRigNotification : uint8
 {
@@ -59,9 +54,7 @@ public:
 		, Class(nullptr)
 		, PreviousName(NAME_None)
 		, PreviousParentName(NAME_None)
-	{
-		ConfigOverrides.SetUsesSubjectName(false);
-	}
+	{}
 	
 	FRigModuleReference(const FName& InName, TSubclassOf<UControlRig> InClass, const FName& InParentModuleName, const FModularRigModel* InModel)
 		: Name(InName)
@@ -71,9 +64,7 @@ public:
 		, PreviousName(NAME_None)
 		, PreviousParentName(NAME_None)
 		, Model(InModel)
-	{
-		ConfigOverrides.SetUsesSubjectName(false);
-	}
+	{}
 
 	UPROPERTY()
 	FName Name;
@@ -97,10 +88,7 @@ public:
 	TMap<FRigElementKey, FRigElementKey> Connections_DEPRECATED; // Connectors to Connection element
 
 	UPROPERTY()
-	TMap<FName, FString> ConfigValues_DEPRECATED;
-
-	UPROPERTY()
-	FControlRigOverrideContainer ConfigOverrides;
+	TMap<FName, FString> ConfigValues;
 
 	UPROPERTY()
 	TMap<FName, FString> Bindings; // ExternalVariableName (current module) -> SourceExternalVariableNamespacedPath (root rig or other module)
@@ -138,7 +126,7 @@ public:
 
 	const FRigModuleReference* GetParentModule() const;
 	const FRigModuleReference* GetRootModule() const;
-
+	
 	friend bool operator==(const FRigModuleReference& A, const FRigModuleReference& B)
 	{
 		return A.ParentModuleName == B.ParentModuleName &&
@@ -148,8 +136,6 @@ public:
 	const FRigConnectorElement* FindPrimaryConnector(const URigHierarchy* InHierarchy) const;
 	TArray<const FRigConnectorElement*> FindConnectors(const URigHierarchy* InHierarchy) const;
 
-	void PatchModelsOnLoad();
-
 private:
 	const FModularRigModel* Model = nullptr;
 
diff --git a/Engine/Plugins/Animation/ControlRig/Source/ControlRigDeveloper/ControlRigDeveloper.Build.cs b/Engine/Plugins/Animation/ControlRig/Source/ControlRigDeveloper/ControlRigDeveloper.Build.cs
--- a/Engine/Plugins/Animation/ControlRig/Source/ControlRigDeveloper/ControlRigDeveloper.Build.cs
+++ b/Engine/Plugins/Animation/ControlRig/Source/ControlRigDeveloper/ControlRigDeveloper.Build.cs
@@ -57,13 +57,6 @@ namespace UnrealBuildTool.Rules
 					}
                 );
 
-                PublicDependencyModuleNames.AddRange(
-	                new string[]
-	                {
-		                "AnimationEditorWidgets",
-	                }
-                );
-
                 PrivateIncludePathModuleNames.Add("ControlRigEditor");
                 DynamicallyLoadedModuleNames.Add("ControlRigEditor");
             }
diff --git a/Engine/Plugins/Animation/ControlRig/Source/ControlRigDeveloper/Private/ControlRigBlueprint.cpp b/Engine/Plugins/Animation/ControlRig/Source/ControlRigDeveloper/Private/ControlRigBlueprint.cpp
--- a/Engine/Plugins/Animation/ControlRig/Source/ControlRigDeveloper/Private/ControlRigBlueprint.cpp
+++ b/Engine/Plugins/Animation/ControlRig/Source/ControlRigDeveloper/Private/ControlRigBlueprint.cpp
@@ -1155,49 +1155,6 @@ void UControlRigBlueprint::UpdateExposedModuleConnectors() const
 	PropagateHierarchyFromBPToInstances();
 }
 
-#if WITH_EDITOR
-
-TArray<FOverrideStatusSubject> UControlRigBlueprint::GetOverrideSubjects() const
-{
-	TArray<FOverrideStatusSubject> Subjects;
-
-	if(const UModularRig* DebuggedRig = Cast<UModularRig>(GetObjectBeingDebugged()))
-	{
-		ModularRigModel.ForEachModule([&Subjects, DebuggedRig](const FRigModuleReference* ModuleReference) -> bool
-		{
-			if(const FRigModuleInstance* ModuleInstance = DebuggedRig->FindModule(ModuleReference->Name))
-			{
-				if(const UControlRig* ModuleRig = ModuleInstance->GetRig())
-				{
-					for(const FControlRigOverrideValue& Override : ModuleReference->ConfigOverrides)
-					{
-						Subjects.Add({ModuleRig, Override.ToPropertyPath()});
-					}
-				}
-			}
-			return true;
-		});
-	}
-
-	return Subjects;
-}
-
-uint32 UControlRigBlueprint::GetOverrideSubjectsHash() const
-{
-	uint32 Hash = 0;
-	
-	ModularRigModel.ForEachModule([&Hash](const FRigModuleReference* ModuleReference) -> bool
-	{
-		Hash = HashCombine(Hash, GetTypeHash(ModuleReference->Name));
-		Hash = HashCombine(Hash, GetTypeHash(ModuleReference->ConfigOverrides));
-		return true;
-	});
-
-	return Hash;
-}
-
-#endif
-
 bool UControlRigBlueprint::ResolveConnector(const FRigElementKey& DraggedKey, const FRigElementKey& TargetKey, bool bSetupUndoRedo)
 {
 	return ResolveConnectorToArray(DraggedKey, {TargetKey}, bSetupUndoRedo);
@@ -1520,7 +1477,6 @@ void UControlRigBlueprint::HandlePackageDone()
 		ModularRigModel.ForEachModule([](const FRigModuleReference* Element) -> bool
 		{
 			(void)Element->Class.LoadSynchronous();
-			const_cast<FRigModuleReference*>(Element)->PatchModelsOnLoad();
 			return true;
 		});
 
diff --git a/Engine/Plugins/Animation/ControlRig/Source/ControlRigDeveloper/Public/ControlRigBlueprint.h b/Engine/Plugins/Animation/ControlRig/Source/ControlRigDeveloper/Public/ControlRigBlueprint.h
--- a/Engine/Plugins/Animation/ControlRig/Source/ControlRigDeveloper/Public/ControlRigBlueprint.h
+++ b/Engine/Plugins/Animation/ControlRig/Source/ControlRigDeveloper/Public/ControlRigBlueprint.h
@@ -22,7 +22,6 @@
 
 #if WITH_EDITOR
 #include "Kismet2/CompilerResultsLog.h"
-#include "Overrides/SOverrideListWidget.h"
 #endif
 
 #include "ControlRigBlueprint.generated.h"
@@ -217,11 +216,6 @@ public:
 
 	UFUNCTION(BlueprintCallable, Category = "Control Rig Blueprint")
 	void UpdateExposedModuleConnectors() const;
-
-#if WITH_EDITOR
-	TArray<FOverrideStatusSubject> GetOverrideSubjects() const;
-	uint32 GetOverrideSubjectsHash() const;
-#endif
 	
 protected:
 
diff --git a/Engine/Plugins/Animation/ControlRig/Source/ControlRigEditor/Private/ControlRigModuleDetails.cpp b/Engine/Plugins/Animation/ControlRig/Source/ControlRigEditor/Private/ControlRigModuleDetails.cpp
--- a/Engine/Plugins/Animation/ControlRig/Source/ControlRigEditor/Private/ControlRigModuleDetails.cpp
+++ b/Engine/Plugins/Animation/ControlRig/Source/ControlRigEditor/Private/ControlRigModuleDetails.cpp
@@ -25,7 +25,6 @@
 #include "Features/IModularFeatures.h"
 #include "IPropertyAccessEditor.h"
 #include "ModularRigRuleManager.h"
-#include "PropertyPath.h"
 #include "ScopedTransaction.h"
 #include "Editor/SRigHierarchyTreeView.h"
 #include "Widgets/SRigVMVariantTagWidget.h"
@@ -643,9 +642,9 @@ void FRigModuleInstanceDetails::CustomizeDetails(IDetailLayoutBuilder& DetailBui
 					continue;
 				}
 
-				TDelegate<void(const FPropertyChangedEvent&)> OnValueChangedDelegate = TDelegate<void(const FPropertyChangedEvent&)>::CreateSP(this, &FRigModuleInstanceDetails::OnConfigValueChanged);
-				DefaultProperty->SetOnPropertyValueChangedWithData(OnValueChangedDelegate);
-				DefaultProperty->SetOnChildPropertyValueChangedWithData(OnValueChangedDelegate);
+				const FSimpleDelegate OnValueChangedDelegate = FSimpleDelegate::CreateSP(this, &FRigModuleInstanceDetails::OnConfigValueChanged, Property->GetFName());
+				DefaultProperty->SetOnPropertyValueChanged(OnValueChangedDelegate);
+				DefaultProperty->SetOnChildPropertyValueChanged(OnValueChangedDelegate);
 
 				FPropertyBindingWidgetArgs BindingArgs;
 				BindingArgs.Property = (FProperty*)Property;
@@ -882,133 +881,28 @@ FRigElementKeyRedirector FRigModuleInstanceDetails::GetConnections() const
 	return FRigElementKeyRedirector();
 }
 
-void FRigModuleInstanceDetails::OnConfigValueChanged(const FPropertyChangedEvent& InPropertyChangedEvent)
+void FRigModuleInstanceDetails::OnConfigValueChanged(const FName InVariableName)
 {
-	UControlRigBlueprint* Blueprint = PerModuleInfos[0].GetBlueprint();
-	if(Blueprint == nullptr)
-	{
-		return;
-	}
-
-	FString PropertyPathString;
-	{
-		TSharedRef< FPropertyPath > PropertyPath = FPropertyPath::CreateEmpty();
-		FProperty* Property = InPropertyChangedEvent.Property;
-		if(!ensure(Property))
-		{
-			return;
-		}
-
-		TMap<FString, int32> PropertyNameStack;
-		InPropertyChangedEvent.GetArrayIndicesPerObject(0, PropertyNameStack);
-
-		if (PropertyNameStack.Num() > 0)
-		{
-			TArray<FString> PropertyNames;
-			PropertyNameStack.GetKeys(PropertyNames);
-
-			// the property names in the array are provided in reverse order
-			for(int32 Index = PropertyNames.Num() - 1; Index >= 0; Index--)
-			{
-				if(!PropertyPath->IsValid())
-				{
-					ensure(InPropertyChangedEvent.GetNumObjectsBeingEdited() > 0);
-					Property = FControlRigOverrideValue::FindProperty(InPropertyChangedEvent.GetObjectBeingEdited(0)->GetClass(), PropertyNames[Index]);
-					if(!ensure(Property != nullptr))
-					{
-						return;
-					}
-				}
-				else
-				{
-					if(const FArrayProperty* ArrayProperty = CastField<FArrayProperty>(Property))
-					{
-						Property = ArrayProperty->Inner;
-					}
-				
-					const FStructProperty* StructProperty = CastField<FStructProperty>(Property);
-					if(!ensure(StructProperty != nullptr))
-					{
-						return;
-					}
-
-					Property = FControlRigOverrideValue::FindProperty(StructProperty->Struct, PropertyNames[Index]);
-					if(!ensure(Property))
-					{
-						return;
-					}
-				}
-				
-				PropertyPath->AddProperty(
-					{
-						TWeakFieldPtr<FProperty>(Property),
-						InPropertyChangedEvent.GetArrayIndex(PropertyNames[Index])
-					}
-				);
-			}
-		}
-		else
-		{
-			PropertyPath->AddProperty(
-				{
-					TWeakFieldPtr<FProperty>(InPropertyChangedEvent.MemberProperty),
-					InPropertyChangedEvent.GetArrayIndex(InPropertyChangedEvent.GetMemberPropertyName().ToString())
-				}
-			);
-		}
-
-		PropertyPathString = PropertyPath->ToString();
-	}
-	
-	// we need to shorten the path if we already have data on something above this.
-	bool bFoundValueWithShorterPath = true;
-	while(bFoundValueWithShorterPath)
-	{
-		bFoundValueWithShorterPath = false;
-
-		for(const FPerModuleInfo& Info : PerModuleInfos)
-		{
-			if (const FRigModuleReference* ModuleReference = Blueprint->ModularRigModel.FindModule(Info.ModuleName))
-			{
-				for(const FControlRigOverrideValue& Override : ModuleReference->ConfigOverrides)
-				{
-					const FString& ParentPath = Override.GetPath();
-					if(FControlRigOverrideContainer::IsChildPathOf(PropertyPathString, ParentPath))
-					{
-						bFoundValueWithShorterPath = true;
-						PropertyPathString = ParentPath;
-						break;
-					}
-				}
-			}
-			
-			if(bFoundValueWithShorterPath)
-			{
-				break;
-			}
-		}
-	}
-
-	TMap<FName, FControlRigOverrideValue> ModuleValues;
-	ModuleValues.Reserve(PerModuleInfos.Num());
-	
+	TMap<FName, FString> ModuleValues;
 	for(const FPerModuleInfo& Info : PerModuleInfos)
 	{
 		if (const FRigModuleInstance* ModuleInstance = Info.GetModule())
 		{
 			if (const UControlRig* ModuleRig = ModuleInstance->GetRig())
 			{
-				FControlRigOverrideValue ConfigValue = FControlRigOverrideValue(PropertyPathString, ModuleRig);
-				ModuleValues.Add(ModuleInstance->Name, ConfigValue);
+				FString ValueStr = ModuleRig->GetVariableAsString(InVariableName);
+				ModuleValues.Add(ModuleInstance->Name, ValueStr);
 			}
 		}
 	}
-
-	FScopedTransaction Transaction(LOCTEXT("ModuleHierarchySetConfigValue", "Set Module Config Value"));
-	for (const TPair<FName, FControlRigOverrideValue>& Value : ModuleValues)
+	
+	for (const TPair<FName, FString>& Value : ModuleValues)
 	{
-		UModularRigController* Controller = Blueprint->GetModularRigController();
-		Controller->SetConfigValueInModule(Value.Key, Value.Value);
+		if (UControlRigBlueprint* Blueprint = PerModuleInfos[0].GetBlueprint())
+		{
+			UModularRigController* Controller = Blueprint->GetModularRigController();
+			Controller->SetConfigValueInModule(Value.Key, InVariableName, Value.Value);
+		}
 	}
 }
 
diff --git a/Engine/Plugins/Animation/ControlRig/Source/ControlRigEditor/Private/ControlRigModuleDetails.h b/Engine/Plugins/Animation/ControlRig/Source/ControlRigEditor/Private/ControlRigModuleDetails.h
--- a/Engine/Plugins/Animation/ControlRig/Source/ControlRigEditor/Private/ControlRigModuleDetails.h
+++ b/Engine/Plugins/Animation/ControlRig/Source/ControlRigEditor/Private/ControlRigModuleDetails.h
@@ -47,7 +47,7 @@ public:
 	TArray<FRigModuleConnector> GetConnectors() const;
 	FRigElementKeyRedirector GetConnections() const;
 
-	void OnConfigValueChanged(const FPropertyChangedEvent& InPropertyChangedEvent);
+	void OnConfigValueChanged(const FName InVariableName);
 	bool OnConnectorTargetChanged(TArray<FRigElementKey> InTargets, FRigModuleConnector InConnector);
 	
 	struct FPerModuleInfo
diff --git a/Engine/Plugins/Animation/ControlRig/Source/ControlRigEditor/Private/EditMode/SControlRigOutliner.cpp b/Engine/Plugins/Animation/ControlRig/Source/ControlRigEditor/Private/EditMode/SControlRigOutliner.cpp
--- a/Engine/Plugins/Animation/ControlRig/Source/ControlRigEditor/Private/EditMode/SControlRigOutliner.cpp
+++ b/Engine/Plugins/Animation/ControlRig/Source/ControlRigEditor/Private/EditMode/SControlRigOutliner.cpp
@@ -41,7 +41,6 @@
 #include "Rigs/FKControlRig.h"
 #include "SEnumCombo.h"
 #include "Components/SkeletalMeshComponent.h"
-#include "Overrides/SOverrideStatusWidget.h"
 
 #define LOCTEXT_NAMESPACE "ControlRigOutliner"
 
@@ -1636,56 +1635,15 @@ TSharedRef<SWidget> SSearchableMultiRigHierarchyTreeView::OnGetOptionsMenu()
 			.AutoHeight()
 			.Padding(3)
 			[
-				SNew(SHorizontalBox)
-
-				+ SHorizontalBox::Slot()
-				.AutoWidth()
-				.Padding(0)
-				[
-					SNew(SEnumComboBox, StaticEnum<EElementNameDisplayMode>())
-					.CurrentValue_Lambda([this]() -> int32
-					{
-						return static_cast<int32>(GetElementNameDisplayMode());
-					})
-					.OnEnumSelectionChanged_Lambda([this](int32 InEnumValue, ESelectInfo::Type)
-					{
-						SetElementNameDisplayMode(static_cast<EElementNameDisplayMode>(InEnumValue));
-					})
-				]
-				
-				+ SHorizontalBox::Slot()
-				.AutoWidth()
-				.Padding(4, 0, 0, 0)
-				.VAlign(VAlign_Center)
-				.HAlign(HAlign_Left)
-				[
-					SNew(SOverrideStatusWidget)
-					.Status_Lambda([this]() -> EOverrideWidgetStatus::Type
-					{
-						return GetElementNameDisplayMode() == EElementNameDisplayMode::AssetDefault ?
-							EOverrideWidgetStatus::None : EOverrideWidgetStatus::ChangedHere;
-					})
-					.MenuContent_Lambda([this]() -> TSharedRef<SWidget>
-					{
-						if(GetElementNameDisplayMode() != EElementNameDisplayMode::AssetDefault)
-						{
-							FMenuBuilder MenuBuilder(true, nullptr);
-							MenuBuilder.AddMenuEntry(
-								LOCTEXT("RemoveOverride", "Remove Override"),
-								LOCTEXT("RemoveElementNameOverrideTooltip", "Removes the override from the element name mode and uses the AssetDefault option."),
-								FSlateIcon(),
-								FUIAction(
-									FExecuteAction::CreateLambda([this]()
-									{
-										SetElementNameDisplayMode(EElementNameDisplayMode::AssetDefault);
-									})
-								)
-							);
-							return MenuBuilder.MakeWidget();
-						}
-						return SNullWidget::NullWidget;
-					})
-				]
+				SNew(SEnumComboBox, StaticEnum<EElementNameDisplayMode>())
+				.CurrentValue_Lambda([this]() -> int32
+				{
+					return static_cast<int32>(GetElementNameDisplayMode());
+				})
+				.OnEnumSelectionChanged_Lambda([this](int32 InEnumValue, ESelectInfo::Type)
+				{
+					SetElementNameDisplayMode(static_cast<EElementNameDisplayMode>(InEnumValue));
+				})
 			]
 			,
 			LOCTEXT("ElementNameDisplayMode", "Name Mode")
diff --git a/Engine/Plugins/Animation/ControlRig/Source/ControlRigEditor/Private/Editor/ControlRigEditor.cpp b/Engine/Plugins/Animation/ControlRig/Source/ControlRigEditor/Private/Editor/ControlRigEditor.cpp
--- a/Engine/Plugins/Animation/ControlRig/Source/ControlRigEditor/Private/Editor/ControlRigEditor.cpp
+++ b/Engine/Plugins/Animation/ControlRig/Source/ControlRigEditor/Private/Editor/ControlRigEditor.cpp
@@ -124,7 +124,6 @@
 #include "ITransportControl.h"
 #include "PropertyPath.h"
 #include "Editor/SRigVMDetailsInspector.h"
-#include "Overrides/OverrideStatusDetailsObjectFilter.h"
 
 #if WITH_RIGVMLEGACYEDITOR
 #include "SKismetInspector.h"
@@ -1569,228 +1568,7 @@ void FControlRigBaseEditor::SetDetailViewForRigModules(const TArray<FName> InMod
 			Objects.Add(ModuleInstance);
 		}
 	}
-
-	if(!Objects.IsEmpty() && CVarRigVMEnablePinOverrides.GetValueOnAnyThread())
-	{
-		TSharedPtr<FOverrideStatusDetailsViewObjectFilter> ObjectFilter = FOverrideStatusDetailsViewObjectFilter::Create();
-
-		ObjectFilter->OnCanMergeObjects().BindLambda([](const UObject* InObjectA, const UObject* InObjectB) -> bool
-		{
-			if(InObjectA && InObjectB)
-			{
-				return InObjectA->IsA<UControlRig>() && InObjectB->IsA<UControlRig>(); 
-			}
-			return false;
-		});
-
-		ObjectFilter->OnCanCreateWidget().BindLambda([](const FOverrideStatusSubject& InSubject) -> bool
-		{
-			static const TSet<FName> CategoriesToIgnore = {
-				TEXT("General"),
-				TEXT("Connections"),
-			};
-			return InSubject.Contains<UControlRig>() && !CategoriesToIgnore.Contains(InSubject.GetCategory());
-		});
-
-		ObjectFilter->OnGetStatus().BindLambda([RigBlueprint](const FOverrideStatusSubject& InSubject)
-		{
-			const FString PropertyPath = InSubject.GetPropertyPathString();
-			return InSubject.GetStatus<UControlRig>(
-				[PropertyPath, RigBlueprint](const UControlRig* InModuleRig) -> TOptional<EOverrideWidgetStatus::Type>
-				{
-					if(const FRigModuleReference* ModuleReference = RigBlueprint->ModularRigModel.FindModule(InModuleRig->GetFName()))
-					{
-						if(PropertyPath.IsEmpty())
-						{
-							if(!ModuleReference->ConfigOverrides.IsEmpty())
-							{
-								return EOverrideWidgetStatus::ChangedInside; 
-							}
-						}
-						else
-						{
-							if(ModuleReference->ConfigOverrides.Contains(PropertyPath, ModuleReference->Name))
-							{
-								return EOverrideWidgetStatus::ChangedHere;
-							}
-
-							if(ModuleReference->ConfigOverrides.ContainsChildPathOf(PropertyPath, ModuleReference->Name))
-							{
-								return EOverrideWidgetStatus::ChangedInside;
-							}
-							if(ModuleReference->ConfigOverrides.ContainsParentPathOf(PropertyPath, ModuleReference->Name))
-							{
-								return EOverrideWidgetStatus::ChangedOutside;
-							}
-						}
-						return EOverrideWidgetStatus::None;
-					}
-					return {};
-				}
-			).Get(EOverrideWidgetStatus::Mixed);
-		});
-
-		ObjectFilter->OnAddOverride().BindLambda([RigBlueprint](const FOverrideStatusSubject& InSubject)
-			{
-				FScopedTransaction Transaction(LOCTEXT("AddOverride", "Add Override"));
-				RigBlueprint->Modify();
-
-				InSubject.ForEach<UControlRig>(
-					[RigBlueprint, InSubject](const UControlRig* InModuleRig)
-					{
-						if(InSubject.HasPropertyPath())
-						{
-							const FControlRigOverrideValue Value(InSubject.GetPropertyPathString(), InModuleRig);
-							if(Value.IsValid())
-							{
-								RigBlueprint->GetModularRigController()->SetConfigValueInModule(InModuleRig->GetFName(), Value);
-							}
-						}
-						else
-						{
-							const TArray<FRigVMExternalVariable> Variables = InModuleRig->GetPublicVariables();
-							for(const FRigVMExternalVariable& Variable : Variables)
-							{
-								const FString PropertyPath = Variable.Name.ToString();
-								const FControlRigOverrideValue Value(PropertyPath, InModuleRig);
-								RigBlueprint->GetModularRigController()->SetConfigValueInModule(InModuleRig->GetFName(), Value);
-							}
-						}
-					}
-				);
-
-				return FReply::Handled();
-			});
-
-		ObjectFilter->OnClearOverride().BindLambda(
-			[this, RigBlueprint](const FOverrideStatusSubject& InSubject)
-			{
-				FScopedTransaction Transaction(LOCTEXT("ClearOverride", "Clear Override"));
-				RigBlueprint->Modify();
-
-				InSubject.ForEach<UControlRig>(
-					[RigBlueprint, InSubject](const UControlRig* InModuleRig)
-					{
-						if(InSubject.HasPropertyPath())
-						{
-							RigBlueprint->GetModularRigController()->ResetConfigValueInModule(
-								InModuleRig->GetFName(), InSubject.GetPropertyPathString(), true);
-						}
-						else
-						{
-							if(const FRigModuleReference* ModuleReference = RigBlueprint->ModularRigModel.FindModule(InModuleRig->GetFName()))
-							{
-								TArray<FString> PathsToClear;
-								for(const FControlRigOverrideValue& Override : ModuleReference->ConfigOverrides)
-								{
-									if(Override.IsValid())
-									{
-										PathsToClear.Add(Override.GetPath());
-									}
-								}
-								for(const FString& Path : PathsToClear)
-								{
-									RigBlueprint->GetModularRigController()->ResetConfigValueInModule(ModuleReference->GetFName(), Path, true);
-								}
-							}
-						}
-					}
-				);
-
-				return FReply::Handled();
-			});
-
-		ObjectFilter->OnResetToDefault().BindLambda(
-			[this, RigBlueprint](const FOverrideStatusSubject& InSubject)
-			{
-				FScopedTransaction Transaction(LOCTEXT("ResetConfigValue", "Reset Config Value"));
-				RigBlueprint->Modify();
-
-				InSubject.ForEach<UControlRig>(
-					[RigBlueprint, InSubject](const UControlRig* InModuleRig)
-					{
-						if(InSubject.HasPropertyPath())
-						{
-							RigBlueprint->GetModularRigController()->ResetConfigValueInModule(
-								InModuleRig->GetFName(), InSubject.GetPropertyPathString(), false);
-						}
-						else
-						{
-							if(const FRigModuleReference* ModuleReference = RigBlueprint->ModularRigModel.FindModule(InModuleRig->GetFName()))
-							{
-								TArray<FString> PathsToClear;
-								for(const FControlRigOverrideValue& Override : ModuleReference->ConfigOverrides)
-								{
-									if(Override.IsValid())
-									{
-										PathsToClear.Add(Override.GetPath());
-									}
-								}
-								for(const FString& Path : PathsToClear)
-								{
-									RigBlueprint->GetModularRigController()->ResetConfigValueInModule(ModuleReference->GetFName(), Path, false);
-								}
-							}
-						}
-					}
-				);
-
-				return FReply::Handled();
-			});
-
-		ObjectFilter->OnValueDiffersFromDefault().BindLambda(
-			[this, RigBlueprint](const FOverrideStatusSubject& InSubject)
-			{
-				return InSubject.GetCommonValue<bool, UControlRig>(
-					[RigBlueprint, InSubject](const UControlRig* InModuleRig)
-					{
-						if(const FRigModuleReference* ModuleReference = RigBlueprint->ModularRigModel.FindModule(InModuleRig->GetFName()))
-						{
-							TArray<TSharedPtr<const FPropertyPath>> PropertyPathsToCheck;
-							if(InSubject.HasPropertyPath())
-							{
-								PropertyPathsToCheck.Add(InSubject.GetPropertyPath());
-							}
-							else
-							{
-								const TArray<FRigVMExternalVariable> Variables = InModuleRig->GetPublicVariables();
-								for(const FRigVMExternalVariable& Variable : Variables)
-								{
-									if(FProperty* Property = InModuleRig->GetClass()->FindPropertyByName(Variable.Name))
-									{
-										PropertyPathsToCheck.Add(FPropertyPath::Create(TWeakFieldPtr<FProperty>(Property)));
-									}
-								}
-							}
-
-							const UClass* Class = ModuleReference->Class.LoadSynchronous();
-							const UObject* CDO = Class->GetDefaultObject();
-							
-							for(const TSharedPtr<const FPropertyPath>& PropertyPath : PropertyPathsToCheck)
-							{
-								const FString PropertyPathString = PropertyPath->ToString();
-								const FString PropertyPathPrefix = PropertyPathString + TEXT("->");
-								for(const FControlRigOverrideValue& Override : ModuleReference->ConfigOverrides)
-								{
-									if(!Override.IdenticalValueInSubject(CDO))
-									{
-										const FString& ValuePathString = Override.GetPath();
-										if(ValuePathString == PropertyPathString || ValuePathString.StartsWith(PropertyPathPrefix))
-										{
-											return true;
-										}
-									}
-								}
-							}
-						}
-						return false;
-					}
-				).Get(false);
-			});
-
-		SetDetailObjectFilter(ObjectFilter);
-	}
-
+	
 	SetDetailObjects(Objects);
 
 	// In case the modules selected are still not available, lets set them again
diff --git a/Engine/Plugins/Animation/ControlRig/Source/ControlRigEditor/Private/Editor/ControlRigEditor.h b/Engine/Plugins/Animation/ControlRig/Source/ControlRigEditor/Private/Editor/ControlRigEditor.h
--- a/Engine/Plugins/Animation/ControlRig/Source/ControlRigEditor/Private/Editor/ControlRigEditor.h
+++ b/Engine/Plugins/Animation/ControlRig/Source/ControlRigEditor/Private/Editor/ControlRigEditor.h
@@ -133,7 +133,6 @@ protected:
 	virtual void OnWrappedPropertyChangedChainEvent(URigVMDetailsViewWrapperObject* InWrapperObject, const FString& InPropertyPath, FPropertyChangedChainEvent& InPropertyChangedChainEvent) = 0;
 	virtual void SetDetailObjects(const TArray<UObject*>& InObjects) = 0;
 	virtual void SetDetailObjectsSuper(const TArray<UObject*>& InObjects) = 0;
-	virtual void SetDetailObjectFilter(TSharedPtr<FDetailsViewObjectFilter> InObjectFilter) = 0;
 	virtual bool DetailViewShowsStruct(UScriptStruct* InStruct) const = 0;
 	virtual TSharedPtr<class SWidget> GetInspector() const = 0;
 	virtual TArray<FName> GetEventQueue() const = 0;
diff --git a/Engine/Plugins/Animation/ControlRig/Source/ControlRigEditor/Private/Editor/ControlRigEditorMode.cpp b/Engine/Plugins/Animation/ControlRig/Source/ControlRigEditor/Private/Editor/ControlRigEditorMode.cpp
--- a/Engine/Plugins/Animation/ControlRig/Source/ControlRigEditor/Private/Editor/ControlRigEditorMode.cpp
+++ b/Engine/Plugins/Animation/ControlRig/Source/ControlRigEditor/Private/Editor/ControlRigEditorMode.cpp
@@ -20,7 +20,6 @@
 #include "Editor/RigVMCompilerResultsTabSummoner.h"
 #include "Editor/RigVMDetailsInspectorTabSummoner.h"
 #include "Editor/RigVMFindReferencesTabSummoner.h"
-#include "Editor/RigOverrideListTabSummoner.h"
 
 #if WITH_RIGVMLEGACYEDITOR
 FControlRigLegacyEditorMode::FControlRigLegacyEditorMode(const TSharedRef<FControlRigLegacyEditor>& InControlRigEditor, bool bCreateDefaultLayout)
@@ -160,11 +159,7 @@ FModularRigLegacyEditorMode::FModularRigLegacyEditorMode(const TSharedRef<FContr
 
 	TabFactories.RegisterFactory(MakeShared<FModularRigModelTabSummoner>(InControlRigEditor));
 	TabFactories.RegisterFactory(MakeShared<FRigModuleAssetBrowserTabSummoner>(InControlRigEditor));
-	if(CVarRigVMEnablePinOverrides.GetValueOnAnyThread())
-	{
-		TabFactories.RegisterFactory(MakeShared<FRigOverrideListTabSummoner>(InControlRigEditor));
-	}
-
+	
 	TabLayout = FTabManager::NewLayout("Standalone_ModularRigEditMode_Layout_v1.3")
 		->AddArea
 		(
@@ -380,11 +375,7 @@ FModularRigEditorMode::FModularRigEditorMode(const TSharedRef<FControlRigEditor>
 
 	TabFactories.RegisterFactory(MakeShared<FModularRigModelTabSummoner>(InControlRigEditor));
 	TabFactories.RegisterFactory(MakeShared<FRigModuleAssetBrowserTabSummoner>(InControlRigEditor));
-	if(CVarRigVMEnablePinOverrides.GetValueOnAnyThread())
-	{
-		TabFactories.RegisterFactory(MakeShared<FRigOverrideListTabSummoner>(InControlRigEditor));
-	}
-
+	
 	TabLayout = FTabManager::NewLayout("Standalone_ModularRigEditMode_Layout_v1.3")
 		->AddArea
 		(
diff --git a/Engine/Plugins/Animation/ControlRig/Source/ControlRigEditor/Private/Editor/ControlRigLegacyEditor.h b/Engine/Plugins/Animation/ControlRig/Source/ControlRigEditor/Private/Editor/ControlRigLegacyEditor.h
--- a/Engine/Plugins/Animation/ControlRig/Source/ControlRigEditor/Private/Editor/ControlRigLegacyEditor.h
+++ b/Engine/Plugins/Animation/ControlRig/Source/ControlRigEditor/Private/Editor/ControlRigLegacyEditor.h
@@ -99,7 +99,6 @@ public:
 
 	virtual void SetDetailObjects(const TArray<UObject*>& InObjects) override { return FControlRigBaseEditor::SetDetailObjectsImpl(InObjects); }
 	virtual void SetDetailObjectsSuper(const TArray<UObject*>& InObjects) override { return IControlRigLegacyEditor::SetDetailObjects(InObjects); }
-	virtual void SetDetailObjectFilter(TSharedPtr<FDetailsViewObjectFilter> InObjectFilter) override { return IControlRigLegacyEditor::SetDetailObjectFilter(InObjectFilter); };
 	virtual void RefreshDetailView() override { return FControlRigBaseEditor::RefreshDetailViewImpl(); }
 	virtual void RefreshDetailViewSuper() override { return IControlRigLegacyEditor::RefreshDetailView(); }
 
diff --git a/Engine/Plugins/Animation/ControlRig/Source/ControlRigEditor/Private/Editor/ControlRigNewEditor.h b/Engine/Plugins/Animation/ControlRig/Source/ControlRigEditor/Private/Editor/ControlRigNewEditor.h
--- a/Engine/Plugins/Animation/ControlRig/Source/ControlRigEditor/Private/Editor/ControlRigNewEditor.h
+++ b/Engine/Plugins/Animation/ControlRig/Source/ControlRigEditor/Private/Editor/ControlRigNewEditor.h
@@ -98,7 +98,6 @@ public:
 
 	virtual void SetDetailObjects(const TArray<UObject*>& InObjects) override { return FControlRigBaseEditor::SetDetailObjectsImpl(InObjects); }
 	virtual void SetDetailObjectsSuper(const TArray<UObject*>& InObjects) override { return IControlRigNewEditor::SetDetailObjects(InObjects); }
-	virtual void SetDetailObjectFilter(TSharedPtr<FDetailsViewObjectFilter> InObjectFilter) override { return IControlRigNewEditor::SetDetailObjectFilter(InObjectFilter); };
 	virtual void RefreshDetailView() override { return FControlRigBaseEditor::RefreshDetailViewImpl(); }
 	virtual void RefreshDetailViewSuper() override { return IControlRigNewEditor::RefreshDetailView(); }
 
diff --git a/Engine/Plugins/Animation/ControlRig/Source/ControlRigEditor/Private/Editor/RigOverrideListTabSummoner.cpp b/Engine/Plugins/Animation/ControlRig/Source/ControlRigEditor/Private/Editor/RigOverrideListTabSummoner.cpp
deleted file mode 100644
diff --git a/Engine/Plugins/Animation/ControlRig/Source/ControlRigEditor/Private/Editor/RigOverrideListTabSummoner.h b/Engine/Plugins/Animation/ControlRig/Source/ControlRigEditor/Private/Editor/RigOverrideListTabSummoner.h
deleted file mode 100644
diff --git a/Engine/Plugins/Animation/ControlRig/Source/ControlRigEditor/Private/Editor/SRigHierarchy.cpp b/Engine/Plugins/Animation/ControlRig/Source/ControlRigEditor/Private/Editor/SRigHierarchy.cpp
--- a/Engine/Plugins/Animation/ControlRig/Source/ControlRigEditor/Private/Editor/SRigHierarchy.cpp
+++ b/Engine/Plugins/Animation/ControlRig/Source/ControlRigEditor/Private/Editor/SRigHierarchy.cpp
@@ -1191,18 +1191,6 @@ TSharedRef< SWidget > SRigHierarchy::CreateFilterMenu()
 		MenuBuilder.AddMenuEntry(Actions.FilteringFlattensHierarchy);
 		MenuBuilder.AddMenuEntry(Actions.HideParentsWhenFiltering);
 
-		const UEnum* NameModeEnum = StaticEnum<EElementNameDisplayMode>();
-		TArray<int32> EnumValueSubset;
-		EnumValueSubset.Reserve(NameModeEnum->NumEnums() - 1);
-		for (int32 i = 0; i < NameModeEnum->NumEnums() - 1; i++)
-		{
-			const int32 Value = static_cast<int32>(NameModeEnum->GetValueByIndex(i));
-			if(Value != static_cast<int32>(EElementNameDisplayMode::AssetDefault))
-			{
-				EnumValueSubset.Add(Value);
-			}
-		}
-
 		MenuBuilder.AddWidget(
 			SNew(SVerticalBox)
 
@@ -1210,8 +1198,7 @@ TSharedRef< SWidget > SRigHierarchy::CreateFilterMenu()
 			.AutoHeight()
 			.Padding(3)
 			[
-				SNew(SEnumComboBox, NameModeEnum)
-				.EnumValueSubset(EnumValueSubset)
+				SNew(SEnumComboBox, StaticEnum<EElementNameDisplayMode>())
 				.CurrentValue_Lambda([this]() -> int32
 				{
 					const IControlRigBaseEditor* StrongEditor = ControlRigEditor.Pin().Get();
