Control Rig: React to hierarchy metadata changes in game thread for widgets

#rnx

diff --git a/Engine/Plugins/Animation/ControlRig/Source/ControlRigEditor/Private/ControlRigElementDetails.cpp b/Engine/Plugins/Animation/ControlRig/Source/ControlRigEditor/Private/ControlRigElementDetails.cpp
--- a/Engine/Plugins/Animation/ControlRig/Source/ControlRigEditor/Private/ControlRigElementDetails.cpp
+++ b/Engine/Plugins/Animation/ControlRig/Source/ControlRigEditor/Private/ControlRigElementDetails.cpp
@@ -30,6 +30,7 @@
 #include "Editor/ControlRigEditor.h"
 #include "Editor/SRigConnectorTargetWidget.h"
 #include "ModularRigRuleManager.h"
+#include "Async/TaskGraphInterfaces.h"
 
 #define LOCTEXT_NAMESPACE "ControlRigElementDetails"
 
@@ -209,7 +210,7 @@ void SRigElementKeyWidget::Construct(const FArguments& InArgs)
 				{
 					if(const TSharedPtr<SRigElementKeyWidget> StrongThisPtr = WeakThisPtr.Pin())
 					{
-						if (StrongThisPtr->OnGetElementType.IsBound())
+						if (StrongThisPtr.IsValid() && StrongThisPtr->OnGetElementType.IsBound())
 						{
 							return (int32)StrongThisPtr->OnGetElementType.Execute();
 						}
@@ -220,11 +221,14 @@ void SRigElementKeyWidget::Construct(const FArguments& InArgs)
 				{
 					if(const TSharedPtr<SRigElementKeyWidget> StrongThisPtr = WeakThisPtr.Pin())
 					{
-						ERigElementType EnumValue = static_cast<ERigElementType>(InEnumValue);
-						StrongThisPtr->OnElementTypeChanged.ExecuteIfBound(EnumValue);
-						StrongThisPtr->UpdateElementNameList();
-						StrongThisPtr->SearchableComboBox->ClearSelection();
-						StrongThisPtr->OnElementNameChanged.ExecuteIfBound(nullptr, ESelectInfo::Direct);
+						if(StrongThisPtr.IsValid())
+						{
+							ERigElementType EnumValue = static_cast<ERigElementType>(InEnumValue);
+							StrongThisPtr->OnElementTypeChanged.ExecuteIfBound(EnumValue);
+							StrongThisPtr->UpdateElementNameList();
+							StrongThisPtr->SearchableComboBox->ClearSelection();
+							StrongThisPtr->OnElementNameChanged.ExecuteIfBound(nullptr, ESelectInfo::Direct);
+						}
 					}
 				})
 		]
@@ -385,7 +389,10 @@ void FRigElementKeyDetails::CustomizeHeader(TSharedRef<IPropertyHandle> InStruct
 			{
 				if(const TSharedPtr<FRigElementKeyDetails> StrongThisPtr = WeakThisPtr.Pin())
 				{
-					return !StrongThisPtr->NameHandle->IsEditConst();
+					if(StrongThisPtr.IsValid())
+					{
+						return !StrongThisPtr->NameHandle->IsEditConst();
+					}
 				}
 				return false;
 			})
@@ -625,7 +632,10 @@ void FRigComponentKeyDetails::CustomizeChildren(TSharedRef<IPropertyHandle> InSt
 			{
 				if(const TSharedPtr<FRigComponentKeyDetails> StrongThisPtr = WeakThisPtr.Pin())
 				{
-					StrongThisPtr->UpdateComponentNameList();
+					if(StrongThisPtr.IsValid())
+					{
+						StrongThisPtr->UpdateComponentNameList();
+					}
 				}
 			}));
 			
@@ -1765,18 +1775,27 @@ void FRigBaseElementDetails::CustomizeMetadata(IDetailLayoutBuilder& DetailBuild
 		}
 
 		TWeakPtr<IPropertyUtilities> WeakPropertyUtilities =  DetailBuilder.GetPropertyUtilities().ToWeakPtr();
-		MetadataHandle = Hierarchy->OnMetadataChanged().AddLambda([this, WeakPropertyUtilities](const FRigElementKey& InKey, const FName&)
-		{
-			if(WeakPropertyUtilities.IsValid())
+		MetadataHandle = Hierarchy->OnMetadataChanged().AddLambda(
+			[this, WeakPropertyUtilities](const FRigElementKey& InKey, const FName&)
 			{
-				const FRigBaseElement* Element = PerElementInfos.Num() == 1 ? PerElementInfos[0].GetElement() : nullptr;
-				if (InKey.Type == ERigElementType::All || (Element && Element->GetKey() == InKey))
+				if(WeakPropertyUtilities.IsValid())
 				{
-					WeakPropertyUtilities.Pin()->ForceRefresh();
+					const FRigBaseElement* Element = PerElementInfos.Num() == 1 ? PerElementInfos[0].GetElement() : nullptr;
+					if (InKey.Type == ERigElementType::All || (Element && Element->GetKey() == InKey))
+					{
+						// run the refresh of the user interface on the next tick on the game thread 
+						FFunctionGraphTask::CreateAndDispatchWhenReady([WeakPropertyUtilities]()
+						{
+							const TSharedPtr<IPropertyUtilities> PropertyUtilities = WeakPropertyUtilities.IsValid() ? WeakPropertyUtilities.Pin() : nullptr;
+							if(PropertyUtilities.IsValid())
+							{
+								PropertyUtilities->ForceRefresh();
+							}
+						}, TStatId(), NULL, ENamedThreads::GameThread);
+					}
 				}
-			}
-		});
-	}
+			});
+		}
 	
 	FRigBaseElement* Element = PerElementInfos[0].Element.Get();
 	TArray<FName> MetadataNames = Element->GetOwner()->GetMetadataNames(Element->GetKey());
