Curve Editor: Tween UI adjustments:
- Change colours of tween functions according to UX suggestions.
- Change order of tween functions.
- Renamed Blend Relative to Move Relative
- Fixed tween colours not getting gamma converted correctly, i.e. FTweenUtilsStyle use FLinearColor(FColor(r,g,b)) now so it correctly gets converted. This fixes the colours from appearing more pale than the UX designers intended.

diff --git a/Engine/Plugins/Animation/ControlRig/Source/ControlRigEditor/Private/EditMode/Tween/ControlRigTweenModels.cpp b/Engine/Plugins/Animation/ControlRig/Source/ControlRigEditor/Private/EditMode/Tween/ControlRigTweenModels.cpp
--- a/Engine/Plugins/Animation/ControlRig/Source/ControlRigEditor/Private/EditMode/Tween/ControlRigTweenModels.cpp
+++ b/Engine/Plugins/Animation/ControlRig/Source/ControlRigEditor/Private/EditMode/Tween/ControlRigTweenModels.cpp
@@ -27,18 +27,25 @@ static TArray<TweeningUtilsEditor::FTweenModelUIEntry> MakeTweenModels(
 {
 	using namespace TweeningUtilsEditor;
 
-	const TAttribute<TWeakPtr<FCurveEditor>> CurveEditorAttr = MakeCurveEditorAttribute(InSequencerAttr);
 	TArray<FTweenModelUIEntry> Result;
-	ForEachCurveTweenable([&InSequencerAttr, &InOwningEditMode, &CurveEditorAttr, &Result]<EBlendFunction BlendFunction>
+	const TAttribute<TWeakPtr<FCurveEditor>> CurveEditorAttr = MakeCurveEditorAttribute(InSequencerAttr);
+	const auto AddCurveTweenable = [&InSequencerAttr, &InOwningEditMode, &CurveEditorAttr, &Result]<EBlendFunction TBlendFunction>
 	{
-		using FModel = TEditorTweenModel<TControlRigTweenModel<BlendFunction>>;
+		using FModel = TEditorTweenModel<TControlRigTweenModel<TBlendFunction>>;
 		TUniquePtr<FModel> TweenModel = MakeUnique<FModel>(
 			CurveEditorAttr,										// TTangentFlatteningTweenProxy arg
 			InSequencerAttr, InOwningEditMode						// FControlRigTweenModel args
 			);
-		Result.Emplace(MoveTemp(TweenModel), FTweenModelDisplayInfo(BlendFunction));
-	});
+		Result.Emplace(MoveTemp(TweenModel), FTweenModelDisplayInfo(TBlendFunction));
+	};
+	static_assert(static_cast<int32>(EBlendFunction::Num) == 7, "You probably want to add the new tween function here.");
 
+	// Could use TweeningUtilsEditor::ForEachCurveTweenable, but the functions are supposed to be listed in this specific order.
+	AddCurveTweenable.template operator()<EBlendFunction::BlendEase>();
+	AddCurveTweenable.template operator()<EBlendFunction::BlendNeighbor>();
+	AddCurveTweenable.template operator()<EBlendFunction::BlendRelative>();
+	AddCurveTweenable.template operator()<EBlendFunction::PushPull>();
+	AddCurveTweenable.template operator()<EBlendFunction::SmoothRough>();
 	Result.Emplace(
 		MakeUnique<TEditorTweenModel<FCurveTimeOffsetTweenModel>>(
 			CurveEditorAttr,		// TTangentFlatteningTweenProxy arg
@@ -46,6 +53,7 @@ static TArray<TweeningUtilsEditor::FTweenModelUIEntry> MakeTweenModels(
 			),
 		FTweenModelDisplayInfo(EBlendFunction::TimeOffset)
 	);
+	AddCurveTweenable.template operator()<EBlendFunction::ControlsToTween>();
 
 	return Result;
 }
diff --git a/Engine/Plugins/Animation/TweeningUtils/Source/TweeningUtilsEditor/Private/Math/Abstraction/KeyBlendingAbstraction.cpp b/Engine/Plugins/Animation/TweeningUtils/Source/TweeningUtilsEditor/Private/Math/Abstraction/KeyBlendingAbstraction.cpp
--- a/Engine/Plugins/Animation/TweeningUtils/Source/TweeningUtilsEditor/Private/Math/Abstraction/KeyBlendingAbstraction.cpp
+++ b/Engine/Plugins/Animation/TweeningUtils/Source/TweeningUtilsEditor/Private/Math/Abstraction/KeyBlendingAbstraction.cpp
@@ -94,9 +94,9 @@ static FStaticBlendFunctionData& GetFunctionData(EBlendFunction InFunction)
 		),
 		FStaticBlendFunctionData(
 			EBlendFunction::BlendRelative, FTweeningUtilsCommands::Get().SetTweenBlendRelative,
-			LOCTEXT("BlendRelative.Label", "Blend Relative"),
+			LOCTEXT("BlendRelative.Label", "Move Relative"),
 			LOCTEXT("BlendRelative.Abbreviation", "BR"),
-			LOCTEXT("BlendRelative.Description", "Blend relative to the next or previous value for selected keys"),
+			LOCTEXT("BlendRelative.Description", "Move relative to the next or previous value for selected keys"),
 			TEXT("BlendRelative")
 		),
 		FStaticBlendFunctionData(
diff --git a/Engine/Plugins/Animation/TweeningUtils/Source/TweeningUtilsEditor/Private/TweeningUtilsCommands.cpp b/Engine/Plugins/Animation/TweeningUtils/Source/TweeningUtilsEditor/Private/TweeningUtilsCommands.cpp
--- a/Engine/Plugins/Animation/TweeningUtils/Source/TweeningUtilsEditor/Private/TweeningUtilsCommands.cpp
+++ b/Engine/Plugins/Animation/TweeningUtils/Source/TweeningUtilsEditor/Private/TweeningUtilsCommands.cpp
@@ -26,7 +26,7 @@ void FTweeningUtilsCommands::RegisterCommands()
 	UI_COMMAND(SetControlsToTween, "Tween ", "Interpolates between the previous and next keys", EUserInterfaceActionType::RadioButton, FInputChord() );
 	UI_COMMAND(SetTweenPushPull, "Push / Pull", "Push or pull the values to the interpolation between the previous and next keys", EUserInterfaceActionType::RadioButton, FInputChord() );
 	UI_COMMAND(SetTweenBlendNeighbor, "Blend Neighbor", "Blend to the next or previous values for selected keys", EUserInterfaceActionType::RadioButton, FInputChord() );
-	UI_COMMAND(SetTweenBlendRelative, "Blend Relative", "Blend relative to the next or previous value for selected keys", EUserInterfaceActionType::RadioButton, FInputChord() );
+	UI_COMMAND(SetTweenBlendRelative, "Move Relative", "Move relative to the next or previous value for selected keys", EUserInterfaceActionType::RadioButton, FInputChord() );
 	UI_COMMAND(SetTweenBlendEase, "Blend Ease", "Blend with an ease falloff to the next or previous value for selected keys", EUserInterfaceActionType::RadioButton, FInputChord() );
 	UI_COMMAND(SetTweenSmoothRough, "Smooth / Rough", "Push adjacent blended keys further together or apart. Smooth is useful for softening noise, like in mocap animations.", EUserInterfaceActionType::RadioButton, FInputChord() );
 	UI_COMMAND(SetTweenTimeOffset, "Time Offset", "Shifts the curve to the left / right by changing the keys' Y values and maintaining frame position.", EUserInterfaceActionType::RadioButton, FInputChord() );
diff --git a/Engine/Plugins/Animation/TweeningUtils/Source/TweeningUtilsEditor/Private/TweeningUtilsStyle.cpp b/Engine/Plugins/Animation/TweeningUtils/Source/TweeningUtilsEditor/Private/TweeningUtilsStyle.cpp
--- a/Engine/Plugins/Animation/TweeningUtils/Source/TweeningUtilsEditor/Private/TweeningUtilsStyle.cpp
+++ b/Engine/Plugins/Animation/TweeningUtils/Source/TweeningUtilsEditor/Private/TweeningUtilsStyle.cpp
@@ -27,13 +27,13 @@ FTweeningUtilsStyle::FTweeningUtilsStyle()
 	
 	// Tween - Slider colors
 	static_assert(static_cast<int32>(EBlendFunction::Num) == 7, "Add a color here");
-	const FLinearColor ControlsToTweenTint(254.f/255.f, 155.f/255.f, 7.f/255.f);
-	const FLinearColor PushPullTint(107.f/255.f, 223.f/255.f, 151.f/255.f);
-	const FLinearColor BlendNeighborTint(133.f/255.f, 185.f/255.f, 73.f/255.f);
-	const FLinearColor BlendRelativeTint(39.f/255.f, 177.f/255.f, 240.f/255.f);
-	const FLinearColor BlendEaseTint(148.f/255.f, 57.f/255.f, 175.f/255.f);
-	const FLinearColor SmoothRoughTint(235.f/255.f, 64.f/255.f, 64.f/255.f);
-	const FLinearColor TimeOffsetTint(96.f/255.f, 65.f/255.f, 234.f/255.f);
+	const FLinearColor ControlsToTweenTint(FColor(247, 216, 43));
+	const FLinearColor PushPullTint(FColor(95, 227, 103));
+	const FLinearColor BlendNeighborTint(FColor(187, 107, 240));
+	const FLinearColor BlendRelativeTint(FColor(134, 138, 253));
+	const FLinearColor BlendEaseTint(FColor(225, 102, 182));
+	const FLinearColor SmoothRoughTint(FColor(68,213, 191));
+	const FLinearColor TimeOffsetTint(FColor(254, 133, 57));
 	Set("TweeningUtils.SetControlsToTween.Color", ControlsToTweenTint);
 	Set("TweeningUtils.SetTweenPushPull.Color", PushPullTint);
 	Set("TweeningUtils.SetTweenBlendNeighbor.Color", BlendNeighborTint);
diff --git a/Engine/Plugins/Animation/TweeningUtils/Source/TweeningUtilsEditor/Private/Widgets/MVC/STweenView.cpp b/Engine/Plugins/Animation/TweeningUtils/Source/TweeningUtilsEditor/Private/Widgets/MVC/STweenView.cpp
--- a/Engine/Plugins/Animation/TweeningUtils/Source/TweeningUtilsEditor/Private/Widgets/MVC/STweenView.cpp
+++ b/Engine/Plugins/Animation/TweeningUtils/Source/TweeningUtilsEditor/Private/Widgets/MVC/STweenView.cpp
@@ -17,10 +17,8 @@ void STweenView::Construct(const FArguments& InArgs)
 	
 	TweenModelAttr = InArgs._TweenModel;
 
-	// STweenView is placed in toolbars.
-	// We want STweenSlider's foreground color to be the same as that of the SComboButton from which you pick functions for color consistency.
-	// Effectively, FComboButtons in the toolbar use this style color.
-	SetForegroundColor(FAppStyle::GetSlateColor("DefaultForeground"));
+	// Assuming that this is placed in a toolbar, we want the of normal and hovered color to be the same, i.e. bright as if hovered.
+	SetForegroundColor(FAppStyle::GetSlateColor("CurveEditor.TweenForeground"));
 	
 	ChildSlot
 	[
diff --git a/Engine/Plugins/Animation/TweeningUtils/Source/TweeningUtilsEditor/Private/Widgets/TweenSliderStyle.cpp b/Engine/Plugins/Animation/TweeningUtils/Source/TweeningUtilsEditor/Private/Widgets/TweenSliderStyle.cpp
--- a/Engine/Plugins/Animation/TweeningUtils/Source/TweeningUtilsEditor/Private/Widgets/TweenSliderStyle.cpp
+++ b/Engine/Plugins/Animation/TweeningUtils/Source/TweeningUtilsEditor/Private/Widgets/TweenSliderStyle.cpp
@@ -33,7 +33,7 @@ void FTweenPointStyle::GetResources(TArray<const FSlateBrush*>& OutBrushes) cons
 FTweenSliderStyle::FTweenSliderStyle()
 	: BarDimensions(250.0, 12.0)
 	, BarBrush(static_cast<FSlateBrush>(FSlateColorBrush(FStyleColors::Black)))
-	, NormalSliderButton(static_cast<FSlateBrush>(FSlateRoundedBoxBrush(FStyleColors::Foreground, 5.0f, FStyleColors::Transparent, 0.f)))
+	, NormalSliderButton(static_cast<FSlateBrush>(FSlateRoundedBoxBrush(FStyleColors::ForegroundHover, 5.0f, FStyleColors::Transparent, 0.f)))
 	, HoveredSliderButton(static_cast<FSlateBrush>(FSlateRoundedBoxBrush(FStyleColors::ForegroundHover, 5.0f, FStyleColors::Transparent, 0.f)))
 	, PressedSliderButton(static_cast<FSlateBrush>(FSlateRoundedBoxBrush(FStyleColors::Foreground, 5.0f, FStyleColors::Transparent, 0.f)))
 	, NormalIconTint(FStyleColors::Black)
diff --git a/Engine/Plugins/Editor/CurveEditorTools/Source/CurveEditorTools/Private/Extensions/Tweening/CurveEditorTweenModels.cpp b/Engine/Plugins/Editor/CurveEditorTools/Source/CurveEditorTools/Private/Extensions/Tweening/CurveEditorTweenModels.cpp
--- a/Engine/Plugins/Editor/CurveEditorTools/Source/CurveEditorTools/Private/Extensions/Tweening/CurveEditorTweenModels.cpp
+++ b/Engine/Plugins/Editor/CurveEditorTools/Source/CurveEditorTools/Private/Extensions/Tweening/CurveEditorTweenModels.cpp
@@ -15,23 +15,23 @@ static TArray<TweeningUtilsEditor::FTweenModelUIEntry> MakeTweenModels(TWeakPtr<
 	using namespace TweeningUtilsEditor;
 	
 	TArray<FTweenModelUIEntry> Result;
-	ForEachCurveTweenable([&InCurveEditor, &Result]<EBlendFunction TBlendFunction>
+	const auto AddCurveTweenable = [&InCurveEditor, &Result]<EBlendFunction TBlendFunction>
 	{
-		if constexpr (!SupportsTweenRange(TBlendFunction))
-		{
-			return;
-		}
-		
 		using FModel = TEditorTweenModel<TContextAwareTweeningModel<TBlendFunction>>;
-
 		TUniquePtr<FModel> TweenModel = MakeUnique<FModel>(
 			InCurveEditor,			// TTangentFlatteningTweenProxy arg
 			InCurveEditor			// FContextAwareTweeningModel args
 			);
 		Result.Emplace(MoveTemp(TweenModel), FTweenModelDisplayInfo(TBlendFunction));
-	});
-	static_assert(!SupportsTweenRange(EBlendFunction::TimeOffset));
+	};
+	static_assert(static_cast<int32>(EBlendFunction::Num) == 7, "You probably want to add the new tween function here.");
 	
+	// Could use TweeningUtilsEditor::ForEachCurveTweenable, but the functions are supposed to be listed in this specific order.
+	AddCurveTweenable.template operator()<EBlendFunction::BlendEase>();
+	AddCurveTweenable.template operator()<EBlendFunction::BlendNeighbor>();
+	AddCurveTweenable.template operator()<EBlendFunction::BlendRelative>();
+	AddCurveTweenable.template operator()<EBlendFunction::PushPull>();
+	AddCurveTweenable.template operator()<EBlendFunction::SmoothRough>();
 	Result.Emplace(
 		MakeUnique<TEditorTweenModel<FCurveTimeOffsetTweenModel>>(
 			InCurveEditor,		// TTangentFlatteningTweenProxy arg
@@ -39,6 +39,7 @@ static TArray<TweeningUtilsEditor::FTweenModelUIEntry> MakeTweenModels(TWeakPtr<
 			),
 		FTweenModelDisplayInfo(EBlendFunction::TimeOffset)
 	);
+	AddCurveTweenable.template operator()<EBlendFunction::ControlsToTween>();
 
 	return Result;
 }
diff --git a/Engine/Source/Editor/EditorStyle/Private/StarshipStyle.cpp b/Engine/Source/Editor/EditorStyle/Private/StarshipStyle.cpp
--- a/Engine/Source/Editor/EditorStyle/Private/StarshipStyle.cpp
+++ b/Engine/Source/Editor/EditorStyle/Private/StarshipStyle.cpp
@@ -7663,6 +7663,10 @@ void FStarshipEditorStyle::FStyle::SetupToolkitStyles()
 
 	// New Curve Editor 
 	{
+		// Foreground colour that sliders, etc. should use. Generally, color when normal and hovered color should be the same (brighter as if hovered).
+		const FSlateColor TweenColor(EStyleColor::ForegroundHover);
+		Set("CurveEditor.TweenForeground", TweenColor);
+		
 		FToolBarStyle CurveEditorToolbar = FStarshipCoreStyle::GetCoreStyle().GetWidgetStyle<FToolBarStyle>("AssetEditorToolbar");
 
 		CurveEditorToolbar.SetButtonPadding(       FMargin(0.0f, 0.0f));
@@ -7684,7 +7688,11 @@ void FStarshipEditorStyle::FStyle::SetupToolkitStyles()
 			
 			CurveEditorToolbar.SetComboLabelMinWidth(90.f);
 			Set("CurveEditorToolBar.TangentMode", CurveEditorToolbar);
-			
+
+			FToolBarStyle TweenToolbar = CurveEditorToolbar;
+			// Toolbar sliders should appear just a bright as when they are hovered.
+			TweenToolbar.ButtonStyle.SetNormalForeground(TweenColor);
+			TweenToolbar.ComboButtonStyle.ButtonStyle.SetNormalForeground(TweenColor);
 			CurveEditorToolbar.SetComboLabelMinWidth(LongestTweeName);
 			Set("CurveEditorTweenToolbar", CurveEditorToolbar);
 			
