Control Rig Edit Mode: Fixed frustum calculation

#rnx

diff --git a/Engine/Plugins/Animation/ControlRig/Source/ControlRigEditor/Private/EditMode/ControlRigEditModeUtil.cpp b/Engine/Plugins/Animation/ControlRig/Source/ControlRigEditor/Private/EditMode/ControlRigEditModeUtil.cpp
--- a/Engine/Plugins/Animation/ControlRig/Source/ControlRigEditor/Private/EditMode/ControlRigEditModeUtil.cpp
+++ b/Engine/Plugins/Animation/ControlRig/Source/ControlRigEditor/Private/EditMode/ControlRigEditModeUtil.cpp
@@ -257,10 +257,11 @@ TOptional<FIntRect> FSelectionHelper::RectangleFromFrustum(const FConvexVolume&
 		{
 			const FVector2D& TopLeft = ScreenPos[3], BottomRight = ScreenPos[1];
 
-			const int32 MinX = FMath::Max<int32>(0, FMath::FloorToInt32(TopLeft.X));
-			const int32 MinY = FMath::Max<int32>(0, FMath::FloorToInt32(TopLeft.Y));
-
 			const FIntPoint ViewportSize = ViewportClient->Viewport->GetSizeXY();
+			
+			const int32 MinX = FMath::Clamp(FMath::FloorToInt32(TopLeft.X), 0, ViewportSize.X-1);
+			const int32 MinY = FMath::Clamp(FMath::FloorToInt32(TopLeft.Y), 0, ViewportSize.Y-1);
+		
 			const int32 MaxX = FMath::Clamp(FMath::CeilToInt32(BottomRight.X), MinX+1, ViewportSize.X);
 			const int32 MaxY = FMath::Clamp(FMath::CeilToInt32(BottomRight.Y), MinY+1, ViewportSize.Y);
 
