Control Rig Outliner: fixed crash using invalid weak ptr

#rnx

diff --git a/Engine/Plugins/Animation/ControlRig/Source/ControlRigEditor/Private/EditMode/SControlRigOutliner.cpp b/Engine/Plugins/Animation/ControlRig/Source/ControlRigEditor/Private/EditMode/SControlRigOutliner.cpp
--- a/Engine/Plugins/Animation/ControlRig/Source/ControlRigEditor/Private/EditMode/SControlRigOutliner.cpp
+++ b/Engine/Plugins/Animation/ControlRig/Source/ControlRigEditor/Private/EditMode/SControlRigOutliner.cpp
@@ -1800,7 +1800,8 @@ bool SSearchableMultiRigHierarchyTreeView::CanToggleModuleManipulators() const
 	const TArray<FMultiRigData> Selection = TreeView->GetSelectedData();
 	for (const FMultiRigData& Data : Selection)
 	{
-		if (Data.WeakControlRig->IsModularRig())
+		const UControlRig* ControlRig = Data.WeakControlRig.Get();
+		if (ControlRig && ControlRig->IsModularRig())
 		{
 			return true;
 		}
