Updates To ChainInfo node:
-converted node to be not mutable
-switched to using direct debug draw calls
-now using struct to contain any segment info outputs
-renamed OutputTransform to InterpolatedTransform

#rnx

#changelist validated

diff --git a/Engine/Plugins/Animation/ControlRig/Source/ControlRig/Private/Units/Highlevel/Hierarchy/RigUnit_ChainInfo.cpp b/Engine/Plugins/Animation/ControlRig/Source/ControlRig/Private/Units/Highlevel/Hierarchy/RigUnit_ChainInfo.cpp
--- a/Engine/Plugins/Animation/ControlRig/Source/ControlRig/Private/Units/Highlevel/Hierarchy/RigUnit_ChainInfo.cpp
+++ b/Engine/Plugins/Animation/ControlRig/Source/ControlRig/Private/Units/Highlevel/Hierarchy/RigUnit_ChainInfo.cpp
@@ -2,7 +2,6 @@
 
 #include "Units/Highlevel/Hierarchy/RigUnit_ChainInfo.h"
 #include "Math/ControlRigMathLibrary.h"
-#include "RigVMFunctions/Debug/RigVMFunction_DebugTransform.h"
 #include "RigVMCore/RigVMDebugDrawSettings.h"
 
 #include UE_INLINE_GENERATED_CPP_BY_NAME(RigUnit_ChainInfo)
@@ -88,7 +87,7 @@ FRigUnit_ChainInfo_Execute()
 	ParamLength = Param * ChainLength;
 
 	// Binary search for current segment containing param based on recorded lengths
-    SegmentIndex = 0;
+    SegmentInfo.SegmentIndex = 0;
 	int StartSegmentIndex = 0;
 	int EndSegmentIndex = Segments.Num() - 1;
 
@@ -99,7 +98,7 @@ FRigUnit_ChainInfo_Execute()
 		{
 			if (SearchIndex == 0)
 			{
-				SegmentIndex = SearchIndex;
+				SegmentInfo.SegmentIndex = SearchIndex;
 				break;
 			}
 			EndSegmentIndex = SearchIndex;
@@ -108,35 +107,35 @@ FRigUnit_ChainInfo_Execute()
 		{
 			if (Segments[SearchIndex].CumLength == ParamLength || SearchIndex + 1 == EndSegmentIndex)
 			{
-				SegmentIndex = SearchIndex + 1;
+				SegmentInfo.SegmentIndex = SearchIndex + 1;
 				break;
 			}
 			StartSegmentIndex = SearchIndex;
 		}
 	}
 
-	const FRigUnit_ChainInfo_Segment& CurrentSegment = Segments[SegmentIndex];
+	const FRigUnit_ChainInfo_Segment& CurrentSegment = Segments[SegmentInfo.SegmentIndex];
 
-	SegmentStartItem = CurrentSegment.StartItem.GetKey();
-	SegmentStartItemIndex = CurrentSegment.StartItemIndex;
-	SegmentEndItem = CurrentSegment.EndItem.GetKey();
-	SegmentEndItemIndex = CurrentSegment.EndItemIndex;
+	SegmentInfo.SegmentStartItem = CurrentSegment.StartItem.GetKey();
+	SegmentInfo.SegmentStartItemIndex = CurrentSegment.StartItemIndex;
+	SegmentInfo.SegmentEndItem = CurrentSegment.EndItem.GetKey();
+	SegmentInfo.SegmentEndItemIndex = CurrentSegment.EndItemIndex;
 
     // Calculate segment length and segment param length
-	if (SegmentIndex == 0) 
+	if (SegmentInfo.SegmentIndex == 0) 
 	{
-		SegmentParamLength = ParamLength;
+		SegmentInfo.SegmentParamLength = ParamLength;
 	}
 	else 
 	{
-		SegmentParamLength = ParamLength - Segments[SegmentIndex - 1].CumLength;
+		SegmentInfo.SegmentParamLength = ParamLength - Segments[SegmentInfo.SegmentIndex - 1].CumLength;
 	}
-    SegmentLength = CurrentSegment.Length;
+    SegmentInfo.SegmentLength = CurrentSegment.Length;
 
     // Calculate segment param
-	if (SegmentLength != 0.f) 
+	if (SegmentInfo.SegmentLength != 0.f) 
 	{
-	    SegmentParam = SegmentParamLength / SegmentLength;
+	    SegmentInfo.SegmentParam = SegmentInfo.SegmentParamLength / SegmentInfo.SegmentLength;
 	}
 
 	if (bCalculateStretch)
@@ -149,7 +148,7 @@ FRigUnit_ChainInfo_Execute()
 
 		if (CurrentSegment.InitialLength != 0.f) 
 		{
-			SegmentStretchFactor = CurrentSegment.Length / CurrentSegment.InitialLength;
+			SegmentInfo.SegmentStretchFactor = CurrentSegment.Length / CurrentSegment.InitialLength;
 		}
 	}
 
@@ -158,51 +157,42 @@ FRigUnit_ChainInfo_Execute()
 	FTransform EndTransform;
 	if (bInitial)
 	{
-		StartTransform = Hierarchy->GetInitialGlobalTransform(SegmentStartItem);
-		EndTransform = Hierarchy->GetInitialGlobalTransform(SegmentEndItem);
+		StartTransform = Hierarchy->GetInitialGlobalTransform(SegmentInfo.SegmentStartItem);
+		EndTransform = Hierarchy->GetInitialGlobalTransform(SegmentInfo.SegmentEndItem);
 	}
 	else
 	{
-		StartTransform = Hierarchy->GetGlobalTransform(SegmentStartItem);
-		EndTransform = Hierarchy->GetGlobalTransform(SegmentEndItem);
+		StartTransform = Hierarchy->GetGlobalTransform(SegmentInfo.SegmentStartItem);
+		EndTransform = Hierarchy->GetGlobalTransform(SegmentInfo.SegmentEndItem);
 	}
-	OutputTransform = FControlRigMathLibrary::LerpTransform(StartTransform, EndTransform, SegmentParam);
+	InterpolatedTransform = FControlRigMathLibrary::LerpTransform(StartTransform, EndTransform, SegmentInfo.SegmentParam);
     
 	// Debug draw
 	if (bDebug) {
-		FRigVMDebugDrawSettings DebugDrawSettings; 
-
-		// Draw interpolated output transform
-		FRigVMFunction_DebugTransformMutableNoSpace::StaticExecute(ExecuteContext, 
-		                                OutputTransform,
-		                                ERigUnitDebugTransformMode::Axes,
-		                                FLinearColor::White,
-		                                0.f,
-		                                DebugScale,
-										FTransform::Identity,
-		                                true,
-										DebugDrawSettings);
-
-		// Draw transform of start item
-		FRigVMFunction_DebugTransformMutableNoSpace::StaticExecute(ExecuteContext, 
-		                                StartTransform,
-		                                ERigUnitDebugTransformMode::Box,
-		                                FLinearColor::Green,
-		                                0.f,
-		                                DebugScale,
-										FTransform::Identity,
-		                                true,
-										DebugDrawSettings);
-
+		FRigVMDebugDrawSettings DebugDrawSettings;
+
+	    // Draw interpolated output transform
+	    ExecuteContext.GetDrawInterface()->DrawAxes(FTransform::Identity, 
+			                                        InterpolatedTransform, 
+													DebugScale, 
+													0.f, 
+													DebugDrawSettings.DepthPriority, 
+													DebugDrawSettings.Lifetime);
+		// Draw transform of start item									
+		StartTransform.SetScale3D(StartTransform.GetScale3D() * DebugScale);
+		ExecuteContext.GetDrawInterface()->DrawBox(FTransform::Identity, 
+			                                       StartTransform, 
+												   FLinearColor::Green, 
+												   0.f, 
+												   DebugDrawSettings.DepthPriority, 
+												   DebugDrawSettings.Lifetime);
 		// Draw transform of end item
-		FRigVMFunction_DebugTransformMutableNoSpace::StaticExecute(ExecuteContext, 
-		                                EndTransform,
-		                                ERigUnitDebugTransformMode::Box,
-		                                FLinearColor::Red,
-		                                0.f,
-		                                DebugScale,
-										FTransform::Identity,
-		                                true,
-										DebugDrawSettings);
+		EndTransform.SetScale3D(EndTransform.GetScale3D() * DebugScale);
+		ExecuteContext.GetDrawInterface()->DrawBox(FTransform::Identity, 
+			                                       EndTransform, 
+												   FLinearColor::Red, 
+												   0.f, 
+												   DebugDrawSettings.DepthPriority, 
+												   DebugDrawSettings.Lifetime);												   
 	}
 }
diff --git a/Engine/Plugins/Animation/ControlRig/Source/ControlRig/Public/Units/Highlevel/Hierarchy/RigUnit_ChainInfo.h b/Engine/Plugins/Animation/ControlRig/Source/ControlRig/Public/Units/Highlevel/Hierarchy/RigUnit_ChainInfo.h
--- a/Engine/Plugins/Animation/ControlRig/Source/ControlRig/Public/Units/Highlevel/Hierarchy/RigUnit_ChainInfo.h
+++ b/Engine/Plugins/Animation/ControlRig/Source/ControlRig/Public/Units/Highlevel/Hierarchy/RigUnit_ChainInfo.h
@@ -22,60 +22,132 @@ struct FRigUnit_ChainInfo_Segment
 		,CumLength(0.f)
 	{}
 
-	/**
-	 * Start item of the chain segment
-	 */
+	/*
+	* Start item of the chain segment
+	*/
 	UPROPERTY()
 	FCachedRigElement StartItem;
 
-	/**
-	 * Start item index of the chain segment
-	 */
+	/*
+	* Start item index of the chain segment
+	*/
 	UPROPERTY()
 	int StartItemIndex;
 
-	/**
-	 * Emd item of the chain segment
-	 */
+	/*
+	* Emd item of the chain segment
+	*/
 	UPROPERTY()
 	FCachedRigElement EndItem;
 
-	/**
-	 * Emd item index of the chain segment
-	 */
+	/*
+	* Emd item index of the chain segment
+	*/
 	UPROPERTY()
 	int EndItemIndex;
 
-	/**
-	 * Initial length of segment
-	 */
+	/*
+	* Initial length of segment
+	*/
 	UPROPERTY()
 	float InitialLength; 
 
-	/**
-	 * Inclusive initial length of all previous segments and this one
-	 */
+	/*
+	* Inclusive initial length of all previous segments and this one
+	*/
 	UPROPERTY()
 	float InitialCumLength;
 
-	/**
-	 * Length of segment
-	 */
+	/*
+	* Length of segment
+	*/
 	UPROPERTY()
 	float Length; 
 
-	/**
-	 * Inclusive length of all previous segments and this one
-	 */
+	/*
+	* Inclusive length of all previous segments and this one
+	*/
 	UPROPERTY()
 	float CumLength;
 };
 
+USTRUCT(BlueprintType)
+struct FRigUnit_ChainInfo_SegmentInfo 
+{
+    GENERATED_BODY()
+
+	FRigUnit_ChainInfo_SegmentInfo()
+		:SegmentIndex(0)
+		,SegmentLength(0.f)	
+		,SegmentParam(0.f)
+		,SegmentParamLength(0.f)
+		,SegmentStartItem(NAME_None, ERigElementType::Bone)
+		,SegmentStartItemIndex(0)
+		,SegmentEndItem(NAME_None, ERigElementType::Bone)
+        ,SegmentEndItemIndex(0)
+		,SegmentStretchFactor(1.f)
+	{}
+
+	/* 
+	* The current segment index
+	*/
+	UPROPERTY()
+	int SegmentIndex;
+ 
+	/* 
+	* The current segment length
+	*/
+	UPROPERTY()
+	float SegmentLength;	
+ 
+	/* 
+	* The current segment parameter from 0 to 1 
+	*/
+	UPROPERTY()
+	float SegmentParam;
+ 
+	/* 
+	* Local segment length of segment param
+	*/
+	UPROPERTY()
+	float SegmentParamLength;
+ 
+	/* 
+	* The item starting the current segment 
+	*/
+	UPROPERTY()
+	FRigElementKey SegmentStartItem;
+ 
+	/* 
+	* The item index starting the current segment 
+	*/
+	UPROPERTY()
+	int SegmentStartItemIndex;
+ 
+	/* 
+	* The item ending the current segment 
+	*/
+	UPROPERTY()
+	FRigElementKey SegmentEndItem;
+ 
+	/* 
+	* The item index ending the current segment 
+	*/
+	UPROPERTY()
+	int SegmentEndItemIndex;
+
+	/* 
+	* Stretch factor of current segment
+	*/
+	UPROPERTY()
+	float SegmentStretchFactor;
+};
+
 /*
- * Retrieves various pieces of info about an interpolated transform hierarchy from an rig element item list
- */
+* Retrieves various pieces of info about an interpolated transform hierarchy from an rig element item list
+*/
 USTRUCT(meta=(DisplayName="Chain Info", Category="Hierarchy", Keywords="Chain"))
-struct FRigUnit_ChainInfo : public FRigUnit_HighlevelBaseMutable
+struct FRigUnit_ChainInfo : public FRigUnit_HighlevelBase
 {
 	GENERATED_BODY()
 
@@ -85,137 +157,81 @@ struct FRigUnit_ChainInfo : public FRigUnit_HighlevelBaseMutable
 		,bInitial(false)
 		,bDebug(false)
 		,DebugScale(1.f)
-		,OutputTransform(FTransform::Identity)
+		,InterpolatedTransform(FTransform::Identity)
 		,ChainLength(0.f)
 		,ParamLength(0.f)
-		,SegmentIndex(0)
-		,SegmentLength(0.f)	
-		,SegmentParam(0.f)
-		,SegmentParamLength(0.f)
-		,SegmentStartItem(NAME_None, ERigElementType::Bone)
-		,SegmentStartItemIndex(0)
-		,SegmentEndItem(NAME_None, ERigElementType::Bone)
-        ,SegmentEndItemIndex(0)
 		,ChainStretchFactor(1.f)
-		,SegmentStretchFactor(1.f)
+        ,SegmentInfo(FRigUnit_ChainInfo_SegmentInfo())
 	{}
 
 	RIGVM_METHOD()
 	virtual void Execute() override;
 
-	/** 
-	 * The items to use to interpret the chain
-	 */
+	/*
+	* The items to use to interpret the chain
+	*/
 	UPROPERTY(meta = (Input))
 	TArray<FRigElementKey> Items;
 
-	/**
-	 * The parameter value down the chain of items from 0 to 1
-	 */	
+	/*
+	* The parameter value down the chain of items from 0 to 1
+	*/	
 	UPROPERTY(meta = (Input, ClampMin=0.f, ClampMax=1.f))
 	float Param;
 
-	/**
-	 * If True calculate stretch factors of chain and current segment
-	 */	
+	/*
+	* If True calculate stretch factors of chain and current segment
+	*/	
 	UPROPERTY(meta = (Input))
 	bool bCalculateStretch;
 
-	/**
-	 * If True use initial transform values for chain
-	 */	
+	/*
+	* If True use initial transform values for chain
+	*/	
 	UPROPERTY(meta = (Input))
 	bool bInitial;
 
-	/**
-	 * Enable debug draw for node
-	 */	
+	/*
+	* Enable debug draw for node
+	*/	
 	UPROPERTY(meta = (Input))
 	bool bDebug;
 
-	/**
-	 * Debug draw scale
-	 */	
+	/*
+	* Debug draw scale
+	*/	
 	UPROPERTY(meta = (Input, UIMin = 0.f))
 	float DebugScale;
 
     /* 
-	 * The interpolated transform at the chain's input parameter
-	 */
+	* The interpolated transform at the chain's input parameter
+	*/
 	UPROPERTY(meta = (Output))
-	FTransform OutputTransform;
+	FTransform InterpolatedTransform;
 
     /* 
-	 * The length of the interpolated chain
-	 */
+	* The length of the interpolated chain
+	*/
 	UPROPERTY(meta = (Output))
 	float ChainLength;
 
     /* 
-	 * The length of the interpolated chain
-	 */
+	* The length of the interpolated chain
+	*/
 	UPROPERTY(meta = (Output))
 	float ParamLength;
-
-    /* 
-	 * The current segment index
-	 */
-	UPROPERTY(meta = (Output))
-	int SegmentIndex;
-
-	/* 
-	 * The current segment length
-	 */
-	UPROPERTY(meta = (Output))
-	float SegmentLength;	
-
-    /* 
-	 * The current segment parameter from 0 to 1 
-	 */
-	UPROPERTY(meta = (Output))
-	float SegmentParam;
-
-    /* 
-	 * Local segment length of segment param
-	 */
-	UPROPERTY(meta = (Output))
-	float SegmentParamLength;
-
-    /* 
-	 * The item starting the current segment 
-	 */
-	UPROPERTY(meta = (Output))
-	FRigElementKey SegmentStartItem;
-
-    /* 
-	 * The item index starting the current segment 
-	 */
-	UPROPERTY(meta = (Output))
-	int SegmentStartItemIndex;
-
-    /* 
-	 * The item ending the current segment 
-	 */
-	UPROPERTY(meta = (Output))
-	FRigElementKey SegmentEndItem;
-
-    /* 
-	 * The item index ending the current segment 
-	 */
-	UPROPERTY(meta = (Output))
-	int SegmentEndItemIndex;
-
+    
 	/* 
-	 * Stretch factor of chain
-	 */
+	* Stretch factor of chain
+	*/
 	UPROPERTY(meta = (Output))
 	float ChainStretchFactor;
 
 	/* 
-	 * Stretch factor of current segment
-	 */
+	* Segment Info
+	*/
 	UPROPERTY(meta = (Output))
-	float SegmentStretchFactor;
+	FRigUnit_ChainInfo_SegmentInfo SegmentInfo;
 
 	// Used to cache the internally used index
 	UPROPERTY()
