Control Rig Override: Improve backwards compat for math types

#rnx

diff --git a/Engine/Plugins/Animation/ControlRig/Source/ControlRig/Private/ControlRigOverride.cpp b/Engine/Plugins/Animation/ControlRig/Source/ControlRig/Private/ControlRigOverride.cpp
--- a/Engine/Plugins/Animation/ControlRig/Source/ControlRig/Private/ControlRigOverride.cpp
+++ b/Engine/Plugins/Animation/ControlRig/Source/ControlRig/Private/ControlRigOverride.cpp
@@ -1,6 +1,7 @@
 // Copyright Epic Games, Inc. All Rights Reserved.
 
 #include "ControlRigOverride.h"
+#include "ControlRig.h"
 #include "RigVMCore/RigVMMemoryCommon.h"
 
 #if WITH_EDITOR
@@ -584,10 +585,11 @@ bool FControlRigOverrideValue::SetFromString(const FString& InValue)
 		if(const FProperty* LeafProperty = GetLeafProperty())
 		{
 			TArray<uint8, TAlignedHeapAllocator<16>> TempDataArray;
-			void* DataPtr = AllocateDataIfRequired(TempDataArray); 
-			
+			void* TempDataPtr = AllocateDataIfRequired(TempDataArray); 
+
+			// first try to import the string value using the property
 			FRigVMMemoryStorageImportErrorContext ErrorPipe(false);
-			LeafProperty->ImportText_Direct(*InValue, DataPtr, nullptr, PPF_None, &ErrorPipe);
+			LeafProperty->ImportText_Direct(*InValue, TempDataPtr, nullptr, PPF_None, &ErrorPipe);
 			if(ErrorPipe.NumErrors == 0)
 			{
 				FreeDataIfRequired(DataArray);
@@ -596,8 +598,123 @@ bool FControlRigOverrideValue::SetFromString(const FString& InValue)
 				CachedStringValue = InValue;
 				return true;
 			}
-
 			FreeDataIfRequired(TempDataArray);
+
+			// we've hit a problem when importing the value
+			if(const FStructProperty* StructProperty = CastField<FStructProperty>(LeafProperty))
+			{
+				// potentially we have a default value for a math type here.
+				const UStruct* Struct = StructProperty->Struct;
+
+				if(Struct == TBaseStructure<FVector2D>::Get())
+				{
+					FString ValueString = InValue;
+					for(int32 TryIndex = 0; TryIndex < 2; TryIndex++)
+					{
+						if(InitFromLegacyString<FVector2D>(ValueString, FVector2D::ZeroVector))
+						{
+							return true;
+						}
+
+						// potentially we have the value as "1.000,2.000"
+						if(TryIndex == 0)
+						{
+							TArray<FString> Parts;
+							if(RigVMStringUtils::SplitString(ValueString, TEXT(","), Parts))
+							{
+								if(Parts.Num() == 2)
+								{
+									ValueString = FString::Printf(TEXT("(X=%s,Y=%s)"), *Parts[0], *Parts[1]);
+									continue;
+								}
+							}
+						}
+						
+						break;
+					}
+				}
+
+				if(Struct == TBaseStructure<FVector>::Get())
+				{
+					FString ValueString = InValue;
+					for(int32 TryIndex = 0; TryIndex < 2; TryIndex++)
+					{
+						if(InitFromLegacyString<FVector>(ValueString, FVector::ZeroVector))
+						{
+							return true;
+						}
+					
+						// potentially we have the value as "1.000,2.000,3.000"
+						if(TryIndex == 0)
+						{
+							TArray<FString> Parts;
+							if(RigVMStringUtils::SplitString(ValueString, TEXT(","), Parts))
+							{
+								if(Parts.Num() == 3)
+								{
+									ValueString = FString::Printf(TEXT("(X=%s,Y=%s,Z=%s)"), *Parts[0], *Parts[1], *Parts[2]);
+									continue;
+								}
+							}
+						}
+						
+						break;
+					}
+				}
+
+				if(Struct == TBaseStructure<FRotator>::Get())
+				{
+					FString ValueString = InValue;
+					for(int32 TryIndex = 0; TryIndex < 2; TryIndex++)
+					{
+						if(InitFromLegacyString<FRotator>(ValueString, FRotator::ZeroRotator))
+						{
+							return true;
+						}
+
+						// potentially we have the value as "1.000,2.000,3.000"
+						if(TryIndex == 0)
+						{
+							TArray<FString> Parts;
+							if(RigVMStringUtils::SplitString(ValueString, TEXT(","), Parts))
+							{
+								if(Parts.Num() == 3)
+								{
+									ValueString = FString::Printf(TEXT("(P=%s,Y=%s,R=%s)"), *Parts[0], *Parts[1], *Parts[2]);
+									continue;
+								}
+							}
+						}
+						
+						break;
+					}
+				}
+
+				if(Struct == TBaseStructure<FQuat>::Get())
+				{
+					if(InitFromLegacyString<FQuat>(InValue, FQuat::Identity))
+					{
+						return true;
+					}
+				}
+
+				if(Struct == TBaseStructure<FTransform>::Get())
+				{
+					if(InitFromLegacyString<FTransform>(InValue, FTransform::Identity))
+					{
+						return true;
+					}
+				}
+			}
+
+			if(SubjectName.IsSet())
+			{
+				UE_LOG(LogControlRig, Warning, TEXT("Unable to import default value '%s' for override '%s' on subject '%s'."), *InValue, *GetPath(), *SubjectName.GetValue().ToString());
+			}
+			else
+			{
+				UE_LOG(LogControlRig, Warning, TEXT("Unable to import default value '%s' for override '%s'"), *InValue, *GetPath());
+			}
 		}
 	}
 	return false;
diff --git a/Engine/Plugins/Animation/ControlRig/Source/ControlRig/Public/ControlRigOverride.h b/Engine/Plugins/Animation/ControlRig/Source/ControlRig/Public/ControlRigOverride.h
--- a/Engine/Plugins/Animation/ControlRig/Source/ControlRig/Public/ControlRigOverride.h
+++ b/Engine/Plugins/Animation/ControlRig/Source/ControlRig/Public/ControlRigOverride.h
@@ -5,6 +5,7 @@
 #include "CoreMinimal.h"
 #include "Hash/CityHash.h"
 #include "UObject/WeakFieldPtr.h"
+#include "RigVMStringUtils.h"
 #include "ControlRigOverride.generated.h"
 
 #if WITH_EDITOR
@@ -75,6 +76,20 @@ public:
 	// returns the raw data of this value
 	const void* GetData() const;
 
+	// returns the raw data of this value
+	template<typename T>
+	T* GetData()
+	{
+		return static_cast<T*>(GetData());
+	}
+
+	// returns the raw data of this value
+	template<typename T>
+	const T* GetData() const
+	{
+		return static_cast<T*>(GetData());
+	}
+
 	// returns the value of this override as a string (or an empty string)  
 	const FString& ToString() const;
 
@@ -134,6 +149,20 @@ private:
 	void FreeDataIfRequired(TArray<uint8, TAlignedHeapAllocator<16>>& InOutDataArray) const;
 	void CopyValue(void* InDestPtr, const void* InSourcePtr) const;
 
+	template<typename T>
+	bool InitFromLegacyString(const FString& InLegacyString, const T& InDefault)
+	{
+		T Value = InDefault;
+		if(Value.InitFromString(InLegacyString))
+		{
+			void* DataPtr = AllocateDataIfRequired(); 
+			CopyValue(DataPtr, &Value);
+			CachedStringValue = InLegacyString;
+			return true;
+		}
+		return false;
+	}
+
 	TOptional<FString> Path;
 	TOptional<FString> CachedStringValue;
 	TOptional<FName> SubjectName;
