Control Rig Edit Mode: propagate the autokey event to the driven controls when using keying on release

#rnx

diff --git a/Engine/Plugins/Animation/ControlRig/Source/ControlRigEditor/Private/EditMode/ControlRigEditMode.cpp b/Engine/Plugins/Animation/ControlRig/Source/ControlRigEditor/Private/EditMode/ControlRigEditMode.cpp
--- a/Engine/Plugins/Animation/ControlRig/Source/ControlRigEditor/Private/EditMode/ControlRigEditMode.cpp
+++ b/Engine/Plugins/Animation/ControlRig/Source/ControlRigEditor/Private/EditMode/ControlRigEditMode.cpp
@@ -1107,6 +1107,13 @@ bool FControlRigEditMode::StartTracking(FEditorViewportClient* InViewportClient,
 		return DragToolHandler.StartTracking(InViewportClient, InViewport);
 	}
 
+	const EAxisList::Type CurrentAxis = InViewportClient->GetCurrentWidgetAxis();
+	if (CurrentAxis == EAxisList::None)
+	{
+		// not manipulating a required axis
+		return false;
+	}
+	
 	return HandleBeginTransform(InViewportClient);
 }
 
diff --git a/Engine/Plugins/Animation/ControlRig/Source/ControlRigEditor/Private/Sequencer/EditModeAnimationUtil.cpp b/Engine/Plugins/Animation/ControlRig/Source/ControlRigEditor/Private/Sequencer/EditModeAnimationUtil.cpp
--- a/Engine/Plugins/Animation/ControlRig/Source/ControlRigEditor/Private/Sequencer/EditModeAnimationUtil.cpp
+++ b/Engine/Plugins/Animation/ControlRig/Source/ControlRigEditor/Private/Sequencer/EditModeAnimationUtil.cpp
@@ -168,6 +168,19 @@ void FControlRigKeyframer::Apply(const FControlRigInteractionScope& InInteractio
 				
 				AutoKeyEvent.Broadcast(InteractingRig, Control, EControlRigSetKey::DoNotCare);				
 			}
+
+			// driven controls
+			if (Control->CanDriveControls())
+			{
+				for (const FRigElementKey& DrivenKey : Control->Settings.DrivenControls)
+				{
+					const bool bHandleDrivenKey = DrivenKey.Type == ERigElementType::Control && !InteractingControls.Contains(DrivenKey);						
+					if (FRigControlElement* DrivenControl = bHandleDrivenKey ? InteractingRig->FindControl(DrivenKey.Name) : nullptr)
+					{
+						AutoKeyEvent.Broadcast(InteractingRig, DrivenControl, EControlRigSetKey::DoNotCare);
+					}
+				}
+			}
 		}
 	}
 }
