Anim In  Engine: Add option to do additive blends with poses. We use the blend function for that since it already supports working off of an initial value.

diff --git a/Engine/Plugins/Animation/ControlRig/Source/ControlRig/Private/Tools/ControlRigPose.cpp b/Engine/Plugins/Animation/ControlRig/Source/ControlRig/Private/Tools/ControlRigPose.cpp
--- a/Engine/Plugins/Animation/ControlRig/Source/ControlRig/Private/Tools/ControlRigPose.cpp
+++ b/Engine/Plugins/Animation/ControlRig/Source/ControlRig/Private/Tools/ControlRigPose.cpp
@@ -121,7 +121,7 @@ void FControlRigControlPose::PastePoseInternal(UControlRig* ControlRig, bool bDo
 				{
 					if (bDoLocal) // -V547  
 					{
-						ControlRig->SetControlLocalTransform(ControlElement->GetFName(), CopyRigControl->LocalTransform, true,Context,bSetupUndo, true/* bFixEulerFlips*/);
+						ControlRig->SetControlLocalTransform(ControlElement->GetFName(), CopyRigControl->LocalTransform, true, Context, bSetupUndo, true/* bFixEulerFlips*/);
 					}
 					else
 					{
@@ -178,8 +178,32 @@ void FControlRigControlPose::PastePoseInternal(UControlRig* ControlRig, bool bDo
 	}
 }
 
-void FControlRigControlPose::BlendWithInitialPoses(FControlRigControlPose& InitialPose, UControlRig* ControlRig, bool bDoKey, bool bDoMirror, float BlendValue)
+void FControlRigControlPose::BlendWithInitialPoses(FControlRigControlPose& InitialPose, UControlRig* ControlRig, bool bDoKey, bool bDoMirror, float BlendValue, bool bDoAdditive)
 {
+
+	auto BlendTransforms = [bDoAdditive, BlendValue](FTransform& InitialVal, FTransform& Val) 
+	{ 
+		FVector Translation, Scale;
+		FQuat Rotation;
+		if (bDoAdditive == false)
+		{
+			Translation = FMath::Lerp(InitialVal.GetTranslation(), Val.GetTranslation(), BlendValue);
+			Rotation = FQuat::Slerp(InitialVal.GetRotation(), Val.GetRotation(), BlendValue); //doing slerp here not fast lerp, can be slow this is for content creation
+			Scale = FMath::Lerp(InitialVal.GetScale3D(), Val.GetScale3D(), BlendValue);
+		}
+		else
+		{
+			Translation = InitialVal.GetTranslation() + (Val.GetTranslation() * BlendValue);
+			//do additive as rotators
+			FRotator RotatorVal = Val.GetRotation().Rotator();
+			FRotator RotatorInitial = InitialVal.GetRotation().Rotator();
+			RotatorVal = RotatorInitial + (RotatorVal * BlendValue);
+			Rotation = FQuat(RotatorVal);
+			Scale = FMath::Lerp(InitialVal.GetScale3D(), Val.GetScale3D(), BlendValue); //scale is just blended
+		}
+		Val = FTransform(Rotation, Translation, Scale);
+	};
+
 	if (InitialPose.CopyOfControls.Num() == 0)
 	{
 		return;
@@ -237,25 +261,15 @@ void FControlRigControlPose::BlendWithInitialPoses(FControlRigControlPose& Initi
 						{
 							FTransform Val = CopyRigControl->LocalTransform;
 							FTransform InitialVal = InitialFound->LocalTransform;
-							FVector Translation, Scale;
-							FQuat Rotation;
-							Translation = FMath::Lerp(InitialVal.GetTranslation(), Val.GetTranslation(), BlendValue);
-							Rotation = FQuat::Slerp(InitialVal.GetRotation(), Val.GetRotation(), BlendValue); //doing slerp here not fast lerp, can be slow this is for content creation
-							Scale = FMath::Lerp(InitialVal.GetScale3D(), Val.GetScale3D(), BlendValue);
-							Val = FTransform(Rotation, Translation, Scale);
-							ControlRig->SetControlLocalTransform(ControlElement->GetFName(), Val, bDoKey,Context,bSetupUndo, true/* bFixEulerFlips*/);
+							BlendTransforms(InitialVal, Val);
+							ControlRig->SetControlLocalTransform(ControlElement->GetFName(), Val, true,Context,bSetupUndo, true/* bFixEulerFlips*/);
 						}
 						else
 						{
 							FTransform Val = CopyRigControl->GlobalTransform;
 							FTransform InitialVal = InitialFound->GlobalTransform;
-							FVector Translation, Scale;
-							FQuat Rotation;
-							Translation = FMath::Lerp(InitialVal.GetTranslation(), Val.GetTranslation(), BlendValue);
-							Rotation = FQuat::Slerp(InitialVal.GetRotation(), Val.GetRotation(), BlendValue); //doing slerp here not fast lerp, can be slow this is for content creation
-							Scale = FMath::Lerp(InitialVal.GetScale3D(), Val.GetScale3D(), BlendValue);
-							Val = FTransform(Rotation, Translation, Scale);
-							ControlRig->SetControlGlobalTransform(ControlElement->GetFName(), Val, bDoKey,Context,bSetupUndo, false /*bPrintPython*/, true/* bFixEulerFlips*/);
+							BlendTransforms(InitialVal, Val);
+							ControlRig->SetControlGlobalTransform(ControlElement->GetFName(), Val, true,Context,bSetupUndo, false /*bPrintPython*/, true/* bFixEulerFlips*/);
 						}
 					}
 					else
@@ -264,21 +278,9 @@ void FControlRigControlPose::BlendWithInitialPoses(FControlRigControlPose& Initi
 						FTransform LocalTransform;
 						bool bIsMatched = MirrorTable.IsMatched(CopyRigControl->Name);
 						MirrorTable.GetMirrorTransform(*CopyRigControl, bDoLocal, bIsMatched, GlobalTransform, LocalTransform);
-						FVector InitialTranslation = InitialFound->GlobalTransform.GetTranslation();
-						FQuat InitialGlobalRotation = InitialFound->GlobalTransform.GetRotation();
-						FVector InitialGlobalScale = InitialFound->GlobalTransform.GetScale3D();
-
-						FVector InitialLocalTranslation = InitialFound->LocalTransform.GetTranslation();
-						FQuat InitialLocationRotation = InitialFound->LocalTransform.GetRotation();
-						FVector InitialLocalScale = InitialFound->LocalTransform.GetScale3D();
-
-						GlobalTransform.SetTranslation(FMath::Lerp(InitialTranslation, GlobalTransform.GetTranslation(), BlendValue));
-						GlobalTransform.SetRotation(FQuat::Slerp(InitialGlobalRotation, GlobalTransform.GetRotation(), BlendValue)); //doing slerp here not fast lerp, can be slow this is for content creation
-						GlobalTransform.SetScale3D(FMath::Lerp(InitialGlobalScale, GlobalTransform.GetScale3D(), BlendValue));
-
-						LocalTransform.SetTranslation(FMath::Lerp(InitialLocalTranslation, LocalTransform.GetTranslation(), BlendValue));
-						LocalTransform.SetRotation(FQuat::Slerp(InitialLocationRotation, LocalTransform.GetRotation(), BlendValue)); //doing slerp here not fast lerp, can be slow this is for content creation
-						LocalTransform.SetScale3D(FMath::Lerp(InitialLocalScale, LocalTransform.GetScale3D(), BlendValue));
+						
+						BlendTransforms(InitialFound->GlobalTransform, GlobalTransform);
+						BlendTransforms(InitialFound->LocalTransform, LocalTransform);
 
 						SetControlMirrorTransform(bDoLocal,ControlRig, ControlElement->GetFName(), bIsMatched, GlobalTransform,LocalTransform,bDoKey,Context,bSetupUndo);							
 					}
@@ -294,14 +296,14 @@ void FControlRigControlPose::BlendWithInitialPoses(FControlRigControlPose& Initi
 				{
 					float InitialVal = InitialFound->Value.Get<float>();
 					float Val = CopyRigControl->Value.Get<float>();
-					Val = FMath::Lerp(InitialVal, Val, BlendValue);
+					Val = bDoAdditive == false ? FMath::Lerp(InitialVal, Val, BlendValue) : InitialVal + (Val * BlendValue);
 					ControlRig->SetControlValue<float>(ControlElement->GetFName(), Val, true, Context, bSetupUndo);
 				}
 				else if (CopyRigControl->ControlType == ERigControlType::Vector2D)
 				{
 					FVector3f InitialVal = InitialFound->Value.Get<FVector3f>();
 					FVector3f Val = CopyRigControl->Value.Get<FVector3f>();
-					Val = FMath::Lerp(InitialVal, Val, BlendValue);
+					Val = bDoAdditive == false ? FMath::Lerp(InitialVal, Val, BlendValue) : InitialVal + (Val * BlendValue);
 					ControlRig->SetControlValue<FVector3f>(ControlElement->GetFName(), Val, true, Context, bSetupUndo);
 				}
 			}
@@ -365,13 +367,25 @@ void UControlRigPoseAsset::SavePose(UControlRig* InControlRig, bool bUseAll)
 	Pose.SavePose(InControlRig,bUseAll);
 }
 
-void UControlRigPoseAsset::PastePose(UControlRig* InControlRig, bool bDoKey, bool bDoMirror)
+void UControlRigPoseAsset::PastePose(UControlRig* InControlRig, bool bDoKey, bool bDoMirror, bool bDoAdditive)
 {
 #if WITH_EDITOR
 	FScopedTransaction ScopedTransaction(LOCTEXT("PastePoseTransaction", "Paste Pose"));
 	InControlRig->Modify();
 #endif
-	Pose.PastePose(InControlRig,bDoKey, bDoMirror);
+	if (bDoAdditive == false)
+	{
+		Pose.PastePose(InControlRig, bDoKey, bDoMirror);
+	}
+	else //if additive we blend with 1 since that uses current pose
+	{
+		FControlRigControlPose TempPose;
+		GetCurrentPose(InControlRig, TempPose);
+		Pose.BlendWithInitialPoses(TempPose, InControlRig, bDoKey, bDoMirror, 1.0, bDoAdditive);
+		InControlRig->Evaluate_AnyThread();
+		Pose.BlendWithInitialPoses(TempPose, InControlRig, bDoKey, bDoMirror, 1.0, bDoAdditive);
+
+	}
 }
 
 void UControlRigPoseAsset::SelectControls(UControlRig* InControlRig, bool bDoMirror, bool bClearSelection)
@@ -426,12 +440,9 @@ TArray<FRigControlCopy> UControlRigPoseAsset::GetCurrentPose(UControlRig* InCont
 	return TempPose.GetPoses();
 }
 
-void UControlRigPoseAsset::BlendWithInitialPoses(FControlRigControlPose& InitialPose, UControlRig* InControlRig, bool bDoKey, bool bDoMirror, float BlendValue)
+void UControlRigPoseAsset::BlendWithInitialPoses(FControlRigControlPose& InitialPose, UControlRig* InControlRig, bool bDoKey, bool bDoMirror, float BlendValue, bool bDoAdditive)
 {
-	if (BlendValue > 0.0f)
-	{
-		Pose.BlendWithInitialPoses(InitialPose, InControlRig, bDoKey, bDoMirror, BlendValue);
-	}
+	Pose.BlendWithInitialPoses(InitialPose, InControlRig, bDoKey, bDoMirror, BlendValue, bDoAdditive);
 }
 
 TArray<FName> UControlRigPoseAsset::GetControlNames() const
diff --git a/Engine/Plugins/Animation/ControlRig/Source/ControlRig/Public/Tools/ControlRigPose.h b/Engine/Plugins/Animation/ControlRig/Source/ControlRig/Public/Tools/ControlRigPose.h
--- a/Engine/Plugins/Animation/ControlRig/Source/ControlRig/Public/Tools/ControlRigPose.h
+++ b/Engine/Plugins/Animation/ControlRig/Source/ControlRig/Public/Tools/ControlRigPose.h
@@ -94,7 +94,7 @@ struct CONTROLRIG_API FControlRigControlPose
 	void SetControlMirrorTransform(bool bDoLocalSpace, UControlRig* ControlRig, const FName& Name, bool bIsMatched,
 		const FTransform& GlobalTrnnsform, const FTransform& LocalTransform, bool bNotify, const  FRigControlModifiedContext& Context, bool bSetupUndo);
 	void PastePoseInternal(UControlRig* ControlRig, bool bDoKey, bool bDoMirror, const TArray<FRigControlCopy>& ControlsToPaste);
-	void BlendWithInitialPoses(FControlRigControlPose& InitialPose, UControlRig* ControlRig, bool bDoKey, bool bDoMirror, float BlendValue);
+	void BlendWithInitialPoses(FControlRigControlPose& InitialPose, UControlRig* ControlRig, bool bDoKey, bool bDoMirror, float BlendValue, bool bDoAdditive = false);
 
 	bool ContainsName(const FName& Name) const;
 	void ReplaceControlName(const FName& Name, const FName& NewName);
@@ -129,7 +129,7 @@ public:
 	void SavePose(UControlRig* InControlRig, bool bUseAll);
 
 	UFUNCTION(BlueprintCallable, Category = "Pose")
-	void PastePose(UControlRig* InControlRig, bool bDoKey = false, bool bDoMirror = false);
+	void PastePose(UControlRig* InControlRig, bool bDoKey = false, bool bDoMirror = false, bool bDoAdditive = false);
 	
 	UFUNCTION(BlueprintCallable, Category = "Pose")
 	void SelectControls(UControlRig* InControlRig, bool bDoMirror = false, bool bClearSelection = true);
@@ -148,7 +148,7 @@ public:
 	UFUNCTION(BlueprintPure, Category = "Pose")
 	bool DoesMirrorMatch(UControlRig* ControlRig, const FName& ControlName) const;
 
-	void BlendWithInitialPoses(FControlRigControlPose& InitialPose, UControlRig* ControlRig, bool bDoKey, bool bdoMirror, float BlendValue);
+	void BlendWithInitialPoses(FControlRigControlPose& InitialPose, UControlRig* ControlRig, bool bDoKey, bool bdoMirror, float BlendValue, bool bDoAdditive = false);
 
 	UPROPERTY(BlueprintReadWrite, EditAnywhere, Category = "Pose")
 	FControlRigControlPose Pose;
diff --git a/Engine/Plugins/Animation/ControlRig/Source/ControlRigEditor/Private/EditMode/SControlRigControlViews.cpp b/Engine/Plugins/Animation/ControlRig/Source/ControlRigEditor/Private/EditMode/SControlRigControlViews.cpp
--- a/Engine/Plugins/Animation/ControlRig/Source/ControlRigEditor/Private/EditMode/SControlRigControlViews.cpp
+++ b/Engine/Plugins/Animation/ControlRig/Source/ControlRigEditor/Private/EditMode/SControlRigControlViews.cpp
@@ -187,6 +187,7 @@ void SControlRigAssetEditableTextBox::OnTextChanged(const FText& InLabel)
 
 bool SControlRigPoseView::bIsKey = false;
 bool SControlRigPoseView::bIsMirror = false;
+bool SControlRigPoseView::bIsAdditive = false;
 
 void SControlRigPoseView::Construct(const FArguments& InArgs)
 {
@@ -333,6 +334,21 @@ void SControlRigPoseView::Construct(const FArguments& InArgs)
 								.IsEnabled(this, &SControlRigPoseView::IsMirrorEnabled)
 							]
 							]
+							+ SHorizontalBox::Slot()
+							.AutoWidth()
+
+							.HAlign(HAlign_Center)
+
+							.Padding(2.5f)
+							[
+							SNew(SCheckBox)
+							.IsChecked(this, &SControlRigPoseView::IsAdditivePoseChecked)
+							.OnCheckStateChanged(this, &SControlRigPoseView::OnAdditivePoseChecked)
+							.Padding(1.0f)
+							[
+								SNew(STextBlock).Text(LOCTEXT("Additive", "Additive"))
+							]
+							]
 							]
 						+ SVerticalBox::Slot()
 							.AutoHeight()
@@ -344,16 +360,14 @@ void SControlRigPoseView::Construct(const FArguments& InArgs)
 								// Only allow spinning if we have a single value
 								.PreventThrottling(true)
 								.Value(this, &SControlRigPoseView::OnGetPoseBlendValue)
-								.ToolTipText(LOCTEXT("BlendTooltip", "Blend between current pose and pose asset. Use Ctrl drag for under and over shoot."))
-								.MinValue(0.0f)
-								.MaxValue(1.0f)
-								.MinSliderValue(0.0f)
-								.MaxSliderValue(1.0f)
+								.ToolTipText(LOCTEXT("BlendTooltip", "Blend between current pose and pose asset."))
+								.MinValue(-1.0f)
+								.MaxValue(2.0f)
+								.MinSliderValue(-1.0)
+								.MaxSliderValue(2.0f)
 								.SliderExponent(1)
 								.Delta(0.005f)
 								.MinDesiredWidth(100.0f)
-								.SupportDynamicSliderMinValue(true)
-								.SupportDynamicSliderMaxValue(true)
 								.OnValueChanged(this, &SControlRigPoseView::OnPoseBlendChanged)
 								.OnValueCommitted(this, &SControlRigPoseView::OnPoseBlendCommited)
 								.OnBeginSliderMovement(this,&SControlRigPoseView::OnBeginSliderMovement)
@@ -536,6 +550,27 @@ bool SControlRigPoseView::IsMirrorEnabled() const
 	return true;
 }
 
+ECheckBoxState SControlRigPoseView::IsAdditivePoseChecked() const
+{
+	if (bIsAdditive)
+	{
+		return ECheckBoxState::Checked;
+	}
+	return ECheckBoxState::Unchecked;
+}
+
+void SControlRigPoseView::OnAdditivePoseChecked(ECheckBoxState NewState)
+{
+	if (NewState == ECheckBoxState::Checked)
+	{
+		bIsAdditive = true;
+	}
+	else
+	{
+		bIsAdditive = false;
+	}
+	UpdateStatusBlocks();
+}
 
 FReply SControlRigPoseView::OnPastePose()
 {
@@ -546,7 +581,7 @@ FReply SControlRigPoseView::OnPastePose()
 		{
 			if (ControlRig)
 			{
-				PoseAsset->PastePose(ControlRig, bIsKey, bIsMirror);
+				PoseAsset->PastePose(ControlRig, bIsKey, bIsMirror,bIsAdditive);
 			}
 		}
 	}
@@ -598,8 +633,8 @@ void SControlRigPoseView::OnPoseBlendChanged(float ChangedVal)
 					PoseAsset->GetCurrentPose(ControlRig, TempPose);
 				}
 
-				PoseAsset->BlendWithInitialPoses(TempPose, ControlRig, false, bIsMirror, PoseBlendValue);
-				PoseAsset->BlendWithInitialPoses(TempPose, ControlRig, false, bIsMirror, PoseBlendValue);
+				PoseAsset->BlendWithInitialPoses(TempPose, ControlRig, false, bIsMirror, PoseBlendValue, bIsAdditive);
+				PoseAsset->BlendWithInitialPoses(TempPose, ControlRig, false, bIsMirror, PoseBlendValue, bIsAdditive);
 			}
 		}
 	}
@@ -635,8 +670,8 @@ void SControlRigPoseView::OnPoseBlendCommited(float ChangedVal, ETextCommit::Typ
 				if (ControlRig)
 				{
 					PoseBlendValue = ChangedVal;
-					PoseAsset->BlendWithInitialPoses(TempPose, ControlRig, bIsKey, bIsMirror, PoseBlendValue);
-					PoseAsset->BlendWithInitialPoses(TempPose, ControlRig, bIsKey, bIsMirror, PoseBlendValue);
+					PoseAsset->BlendWithInitialPoses(TempPose, ControlRig, bIsKey, bIsMirror, PoseBlendValue, bIsAdditive);
+					PoseAsset->BlendWithInitialPoses(TempPose, ControlRig, bIsKey, bIsMirror, PoseBlendValue, bIsAdditive);
 					bIsBlending = false;
 					PoseBlendValue = 0.0f;
 
diff --git a/Engine/Plugins/Animation/ControlRig/Source/ControlRigEditor/Private/EditMode/SControlRigControlViews.h b/Engine/Plugins/Animation/ControlRig/Source/ControlRigEditor/Private/EditMode/SControlRigControlViews.h
--- a/Engine/Plugins/Animation/ControlRig/Source/ControlRigEditor/Private/EditMode/SControlRigControlViews.h
+++ b/Engine/Plugins/Animation/ControlRig/Source/ControlRigEditor/Private/EditMode/SControlRigControlViews.h
@@ -53,6 +53,8 @@ private:
 	void OnKeyPoseChecked(ECheckBoxState NewState);
 	ECheckBoxState IsMirrorPoseChecked() const;
 	void OnMirrorPoseChecked(ECheckBoxState NewState);
+	ECheckBoxState IsAdditivePoseChecked() const;
+	void OnAdditivePoseChecked(ECheckBoxState NewState);
 	bool IsMirrorEnabled() const;
 	void OnPoseBlendChanged(float ChangedVal);
 	void OnPoseBlendCommited(float ChangedVal, ETextCommit::Type Type);
@@ -68,6 +70,7 @@ private:
 
 	static bool bIsKey;
 	static bool bIsMirror;
+	static bool bIsAdditive;
 	float PoseBlendValue;
 	bool bIsBlending;
 	bool bSliderStartedTransaction;
