Anim In Engine: Motion trail fixes
Last set of Motion trail tweaks based on animator feedback 1)  Turn off showing selected trails but instead will automatically pin the selection when the trails are turned on. 2) Add modifier keys (shift/ctrl/alt) for adding, toggling, removing trails from the motion button to avoid having to go into the menu itself.  3)  Changed evaluating setup to actually be if the trail tool is moving something verus just selection, also if directly manipulating will draw the whole trail as lines between the keys.  4) Added option to evaluate the whole range or just the view range when not showing the full range, defaults to only evaluate the view range which makes tool more useable for very large evaluation ranges. 4) Changed offset transform to offset every key fully(including rotation), so you can do things like move the offset down the length of a sword when parented to a control and the tip will move with rotation.  5) Misc fixes like add max number of pinnable and FOV fix for drawing keys

diff --git a/Engine/Plugins/Animation/ControlRig/Source/ControlRigEditor/Private/EditMode/ControlRigEditMode.cpp b/Engine/Plugins/Animation/ControlRig/Source/ControlRigEditor/Private/EditMode/ControlRigEditMode.cpp
--- a/Engine/Plugins/Animation/ControlRig/Source/ControlRigEditor/Private/EditMode/ControlRigEditMode.cpp
+++ b/Engine/Plugins/Animation/ControlRig/Source/ControlRigEditor/Private/EditMode/ControlRigEditMode.cpp
@@ -170,6 +170,8 @@ namespace ControlRigSelectionConstants
 	static const float BodyTraceDistance = 100000.0f;
 }
 
+bool FControlRigEditMode::bDoPostPoseUpdate = true;
+
 FControlRigEditMode::FControlRigEditMode()
 	: PendingFocus(FPendingWidgetFocus::MakeNoTextEdit())
 	, bIsChangingControlShapeTransform(false)
@@ -378,7 +380,6 @@ void FControlRigEditMode::SetUpDetailPanel() const
 		if (ModeTools.IsValid())
 		{
 			ModeTools->SetSequencer(WeakSequencer.Pin());
-			ModeTools->SetSettingsDetailsObject(GetMutableSettings());
 		}	
 	}
 }
@@ -6265,6 +6266,11 @@ void FControlRigEditMode::OnPoseInitialized()
 
 void FControlRigEditMode::PostPoseUpdate(const FEditorViewportClient* InViewportClient, const TArray<TWeakObjectPtr<UControlRig>>& InRigs) const
 {
+	if (bDoPostPoseUpdate == false)
+	{
+		return;
+	}
+
 	using namespace ControlRigEditMode::Shapes;
 
 	const bool bIsGameView = InViewportClient ? InViewportClient->IsInGameView() : false;
diff --git a/Engine/Plugins/Animation/ControlRig/Source/ControlRigEditor/Private/Editor/ControlRigViewportToolbarExtensions.cpp b/Engine/Plugins/Animation/ControlRig/Source/ControlRigEditor/Private/Editor/ControlRigViewportToolbarExtensions.cpp
--- a/Engine/Plugins/Animation/ControlRig/Source/ControlRigEditor/Private/Editor/ControlRigViewportToolbarExtensions.cpp
+++ b/Engine/Plugins/Animation/ControlRig/Source/ControlRigEditor/Private/Editor/ControlRigViewportToolbarExtensions.cpp
@@ -219,13 +219,13 @@ TSharedRef<SWidget> CreateFramesBeforeWidget(UMotionTrailToolOptions* Settings)
 									}
 								)
 								.Value_Lambda(
-									[Settings]() -> float
+									[Settings]() -> int32
 									{
 										return Settings->FramesBefore;
 									}
 								)
 								.OnValueChanged_Lambda(
-									[Settings](float InValue)
+									[Settings](int32 InValue)
 									{
 										Settings->FramesBefore = InValue;
 									}
@@ -274,13 +274,13 @@ TSharedRef<SWidget> CreateFramesAfterWidget(UMotionTrailToolOptions* Settings)
 									}
 								)
 								.Value_Lambda(
-									[Settings]() -> float
+									[Settings]() -> int32
 									{
 										return Settings->FramesAfter;
 									}
 								)
 								.OnValueChanged_Lambda(
-									[Settings](float InValue)
+									[Settings](int32 InValue)
 									{
 										Settings->FramesAfter = InValue;
 									}
@@ -364,6 +364,59 @@ TSharedRef<SWidget> CreateTrailStyleWidget(UMotionTrailToolOptions* Settings)
 }
 
 
+TSharedRef<SWidget> CreateMaxNumberPinnedWidget(UMotionTrailToolOptions* Settings)
+{
+	// clang-format off
+	return
+		SNew(SBox)
+		.HAlign(HAlign_Right)
+		[
+			SNew(SBox)
+				.Padding(FMargin(0.0f, 0.0f, 0.0f, 0.0f))
+				.WidthOverride(100.0f)
+				[
+					SNew(SBorder)
+						.BorderImage(FAppStyle::Get().GetBrush("Menu.WidgetBorder"))
+						.Padding(FMargin(1.0f))
+						[
+							SNew(SSpinBox<int32>)
+								.MinValue(5)
+								.MaxValue(100)
+								.MinDesiredWidth(50)
+								.ToolTipText_Lambda(
+									[Settings]() -> FText
+									{
+										return FText::AsNumber(
+											Settings->MaxNumberPinned
+										);
+									}
+								)
+								.Value_Lambda(
+									[Settings]() -> int32
+									{
+										return Settings->MaxNumberPinned;
+									}
+								)
+								.OnValueChanged_Lambda(
+									[Settings](int32 InValue)
+									{
+										Settings->MaxNumberPinned = InValue;
+									}
+								)
+						]
+				]
+		];
+
+	// clang-format on
+}
+
+FToolMenuEntry CreateMaxNumberPinned(UMotionTrailToolOptions* Settings)
+{
+	FToolMenuEntry Entry = FToolMenuEntry::InitWidget("MaxNumberPinned", CreateMaxNumberPinnedWidget(Settings), LOCTEXT("MaxNumberPinned", "Max Number Pinned"));
+	return Entry;
+}
+
+
 FToolMenuEntry CreateTrailStyle(UMotionTrailToolOptions* Settings)
 {
 	FToolMenuEntry Entry = FToolMenuEntry::InitWidget("TrailStyle", CreateTrailStyleWidget(Settings), LOCTEXT("TrailStyle", "Trail Style"));
@@ -379,6 +432,12 @@ void CreatePinnedSubMenu(UToolMenu* InSubMenu, UMotionTrailToolOptions* Settings
 	FUIAction PinSelectedAction(
 		FExecuteAction::CreateLambda([Settings]()
 			{
+				if (Settings->bShowTrails == false)
+				{
+					Settings->bShowTrails = true;
+					FPropertyChangedEvent ShowTrailEvent(UMotionTrailToolOptions::StaticClass()->FindPropertyByName(GET_MEMBER_NAME_CHECKED(UMotionTrailToolOptions, bShowTrails)));
+					Settings->PostEditChangeProperty(ShowTrailEvent);
+				}
 				Settings->PinSelection();
 			}
 		),
@@ -424,6 +483,12 @@ void CreatePinnedSubMenu(UToolMenu* InSubMenu, UMotionTrailToolOptions* Settings
 									.SceneComponent(Component)
 									.OnSocketChosen_Lambda([Settings, Component](FName InSocketName) mutable
 										{
+											if (Settings->bShowTrails == false)
+											{
+												Settings->bShowTrails = true;
+												FPropertyChangedEvent ShowTrailEvent(UMotionTrailToolOptions::StaticClass()->FindPropertyByName(GET_MEMBER_NAME_CHECKED(UMotionTrailToolOptions, bShowTrails)));
+												Settings->PostEditChangeProperty(ShowTrailEvent);
+											}
 											Settings->PinComponent(Component, InSocketName);
 										}
 									);
@@ -550,8 +615,8 @@ FToolMenuEntry CreateTrailColor(UMotionTrailToolOptions* Settings, FName Propert
 
 	return Entry;
 }
-
-TSharedRef<SWidget> CreateDoublePropertyWidget(UMotionTrailToolOptions* Settings, FName PropertyName)
+template<typename NumericType>
+TSharedRef<SWidget> CreatePropertyWidget(UObject* Settings, FName PropertyName)
 {
 	// clang-format off
 	return
@@ -566,7 +631,7 @@ TSharedRef<SWidget> CreateDoublePropertyWidget(UMotionTrailToolOptions* Settings
 						.BorderImage(FAppStyle::Get().GetBrush("Menu.WidgetBorder"))
 						.Padding(FMargin(1.0f))
 						[
-							SNew(SSpinBox<double>)
+							SNew(SSpinBox<NumericType>)
 								.MinValue(0)
 								.MaxValue(100)
 								.MinDesiredWidth(50)
@@ -574,25 +639,25 @@ TSharedRef<SWidget> CreateDoublePropertyWidget(UMotionTrailToolOptions* Settings
 									[Settings, PropertyName]() -> FText
 									{
 										FTrackInstancePropertyBindings Binding(PropertyName, PropertyName.ToString());
-										double Value = Binding.GetCurrentValue<double>(*Settings);
+										NumericType Value = Binding.GetCurrentValue<NumericType>(*Settings);
 										return FText::AsNumber(
 											Value
 										);
 									}
 								)
 								.Value_Lambda(
-									[Settings, PropertyName]() -> double
+									[Settings, PropertyName]() -> NumericType
 									{
 										FTrackInstancePropertyBindings Binding(PropertyName, PropertyName.ToString());
-										double Value = Binding.GetCurrentValue<double>(*Settings);
+										NumericType Value = Binding.GetCurrentValue<NumericType>(*Settings);
 										return Value;
 									}
 								)
 								.OnValueChanged_Lambda(
-									[Settings, PropertyName](double InValue)
+									[Settings, PropertyName](NumericType InValue)
 									{
 										FTrackInstancePropertyBindings Binding(PropertyName, PropertyName.ToString());
-										Binding.CallFunction<double>(*Settings, InValue);
+										Binding.CallFunction<NumericType>(*Settings, InValue);
 										FPropertyChangedEvent Event(Binding.GetProperty(*Settings));
 										Settings->PostEditChangeProperty(Event);
 									}
@@ -604,10 +669,11 @@ TSharedRef<SWidget> CreateDoublePropertyWidget(UMotionTrailToolOptions* Settings
 	// clang-format on
 }
 
-FToolMenuEntry CreateDoubleProperty(UMotionTrailToolOptions* Settings, FName PropertyName)
+template<typename NumericType>
+FToolMenuEntry CreateProperty(UObject* Settings, FName PropertyName)
 {
 	FText Text = FText::FromString(PropertyName.ToString());
-	FToolMenuEntry Entry = FToolMenuEntry::InitWidget(PropertyName, CreateDoublePropertyWidget(Settings, PropertyName), Text);
+	FToolMenuEntry Entry = FToolMenuEntry::InitWidget(PropertyName, CreatePropertyWidget<NumericType>(Settings, PropertyName), Text);
 
 	return Entry;
 }
@@ -678,16 +744,16 @@ void CreateAdvancedSubMenu(UToolMenu* InSubMenu, UMotionTrailToolOptions* Settin
 	}
 	{
 		const FName DoubleProperty("KeySize");
-		TrailSettings.AddEntry(CreateDoubleProperty(Settings, DoubleProperty));
+		TrailSettings.AddEntry(CreateProperty<double>(Settings, DoubleProperty));
 	}
 
 	{
 		const FName DoubleProperty("MarkSize");
-		TrailSettings.AddEntry(CreateDoubleProperty(Settings, DoubleProperty));
+		TrailSettings.AddEntry(CreateProperty<double>(Settings, DoubleProperty));
 	}
 	{
 		const FName DoubleProperty("TrailThickness");
-		TrailSettings.AddEntry(CreateDoubleProperty(Settings, DoubleProperty));
+		TrailSettings.AddEntry(CreateProperty<double>(Settings, DoubleProperty));
 	}
 
 	FToolMenuSection& ColorSettings = InSubMenu->AddSection("ColorSettings", LOCTEXT("ColorSettings", "Color Settings"));
@@ -720,6 +786,10 @@ void CreateAdvancedSubMenu(UToolMenu* InSubMenu, UMotionTrailToolOptions* Settin
 		ColorSettings.AddEntry(CreateTrailColor(Settings, ColorProperty));
 	}
 
+	FToolMenuSection& PinSettings = InSubMenu->AddSection("PinSettings", LOCTEXT("PinSettings", "Pin Settings"));
+	{
+		PinSettings.AddEntry(CreateMaxNumberPinned(Settings));
+	}
 
 }
 
@@ -796,7 +866,7 @@ void CreateMotionTrailMenu(UToolMenu* InMenu)
 	}
 	PathModeSection.AddEntry(CreateFramesBefore(Settings));
 	PathModeSection.AddEntry(CreateFramesAfter(Settings));
-
+	
 	CreatePinnedSubMenu(InMenu, Settings);
 	/* in case design wants to go back to a subment
 	FToolMenuSection& PinnedSection = InMenu->AddSection(TEXT("Pinned"), LOCTEXT("Pinned", "Pinned"));
@@ -842,7 +912,7 @@ void CreateMotionTrailMenu(UToolMenu* InMenu)
 	}
 	PathOptionsMenu.AddEntry(CreateTrailStyle(Settings));
 
-	PathOptionsMenu.AddSubMenu(TEXT("Advanced"), LOCTEXT("Advanced", "Advanced"), LOCTEXT("Advanced_tooltip", "Advanced display options"),
+	PathOptionsMenu.AddSubMenu(TEXT("Advanced"), LOCTEXT("Advanced", "Advanced"), LOCTEXT("Advanced_tooltip", "Advanced options"),
 		FNewToolMenuDelegate::CreateLambda([Settings](UToolMenu* InSubMenu)
 			{
 				CreateAdvancedSubMenu(InSubMenu, Settings);
@@ -1015,9 +1085,51 @@ void UE::ControlRig::PopulateControlRigViewportToolbarTransformSubmenu(const FNa
 			CheckboxMenuAction.ExecuteAction = FToolMenuExecuteAction::CreateLambda(
 				[Settings](const FToolMenuContext& Context) -> void
 				{
+					//if trails are already shown and we do a modifier we leave them on but do a specific action
+					const bool bAltDown = FSlateApplication::Get().GetModifierKeys().IsAltDown();
+					const bool bControlDown = FSlateApplication::Get().GetModifierKeys().IsControlDown();
+					const bool bShiftDown = FSlateApplication::Get().GetModifierKeys().IsShiftDown();
+					const bool bModifierDown = bAltDown || bControlDown || bShiftDown;
+					auto HandleModifier = [Settings, bAltDown, bControlDown, bShiftDown]() {
+
+						if (bAltDown)
+						{
+							Settings->UnPinSelection();
+							return;
+						}
+						if (bControlDown)
+						{
+							Settings->DeleteAllPinned();
+							Settings->PinSelection();
+							return;
+						}
+						if (bShiftDown)
+						{
+							Settings->PinSelection();
+							return;
+						}
+					};
+					//if we are on and we have a modifier we do the operations then bail out leaving them on
+					if (Settings->bShowTrails && bModifierDown)
+					{
+						HandleModifier();
+						return;
+					}
 					Settings->bShowTrails = !Settings->bShowTrails;
 					FPropertyChangedEvent ShowTrailEvent(UMotionTrailToolOptions::StaticClass()->FindPropertyByName(GET_MEMBER_NAME_CHECKED(UMotionTrailToolOptions, bShowTrails)));
 					Settings->PostEditChangeProperty(ShowTrailEvent);
+					//if we are now on and we 1) Pin if no pinned trails  or 2) handle the modifier if there are pinned
+					if (Settings->bShowTrails)
+					{
+						if (Settings->GetNumPinned() == 0)
+						{
+							Settings->PinSelection();
+						}
+						else if (bModifierDown)
+						{
+							HandleModifier();
+						}
+					}
 				}
 			);
 			CheckboxMenuAction.GetActionCheckState = FToolMenuGetActionCheckState::CreateLambda(
@@ -1039,7 +1151,7 @@ void UE::ControlRig::PopulateControlRigViewportToolbarTransformSubmenu(const FNa
 			"MotionPaths",
 			LOCTEXT("MotionPathsLabel", "Motion Paths"),
 			// TODO: Update this and other labels/tooltips in this file.
-			LOCTEXT("MotionPathsTooltip", "Check to enable motion paths. Submenu contains settings for motion paths."),
+			LOCTEXT("MotionPathsTooltip", "Check to enable motion paths. Submenu contains settings for motion paths.\nHotkeys:\nUse SHIFT to add selected items to pin list\nUse CTRL to reset pin list to just the selected item\nUse ALT to remove selected item from pin list"),
 			MakeMenuDelegate,
 			CheckboxMenuAction,
 			EUserInterfaceActionType::ToggleButton,
@@ -1088,6 +1200,50 @@ void UE::ControlRig::PopulateControlRigViewportToolbarTransformSubmenu(const FNa
 	}
 }
 
+void CreateAxisOnSelectionMenu(UToolMenu* AnimationShowFlagsSubmenu, FToolMenuSection& UnnamedSection, UControlRigEditModeSettings* Settings)
+{
+	UUnrealEdViewportToolbarContext* Context = AnimationShowFlagsSubmenu->FindContext<UUnrealEdViewportToolbarContext>();
+
+	FNewToolMenuDelegate AxisOnSelectionMenuDelegate = FNewToolMenuDelegate::CreateLambda(
+		[Settings](UToolMenu* Submenu)
+		{
+			FToolMenuSection& AxisOnSelection =
+				Submenu->FindOrAddSection("AxisOnSelection", LOCTEXT("AxisOnSelectionLabel", "Axis On Selection"));
+
+			const FName DoubleProperty("AxisScale");
+			FToolMenuEntry NumericEntry = CreateProperty<float>(Settings, DoubleProperty);
+
+			AxisOnSelection.AddEntry(NumericEntry);
+		}
+	);
+
+	FToolMenuEntry Entry = FToolMenuEntry::InitSubMenu(
+		"AxisOnSelection",
+		LOCTEXT("AxisOnSelectionLabel", "Axis On Selection"),
+		LOCTEXT("AxisOnSelectionTooltip", "Should we show axes for the selected elements"),
+		AxisOnSelectionMenuDelegate,
+		FToolUIAction(
+			FToolMenuExecuteAction::CreateLambda([Settings,Context](const FToolMenuContext& InContext)
+			{
+				Settings->bDisplayAxesOnSelection = !Settings->bDisplayAxesOnSelection;
+				UControlRigEditModeSettings::OnSettingsChange.Broadcast(Settings);
+				if (Context)
+				{
+					Context->RefreshViewport();
+				}
+			}),
+			FToolMenuGetActionCheckState::CreateLambda([Settings](const FToolMenuContext& InContext)
+			{
+				return Settings->bDisplayAxesOnSelection ? ECheckBoxState::Checked
+								: ECheckBoxState::Unchecked;
+			})
+		),
+		EUserInterfaceActionType::ToggleButton
+	);
+
+	UnnamedSection.AddEntry(Entry);
+	
+}
 
 void UE::ControlRig::PopulateControlRigViewportToolbarShowSubmenu(const FName InMenuName)
 {
@@ -1109,36 +1265,7 @@ void UE::ControlRig::PopulateControlRigViewportToolbarShowSubmenu(const FName In
 				if (UControlRigEditModeSettings* const Settings = GetMutableDefault<UControlRigEditModeSettings>())
 				{
 					{
-						FUIAction Action;
-						Action.ExecuteAction = FExecuteAction::CreateLambda(
-							[Settings, Context]() -> void
-							{
-								Settings->bDisplayAxesOnSelection = !Settings->bDisplayAxesOnSelection;
-								UControlRigEditModeSettings::OnSettingsChange.Broadcast(Settings);
-								if (Context)
-								{
-									Context->RefreshViewport();
-								}
-							}
-						);
-						Action.GetActionCheckState = FGetActionCheckState::CreateLambda(
-							[Settings]() -> ECheckBoxState
-							{
-								return Settings->bDisplayAxesOnSelection ? ECheckBoxState::Checked
-									: ECheckBoxState::Unchecked;
-							}
-						);
-
-						FToolMenuEntry Entry = FToolMenuEntry::InitMenuEntry(
-							"AxisOnSelection",
-							LOCTEXT("AxisOnSelectionLabel", "Axis On Selection"),
-							LOCTEXT("AxisOnSelectionTooltip", "Should we show axes for the selected elements"),
-							FSlateIcon(),
-							Action,
-							EUserInterfaceActionType::ToggleButton
-						);
-						UnnamedSection.AddEntry(Entry);
-
+						CreateAxisOnSelectionMenu(AnimationShowFlagsSubmenu, UnnamedSection, Settings);
 					}
 					{
 						FUIAction Action;
diff --git a/Engine/Plugins/Animation/ControlRig/Source/ControlRigEditor/Public/EditMode/ControlRigEditMode.h b/Engine/Plugins/Animation/ControlRig/Source/ControlRigEditor/Public/EditMode/ControlRigEditMode.h
--- a/Engine/Plugins/Animation/ControlRig/Source/ControlRigEditor/Public/EditMode/ControlRigEditMode.h
+++ b/Engine/Plugins/Animation/ControlRig/Source/ControlRigEditor/Public/EditMode/ControlRigEditMode.h
@@ -265,6 +265,8 @@ public:
 		return FRigElementTypeHelper::ToMask(ERigElementType::Control);
 	}
 
+
+
 protected:
 
 	EControlRigInteractionTransformSpace GetTransformSpace() const;
@@ -723,6 +725,9 @@ private:
 
 	static void EvaluateRig(UControlRig* InControlRig);
 	
+	//to disable post pose update, needed for offfline evlauations
+	static bool bDoPostPoseUpdate;
+
 private:
 
 	// get default/mutable settings
@@ -754,4 +759,23 @@ private:
 	friend class FControlRigEditModeGenericDetails;
 	friend class UControlRigEditModeDelegateHelper;
 	friend class SControlRigEditModeTools;
-};
\ No newline at end of file
+
+public:
+
+	class FTurnOffPosePoseUpdate
+	{
+	public:
+		FTurnOffPosePoseUpdate()
+		{
+			bLastVal = FControlRigEditMode::bDoPostPoseUpdate;
+			FControlRigEditMode::bDoPostPoseUpdate = false;
+		}
+		~FTurnOffPosePoseUpdate()
+		{
+			FControlRigEditMode::bDoPostPoseUpdate = bLastVal;
+		}
+	private:
+		bool bLastVal;
+	};
+};
+
diff --git a/Engine/Plugins/Editor/SequencerAnimTools/Source/SequencerAnimTools/Private/Sequencer/MotionTrailMovieSceneKey.cpp b/Engine/Plugins/Editor/SequencerAnimTools/Source/SequencerAnimTools/Private/Sequencer/MotionTrailMovieSceneKey.cpp
--- a/Engine/Plugins/Editor/SequencerAnimTools/Source/SequencerAnimTools/Private/Sequencer/MotionTrailMovieSceneKey.cpp
+++ b/Engine/Plugins/Editor/SequencerAnimTools/Source/SequencerAnimTools/Private/Sequencer/MotionTrailMovieSceneKey.cpp
@@ -14,7 +14,7 @@
 #include "MovieSceneSection.h"
 #include "PrimitiveDrawingUtils.h"
 #include "SceneView.h"
-
+#include "Misc/Optional.h"
 
 namespace UE
 {
@@ -38,7 +38,7 @@ struct HMotionTrailMovieSceneKeyProxy : public HBaseTrailProxy
 };
 IMPLEMENT_HIT_PROXY(HMotionTrailMovieSceneKeyProxy, HBaseTrailProxy);
 
-void FMotionTraiMovieScenelKeyTool::Render(const FGuid& Guid, const FSceneView* View, FPrimitiveDrawInterface* PDI)
+void FMotionTraiMovieScenelKeyTool::Render(const FGuid& Guid, const FSceneView* View, FPrimitiveDrawInterface* PDI, bool bTrailsIsEvaluating)
 {
 	const bool bIsVisible = (UMotionTrailToolOptions::GetTrailOptions()->bShowKeys);
 
@@ -47,16 +47,42 @@ void FMotionTraiMovieScenelKeyTool::Render(const FGuid& Guid, const FSceneView*
 		ClearSelection();
 		return;
 	}
-	const FTransform OffsetTransform = OwningTrail->GetOffsetTransform();
+	const FTransform OffsetTransform = FTransform::Identity;//OwningTrail->GetOffsetTransform();
 	const FTransform ParentSpaceTransform = OwningTrail->GetParentSpaceTransform();
 	const bool bHitTesting = PDI && PDI->IsHitTesting();
-	const float KeySize = UMotionTrailToolOptions::GetTrailOptions()->KeySize;
+	float KeySize = UMotionTrailToolOptions::GetTrailOptions()->KeySize;
+	TOptional<FVector> LastPoint;
+	const float TrailThickness = UMotionTrailToolOptions::GetTrailOptions()->TrailThickness;
+	const FLinearColor TrailColor = UMotionTrailToolOptions::GetTrailOptions()->DefaultColor;
+	//hack for PDI->DrawPoint not scaling perspective keys see FViewElementPDI::DrawPoint
+	const bool bIsPerspective = (View->ViewMatrices.GetProjectionMatrix().M[3][3] < 1.0f) ? true : false;
+	if (bIsPerspective)
+	{
+		const float ZoomFactor = FMath::Min<float>(View->ViewMatrices.GetProjectionMatrix().M[0][0], View->ViewMatrices.GetProjectionMatrix().M[1][1]);
+		KeySize = KeySize / ZoomFactor;
+	}
 	for (const TPair<FFrameNumber, TUniquePtr<FTrailKeyInfo>>& FrameKeyPair : Keys.Keys)
 	{
 		if (CachedViewFrameRange.Contains(FrameKeyPair.Value->FrameNumber))
 		{
+			
 			FLinearColor KeyColor = CachedSelection.Contains(FrameKeyPair.Value.Get()) ? UMotionTrailToolOptions::GetTrailOptions()->SelectedKeyColor :  UMotionTrailToolOptions::GetTrailOptions()->KeyColor;
 			const FTransform Transform = FrameKeyPair.Value->Transform * OffsetTransform * ParentSpaceTransform;
+			if (bTrailsIsEvaluating)
+			{
+				if (LastPoint.IsSet())
+				{
+					const FVector CurPoint = Transform.GetLocation();
+					PDI->DrawLine(LastPoint.GetValue(), CurPoint, TrailColor, SDPG_Foreground, TrailThickness);
+					LastPoint = CurPoint;
+				}
+				else
+				{
+					LastPoint = Transform.GetLocation();
+				}
+			}
+
+
 			if (bHitTesting)
 			{
 				PDI->SetHitProxy(new HMotionTrailMovieSceneKeyProxy(Guid, Transform, FrameKeyPair.Value.Get()));
@@ -75,43 +101,7 @@ void FMotionTraiMovieScenelKeyTool::Render(const FGuid& Guid, const FSceneView*
 void FMotionTraiMovieScenelKeyTool::DrawHUD(const FSceneView* View, FCanvas* Canvas)
 {
 	return;
-	/* no longer show frame number 
-	const bool bIsVisible = (UMotionTrailToolOptions::GetTrailOptions()->bShowKeys && UMotionTrailToolOptions::GetTrailOptions()->bShowFrameNumber);
-	if (!bIsVisible)
-	{
-		return;
-	}
-
-	FFrameRate TickResolution = OwningTrail->GetSequencer()->GetFocusedTickResolution();
-	FFrameRate DisplayRate = OwningTrail->GetSequencer()->GetFocusedDisplayRate();
-	const FTransform OffsetTransform = OwningTrail->GetOffsetTransform();
-	const FTransform ParentSpaceTransform = OwningTrail->GetParentSpaceTransform();
-
-	const float KeySize = UMotionTrailToolOptions::GetTrailOptions()->KeySize;
-	for (const TPair<FFrameNumber, TUniquePtr<FTrailKeyInfo>>& FrameKeyPair : Keys.Keys)
-	{
-		if (CachedViewFrameRange.Contains(FrameKeyPair.Value->FrameNumber))
-		{
-			FLinearColor KeyColor = CachedSelection.Contains(FrameKeyPair.Value.Get()) ? FLinearColor(1.0, 1.0, 0.0f) : UMotionTrailToolOptions::GetTrailOptions()->KeyColor;
-
-
-			FVector2D PixelLocation;
-			const FTransform Transform = FrameKeyPair.Value->Transform * OffsetTransform * ParentSpaceTransform;
-
-			if (View->WorldToPixel(Transform.GetLocation(), PixelLocation))
-			{
-				PixelLocation /= Canvas->GetDPIScale();
-
-				int32 FrameNumber = FFrameRate::TransformTime(FFrameTime(FrameKeyPair.Value->FrameNumber), TickResolution, DisplayRate).RoundToFrame().Value;
-				FCanvasTextItem TextItem(PixelLocation, FText::FromString(FString::FromInt(FrameNumber)), GEngine->GetMediumFont(), KeyColor);
-				TextItem.Scale = FVector2D::UnitVector;
-				//TextItem.EnableShadow(FLinearColor::Black);
-				TextItem.Draw(Canvas);
-
-			}
-		}
-	}
-	*/
+	
 }
 
 bool FMotionTraiMovieScenelKeyTool::IsSelected(FVector& OutVectorPosition) const
@@ -129,7 +119,8 @@ void FMotionTraiMovieScenelKeyTool::GetSelectedKeyPositions(TArray<FVector>& Out
 {
 	if (CachedSelection.Num() > 0)
 	{
-		const FTransform OffsetTransform = OwningTrail->GetOffsetTransform();
+		//leaving this around in case we also want to add an extra additive offset to move the whole trail without worrying about rotation effects
+		const FTransform OffsetTransform = FTransform::Identity;//OwningTrail->GetOffsetTransform();
 		const FTransform ParentSpaceTransform = OwningTrail->GetParentSpaceTransform();
 		for (FTrailKeyInfo* KeyInfo : CachedSelection)
 		{
@@ -194,7 +185,7 @@ void FMotionTraiMovieScenelKeyTool::UpdateSelectedKeysTransform() const
 {
 	if (CachedSelection.Num() > 0)
 	{
-		const FTransform OffsetTransform = OwningTrail->GetOffsetTransform();
+		const FTransform OffsetTransform = FTransform::Identity;//OwningTrail->GetOffsetTransform();
 		const FTransform ParentSpaceTransform = OwningTrail->GetParentSpaceTransform();
 		FVector NewGizmoLocation = FVector::ZeroVector;
 		for (FTrailKeyInfo* KeyInfo : CachedSelection)
diff --git a/Engine/Plugins/Editor/SequencerAnimTools/Source/SequencerAnimTools/Private/Sequencer/MotionTrailMovieSceneKey.h b/Engine/Plugins/Editor/SequencerAnimTools/Source/SequencerAnimTools/Private/Sequencer/MotionTrailMovieSceneKey.h
--- a/Engine/Plugins/Editor/SequencerAnimTools/Source/SequencerAnimTools/Private/Sequencer/MotionTrailMovieSceneKey.h
+++ b/Engine/Plugins/Editor/SequencerAnimTools/Source/SequencerAnimTools/Private/Sequencer/MotionTrailMovieSceneKey.h
@@ -215,7 +215,7 @@ public:
 	{}
 	~FMotionTraiMovieScenelKeyTool() {};
 
-	void Render(const FGuid& Guid, const FSceneView* View, FPrimitiveDrawInterface* PDI);
+	void Render(const FGuid& Guid, const FSceneView* View, FPrimitiveDrawInterface* PDI, bool bTrailIsEvaluating);
 	void DrawHUD(const FSceneView* View, FCanvas* Canvas);
 
 	bool HandleClick(const FGuid& Guid, FEditorViewportClient* InViewportClient, HHitProxy* HitProxy, FInputClick Click);
diff --git a/Engine/Plugins/Editor/SequencerAnimTools/Source/SequencerAnimTools/Private/Sequencer/MovieSceneTransformTrail.cpp b/Engine/Plugins/Editor/SequencerAnimTools/Source/SequencerAnimTools/Private/Sequencer/MovieSceneTransformTrail.cpp
--- a/Engine/Plugins/Editor/SequencerAnimTools/Source/SequencerAnimTools/Private/Sequencer/MovieSceneTransformTrail.cpp
+++ b/Engine/Plugins/Editor/SequencerAnimTools/Source/SequencerAnimTools/Private/Sequencer/MovieSceneTransformTrail.cpp
@@ -23,6 +23,7 @@
 #include "Framework/Application/SlateApplication.h"
 #include "PrimitiveDrawInterface.h"
 #include "Components/SkeletalMeshComponent.h"
+#include "EditMode/ControlRigEditMode.h"
 
 #define LOCTEXT_NAMESPACE "SequencerAnimTools"
 
@@ -132,21 +133,42 @@ void FMovieSceneTransformTrail::HasStartedEvaluating()
 	PreviousCachedDrawData = CachedDrawData;
 }
 
-void FMovieSceneTransformTrail::UpdateFinished(const TRange<FFrameNumber>& Range, const TArray<int32>& IndicesToCalcluate)
+void FMovieSceneTransformTrail::UpdateFinished(const TRange<FFrameNumber>& Range, const TArray<int32>& IndicesToCalcluate, bool bDoneCalculating)
 {
-	FTrail::UpdateFinished(Range, IndicesToCalcluate);
+	FTrail::UpdateFinished(Range, IndicesToCalcluate, bDoneCalculating);
+	if (HasOffsetTransform())
+	{
+		TSet<int32> CalculatedIndices;
+		for (int32 Count = 0; Count < IndicesToCalcluate.Num(); ++Count)
+		{
+			int32 Index = IndicesToCalcluate[Count];
+			if (CalculatedIndices.Contains(Index) == false)
+			{
+				CalculatedIndices.Add(Index);
+				ArrayOfTransforms->Transforms[Index] = OffsetTransform * ArrayOfTransforms->Transforms[Index];
+			}
+		}
+	}
 	//if we are in a parent space we need to put transforms in that space
 	if (ParentSpace.bIsValid && IndicesToCalcluate.Num() > 0)
 	{
+		TSet<int32> CalculatedIndices;
 		for (int32 Count = 0; Count < IndicesToCalcluate.Num(); ++Count)
 		{
 			int32 Index = IndicesToCalcluate[Count];
-			ArrayOfTransforms->Transforms[Index] = ArrayOfTransforms->Transforms[Index].GetRelativeTransform(ParentSpace.ParentSpace.WorldTransforms->Transforms[Index]);
-			ParentArrayOfTransforms->Transforms[Index] = ParentArrayOfTransforms->Transforms[Index].GetRelativeTransform(ParentSpace.ParentSpace.WorldTransforms->Transforms[Index]);
+			if (CalculatedIndices.Contains(Index) == false)
+			{
+				CalculatedIndices.Add(Index);
+				ArrayOfTransforms->Transforms[Index] = ArrayOfTransforms->Transforms[Index].GetRelativeTransform(ParentSpace.ParentSpace.WorldTransforms->Transforms[Index]);
+				ParentArrayOfTransforms->Transforms[Index] = ParentArrayOfTransforms->Transforms[Index].GetRelativeTransform(ParentSpace.ParentSpace.WorldTransforms->Transforms[Index]);
+			}
 		}
 
 	}
-
+	if (bDoneCalculating)
+	{
+		bIsTracking = false;
+	}
 	if (bIsTracking == false)
 	{
 		KeyTool->DirtyKeyTransforms();
@@ -512,8 +534,10 @@ void FMovieSceneTransformTrail::ReadyToDrawTrail(FColorState& ColorState, const
 {
 	ColorState.ReadyForTrail(bIsPinned, GetDrawInfo()->GetStyle());
 	bool bCalculateColor = (ColorState.GetStyle() == EMotionTrailTrailStyle::Time);
-		
-	if (InCurrentFramesInfo->bViewRangeIsEvalRange == false || (CachedDrawData.PointsToDraw.Num() != CachedDrawData.Color.Num()) || (CachedDrawData.PointsToDraw.Num() != InCurrentFramesInfo->CurrentFrames.Num()))
+	GetParentSpaceTransform();
+
+	if (InCurrentFramesInfo->bViewRangeIsEvalRange == false || (CachedDrawData.PointsToDraw.Num() != CachedDrawData.Color.Num()) || (CachedDrawData.PointsToDraw.Num() != InCurrentFramesInfo->CurrentFrames.Num())
+		|| (CachedDrawData.Color.Num() != InCurrentFramesInfo->CurrentFrames.Num()))
 	{
 		GetTrajectoryPointsForDisplay(*InCurrentFramesInfo, bIsEvaluating, CachedDrawData.PointsToDraw,CachedDrawData.Frames);
 		bCalculateColor = true;
@@ -551,39 +575,42 @@ TOptional<FLinearColor> FMovieSceneTransformTrail::GetOverrideColor() const
 	return Color;
 }
 
-void FMovieSceneTransformTrail::InternalDrawTrail(const FGuid& Guid, const FSceneView* View, FPrimitiveDrawInterface* PDI, const FDrawCacheData& InDrawData, bool bHitTesting) const
+void FMovieSceneTransformTrail::InternalDrawTrail(const FGuid& Guid, const FSceneView* View, FPrimitiveDrawInterface* PDI, const FDrawCacheData& InDrawData, bool bHitTesting, bool bTrailIsEvaluating) const
 {
-	const TOptional<FLinearColor> OverrideColor = GetOverrideColor();
-	const float TrailThickness = UMotionTrailToolOptions::GetTrailOptions()->TrailThickness;
-	FVector LastPoint = InDrawData.PointsToDraw[0];
-
-	for (int32 Idx = 1; Idx < InDrawData.PointsToDraw.Num(); Idx++)
+	if (bTrailIsEvaluating == false) // if evaluating we now just draw betweeen the keys
 	{
-		const FFrameNumber CurFrame = InDrawData.Frames[Idx - 1];
-		if (bHitTesting)
-		{
-			PDI->SetHitProxy(new HNewMotionTrailProxy(Guid, LastPoint, CurFrame));
-		}
-		const FVector CurPoint = InDrawData.PointsToDraw[Idx];
-		const FLinearColor CurColor = OverrideColor.IsSet() == false ? InDrawData.Color[Idx - 1] : OverrideColor.GetValue();
-		PDI->DrawLine(LastPoint, CurPoint, CurColor, SDPG_Foreground, TrailThickness);
-		LastPoint = CurPoint;
-		if (bHitTesting)
+		const TOptional<FLinearColor> OverrideColor = GetOverrideColor();
+		const float TrailThickness = UMotionTrailToolOptions::GetTrailOptions()->TrailThickness;
+		FVector LastPoint = InDrawData.PointsToDraw[0];
+
+		for (int32 Idx = 1; Idx < InDrawData.PointsToDraw.Num(); Idx++)
 		{
-			PDI->SetHitProxy(nullptr);
+			const FFrameNumber CurFrame = InDrawData.Frames[Idx - 1];
+			if (bHitTesting)
+			{
+				PDI->SetHitProxy(new HNewMotionTrailProxy(Guid, LastPoint, CurFrame));
+			}
+			const FVector CurPoint = InDrawData.PointsToDraw[Idx];
+			const FLinearColor CurColor = OverrideColor.IsSet() == false ? InDrawData.Color[Idx - 1] : OverrideColor.GetValue();
+			PDI->DrawLine(LastPoint, CurPoint, CurColor, SDPG_Foreground, TrailThickness);
+			LastPoint = CurPoint;
+			if (bHitTesting)
+			{
+				PDI->SetHitProxy(nullptr);
+			}
 		}
 	}
 }
 
-void FMovieSceneTransformTrail::Render(const FGuid& Guid, const FSceneView* View, FPrimitiveDrawInterface* PDI)
+void FMovieSceneTransformTrail::Render(const FGuid& Guid, const FSceneView* View, FPrimitiveDrawInterface* PDI, bool bTrailIsEvaluating)
 {
 	if (CachedDrawData.PointsToDraw.Num() > 1)
 	{
 		const bool bHitTesting = PDI && PDI->IsHitTesting();
-		InternalDrawTrail(Guid, View, PDI, CachedDrawData, bHitTesting);
+		InternalDrawTrail(Guid, View, PDI, CachedDrawData, bHitTesting, bTrailIsEvaluating);
 	}
 	//render key tool
-	KeyTool->Render(Guid, View, PDI);
+	KeyTool->Render(Guid, View, PDI, bTrailIsEvaluating);
 }
 
 void FMovieSceneTransformTrail::RenderEvaluating(const FGuid& Guid, const FSceneView* View, FPrimitiveDrawInterface* PDI)
@@ -591,7 +618,7 @@ void FMovieSceneTransformTrail::RenderEvaluating(const FGuid& Guid, const FScene
 	if (PreviousCachedDrawData.PointsToDraw.Num() > 0)
 	{
 		const bool bHitTesting = false;
-		InternalDrawTrail(Guid, View, PDI, PreviousCachedDrawData, bHitTesting);
+		InternalDrawTrail(Guid, View, PDI, PreviousCachedDrawData, bHitTesting,false); //if rendering cached we say not evaluating
 	}
 }
 
@@ -722,6 +749,7 @@ bool FMovieSceneTransformTrail::IsTrailSelected() const
 
 void FMovieSceneTransformTrail::ClearOffsetTransform()
 {
+	FTrail::ClearOffsetTransform();
 	OffsetTransform = FTransform::Identity;
 	if (bIsOffsetMode)
 	{
@@ -738,33 +766,19 @@ void FMovieSceneTransformTrail::SetOffsetMode()
 		bIsOffsetMode = true;
 		SelectNone();
 		bIsSelected = true;
-		//set selected position aka gizmo location  to be the average of the points drawn
-		if (CachedDrawData.PointsToDraw.Num() > 0)
-		{
-			FVector Min(DBL_MAX, DBL_MAX, DBL_MAX);
-			FVector Max(DBL_MIN, DBL_MIN, DBL_MIN);
-			for (const FVector& Vector : CachedDrawData.PointsToDraw)
-			{
-				for (int32 Index = 0; Index <= 2; ++Index)
-				{
-					if (Vector[Index] < Min[Index])
-					{
-						Min[Index] = Vector[Index];
-					}
-					if (Vector[Index] > Max[Index])
-					{
-						Max[Index] = Vector[Index];
-					}
-				}
-			}
-			for (int32 Index = 0; Index <= 2; ++Index)
-			{
-				SelectedPos[Index] = (Max[Index] + Min[Index]) * 0.5;
-			}
-		}
+		//set selected position aka gizmo location  to be location of the current object
+		FFrameTime StartTime = GetSequencer()->GetLocalTime().Time;
+		const FTransform Transform = ArrayOfTransforms->Interp(StartTime.GetFrame(), CurrentFramesInfo->TransformIndices, CurrentFramesInfo->CurrentFrames);
+		SelectedPos = Transform.GetLocation();
 	}
 
 }
+
+bool FMovieSceneTransformTrail::IsTracking() const 
+{ 
+	return bIsTracking; 
+}
+
 bool FMovieSceneTransformTrail::StartTracking()
 {
 	bIsTracking = true;
@@ -791,7 +805,7 @@ bool FMovieSceneTransformTrail::ApplyDelta(const FVector& Pos, const FRotator& R
 				UMotionTrailToolOptions::GetTrailOptions()->SetHasOffset(Index,true);
 			}
 			SelectedPos += Pos;
-			ClearCachedData();
+			ForceEvaluateNextTick();
 			return true;
 		}
 	}
@@ -800,7 +814,7 @@ bool FMovieSceneTransformTrail::ApplyDelta(const FVector& Pos, const FRotator& R
 
 bool FMovieSceneTransformTrail::EndTracking()
 {
-	bIsTracking = false;
+	//bIsTracking = false; We don't do this since we need to stop tracking once evaluating is done
 	return false;
 }
 
@@ -967,6 +981,7 @@ bool FMovieSceneComponentTransformTrail::ApplyDelta(const FVector& Pos, const FR
 		TArray<FTransform> Transforms;
 		TArray<FTransform> ParentTransforms;
 		TArray<FFrameNumber> Times;
+		FTransform InverseOffsetTransform = OffsetTransform.Inverse();
 		if (IsTrailSelected())
 		{
 			TArray<FTrailKeyInfo*> Keys;
@@ -976,6 +991,7 @@ bool FMovieSceneComponentTransformTrail::ApplyDelta(const FVector& Pos, const FR
 				FTransform NewTransform(KeyInfo->Transform);
 				NewTransform.SetLocation(NewTransform.GetLocation() + Pos);
 				KeyInfo->Transform = NewTransform;
+				NewTransform = OffsetTransform.GetRelativeTransformReverse(NewTransform);
 				Transforms.Add(NewTransform);
 				ParentTransforms.Add(KeyInfo->ParentTransform);
 				Times.Add(KeyInfo->FrameNumber);
@@ -989,6 +1005,8 @@ bool FMovieSceneComponentTransformTrail::ApplyDelta(const FVector& Pos, const FR
 				FTransform NewTransform(KeyInfo->Transform);
 				NewTransform.SetLocation(NewTransform.GetLocation() + Pos);
 				KeyInfo->Transform = NewTransform;
+				NewTransform = OffsetTransform.GetRelativeTransformReverse(NewTransform);
+
 				Transforms.Add(NewTransform);
 				ParentTransforms.Add(KeyInfo->ParentTransform);
 				Times.Add(KeyInfo->FrameNumber);
@@ -1411,6 +1429,9 @@ bool FMovieSceneControlRigTransformTrail::ApplyDelta(const FVector& Pos, const F
 				FTransform NewTransform(KeyInfo->Transform);
 				NewTransform.SetLocation(NewTransform.GetLocation() + Pos);
 				KeyInfo->Transform = NewTransform;
+
+				NewTransform = OffsetTransform.GetRelativeTransformReverse(NewTransform);
+
 				NewTransform = NewTransform.GetRelativeTransform(KeyInfo->ParentTransform);
 				GetSequencer()->GetEvaluationTemplate().EvaluateSynchronousBlocking(MovieSceneContext);
 				ControlRig->Evaluate_AnyThread();
@@ -1422,6 +1443,7 @@ bool FMovieSceneControlRigTransformTrail::ApplyDelta(const FVector& Pos, const F
 				UE::AIE::FSetTransformHelpers::SetControlTransform(EulerTransform,ControlRig, ControlElement, Context);		
 			}
 		};
+		FControlRigEditMode::FTurnOffPosePoseUpdate  TurnOff; //stop flashing
 
 		if (IsTrailSelected())
 		{
@@ -1441,6 +1463,12 @@ bool FMovieSceneControlRigTransformTrail::ApplyDelta(const FVector& Pos, const F
 			}
 		}
 		KeyTool->UpdateSelectedKeysTransform();
+		GetSequencer()->ForceEvaluate();
+		ControlRig->Evaluate_AnyThread();
+		if (ControlRig->GetObjectBinding())
+		{
+			ControlRig->EvaluateSkeletalMeshComponent(0.0);
+		}
 		return true;
 	}
 	return false;
diff --git a/Engine/Plugins/Editor/SequencerAnimTools/Source/SequencerAnimTools/Private/Sequencer/MovieSceneTransformTrail.h b/Engine/Plugins/Editor/SequencerAnimTools/Source/SequencerAnimTools/Private/Sequencer/MovieSceneTransformTrail.h
--- a/Engine/Plugins/Editor/SequencerAnimTools/Source/SequencerAnimTools/Private/Sequencer/MovieSceneTransformTrail.h
+++ b/Engine/Plugins/Editor/SequencerAnimTools/Source/SequencerAnimTools/Private/Sequencer/MovieSceneTransformTrail.h
@@ -72,7 +72,7 @@ public:
 	// Begin FTrail interface
 	virtual FTrailCurrentStatus UpdateTrail(const FNewSceneContext& NewSceneContext) override;
 	virtual void Interp(const FFrameNumber& Time, FTransform& OutTransform, FTransform& OutParentTransform) override;
-	virtual void UpdateFinished(const TRange<FFrameNumber>& UpdatedRange, const TArray<int32>& IndicesToCalcluate) override;
+	virtual void UpdateFinished(const TRange<FFrameNumber>& UpdatedRange, const TArray<int32>& IndicesToCalcluate, bool bDoneCalculating) override;
 	virtual void AddImportantTimes(TSet<FFrameNumber>& InOutImportantTimes) override;
 	virtual FTransform GetOffsetTransform() const override { return OffsetTransform; }
 	virtual void ClearOffsetTransform();
@@ -86,7 +86,7 @@ public:
 	virtual void ActorChanged(AActor* InActor);
 	virtual void UpdateNeedsEvaluation(const FGuid& InTrailGuid, FSequencerTrailHierarchy* SequencerHierarchy);
 
-	virtual void Render(const FGuid& Guid, const FSceneView* View,  FPrimitiveDrawInterface* PDI) override;
+	virtual void Render(const FGuid& Guid, const FSceneView* View,  FPrimitiveDrawInterface* PDI, bool bTrailIsEvaluating) override;
 	virtual void RenderEvaluating(const FGuid& Guid, const FSceneView* View, FPrimitiveDrawInterface* PDI) override;
 	virtual void DrawHUD(const FSceneView* View, FCanvas* Canvas) override;
 
@@ -98,6 +98,8 @@ public:
 	virtual bool StartTracking() override;
 	virtual bool ApplyDelta(const FVector& Pos, const FRotator& Rot, const FVector& WidgetLocation, bool bApplyToOffset) override;
 	virtual bool EndTracking() override;
+	virtual bool IsTracking() const override;
+
 	virtual void HasStartedEvaluating() override;
 	virtual void TranslateSelectedKeys(bool bRight) override;
 	virtual void DeleteSelectedKeys() override;
@@ -125,7 +127,7 @@ protected:
 	virtual bool HandleAltClick(FEditorViewportClient* InViewportClient, HNewMotionTrailProxy* Proxy, FInputClick Click);
 
 	//internal draw trail
-	void InternalDrawTrail(const FGuid& Guid, const FSceneView* View, FPrimitiveDrawInterface* PDI, const FDrawCacheData& InDrawData, bool bHitTesting) const;
+	void InternalDrawTrail(const FGuid& Guid, const FSceneView* View, FPrimitiveDrawInterface* PDI, const FDrawCacheData& InDrawData, bool bHitTesting, bool bTrailIsEvaluating) const;
 
 	//if selected or selected and offset mode we draw different colors
 	TOptional<FLinearColor> GetOverrideColor() const;
@@ -141,7 +143,6 @@ protected:
 	bool bIsOffsetMode;
 	FVector SelectedPos;
 	bool bIsTracking; 
-
 	//main tool for adjusting keys
 	TUniquePtr<FMotionTraiMovieScenelKeyTool> KeyTool;
 
diff --git a/Engine/Plugins/Editor/SequencerAnimTools/Source/SequencerAnimTools/Private/Sequencer/SequencerTrailHierarchy.cpp b/Engine/Plugins/Editor/SequencerAnimTools/Source/SequencerAnimTools/Private/Sequencer/SequencerTrailHierarchy.cpp
--- a/Engine/Plugins/Editor/SequencerAnimTools/Source/SequencerAnimTools/Private/Sequencer/SequencerTrailHierarchy.cpp
+++ b/Engine/Plugins/Editor/SequencerAnimTools/Source/SequencerAnimTools/Private/Sequencer/SequencerTrailHierarchy.cpp
@@ -43,6 +43,13 @@ void FColorState::Setup(FTrailHierarchy* TrailHierarchy)
 
 EMotionTrailTrailStyle FColorState::GetStyle() const
 {
+	if (Options)
+	{
+		return Options->TrailStyle;
+	}
+	return PinnedStyle;
+
+	/* Pinned style doesn't matter for now but may come back
 	if (bIsPinned)
 	{
 		return PinnedStyle;
@@ -52,6 +59,7 @@ EMotionTrailTrailStyle FColorState::GetStyle() const
 		return Options->TrailStyle;
 	}
 	return PinnedStyle;
+	*/
 }
 
 void FColorState::ReadyForTrail(bool bInIsPinned, EMotionTrailTrailStyle InPinnedStyle)
@@ -75,6 +83,19 @@ FFrameNumber FSequencerTrailHierarchy::GetLocalTime() const
 using namespace UE::AIE;
 TArray<TUniquePtr <FTrail::FMotionTrailState >> FSequencerTrailHierarchy::PreviouslyPinnedTrails;
 
+
+FSequencerTrailHierarchy::FSequencerTrailHierarchy(TWeakPtr<ISequencer> InWeakSequencer)
+	: FTrailHierarchy()
+	, WeakSequencer(InWeakSequencer)
+	, ObjectsTracked()
+	, ControlsTracked()
+	, HierarchyRenderer(MakeUnique<FTrailHierarchyRenderer>(this, UMotionTrailToolOptions::GetTrailOptions()))
+	, OnActorAddedToSequencerHandle()
+	, OnSelectionChangedHandle()
+	, OnViewOptionsChangedHandle()
+	, ControlRigDelegateHandles()
+{
+}
 void FSequencerTrailHierarchy::Initialize()
 {
 	TSharedPtr<ISequencer> Sequencer = WeakSequencer.Pin();
@@ -82,7 +103,6 @@ void FSequencerTrailHierarchy::Initialize()
 	{
 		return;
 	}
-	UpdateViewAndEvalRange();
 
 	TArray<FGuid> SequencerSelectedObjects;
 	Sequencer->GetSelectedObjects(SequencerSelectedObjects);
@@ -200,9 +220,6 @@ void FSequencerTrailHierarchy::Initialize()
 		}
 		PreviouslyPinnedTrails.Reset();
 	}
-
-	
-	
 }
 
 void FSequencerTrailHierarchy::Destroy()
@@ -352,7 +369,7 @@ bool FSequencerTrailHierarchy::IsTrailEvaluating(const FGuid& InGuid, bool bIndi
 		{
 			if (const TUniquePtr<FTrail>* Trail = GetAllTrails().Find(InGuid))
 			{
-				if (Trail->Get()->IsAnythingSelected())
+				if (Trail->Get()->IsTracking()) //direct if moving a key/offset
 				{
 					return false;
 				}
@@ -404,6 +421,10 @@ void FSequencerTrailHierarchy::EvaluateSequencerAndSetTransforms()
 						}
 					}
 				}
+				else
+				{
+					GuidTrailPair.Value->ClearCachedData();
+				}
 			}
 			CurrentFramesInfo.AddMustHaveTimes(MustHaveTimes, CurrentFrame);
 
@@ -425,11 +446,14 @@ void FSequencerTrailHierarchy::EvaluateSequencerAndSetTransforms()
 			
 			const TRange<FFrameNumber> Range = TickEvalRange;
 			const TArray<int32>& IndicesToCalcluate = CurrentFramesInfo.IndicesToCalculate();
-			for (const FGuid& CurGuid : EvaluatingTrails)
+			if (CurrentFramesInfo.TransformIndices.Num() > 0 || IndicesToCalcluate.Num() > 0)
 			{
-				if (const TUniquePtr<FTrail>* Trail = GetAllTrails().Find(CurGuid))
+				for (const FGuid& CurGuid : EvaluatingTrails)
 				{
-					(*Trail)->UpdateFinished(Range, IndicesToCalcluate);
+					if (const TUniquePtr<FTrail>* Trail = GetAllTrails().Find(CurGuid))
+					{
+						(*Trail)->UpdateFinished(Range, IndicesToCalcluate, (bKeepCalculating == false));
+					}
 				}
 			}
 
@@ -605,6 +629,7 @@ void FSequencerTrailHierarchy::UpdateSequencerBindings(const TArray<FGuid>& Sequ
 				}
 				if (bAddedControlRig)
 				{
+					ClearSelection();
 					continue;
 				}
 			}
@@ -732,7 +757,6 @@ void FSequencerTrailHierarchy::UpdateViewAndEvalRange()
 
 	FFrameRate TickResolution = Sequencer->GetFocusedTickResolution();
 	FFrameRate DisplayRate = Sequencer->GetFocusedDisplayRate();
-
 	TickEvalRange  = Sequencer->GetFocusedMovieSceneSequence()->GetMovieScene()->GetPlaybackRange();
 	if (!UMotionTrailToolOptions::GetTrailOptions()->bShowFullTrail)
 	{
@@ -740,6 +764,14 @@ void FSequencerTrailHierarchy::UpdateViewAndEvalRange()
 		const FFrameNumber TicksBefore = FFrameRate::TransformTime(FFrameNumber(UMotionTrailToolOptions::GetTrailOptions()->FramesBefore), DisplayRate, TickResolution).FloorToFrame();
 		const FFrameNumber TicksAfter = FFrameRate::TransformTime(FFrameNumber(UMotionTrailToolOptions::GetTrailOptions()->FramesAfter), DisplayRate, TickResolution).FloorToFrame();
 		TickViewRange = TRange<FFrameNumber>(SequenceTime.GetFrame() - TicksBefore, SequenceTime.GetFrame() + TicksAfter);
+		if (TickViewRange.GetLowerBoundValue() < TickEvalRange.GetLowerBoundValue())
+		{
+			TickViewRange.SetLowerBoundValue(TickEvalRange.GetLowerBoundValue());
+		}
+		if (TickViewRange.GetUpperBoundValue() > TickEvalRange.GetUpperBoundValue())
+		{
+			TickViewRange.SetUpperBoundValue(TickEvalRange.GetUpperBoundValue());
+		}
 	}
 	else
 	{
@@ -762,6 +794,9 @@ FGuid FSequencerTrailHierarchy::AddComponentToHierarchy(const FGuid& InBindingGu
 	CurTrail->ForceEvaluateNextTick();
 
 	AddTrail(ObjectsTracked[CompToAdd], MoveTemp(CurTrail));
+
+	ClearSelection();
+
 	return CurTrailGuid;
 }
 
@@ -832,6 +867,18 @@ FGuid FSequencerTrailHierarchy::AddControlRigTrail(USkeletalMeshComponent* Compo
 	return NewGuid;
 }
 
+void FSequencerTrailHierarchy::ClearSelection()
+{
+	const bool bShiftDown = FSlateApplication::Get().GetModifierKeys().IsShiftDown();
+	if (!bShiftDown)
+	{
+		for (const TPair<FGuid, TUniquePtr<FTrail>>& GuidTrailPair : GetAllTrails())
+		{
+			GuidTrailPair.Value->SelectNone();
+		}
+	}
+}
+
 void FSequencerTrailHierarchy::RegisterControlRigDelegates(USkeletalMeshComponent* Component, UMovieSceneControlRigParameterTrack* CRParameterTrack)
 {
 	UControlRig* ControlRig = CRParameterTrack->GetControlRig();
@@ -883,6 +930,8 @@ void FSequencerTrailHierarchy::RegisterControlRigDelegates(USkeletalMeshComponen
 				VisibilityManager.ControlSelected.Remove(TrailGuid);
 				RemoveTrailIfNotAlwaysVisible(TrailGuid);
 			}
+
+			ClearSelection();
 		}
 
 			
@@ -931,8 +980,10 @@ void FSequencerTrailHierarchy::RegisterMotionTrailOptionDelegates()
 	if (UMotionTrailToolOptions* TrailOptions = UMotionTrailToolOptions::GetTrailOptions())
 	{
 		TrailOptions->OnPinSelection.AddRaw(this, &FSequencerTrailHierarchy::OnPinSelection);
+		TrailOptions->OnUnPinSelection.AddRaw(this, &FSequencerTrailHierarchy::OnUnPinSelection);
 		TrailOptions->OnPinComponent.AddRaw(this, &FSequencerTrailHierarchy::OnPinComponent);
 		TrailOptions->OnDeletePinned.AddRaw(this, &FSequencerTrailHierarchy::OnDeletePinned);
+		TrailOptions->OnDeleteAllPinned.AddRaw(this, &FSequencerTrailHierarchy::OnDeleteAllPinned);
 		TrailOptions->OnPutPinnedInSpace.AddRaw(this, &FSequencerTrailHierarchy::OnPutPinnedInSpace);
 		TrailOptions->OnSetLinearColor.AddRaw(this, &FSequencerTrailHierarchy::OnSetLinearColor);
 		TrailOptions->OnSetHasOffset.AddRaw(this, &FSequencerTrailHierarchy::OnSetHasOffset);
@@ -944,8 +995,10 @@ void FSequencerTrailHierarchy::UnRegisterMotionTrailOptionDelegates()
 	if (UMotionTrailToolOptions* TrailOptions = UMotionTrailToolOptions::GetTrailOptions())
 	{
 		TrailOptions->OnPinSelection.RemoveAll(this);
+		TrailOptions->OnUnPinSelection.RemoveAll(this);
 		TrailOptions->OnPinComponent.RemoveAll(this);
 		TrailOptions->OnDeletePinned.RemoveAll(this);
+		TrailOptions->OnDeleteAllPinned.RemoveAll(this);
 		TrailOptions->OnPutPinnedInSpace.RemoveAll(this);
 		TrailOptions->OnSetLinearColor.RemoveAll(this);
 		TrailOptions->OnSetHasOffset.RemoveAll(this);
@@ -956,13 +1009,19 @@ void FSequencerTrailHierarchy::PinTrail(FGuid InGuid)
 {
 	if (const TUniquePtr<FTrail>* Trail = GetAllTrails().Find(InGuid))
 	{
+		UMotionTrailToolOptions* Settings = UMotionTrailToolOptions::GetTrailOptions();
+		if (Settings->GetNumPinned() >= Settings->MaxNumberPinned)
+		{
+			UE_LOG(LogTemp, Warning, TEXT("MotionTrails: Cannot Pin trail %s Max number reached. Please delete pinned trail if you want to add this one."), *(*Trail)->GetName().ToString());
+			return;
+		}
+
 		VisibilityManager.SetTrailAlwaysVisible(InGuid, true);
 		if ((*Trail)->GetDrawInfo())
 		{
 			(*Trail)->GetDrawInfo()->SetColor(UMotionTrailToolOptions::GetTrailOptions()->DefaultColor);
 			(*Trail)->GetDrawInfo()->SetStyle(UMotionTrailToolOptions::GetTrailOptions()->TrailStyle);
 
-			UMotionTrailToolOptions* Settings = UMotionTrailToolOptions::GetTrailOptions();
 			UMotionTrailToolOptions::FPinnedTrail PinnedTrail;
 			PinnedTrail.TrailGuid = InGuid;
 			PinnedTrail.TrailName = (*Trail)->GetName();
@@ -984,6 +1043,18 @@ void FSequencerTrailHierarchy::OnPinSelection()
 	}
 }
 
+void FSequencerTrailHierarchy::OnUnPinSelection()
+{
+	for (const TPair<FGuid, TUniquePtr<FTrail>>& GuidTrailPair : GetAllTrails())
+	{
+		if (VisibilityManager.IsTrailAlwaysVisible(GuidTrailPair.Key) == true 
+			&& (VisibilityManager.ControlSelected.Contains(GuidTrailPair.Key) || VisibilityManager.Selected.Contains(GuidTrailPair.Key)))
+		{
+			VisibilityManager.SetTrailAlwaysVisible(GuidTrailPair.Key, false);
+		}
+	}
+}
+
 void FSequencerTrailHierarchy::OnPinComponent(USceneComponent* InSceneComponent, FName InSocketName)
 {
 	PinComponent(InSceneComponent, InSocketName);
@@ -1070,7 +1141,23 @@ FGuid FSequencerTrailHierarchy::PinComponent(USceneComponent* InSceneComponent,
 void FSequencerTrailHierarchy::OnDeletePinned(FGuid InGuid)
 {
 	VisibilityManager.SetTrailAlwaysVisible(InGuid, false);
-	RemoveTrail(InGuid);
+	if (VisibilityManager.ControlSelected.Contains(InGuid) == false && VisibilityManager.Selected.Contains(InGuid) == false)
+	{
+		RemoveTrail(InGuid);
+	}
+}
+
+void FSequencerTrailHierarchy::OnDeleteAllPinned()
+{
+	TArray<FGuid> AlwaysVisibleTrails = VisibilityManager.AlwaysVisible.Array();
+	for (const FGuid& Guid: AlwaysVisibleTrails)
+	{
+		if (VisibilityManager.ControlSelected.Contains(Guid) == false && VisibilityManager.Selected.Contains(Guid) == false)
+		{
+			RemoveTrail(Guid);
+		}
+	}
+	VisibilityManager.AlwaysVisible.Reset();
 }
 
 void FSequencerTrailHierarchy::OnPutPinnedInSpace(FGuid InGuid, AActor* InActor, FName InComponentName)
diff --git a/Engine/Plugins/Editor/SequencerAnimTools/Source/SequencerAnimTools/Private/Sequencer/SequencerTrailHierarchy.h b/Engine/Plugins/Editor/SequencerAnimTools/Source/SequencerAnimTools/Private/Sequencer/SequencerTrailHierarchy.h
--- a/Engine/Plugins/Editor/SequencerAnimTools/Source/SequencerAnimTools/Private/Sequencer/SequencerTrailHierarchy.h
+++ b/Engine/Plugins/Editor/SequencerAnimTools/Source/SequencerAnimTools/Private/Sequencer/SequencerTrailHierarchy.h
@@ -35,18 +35,8 @@ enum class EBindingVisibilityState
 class FSequencerTrailHierarchy : public FTrailHierarchy
 {
 public:
-	FSequencerTrailHierarchy(TWeakPtr<ISequencer> InWeakSequencer)
-		: FTrailHierarchy()
-		, WeakSequencer(InWeakSequencer)
-		, ObjectsTracked()
-		, ControlsTracked()
-		, HierarchyRenderer(MakeUnique<FTrailHierarchyRenderer>(this))
-		, OnActorAddedToSequencerHandle()
-		, OnSelectionChangedHandle()
-		, OnViewOptionsChangedHandle()
-		, ControlRigDelegateHandles()
-	{
-	}
+	FSequencerTrailHierarchy(TWeakPtr<ISequencer> InWeakSequencer);
+	
 	//New
 	void EvaluateActor(const FGuid& InKey, FActorForWorldTransforms& InActor,TSharedPtr<UE::AIE::FArrayOfTransforms>& WorldTransforms, TSharedPtr<UE::AIE::FArrayOfTransforms>& ParentWorldTransforms);
 	void EvaluateControlRig(const FGuid& InKey, UControlRig* InControlRig, const FName& InControlName, TSharedPtr<UE::AIE::FArrayOfTransforms>& WorldTransforms);
@@ -134,10 +124,12 @@ private:
 	//cached guid to check for updates
 	FGuid LastValidMovieSceneGuid;
 protected: //callbacks
-
+	void ClearSelection(); 
 	void RegisterMotionTrailOptionDelegates();
 	void UnRegisterMotionTrailOptionDelegates();
 	void OnPinSelection();
+	void OnUnPinSelection();
+	void OnDeleteAllPinned();
 	void OnPinComponent(USceneComponent* InSceneComponent, FName SocketName);
 	void OnDeletePinned(FGuid InGuid);
 	void OnPutPinnedInSpace(FGuid InGuid, AActor* InActor, FName InComponentName);
diff --git a/Engine/Plugins/Editor/SequencerAnimTools/Source/SequencerAnimTools/Private/TrailHierarchy.cpp b/Engine/Plugins/Editor/SequencerAnimTools/Source/SequencerAnimTools/Private/TrailHierarchy.cpp
--- a/Engine/Plugins/Editor/SequencerAnimTools/Source/SequencerAnimTools/Private/TrailHierarchy.cpp
+++ b/Engine/Plugins/Editor/SequencerAnimTools/Source/SequencerAnimTools/Private/TrailHierarchy.cpp
@@ -15,6 +15,10 @@
 #include "Widgets/Text/STextBlock.h"
 #include "CanvasTypes.h"
 #include "Sequencer/SequencerTrailHierarchy.h"
+#include "EditMode/ControlRigEditMode.h"
+#include "EditorModeManager.h"
+#include "ScopedTransaction.h"
+
 
 #define LOCTEXT_NAMESPACE "MotionTrailEditorMode"
 
@@ -77,7 +81,7 @@ void FTrail::GetColor(const FFrameNumber& CurrentTime, FColorState& State)
 
 void FTrailHierarchyRenderer::Render(const FSceneView* View, FPrimitiveDrawInterface* PDI)
 {
-	/*
+	
 	if (CachedOptions)
 	{
 		//options object won't change over lifetime of the renderer functions
@@ -108,64 +112,18 @@ void FTrailHierarchyRenderer::Render(const FSceneView* View, FPrimitiveDrawInter
 				//getting manipulated we still want to draw it. If it's not evaluating though we don't use sparse time values, instead do every frame.
 				const bool bTrailDirectlyEvaluating = OwningHierarchy->IsTrailEvaluating(GuidTrailPair.Key, false);
 				GuidTrailPair.Value->ReadyToDrawTrail(ColorState, CurrentFramesInfo, bTrailDirectlyEvaluating, OwningHierarchy->GetVisibilityManager().IsTrailAlwaysVisible(GuidTrailPair.Key));
-				GuidTrailPair.Value->Render(GuidTrailPair.Key, View, PDI);
+				GuidTrailPair.Value->Render(GuidTrailPair.Key, View, PDI, bTrailDirectlyEvaluating);
 			}
 			
 		}
 		const FTimespan RenderTimespan = FDateTime::Now() - RenderStartTime;
 		OwningHierarchy->GetTimingStats().Add("FTrailHierarchyRenderer::Render", RenderTimespan);
 	}
-	*/
-
-	const FDateTime RenderStartTime = FDateTime::Now();
-
-	const int32 NumEvalTimes = int32((OwningHierarchy->GetViewFrameRange().GetUpperBoundValue().Value - OwningHierarchy->GetViewFrameRange().GetLowerBoundValue().Value) / OwningHierarchy->GetFramesPerFrame().Value) + 1;
-	const int32 NumLinesReserveSize = int32(NumEvalTimes * OwningHierarchy->GetAllTrails().Num() * 1.3);
-	PDI->AddReserveLines(SDPG_Foreground, NumLinesReserveSize);
-
-	TArray<FGuid> TrailsToRemove; //if trail is not visible we just get rid of it
-	// <parent, current>
-	FColorState ColorState;
-	ColorState.Setup(OwningHierarchy);
-	const FCurrentFramesInfo* CurrentFramesInfo = OwningHierarchy->GetCurrentFramesInfo();
-	for (const TPair<FGuid, TUniquePtr<FTrail>>& GuidTrailPair : OwningHierarchy->GetAllTrails())
-	{
-		FTrajectoryDrawInfo* CurDrawInfo = GuidTrailPair.Value->GetDrawInfo();
-		if (CurDrawInfo && OwningHierarchy->GetVisibilityManager().IsTrailVisible(GuidTrailPair.Key, GuidTrailPair.Value.Get()))
-		{
-			if (OwningHierarchy->IsTrailEvaluating(GuidTrailPair.Key, true))
-			{
-				GuidTrailPair.Value->RenderEvaluating(GuidTrailPair.Key, View, PDI);
-				continue;
-			}
-			if (GuidTrailPair.Value.Get()->GetCacheState() != ETrailCacheState::UpToDate)
-			{
-				continue;
-			}
-			//the above check so that if we are evaluating something that is getting drawn indirectly we skip drawing it, but if it is
-			//getting manipulated we still want to draw it. If it's not evaluating though we don't use sparse time values, instead do every frame.
-			const bool bTrailDirectlyEvaluating = OwningHierarchy->IsTrailEvaluating(GuidTrailPair.Key, false);
-			GuidTrailPair.Value->ReadyToDrawTrail(ColorState, CurrentFramesInfo, bTrailDirectlyEvaluating, OwningHierarchy->GetVisibilityManager().IsTrailAlwaysVisible(GuidTrailPair.Key));
-			GuidTrailPair.Value->Render(GuidTrailPair.Key, View, PDI);
-		}
-		else
-		{
-			TrailsToRemove.Add(GuidTrailPair.Key);
-		}
-
-	}
-	for (const FGuid& Key : TrailsToRemove)
-	{
-		OwningHierarchy->RemoveTrail(Key);
-	}
-	const FTimespan RenderTimespan = FDateTime::Now() - RenderStartTime;
-	OwningHierarchy->GetTimingStats().Add("FTrailHierarchyRenderer::Render", RenderTimespan);
-
 }
 
 void FTrailHierarchyRenderer::DrawHUD(const FSceneView* View, FCanvas* Canvas)
 {
-	/*
+	
 	if (CachedOptions && CachedOptions->bShowMarks)
 	{
 		const FDateTime DrawHUDStartTime = FDateTime::Now();
@@ -182,56 +140,12 @@ void FTrailHierarchyRenderer::DrawHUD(const FSceneView* View, FCanvas* Canvas)
 		{
 			if (OwningHierarchy->GetVisibilityManager().IsTrailVisible(GuidTrailPair.Key, GuidTrailPair.Value.Get(), CachedOptions->bShowSelectedTrails))
 			{
-				if (OwningHierarchy->IsTrailEvaluating(GuidTrailPair.Key, true) || GuidTrailPair.Value.Get()->GetCacheState() != ETrailCacheState::UpToDate)
-				{
-					continue;
-				}
-				TArray<FVector2D> Marks, MarkNormals;
-				const bool bTrailDirectlyEvaluating = OwningHierarchy->IsTrailEvaluating(GuidTrailPair.Key, false);
-				GuidTrailPair.Value->GetTickPointsForDisplay(Transform, *CurrentFramesInfo, bTrailDirectlyEvaluating, Marks, MarkNormals);
-				if (Marks.Num() > 0)
-				{
-					const FLinearColor Color = UMotionTrailToolOptions::GetTrailOptions()->MarkColor;
-					for (int32 Idx = 0; Idx < Marks.Num(); Idx++)
-					{
-						const FVector2D StartPoint = Marks[Idx] - MarkNormals[Idx] * UMotionTrailToolOptions::GetTrailOptions()->MarkSize;
-						const FVector2D EndPoint = Marks[Idx] + MarkNormals[Idx] * UMotionTrailToolOptions::GetTrailOptions()->MarkSize;
-						FCanvasLineItem LineItem = FCanvasLineItem(StartPoint, EndPoint);
-						LineItem.SetColor(Color);
-						Canvas->DrawItem(LineItem);
-					}
-				}
-			}
-			GuidTrailPair.Value->DrawHUD(View, Canvas);
-
-		}
-
-		const FTimespan DrawHUDTimespan = FDateTime::Now() - DrawHUDStartTime;
-		OwningHierarchy->GetTimingStats().Add("FTrailHierarchyRenderer::DrawHUD", DrawHUDTimespan);
-	}*/
-	if (UMotionTrailToolOptions::GetTrailOptions()->bShowMarks)
-	{
-		const FDateTime DrawHUDStartTime = FDateTime::Now();
-		const FFrameNumber FramesPerMark = OwningHierarchy->GetFramesPerFrame();
-
-		const int32 NumMarksPerTrail = int32((OwningHierarchy->GetViewFrameRange().GetUpperBoundValue().Value - OwningHierarchy->GetViewFrameRange().GetLowerBoundValue().Value) / FramesPerMark.Value) + 1;
-		const int32 PredictedNumMarks = int32(NumMarksPerTrail * OwningHierarchy->GetAllTrails().Num() * 1.3);
-
-		Canvas->GetBatchedElements(FCanvas::EElementType::ET_Line)->AddReserveLines(PredictedNumMarks);
-		FTrailScreenSpaceTransform Transform(View, Canvas->GetDPIScale());
-		const FCurrentFramesInfo* CurrentFramesInfo = OwningHierarchy->GetCurrentFramesInfo();
-
-		for (const TPair<FGuid, TUniquePtr<FTrail>>& GuidTrailPair : OwningHierarchy->GetAllTrails())
-		{
-			if (OwningHierarchy->GetVisibilityManager().IsTrailVisible(GuidTrailPair.Key, GuidTrailPair.Value.Get()))
-			{
-				if (OwningHierarchy->IsTrailEvaluating(GuidTrailPair.Key, true) || GuidTrailPair.Value.Get()->GetCacheState() != ETrailCacheState::UpToDate)
+				if (OwningHierarchy->IsTrailEvaluating(GuidTrailPair.Key, false) || GuidTrailPair.Value.Get()->GetCacheState() != ETrailCacheState::UpToDate)
 				{
 					continue;
 				}
 				TArray<FVector2D> Marks, MarkNormals;
-				const bool bTrailDirectlyEvaluating = OwningHierarchy->IsTrailEvaluating(GuidTrailPair.Key, false);
-				GuidTrailPair.Value->GetTickPointsForDisplay(Transform, *CurrentFramesInfo, bTrailDirectlyEvaluating, Marks, MarkNormals);
+				GuidTrailPair.Value->GetTickPointsForDisplay(Transform, *CurrentFramesInfo, false, Marks, MarkNormals);
 				if (Marks.Num() > 0)
 				{
 					const FLinearColor Color = UMotionTrailToolOptions::GetTrailOptions()->MarkColor;
@@ -356,7 +270,8 @@ public:
 
 void FTrailHierarchy::OpenContextMenu(const FGuid& TrailGuid)
 {
-
+	/* Turned off for now 
+	* 
 	FLevelEditorModule& LevelEditorModule = FModuleManager::GetModuleChecked<FLevelEditorModule>("LevelEditor");
 	TSharedPtr< ILevelEditor > LevelEditor = LevelEditorModule.GetFirstLevelEditor();
 
@@ -377,7 +292,7 @@ void FTrailHierarchy::OpenContextMenu(const FGuid& TrailGuid)
 			FPopupTransitionEffect(FPopupTransitionEffect::ContextMenu)
 		);
 	}
-	
+	*/
 }
 
 bool FTrailHierarchy::IsHitByClick(HHitProxy* InHitProxy)
@@ -389,6 +304,7 @@ bool FTrailHierarchy::IsHitByClick(HHitProxy* InHitProxy)
 			return true;
 		}
 	}
+	
 	return false;
 }
 
@@ -415,6 +331,21 @@ bool FTrailHierarchy::HandleClick(FEditorViewportClient* InViewportClient, HHitP
 		}
 	}
 
+	if (bClickIsHandled)
+	{
+		FScopedTransaction ScopedTransaction(LOCTEXT("HandleClick", "Handle Click"), !GIsTransacting);
+
+		if (UMotionTrailToolOptions::GetTrailOptions()->bShowSelectedTrails == false && Click.bShiftIsDown == false)
+		{
+			//use control rig edit mode to clear selction, since it handles everything
+			if (FControlRigEditMode* ControlRigEditMode = static_cast<FControlRigEditMode*>(GLevelEditorModeTools().GetActiveMode(FControlRigEditMode::ModeName)))
+			{
+				ControlRigEditMode->ClearSelection();
+			}
+		}
+	}
+	
+
 	return bClickIsHandled;
 }
 
diff --git a/Engine/Plugins/Editor/SequencerAnimTools/Source/SequencerAnimTools/Private/TrajectoryDrawInfo.cpp b/Engine/Plugins/Editor/SequencerAnimTools/Source/SequencerAnimTools/Private/TrajectoryDrawInfo.cpp
--- a/Engine/Plugins/Editor/SequencerAnimTools/Source/SequencerAnimTools/Private/TrajectoryDrawInfo.cpp
+++ b/Engine/Plugins/Editor/SequencerAnimTools/Source/SequencerAnimTools/Private/TrajectoryDrawInfo.cpp
@@ -65,7 +65,7 @@ void FFrameCalculator::SetUpFrameCalculator(const UE::AIE::FFrameTimeByIndex& In
 
 	Reset();
 	CalculateStartEndIndices();
-	NumOfBuckets = FMath::RoundToInt32((float)(CurrentRangeEndIndex + 1 - CurrentRangeStartIndex) / (float)(IndicesToCalculatePerBucket));
+	NumOfBuckets = FMath::FloorToInt32((float)(CurrentRangeEndIndex + 1 - CurrentRangeStartIndex) / (float)(IndicesToCalculatePerBucket)) + 1;
 }
 
 void FFrameCalculator::CalculateStartEndIndices()
@@ -86,7 +86,7 @@ void FFrameCalculator::AddMustHaveIndices(const TArray<int32>& InMustHaveIndices
 void FFrameCalculator::Reset()
 {
 	CurrentRange = 0;
-	CurrentBucket = -1; //-1 bucket is just MustHaveTimes
+	CurrentBucket = 0; //-1 bucket is just MustHaveTimes, only for mouse down
 	IndicesToCalculate.SetNum(0);
 }
 
@@ -103,6 +103,7 @@ bool FFrameCalculator::CalculateIndices()
 	{
 		CurrentBucket = -1;
 		CurrentRange = 0;
+		return true;
 	}
 	if (CurrentBucket >= 0)
 	{
@@ -151,7 +152,7 @@ void FCurrentFramesInfo::SetUpFrameTimes(const TRange<FFrameNumber>& InEvalFrame
 	TransformIndices.Reserve(CurrentFrameTimes.NumFrames);
 	TransformIndices.Reset();
 
-	FrameCalculator.SetUpFrameCalculator(CurrentFrameTimes, ViewRange);
+	FrameCalculator.SetUpFrameCalculator(CurrentFrameTimes, InEvalFrameRange);
 }
 
 void FCurrentFramesInfo::AddMustHaveTimes(const TSet<FFrameNumber>& InMustHaveTimes, const FFrameNumber& CurrentFrame)
diff --git a/Engine/Plugins/Editor/SequencerAnimTools/Source/SequencerAnimTools/Public/Trail.h b/Engine/Plugins/Editor/SequencerAnimTools/Source/SequencerAnimTools/Public/Trail.h
--- a/Engine/Plugins/Editor/SequencerAnimTools/Source/SequencerAnimTools/Public/Trail.h
+++ b/Engine/Plugins/Editor/SequencerAnimTools/Source/SequencerAnimTools/Public/Trail.h
@@ -120,14 +120,14 @@ public:
 	virtual void Interp(const FFrameNumber& Time, FTransform& OutTransform, FTransform& OutParentTransform) { OutTransform = FTransform::Identity; OutParentTransform = FTransform::Identity;
 	}
 	//when multiple frame update is finally finished
-	virtual void UpdateFinished(const TRange<FFrameNumber>& UpdatedRange, const TArray<int32>& IndicesToCalcluate) { ClearCachedData();  CacheState = ETrailCacheState::UpToDate; }
+	virtual void UpdateFinished(const TRange<FFrameNumber>& UpdatedRange, const TArray<int32>& IndicesToCalcluate, bool bDoneCalcuating) { ClearCachedData();  CacheState = ETrailCacheState::UpToDate; }
 	virtual FText GetName() const = 0;
 
 	//add set of important times for this trail to be used by the incremental evaluations, usually these are the edited times
 	virtual void AddImportantTimes(TSet<FFrameNumber>& InOutImportantTimes) {};
 	
 	//Additional Render/HitTest event handling for specific trails, usually let default renderer handle it
-	virtual void Render(const FGuid& Guid, const FSceneView* View, FPrimitiveDrawInterface* PDI) {};
+	virtual void Render(const FGuid& Guid, const FSceneView* View, FPrimitiveDrawInterface* PDI, bool bTrailIsEvaluating) {};
 	virtual void RenderEvaluating(const FGuid& Guid, const FSceneView* View, FPrimitiveDrawInterface* PDI) {};
 	virtual void DrawHUD(const FSceneView* View, FCanvas* Canvas) {};
 	virtual bool HandleClick(const FGuid& Guid, FEditorViewportClient* InViewportClient, HHitProxy* HitProxy, FInputClick Click) { return false; }
@@ -141,6 +141,7 @@ public:
 	virtual bool StartTracking() { return false; }
 	virtual bool ApplyDelta(const FVector& Pos, const FRotator& Rot, const FVector& WidgetLocation, bool bApplyToOffset) { return false; }
 	virtual bool EndTracking() { return false; }
+	virtual bool IsTracking() const { return false; }
 	virtual void TranslateSelectedKeys(bool bRight) {};
 	virtual void DeleteSelectedKeys() {};
 	virtual void SelectNone() {};
@@ -159,20 +160,20 @@ public:
 	{
 		if (FTrajectoryDrawInfo* DI = GetDrawInfo())
 		{
-			DI->GetTrajectoryPointsForDisplay(GetOffsetTransform(), GetParentSpaceTransform(), InCurrentFramesInfo, bIsEvaluating, OutPoints, OutFrames);
+			DI->GetTrajectoryPointsForDisplay(FTransform::Identity, GetParentSpaceTransform(), InCurrentFramesInfo, bIsEvaluating, OutPoints, OutFrames);
 		}
 	}
 	virtual void GetTickPointsForDisplay(const FTrailScreenSpaceTransform& InScreenSpaceTransform, const FCurrentFramesInfo& InCurrentFramesInfo, bool bIsEvaluating, TArray<FVector2D>& OutTicks, TArray<FVector2D>& OutTickTangents)
 	{
 		if (FTrajectoryDrawInfo* DI = GetDrawInfo())
 		{
-			DI->GetTickPointsForDisplay(GetOffsetTransform(), GetParentSpaceTransform(), InScreenSpaceTransform, InCurrentFramesInfo, bIsEvaluating, OutTicks, OutTickTangents);
+			DI->GetTickPointsForDisplay(FTransform::Identity, GetParentSpaceTransform(), InScreenSpaceTransform, InCurrentFramesInfo, bIsEvaluating, OutTicks, OutTickTangents);
 		}
 	}
 
 	//offset the trail from current position
 	virtual FTransform GetOffsetTransform() const { return FTransform::Identity; }
-	virtual void ClearOffsetTransform() {};
+	virtual void ClearOffsetTransform() { ForceEvaluateNextTick(); };
 	virtual bool HasOffsetTransform() { return !GetOffsetTransform().Equals(FTransform::Identity); }
 	//set space and get space transform
 	virtual FTransform GetParentSpaceTransform() const { return FTransform::Identity; }
diff --git a/Engine/Plugins/Editor/SequencerAnimTools/Source/SequencerAnimTools/Public/TrailHierarchy.h b/Engine/Plugins/Editor/SequencerAnimTools/Source/SequencerAnimTools/Public/TrailHierarchy.h
--- a/Engine/Plugins/Editor/SequencerAnimTools/Source/SequencerAnimTools/Public/TrailHierarchy.h
+++ b/Engine/Plugins/Editor/SequencerAnimTools/Source/SequencerAnimTools/Public/TrailHierarchy.h
@@ -6,7 +6,8 @@
 
 class HHitProxy;
 struct FFrameNumber;
-
+class UMotionTrailToolOptions;
+class FTrailHierarchy;
 
 namespace UE
 {
@@ -15,9 +16,9 @@ namespace SequencerAnimTools
 
 struct FTrailVisibilityManager
 {
-	bool IsTrailVisible(const FGuid& Guid, const FTrail* Trail) const
+	bool IsTrailVisible(const FGuid& Guid, const FTrail* Trail, bool bShowSelected = true) const
 	{
-		return !InactiveMask.Contains(Guid) && !VisibilityMask.Contains(Guid) && (AlwaysVisible.Contains(Guid) || Selected.Contains(Guid) || ControlSelected.Contains(Guid)
+		return !InactiveMask.Contains(Guid) && !VisibilityMask.Contains(Guid) && (AlwaysVisible.Contains(Guid) || (bShowSelected == true && (Selected.Contains(Guid) || ControlSelected.Contains(Guid)))
 			|| Trail->IsAnythingSelected()) && Guid.IsValid();
 	}
 
@@ -64,15 +65,17 @@ public:
 class FTrailHierarchyRenderer : public ITrailHierarchyRenderer
 {
 public:
-	FTrailHierarchyRenderer(class FTrailHierarchy* InOwningHierarchy)
-		: OwningHierarchy(InOwningHierarchy)
+	FTrailHierarchyRenderer(FTrailHierarchy* InOwningHierarchy, UMotionTrailToolOptions* InOptions)
+		: OwningHierarchy(InOwningHierarchy), CachedOptions(InOptions)
 	{}
 
 	virtual void Render(const FSceneView* View, FPrimitiveDrawInterface* PDI) override;
 	virtual void DrawHUD(const FSceneView* View, FCanvas* Canvas) override;
 
 private:
-	class FTrailHierarchy* OwningHierarchy;
+	FTrailHierarchy* OwningHierarchy;
+	UMotionTrailToolOptions* CachedOptions;
+
 };
 
 class FTrailHierarchy
@@ -147,6 +150,7 @@ protected:
 protected:
 	TRange<FFrameNumber> TickViewRange;
 	TRange<FFrameNumber> TickEvalRange;
+
 	FFrameNumber TicksPerSegment;
 	TRange<FFrameNumber> LastTickEvalRange = TRange<FFrameNumber>(0, 0);
 	FFrameNumber LastTicksPerSegment = FFrameNumber(1);
diff --git a/Engine/Source/Editor/Sequencer/Private/Tools/SMotionTrailOptions.cpp b/Engine/Source/Editor/Sequencer/Private/Tools/SMotionTrailOptions.cpp
--- a/Engine/Source/Editor/Sequencer/Private/Tools/SMotionTrailOptions.cpp
+++ b/Engine/Source/Editor/Sequencer/Private/Tools/SMotionTrailOptions.cpp
@@ -78,6 +78,11 @@ void UMotionTrailToolOptions::ResetPinnedItems()
 
 void UMotionTrailToolOptions::AddPinned(const FPinnedTrail& InPinnedTrail)
 {
+	if (PinnedTrails.Num() >= MaxNumberPinned)
+	{
+		UE_LOG(LogTemp, Warning, TEXT("MotionTrails: Cannot Pin trail %s Max number reached. Please delete pinned trail if you want to add this one."), *InPinnedTrail.TrailName.ToString());
+		return;
+	}
 	if (UMotionTrailToolOptions::FPinnedTrail* Trail = PinnedTrails.FindByKey(InPinnedTrail.TrailGuid))
 	{
 		return;
@@ -96,6 +101,11 @@ void UMotionTrailToolOptions::PinSelection() const
 	OnPinSelection.Broadcast(); //motion trails will handle this
 }
 
+void UMotionTrailToolOptions::UnPinSelection() 
+{
+	OnUnPinSelection.Broadcast(); 
+}
+
 void UMotionTrailToolOptions::PinComponent(USceneComponent* InSceneComponent, const FName& InSocketName) const
 {
 	OnPinComponent.Broadcast(InSceneComponent, InSocketName);
@@ -111,6 +121,12 @@ void UMotionTrailToolOptions::DeletePinned(int32 Index)
 	}
 }
 
+void UMotionTrailToolOptions::DeleteAllPinned() 
+{
+	PinnedTrails.SetNum(0);
+	OnDeleteAllPinned.Broadcast();
+}
+
 void UMotionTrailToolOptions::PutPinnnedInSpace(int32 Index, AActor* InActor, const FName& InComponentName)
 {
 	if (Index >= 0 && Index < PinnedTrails.Num())
diff --git a/Engine/Source/Editor/Sequencer/Public/Tools/MotionTrailOptions.h b/Engine/Source/Editor/Sequencer/Public/Tools/MotionTrailOptions.h
--- a/Engine/Source/Editor/Sequencer/Public/Tools/MotionTrailOptions.h
+++ b/Engine/Source/Editor/Sequencer/Public/Tools/MotionTrailOptions.h
@@ -32,7 +32,7 @@ class SEQUENCER_API UMotionTrailToolOptions : public UObject
 public:
 	UMotionTrailToolOptions()
 		: bShowTrails(false)
-		, bShowSelectedTrails(true)
+		, bShowSelectedTrails(false)
 		, TrailStyle(EMotionTrailTrailStyle::Default)
 		, DefaultColor(.22,0.15,1.0)
 		, TimePreColor(.22, 0.35, 0.8)
@@ -52,6 +52,7 @@ public:
 		, bShowMarks(false)
 		, MarkColor(0.25,1.0,0.15)
 		, MarkSize(5.0)
+		, MaxNumberPinned(10)
 		, bLockMarksToFrames_DEPRECATED(true)
 		, SecondsPerMark_DEPRECATED(0.1)
 	{}
@@ -139,6 +140,9 @@ public:
 	UPROPERTY(EditAnywhere, Config, Category = Marks, Meta = (ClampMin = "0.0"))
 	double MarkSize;
 
+	/** Max number of pinned trails*/
+	UPROPERTY(EditAnywhere, Config, Category = AdvancedDisplay)
+	int32 MaxNumberPinned;
 
 	/** Deprecated in 5.6 */
 	UPROPERTY()
@@ -178,9 +182,11 @@ public:
 
 	int32 GetIndexFromGuid(FGuid InGuid) const;
 	void PinSelection() const;
+	void UnPinSelection();
 	void AddPinned(const FPinnedTrail& InPinnedTrail);
 	void PinComponent(USceneComponent* InSceneComponent, const FName& InSocketName) const;
 	void DeletePinned(int32 Index);
+	void DeleteAllPinned();
 	void PutPinnnedInSpace(int32 Index, AActor* InActor, const FName& InComponentName);
 	void SetLinearColor(int32 Index, const FLinearColor& Color);
 	void SetHasOffset(int32 Index, bool bHasOffset);
@@ -192,6 +198,9 @@ public:
 	DECLARE_MULTICAST_DELEGATE(FOnPinSelection);
 	FOnPinSelection OnPinSelection;
 
+	DECLARE_MULTICAST_DELEGATE(FOnUnPinSelection);
+	FOnUnPinSelection OnUnPinSelection;
+
 	DECLARE_MULTICAST_DELEGATE_OneParam(FOnAddPinned, FGuid);
 	FOnAddPinned OnAddPinned;
 
@@ -201,6 +210,9 @@ public:
 	DECLARE_MULTICAST_DELEGATE_OneParam(FOnDeletePinned, FGuid);
 	FOnDeletePinned OnDeletePinned;
 
+	DECLARE_MULTICAST_DELEGATE(FOnDeleteAllPinned);
+	FOnDeleteAllPinned OnDeleteAllPinned;
+
 	DECLARE_MULTICAST_DELEGATE_ThreeParams(FOnPutPinnedInSpace, FGuid, AActor*, FName);
 	FOnPutPinnedInSpace OnPutPinnedInSpace;
 
