Anim In Engine: Improved Anim Details

* All selected properties change their value relatively when scrubbing property values
* All selected properties adopt a value when entering values
* Math operations can be applied to all selected properties. Supported operators: add (=+, +=), subtract (=-, -=), multiply (*, =*, *=), divide (/, =/, /=)
* Value cells can now be multi selected
* Selection can now be painted
* Added an option to set the number of fractional digits displayed
* The top level category row can no longer be collapsed
* Clicking the top level category clears selection
* Fixed issues with selection

diff --git a/Engine/Plugins/Animation/ControlRig/Source/ControlRigEditor/Private/AnimDetails/AnimDetailsMathOperation.h b/Engine/Plugins/Animation/ControlRig/Source/ControlRigEditor/Private/AnimDetails/AnimDetailsMathOperation.h
new file mode 100644
--- /dev/null
+++ b/Engine/Plugins/Animation/ControlRig/Source/ControlRigEditor/Private/AnimDetails/AnimDetailsMathOperation.h
@@ -0,0 +1,149 @@
+// Copyright Epic Games, Inc. All Rights Reserved.
+
+#pragma once
+
+#include "Containers/Map.h"
+#include "Containers/UnrealString.h"
+#include "HAL/Platform.h"
+#include "UObject/NameTypes.h"
+
+namespace UE::ControlRigEditor
+{
+	/** Math operators for a control proxy value. */
+	enum class EAnimDetailsMathOperator : uint8
+	{
+		None,
+		Add,
+		Substract,
+		Multiply,
+		Divide
+	};
+
+	/** Defines a mathematical operation by an operator and a RHS value */
+	template <typename NumericType>
+	struct FAnimDetailsMathOperation
+	{
+		FAnimDetailsMathOperation(EAnimDetailsMathOperator InMathOperator, NumericType InRHSValue)
+			: MathOperator(InMathOperator)
+			, RHSValue(InRHSValue)
+		{}
+
+		const EAnimDetailsMathOperator MathOperator;
+		const NumericType RHSValue;
+	};
+
+	namespace FAnimDetailsMathParser
+	{
+		/** 
+		 * Extracts a math operation from a string. The result is optional and only set if a valid math operation could be parsed.
+		 *
+		 * Leading valid LHS values are accepted, but not parsed.
+		 * 
+		 * Currently supported operations, whereas x is an arbitrary number of numeric type:
+		 * Add: =+x or +=x
+		 * Substract: =-x or -=x
+		 * Multiply: *x or =*x or *=x
+		 * Divide: /x or =/x or /=x
+		 */
+		template <typename NumericType>
+		TOptional<FAnimDetailsMathOperation<NumericType>> FromString(const FString& InString)
+		{
+			// Map of operator names with their related math operator
+			static const TMap<FName, EAnimDetailsMathOperator> MathOperatorNameToMathOperatorMap =
+			{
+				{ "+=", EAnimDetailsMathOperator::Add },
+				{ "=+", EAnimDetailsMathOperator::Add },
+
+				{ "-=", EAnimDetailsMathOperator::Substract },
+				{ "=-", EAnimDetailsMathOperator::Substract },
+
+				{ "*", EAnimDetailsMathOperator::Multiply },
+				{ "=*", EAnimDetailsMathOperator::Multiply },
+				{ "*=", EAnimDetailsMathOperator::Multiply },
+
+				{ "/", EAnimDetailsMathOperator::Divide },
+				{ "=/", EAnimDetailsMathOperator::Divide },
+				{ "/=", EAnimDetailsMathOperator::Divide },
+			};
+
+			TArray<FString> NoWhiteSpaceArr;
+			InString.ParseIntoArrayWS(NoWhiteSpaceArr);
+
+			FString TrimmedString;
+			for (const FString& NoWhiteSpaceString : NoWhiteSpaceArr)
+			{
+				TrimmedString += NoWhiteSpaceString;
+			}
+
+			if (TrimmedString.IsEmpty())
+			{
+				return TOptional<FAnimDetailsMathOperation<NumericType>>();
+			}
+
+			// Remove any valid LHS value
+			for (const TTuple<FName, EAnimDetailsMathOperator>& MathOperatorNameToMathOperatorPair : MathOperatorNameToMathOperatorMap)
+			{
+				const int32 OperatorIndex = TrimmedString.Find(MathOperatorNameToMathOperatorPair.Key.ToString());
+				
+				if (OperatorIndex == INDEX_NONE)
+				{
+					// The string does not contain this operator.
+					continue;
+				}
+				else if (OperatorIndex == 0)
+				{
+					// Starts with an operator, there's no LHS value.
+					break;
+				}
+				else
+				{
+					// Try chop the LHS value
+					const FString LeadingNumberString = TrimmedString.Left(OperatorIndex + 1);
+					NumericType LeadingNumber;
+					if (LexTryParseString(LeadingNumber, *LeadingNumberString))
+					{
+						TrimmedString = TrimmedString.RightChop(OperatorIndex);
+						break;
+					}
+					else
+					{
+						// There is no LHS Value but another leading string
+						return TOptional<FAnimDetailsMathOperation<NumericType>>();
+					}
+				}
+			}
+
+			int32 OperatorSize = -1;
+			const EAnimDetailsMathOperator* OperatorPtr = nullptr;
+			
+			// Find single char operators, e.g. '*'
+			OperatorPtr = MathOperatorNameToMathOperatorMap.Find(*TrimmedString.Left(1));
+			if (OperatorPtr)
+			{
+				OperatorSize = 1;
+			}
+			else
+			{
+				// Find two char size math operators, e.g. "+="
+				OperatorPtr = MathOperatorNameToMathOperatorMap.Find(*TrimmedString.Left(2));
+				OperatorSize = 2;
+			}
+
+			if (!OperatorPtr)
+			{
+				return TOptional<FAnimDetailsMathOperation<NumericType>>();
+			}
+
+			// Remove the math operator from the string and get the RHS value
+			const FString TrimmedRHSValueString = TrimmedString.RightChop(OperatorSize);
+
+			NumericType RHSValue;
+			if (LexTryParseString<NumericType>(RHSValue, *TrimmedRHSValueString))
+			{
+				return FAnimDetailsMathOperation<NumericType>(*OperatorPtr, RHSValue);
+			}
+
+			return TOptional<FAnimDetailsMathOperation<NumericType>>();
+		}
+	}
+}
diff --git a/Engine/Plugins/Animation/ControlRig/Source/ControlRigEditor/Private/AnimDetails/Customizations/AnimDetailsProxyDetails.cpp b/Engine/Plugins/Animation/ControlRig/Source/ControlRigEditor/Private/AnimDetails/Customizations/AnimDetailsProxyDetails.cpp
new file mode 100644
--- /dev/null
+++ b/Engine/Plugins/Animation/ControlRig/Source/ControlRigEditor/Private/AnimDetails/Customizations/AnimDetailsProxyDetails.cpp
@@ -0,0 +1,163 @@
+// Copyright Epic Games, Inc. All Rights Reserved.
+
+#include "AnimDetailsProxyDetails.h"
+
+#include "AnimDetails/AnimDetailsProxyManager.h"
+#include "AnimDetails/AnimDetailsSelection.h"
+#include "AnimDetails/Proxies/AnimDetailsProxyBase.h"
+#include "AnimDetails/Widgets/SAnimDetailsPropertySelectionBorder.h"
+#include "DetailCategoryBuilder.h"
+#include "DetailLayoutBuilder.h"
+#include "DetailWidgetRow.h"
+#include "EditMode/ControlRigEditMode.h"
+#include "EditorModeManager.h"
+#include "IDetailGroup.h"
+
+#define LOCTEXT_NAMESPACE "AnimDetailsProxyDetails"
+
+namespace UE::ControlRigEditor
+{
+	FAnimDetailsProxyDetails::FAnimDetailsProxyDetails()
+	{}
+
+	TSharedRef<IDetailCustomization> FAnimDetailsProxyDetails::MakeInstance()
+	{
+		return MakeShared<FAnimDetailsProxyDetails>();
+	}
+
+	void FAnimDetailsProxyDetails::CustomizeDetails(IDetailLayoutBuilder& DetailBuilder)
+	{
+		TArray<TWeakObjectPtr<UObject>> EditedObjects;
+		DetailBuilder.GetObjectsBeingCustomized(EditedObjects);
+
+		bool bIsIndividual = false;
+		TArray<UAnimDetailsProxyBase*> AllProxies;
+		for (int32 ProxyIndex = 0; ProxyIndex < EditedObjects.Num(); ProxyIndex++)
+		{
+			UAnimDetailsProxyBase* Proxy = Cast<UAnimDetailsProxyBase>(EditedObjects[ProxyIndex].Get());
+			if (Proxy)
+			{
+				AllProxies.Add(Proxy);
+
+				bIsIndividual |= Proxy->bIsIndividual;
+			}
+		}
+
+		if (AllProxies.IsEmpty())
+		{
+			return;
+		}
+		const UAnimDetailsProxyBase* FirstProxy = AllProxies[0];
+
+		// Find the category name text
+		FText DisplayNameText;
+		if (AllProxies.Num() == 1 && AllProxies[0])
+		{
+			DisplayNameText = AllProxies[0]->GetDisplayNameText();
+			AllProxies[0]->UpdatePropertyDisplayNames(DetailBuilder);
+		}
+		else
+		{
+			FString DisplayString = TEXT("Multiple");
+			DisplayNameText = FText::FromString(*DisplayString);
+		}
+
+		// Create a custom row to display the header instead of using the category row, so it cannot be collapsed
+		IDetailCategoryBuilder& NoCategory = DetailBuilder.EditCategory("NoCategory");
+
+		if (!bIsIndividual)
+		{
+			const FText& NoFilterText = FText::GetEmpty();
+			NoCategory.AddCustomRow(NoFilterText)
+				.WholeRowContent()
+				[
+					SNew(SBorder)
+					.VAlign(VAlign_Center)
+					.OnMouseButtonDown(this, &FAnimDetailsProxyDetails::OnHeaderCategoryRowClicked)
+					.BorderImage(FAppStyle::Get().GetBrush("DetailsView.CategoryTop"))
+					[
+						SNew(STextBlock)
+						.Text(DisplayNameText)
+						.Font(FAppStyle::Get().GetFontStyle("DetailsView.CategoryFontStyle"))
+						.TextStyle(FAppStyle::Get(), "DetailsView.CategoryTextStyle")
+					]
+				];
+		}
+
+		// Add properties to anim details specific categories instead of the default category
+		const FName CategoryName = FirstProxy->GetCategoryName();
+
+		IDetailCategoryBuilder& DefaultCategory = DetailBuilder.EditCategory(CategoryName);
+		TArray<TSharedRef<IPropertyHandle>> PropertyHandles;
+		DefaultCategory.GetDefaultProperties(PropertyHandles);
+		
+		constexpr bool bForAdvanced = false;
+		constexpr bool bStartExpanded = true;
+		IDetailGroup* IndividualPropertyGroupPtr = bIsIndividual ? &DefaultCategory.AddGroup(CategoryName, DisplayNameText, bForAdvanced, bStartExpanded) : nullptr;
+		
+		for (const TSharedRef<IPropertyHandle>& PropertyHandle : PropertyHandles)
+		{
+			uint32 NumChildren;
+			PropertyHandle->GetNumChildren(NumChildren);
+
+			if (bIsIndividual && NumChildren < 2)
+			{
+				if (const TSharedPtr<IPropertyHandle>& ChildHandle = PropertyHandle->GetChildHandle(0))
+				{
+					// Hack to turn off shift copy/paste
+					ChildHandle->GetProperty()->SetPropertyFlags(CPF_TextExportTransient); 
+
+					DefaultCategory.SetCategoryVisibility(false);
+
+					TArray<UObject*> OuterObjects;
+					ChildHandle->GetOuterObjects(OuterObjects);
+
+					// Show the name of the control instead of the property. 
+					// Otherwise the proxy property name would be e.g. "Bool" when multi editing two bool controls (the default details view behaviour)
+					const FName PropertyName = ChildHandle->GetProperty()->GetFName();
+
+					FText DisplayName;
+					TOptional<FText> StructDisplayName;
+					FirstProxy->GetLocalizedPropertyName(PropertyName, DisplayName, StructDisplayName);
+
+					NoCategory.AddProperty(ChildHandle)
+						.DisplayName(DisplayName);
+
+				}
+			}
+			else
+			{
+				// Add properties with more than one child property, e.g. vectors.
+				// Note, we don't need a selection border here. Selection for structs with more than one 
+				// property is handled in their customization, see FAnimDetailValueCustomization.			
+				if (IndividualPropertyGroupPtr)
+				{
+					IndividualPropertyGroupPtr->AddPropertyRow(PropertyHandle);
+				}
+				else
+				{
+					DefaultCategory.SetCategoryVisibility(false);
+
+					NoCategory.AddProperty(PropertyHandle);
+					NoCategory.InitiallyCollapsed(false);
+				}
+			}
+		}
+	}
+
+	FReply FAnimDetailsProxyDetails::OnHeaderCategoryRowClicked(const FGeometry& MyGeometry, const FPointerEvent& MouseEvent)
+	{
+		// Reset selection
+		FControlRigEditMode* EditMode = static_cast<FControlRigEditMode*>(GLevelEditorModeTools().GetActiveMode(FControlRigEditMode::ModeName));
+		UAnimDetailsProxyManager* ProxyManager = EditMode ? EditMode->GetAnimDetailsProxyManager() : nullptr;
+		UAnimDetailsSelection* Selection = ProxyManager ? ProxyManager->GetAnimDetailsSelection() : nullptr;
+		if (Selection)
+		{
+			Selection->ClearSelection();
+		}
+
+		return FReply::Handled();
+	}
+}
+
+#undef LOCTEXT_NAMESPACE
diff --git a/Engine/Plugins/Animation/ControlRig/Source/ControlRigEditor/Private/AnimDetails/Customizations/AnimDetailsProxyManagerDetails.cpp b/Engine/Plugins/Animation/ControlRig/Source/ControlRigEditor/Private/AnimDetails/Customizations/AnimDetailsProxyManagerDetails.cpp
new file mode 100644
--- /dev/null
+++ b/Engine/Plugins/Animation/ControlRig/Source/ControlRigEditor/Private/AnimDetails/Customizations/AnimDetailsProxyManagerDetails.cpp
@@ -0,0 +1,79 @@
+// Copyright Epic Games, Inc. All Rights Reserved.
+
+#include "AnimDetailsProxyManagerDetails.h"
+
+#include "AnimDetails/AnimDetailsProxyManager.h"
+#include "AnimDetails/Proxies/AnimDetailsProxyBase.h"
+#include "DetailCategoryBuilder.h"
+#include "DetailLayoutBuilder.h"
+#include "ISequencer.h"
+
+namespace UE::ControlRigEditor
+{
+	TSharedRef<IDetailCustomization> FAnimDetailProxyManagerDetails::MakeInstance()
+	{
+		return MakeShared<FAnimDetailProxyManagerDetails>();
+	}
+
+	void FAnimDetailProxyManagerDetails::CustomizeDetails(IDetailLayoutBuilder& DetailLayout)
+	{
+		TArray<TWeakObjectPtr<UObject>> ObjectsBeingCustomized;
+		DetailLayout.GetObjectsBeingCustomized(ObjectsBeingCustomized);
+
+		for (const TWeakObjectPtr<UObject>& ObjectBeingCustomized : ObjectsBeingCustomized)
+		{
+			UAnimDetailsProxyManager* ProxyManager = Cast<UAnimDetailsProxyManager>(ObjectBeingCustomized.Get());
+			if (!ProxyManager)
+			{
+				continue;
+			}
+
+			FAnimDetailsFilter& Filter = ProxyManager->GetAnimDetailsFilter();
+			const TArray<UAnimDetailsProxyBase*> FilteredProxies = Filter.GetFilteredProxies();
+
+			TMap<FName, TArray<UObject*>> DetailRowIDToGroupedProxiesMap;
+			TMap<FName, TArray<UObject*>> DetailRowIDToIndividualProxiesMap;
+			for (UAnimDetailsProxyBase* Proxy : FilteredProxies)
+			{
+				if (!Proxy)
+				{
+					continue;
+				}
+
+				const FName DetailsRowID = Proxy->GetDetailRowID();
+				if (Proxy->bIsIndividual)
+				{
+					DetailRowIDToIndividualProxiesMap.FindOrAdd(DetailsRowID).Add(Proxy);
+				}
+				else
+				{
+					DetailRowIDToGroupedProxiesMap.FindOrAdd(DetailsRowID).Add(Proxy);
+				}
+			}
+
+			// Mock up the details view as per anim details specifics, whereas grouped proxies show first, followed by individual proxies.
+			IDetailCategoryBuilder& CategoryBuilder_None = DetailLayout.EditCategory("nocategory");
+			for (const TTuple<FName, TArray<UObject*>>& FieldTypeToProxyPair : DetailRowIDToGroupedProxiesMap)
+			{
+				CategoryBuilder_None.AddExternalObjects(
+					FieldTypeToProxyPair.Value,
+					EPropertyLocation::Default,
+					FAddPropertyParams()
+					.HideRootObjectNode(true)
+				);
+			}
+
+			// Note, individual proxies are displayed individually, but are still multi-edited across multiple control rigs
+			IDetailCategoryBuilder& CategoryBuilder_Attributes = DetailLayout.EditCategory("Attributes");
+			for (const TTuple<FName, TArray<UObject*>>& ControlNameToIndividualProxiesPair : DetailRowIDToIndividualProxiesMap)
+			{
+				CategoryBuilder_Attributes.AddExternalObjects(
+					ControlNameToIndividualProxiesPair.Value,
+					EPropertyLocation::Default,
+					FAddPropertyParams()
+					.HideRootObjectNode(true)
+				);
+			}
+		}
+	}
+}
diff --git a/Engine/Plugins/Animation/ControlRig/Source/ControlRigEditor/Private/AnimDetails/Proxies/AnimDetailsProxyBase.cpp b/Engine/Plugins/Animation/ControlRig/Source/ControlRigEditor/Private/AnimDetails/Proxies/AnimDetailsProxyBase.cpp
new file mode 100644
--- /dev/null
+++ b/Engine/Plugins/Animation/ControlRig/Source/ControlRigEditor/Private/AnimDetails/Proxies/AnimDetailsProxyBase.cpp
@@ -0,0 +1,871 @@
+// Copyright Epic Games, Inc. All Rights Reserved.
+
+#include "AnimDetailsProxyBase.h"
+
+#include "AnimDetails/AnimDetailsMultiEditUtil.h"
+#include "AnimDetails/AnimDetailsProxyManager.h"
+#include "AnimDetailsProxyTransform.h"
+#include "AnimDetailsProxyVector2D.h"
+#include "ConstraintsManager.h"
+#include "ControlRig.h"
+#include "EditMode/ControlRigEditMode.h"
+#include "IDetailKeyframeHandler.h"
+#include "ISequencer.h"
+#include "LevelEditorViewport.h"
+#include "MovieSceneCommonHelpers.h"
+#include "MVVM/SectionModelStorageExtension.h"
+#include "MVVM/ViewModels/ChannelModel.h"
+#include "MVVM/ViewModels/SectionModel.h"
+#include "MVVM/ViewModels/SequencerEditorViewModel.h"
+#include "MVVM/ViewModels/TrackModel.h"
+#include "PropertyHandle.h"
+#include "Rigs/RigHierarchyDefines.h"
+#include "ScopedTransaction.h"
+#include "Sequencer/ControlRigSequencerHelpers.h"
+#include "Sequencer/MovieSceneControlRigParameterSection.h"
+#include "Sequencer/MovieSceneControlRigParameterTrack.h"
+#include "SequencerAddKeyOperation.h"
+#include "Tracks/MovieScenePropertyTrack.h"
+#include "Units/RigUnitContext.h"
+
+#include UE_INLINE_GENERATED_CPP_BY_NAME(AnimDetailsProxyBase)
+
+namespace UE::ControlRigEditor
+{
+	namespace KeyUtils
+	{
+		/** Adds a key to the specified track */
+		static void KeyTrack(const TSharedPtr<ISequencer>& Sequencer, UAnimDetailsProxyBase* Proxy, UMovieScenePropertyTrack* Track, EControlRigContextChannelToKey ChannelToKey)
+		{
+			using namespace UE::Sequencer;
+
+			if (!Sequencer.IsValid() || !Proxy || !Track)
+			{
+				return;
+			}
+
+			const FFrameNumber Time = Sequencer->GetLocalTime().Time.FloorToFrame();
+			float Weight = 0.0;
+
+			UMovieSceneSection* Section = Track->FindOrExtendSection(Time, Weight);
+
+			FScopedTransaction PropertyChangedTransaction(NSLOCTEXT("AnimDetailsProxyBase", "KeyProperty", "Key Property"), !GIsTransacting);
+			if (!Section || !Section->TryModify())
+			{
+				PropertyChangedTransaction.Cancel();
+				return;
+			}
+
+			const TSharedPtr<FSequencerEditorViewModel> EditorViewModel = Sequencer->GetViewModel();
+			const FViewModelPtr RootModel = EditorViewModel.IsValid() ? EditorViewModel->GetRootModel() : nullptr;
+			const FSectionModelStorageExtension* SectionModelStorage = RootModel.IsValid() ? RootModel->CastDynamic<FSectionModelStorageExtension>() : nullptr;
+			const TSharedPtr<FSectionModel> SectionHandle = SectionModelStorage ? SectionModelStorage->FindModelForSection(Section) : nullptr;
+			const TSharedPtr<FViewModel> ViewModel = SectionHandle.IsValid() ? SectionHandle->GetParentTrackModel().AsModel() : nullptr;
+			if (!EditorViewModel.IsValid() ||
+				!RootModel.IsValid() ||
+				!SectionModelStorage || 
+				!SectionHandle.IsValid() || 
+				!ViewModel.IsValid())
+			{
+				return;
+			}
+
+			TArray<TSharedRef<IKeyArea>> KeyAreas;
+			const TParentFirstChildIterator<FChannelGroupModel> KeyAreaNodes = ViewModel->GetDescendantsOfType<FChannelGroupModel>();
+			
+			for (const TViewModelPtr<FChannelGroupModel>& KeyAreaNode : KeyAreaNodes)
+			{
+				for (const TWeakViewModelPtr<FChannelModel>& Channel : KeyAreaNode->GetChannels())
+				{
+					if (const TSharedPtr<FChannelModel> ChannelModel = Channel.Pin())
+					{
+						const EControlRigContextChannelToKey ThisChannelToKey = Proxy->GetChannelToKeyFromChannelName(ChannelModel->GetChannelName().ToString());
+						if ((int32)ChannelToKey & (int32)ThisChannelToKey)
+						{
+							KeyAreas.Add(ChannelModel->GetKeyArea().ToSharedRef());
+						}
+					}
+				}
+			}
+
+			const TSharedPtr<FTrackModel> TrackModel = SectionHandle->FindAncestorOfType<FTrackModel>();
+			FAddKeyOperation::FromKeyAreas(TrackModel->GetTrackEditor().Get(), KeyAreas).Commit(Time, *Sequencer);
+		}
+
+		static EPropertyKeyedStatus GetChannelKeyStatus(
+			FMovieSceneChannel* InChannel, 
+			EPropertyKeyedStatus InSectionKeyedStatus, 
+			const TRange<FFrameNumber>& InRange, 
+			int32& OutEmptyChannelCount)
+		{
+			if (!InChannel)
+			{
+				return InSectionKeyedStatus;
+			}
+
+			if (InChannel->GetNumKeys() == 0)
+			{
+				++OutEmptyChannelCount;
+				return InSectionKeyedStatus;
+			}
+
+			InSectionKeyedStatus = FMath::Max(InSectionKeyedStatus, EPropertyKeyedStatus::KeyedInOtherFrame);
+
+			TArray<FFrameNumber> KeyTimes;
+			InChannel->GetKeys(InRange, &KeyTimes, nullptr);
+			if (KeyTimes.IsEmpty())
+			{
+				++OutEmptyChannelCount;
+			}
+			else
+			{
+				InSectionKeyedStatus = FMath::Max(InSectionKeyedStatus, EPropertyKeyedStatus::PartiallyKeyed);
+			}
+
+			return InSectionKeyedStatus;
+		}
+
+		static EPropertyKeyedStatus GetKeyedStatusInSection(
+			const UControlRig* ControlRig, 
+			const FName& ControlName, 
+			const UMovieSceneControlRigParameterSection* Section, 
+			const TRange<FFrameNumber>& Range, 
+			const EControlRigContextChannelToKey ChannelToKey)
+		{
+			int32 EmptyChannelCount = 0;
+			EPropertyKeyedStatus SectionKeyedStatus = EPropertyKeyedStatus::NotKeyed;
+			
+			const FRigControlElement* ControlElement = ControlRig ? ControlRig->FindControl(ControlName) : nullptr;
+			if (!ControlElement)
+			{
+				return SectionKeyedStatus;
+			}
+			
+			switch (ControlElement->Settings.ControlType)
+			{
+			case ERigControlType::Bool:
+			{
+				const TArrayView<FMovieSceneBoolChannel*> BoolChannels = FControlRigSequencerHelpers::GetBoolChannels(ControlRig, ControlElement->GetKey().Name, Section);
+				for (FMovieSceneChannel* Channel : BoolChannels)
+				{
+					SectionKeyedStatus = GetChannelKeyStatus(Channel, SectionKeyedStatus, Range, EmptyChannelCount);
+				}
+
+				break;
+			}
+			case ERigControlType::Integer:
+			{
+				const TArrayView<FMovieSceneIntegerChannel*> IntegarChannels = FControlRigSequencerHelpers::GetIntegerChannels(ControlRig, ControlElement->GetKey().Name, Section);
+				for (FMovieSceneChannel* Channel : IntegarChannels)
+				{
+					SectionKeyedStatus = GetChannelKeyStatus(Channel, SectionKeyedStatus, Range, EmptyChannelCount);
+				}
+
+				const TArrayView<FMovieSceneByteChannel*>  EnumChannels = FControlRigSequencerHelpers::GetByteChannels(ControlRig, ControlElement->GetKey().Name, Section);
+				for (FMovieSceneChannel* Channel : EnumChannels)
+				{
+					SectionKeyedStatus = GetChannelKeyStatus(Channel, SectionKeyedStatus, Range, EmptyChannelCount);
+				}
+
+				break;
+			}
+			case ERigControlType::Position:
+			case ERigControlType::Transform:
+			case ERigControlType::TransformNoScale:
+			case ERigControlType::EulerTransform:
+			case ERigControlType::Float:
+			case ERigControlType::ScaleFloat:
+			case ERigControlType::Vector2D:
+			{
+				const int32 IChannelToKey = (int32)ChannelToKey;
+				const TArrayView<FMovieSceneFloatChannel*> FloatChannels = FControlRigSequencerHelpers::GetFloatChannels(ControlRig, ControlElement->GetKey().Name, Section);
+
+				int32 Num = 0;
+				if (FloatChannels.Num() > Num)
+				{
+					if (IChannelToKey & int32(EControlRigContextChannelToKey::TranslationX))
+					{
+						SectionKeyedStatus = GetChannelKeyStatus(FloatChannels[Num], SectionKeyedStatus, Range, EmptyChannelCount);
+					}
+				}
+				else
+				{
+					break;
+				}
+				++Num;
+
+				if (FloatChannels.Num() > Num)
+				{
+					if (IChannelToKey & int32(EControlRigContextChannelToKey::TranslationY))
+					{
+						SectionKeyedStatus = GetChannelKeyStatus(FloatChannels[Num], SectionKeyedStatus, Range, EmptyChannelCount);
+					}
+				}
+				else
+				{
+					break;
+				}
+				++Num;
+
+				if (FloatChannels.Num() > Num)
+				{
+					if (IChannelToKey & int32(EControlRigContextChannelToKey::TranslationZ))
+					{
+						SectionKeyedStatus = GetChannelKeyStatus(FloatChannels[Num], SectionKeyedStatus, Range, EmptyChannelCount);
+					}
+				}
+				else
+				{
+					break;
+				}
+				++Num;
+
+				if (FloatChannels.Num() > Num)
+				{
+					if (IChannelToKey & int32(EControlRigContextChannelToKey::RotationX))
+					{
+						SectionKeyedStatus = GetChannelKeyStatus(FloatChannels[Num], SectionKeyedStatus, Range, EmptyChannelCount);
+					}
+				}
+				else
+				{
+					break;
+				}
+				++Num;
+
+				if (FloatChannels.Num() > Num)
+				{
+					if (IChannelToKey & int32(EControlRigContextChannelToKey::RotationY))
+					{
+						SectionKeyedStatus = GetChannelKeyStatus(FloatChannels[Num], SectionKeyedStatus, Range, EmptyChannelCount);
+					}
+				}
+				else
+				{
+					break;
+				}
+				++Num;
+
+				if (FloatChannels.Num() > Num)
+				{
+					if (IChannelToKey & int32(EControlRigContextChannelToKey::RotationZ))
+					{
+						SectionKeyedStatus = GetChannelKeyStatus(FloatChannels[Num], SectionKeyedStatus, Range, EmptyChannelCount);
+					}
+				}
+				else
+				{
+					break;
+				}
+				++Num;
+
+				if (FloatChannels.Num() > Num)
+				{
+					if (IChannelToKey & int32(EControlRigContextChannelToKey::ScaleX))
+					{
+						SectionKeyedStatus = GetChannelKeyStatus(FloatChannels[Num], SectionKeyedStatus, Range, EmptyChannelCount);
+					}
+				}
+				else
+				{
+					break;
+				}
+				++Num;
+
+				if (FloatChannels.Num() > Num)
+				{
+					if (IChannelToKey & int32(EControlRigContextChannelToKey::ScaleY))
+					{
+						SectionKeyedStatus = GetChannelKeyStatus(FloatChannels[Num], SectionKeyedStatus, Range, EmptyChannelCount);
+					}
+				}
+				else
+				{
+					break;
+				}
+				++Num;
+
+				if (FloatChannels.Num() > Num)
+				{
+					if (IChannelToKey & int32(EControlRigContextChannelToKey::ScaleZ))
+					{
+						SectionKeyedStatus = GetChannelKeyStatus(FloatChannels[Num], SectionKeyedStatus, Range, EmptyChannelCount);
+					}
+				}
+				break;
+			}
+			case ERigControlType::Scale:
+			{
+
+				int32 IChannelToKey = (int32)ChannelToKey;
+				const TArrayView<FMovieSceneFloatChannel*> FloatChannels = FControlRigSequencerHelpers::GetFloatChannels(ControlRig, ControlElement->GetKey().Name, Section);
+
+				int32 Num = 0;
+				if (FloatChannels.Num() > Num)
+				{
+					if (IChannelToKey & int32(EControlRigContextChannelToKey::ScaleX))
+					{
+						SectionKeyedStatus = GetChannelKeyStatus(FloatChannels[Num], SectionKeyedStatus, Range, EmptyChannelCount);
+					}
+				}
+				else
+				{
+					break;
+				}
+				++Num;
+
+				if (FloatChannels.Num() > Num)
+				{
+					if (IChannelToKey & int32(EControlRigContextChannelToKey::ScaleY))
+					{
+						SectionKeyedStatus = GetChannelKeyStatus(FloatChannels[Num], SectionKeyedStatus, Range, EmptyChannelCount);
+					}
+				}
+				else
+				{
+					break;
+				}
+				++Num;
+
+				if (FloatChannels.Num() > Num)
+				{
+					if (IChannelToKey & int32(EControlRigContextChannelToKey::ScaleZ))
+					{
+						SectionKeyedStatus = GetChannelKeyStatus(FloatChannels[Num], SectionKeyedStatus, Range, EmptyChannelCount);
+					}
+				}
+				break;
+			}
+			case ERigControlType::Rotator:
+			{
+				int32 IChannelToKey = (int32)ChannelToKey;
+				TArrayView<FMovieSceneFloatChannel*> FloatChannels = FControlRigSequencerHelpers::GetFloatChannels(ControlRig, ControlElement->GetKey().Name, Section);
+				int32 Num = 0;
+				if (FloatChannels.Num() > Num)
+				{
+					if (IChannelToKey & int32(EControlRigContextChannelToKey::RotationX))
+					{
+						SectionKeyedStatus = GetChannelKeyStatus(FloatChannels[Num], SectionKeyedStatus, Range, EmptyChannelCount);
+					}
+				}
+				else
+				{
+					break;
+				}
+				++Num;
+
+				if (FloatChannels.Num() > Num)
+				{
+					if (IChannelToKey & int32(EControlRigContextChannelToKey::RotationY))
+					{
+						SectionKeyedStatus = GetChannelKeyStatus(FloatChannels[Num], SectionKeyedStatus, Range, EmptyChannelCount);
+					}
+				}
+				else
+				{
+					break;
+				}
+				++Num;
+
+				if (FloatChannels.Num() > Num)
+				{
+					if (IChannelToKey & int32(EControlRigContextChannelToKey::RotationZ))
+					{
+						SectionKeyedStatus = GetChannelKeyStatus(FloatChannels[Num], SectionKeyedStatus, Range, EmptyChannelCount);
+					}
+				}
+				else
+				{
+					break;
+				}
+				break;
+			}
+			}
+			if (EmptyChannelCount == 0 && SectionKeyedStatus == EPropertyKeyedStatus::PartiallyKeyed)
+			{
+				SectionKeyedStatus = EPropertyKeyedStatus::KeyedInFrame;
+			}
+
+			return SectionKeyedStatus;
+		}
+
+		static EPropertyKeyedStatus GetKeyedStatusInTrack(
+			const UControlRig* ControlRig, 
+			const FName& ControlName, 
+			const UMovieSceneControlRigParameterTrack* Track, 
+			const TRange<FFrameNumber>& Range,
+			const EControlRigContextChannelToKey ChannelToKey)
+		{
+			EPropertyKeyedStatus SectionKeyedStatus = EPropertyKeyedStatus::NotKeyed;
+
+			const FRigControlElement* ControlElement = ControlRig ? ControlRig->FindControl(ControlName) : nullptr;
+			if (!ControlElement)
+			{
+				return SectionKeyedStatus;
+			}
+
+			for (const UMovieSceneSection* BaseSection : Track->GetAllSections())
+			{
+				const UMovieSceneControlRigParameterSection* Section = Cast<UMovieSceneControlRigParameterSection>(BaseSection);
+				if (!Section)
+				{
+					continue;
+				}
+				EPropertyKeyedStatus NewSectionKeyedStatus = GetKeyedStatusInSection(ControlRig, ControlName, Section, Range, ChannelToKey);
+				SectionKeyedStatus = FMath::Max(SectionKeyedStatus, NewSectionKeyedStatus);
+
+				// Maximum Status Reached no need to iterate further
+				if (SectionKeyedStatus == EPropertyKeyedStatus::KeyedInFrame)
+				{
+					return SectionKeyedStatus;
+				}
+			}
+
+			return SectionKeyedStatus;
+		}
+
+		static EPropertyKeyedStatus GetKeyedStatusInSection(
+			const UMovieSceneSection* Section, 
+			const TRange<FFrameNumber>& Range, 
+			const EControlRigContextChannelToKey ChannelToKey, 
+			const int32 MaxNumIndices)
+		{
+			EPropertyKeyedStatus SectionKeyedStatus = EPropertyKeyedStatus::NotKeyed;
+
+			FMovieSceneChannelProxy& ChannelProxy = Section->GetChannelProxy();
+
+			TArray<int32, TFixedAllocator<3>> ChannelIndices;
+			switch (ChannelToKey)
+			{
+			case EControlRigContextChannelToKey::Translation:
+				ChannelIndices = { 0, 1, 2 };
+				break;
+			case EControlRigContextChannelToKey::TranslationX:
+				ChannelIndices = { 0 };
+				break;
+			case EControlRigContextChannelToKey::TranslationY:
+				ChannelIndices = { 1 };
+				break;
+			case EControlRigContextChannelToKey::TranslationZ:
+				ChannelIndices = { 2 };
+				break;
+			case EControlRigContextChannelToKey::Rotation:
+				ChannelIndices = { 3, 4, 5 };
+				break;
+			case EControlRigContextChannelToKey::RotationX:
+				ChannelIndices = { 3 };
+				break;
+			case EControlRigContextChannelToKey::RotationY:
+				ChannelIndices = { 4 };
+				break;
+			case EControlRigContextChannelToKey::RotationZ:
+				ChannelIndices = { 5 };
+				break;
+			case EControlRigContextChannelToKey::Scale:
+				ChannelIndices = { 6, 7, 8 };
+				break;
+			case EControlRigContextChannelToKey::ScaleX:
+				ChannelIndices = { 6 };
+				break;
+			case EControlRigContextChannelToKey::ScaleY:
+				ChannelIndices = { 7 };
+				break;
+			case EControlRigContextChannelToKey::ScaleZ:
+				ChannelIndices = { 8 };
+				break;
+			}
+
+			int32 EmptyChannelCount = 0;
+			for (const FMovieSceneChannelEntry& ChannelEntry : ChannelProxy.GetAllEntries())
+			{
+				if (ChannelEntry.GetChannelTypeName() != FMovieSceneDoubleChannel::StaticStruct()->GetFName() &&
+					ChannelEntry.GetChannelTypeName() != FMovieSceneFloatChannel::StaticStruct()->GetFName() &&
+					ChannelEntry.GetChannelTypeName() != FMovieSceneBoolChannel::StaticStruct()->GetFName() &&
+					ChannelEntry.GetChannelTypeName() != FMovieSceneIntegerChannel::StaticStruct()->GetFName() &&
+					ChannelEntry.GetChannelTypeName() != FMovieSceneByteChannel::StaticStruct()->GetFName())
+				{
+					continue;
+				}
+
+				const TConstArrayView<FMovieSceneChannel*> Channels = ChannelEntry.GetChannels();
+
+				int32 ChannelIndex = 0;
+				for (FMovieSceneChannel* Channel : Channels)
+				{
+					if (ChannelIndex >= MaxNumIndices)
+					{
+						break;
+					}
+					if (ChannelIndices.Contains(ChannelIndex++) == false)
+					{
+						continue;
+					}
+					if (Channel->GetNumKeys() == 0)
+					{
+						++EmptyChannelCount;
+						continue;
+					}
+
+					SectionKeyedStatus = FMath::Max(SectionKeyedStatus, EPropertyKeyedStatus::KeyedInOtherFrame);
+
+					TArray<FFrameNumber> KeyTimes;
+					Channel->GetKeys(Range, &KeyTimes, nullptr);
+					if (KeyTimes.IsEmpty())
+					{
+						++EmptyChannelCount;
+					}
+					else
+					{
+						SectionKeyedStatus = FMath::Max(SectionKeyedStatus, EPropertyKeyedStatus::PartiallyKeyed);
+					}
+				}
+
+				break; // just do it for one type
+			}
+
+			if (EmptyChannelCount == 0 && SectionKeyedStatus == EPropertyKeyedStatus::PartiallyKeyed)
+			{
+				SectionKeyedStatus = EPropertyKeyedStatus::KeyedInFrame;
+			}
+
+			return SectionKeyedStatus;
+		}
+
+		static EPropertyKeyedStatus GetKeyedStatusInTrack(
+			const UMovieScenePropertyTrack* Track, 
+			const TRange<FFrameNumber>& Range, 
+			const EControlRigContextChannelToKey ChannelToKey, 
+			const int32 MaxNumIndices)
+		{
+			EPropertyKeyedStatus SectionKeyedStatus = EPropertyKeyedStatus::NotKeyed;
+			for (UMovieSceneSection* BaseSection : Track->GetAllSections())
+			{
+				if (!BaseSection)
+				{
+					continue;
+				}
+				EPropertyKeyedStatus NewSectionKeyedStatus = GetKeyedStatusInSection(BaseSection, Range, ChannelToKey, MaxNumIndices);
+				SectionKeyedStatus = FMath::Max(SectionKeyedStatus, NewSectionKeyedStatus);
+
+				// Maximum Status Reached no need to iterate further
+				if (SectionKeyedStatus == EPropertyKeyedStatus::KeyedInFrame)
+				{
+					return SectionKeyedStatus;
+				}
+			}
+
+			return SectionKeyedStatus;
+		}
+	}
+}
+
+void UAnimDetailsProxyBase::SetControlFromControlRig(UControlRig* InControlRig, const FName& InName)
+{
+	SequencerItem.Reset();
+
+	URigHierarchy* Hierarchy = InControlRig ? InControlRig->GetHierarchy() : nullptr;
+	if (Hierarchy)
+	{
+		WeakControlRig = InControlRig;
+		CachedRigElement = FCachedRigElement(FRigElementKey(InName, ERigElementType::Control), Hierarchy);
+	}
+}
+
+void UAnimDetailsProxyBase::SetControlFromSequencerBinding(UObject* InObject, const TWeakObjectPtr<UMovieSceneTrack>& InTrack, const TSharedPtr<FTrackInstancePropertyBindings>& InBinding)
+{
+	WeakControlRig.Reset();
+	CachedRigElement.Reset();
+
+	if (InObject &&
+		InTrack.IsValid() &&
+		InBinding.IsValid())
+	{
+		SequencerItem = FAnimDetailsSequencerProxyItem(*InObject, *InTrack.Get(), InBinding.ToSharedRef());
+	}
+	else
+	{
+		SequencerItem.Reset();
+	}
+}
+
+UControlRig* UAnimDetailsProxyBase::GetControlRig() const
+{
+	return WeakControlRig.Get();
+}
+
+FRigControlElement* UAnimDetailsProxyBase::GetControlElement() const
+{
+	const URigHierarchy* Hierarchy = WeakControlRig.IsValid() ? WeakControlRig->GetHierarchy() : nullptr;
+	const FRigControlElement* ControlElement = Hierarchy ? Cast<FRigControlElement>(CachedRigElement.GetElement(Hierarchy)) : nullptr;
+
+	// @todo There is no specific reason to prevent from getting a non const ptr from FCachedRigElement. 
+	// The related change turns out relatively large, so delay it for later and const cast here.
+	return const_cast<FRigControlElement*>(ControlElement);
+}
+
+FRigElementKey UAnimDetailsProxyBase::GetControlElementKey() const
+{
+	return CachedRigElement.GetResolvedKey();
+}
+
+void UAnimDetailsProxyBase::PropagonateValues()
+{
+	FRigControlModifiedContext Context;
+	Context.SetKey = EControlRigSetKey::DoNotCare;
+
+	FRigControlModifiedContext NotifyDrivenContext;
+
+	UWorld* World = GCurrentLevelEditingViewportClient ? GCurrentLevelEditingViewportClient->GetWorld() : nullptr;
+	const FConstraintsManagerController& Controller = FConstraintsManagerController::Get(World);
+	Controller.EvaluateAllConstraints();
+
+	UControlRig* ControlRig = GetControlRig();
+	FRigControlElement* ControlElement = GetControlElement();
+	if (ControlRig && ControlElement)
+	{
+		SetControlRigElementValueFromCurrent(ControlRig, ControlElement, Context);
+		FControlRigEditMode::NotifyDrivenControls(ControlRig, ControlElement->GetKey(), NotifyDrivenContext);
+			
+		ControlRig->Evaluate_AnyThread();
+	}
+	else if (SequencerItem.IsValid())
+	{
+		SetBindingValueFromCurrent(SequencerItem.GetBoundObject(), SequencerItem.GetBinding(), Context);
+	}
+}
+
+FText UAnimDetailsProxyBase::GetDisplayNameText() const
+{
+	const URigHierarchy* Hierarchy = WeakControlRig.IsValid() ? WeakControlRig->GetHierarchy() : nullptr;
+	const FRigControlElement* ControlElement = WeakControlRig.IsValid() ? GetControlElement() : nullptr;
+	if (Hierarchy && ControlElement)
+	{
+		return Hierarchy->GetDisplayNameForUI(ControlElement);
+	}
+	else if (const UObject* BoundObject = SequencerItem.GetBoundObject())
+	{
+		if (const AActor* Actor = Cast<AActor>(BoundObject))
+		{
+			return FText::FromString(Actor->GetActorLabel());
+		}
+		else if (const UActorComponent* Component = Cast<UActorComponent>(BoundObject))
+		{
+			return FText::FromString(*Component->GetName());
+		}
+	}
+
+	return FText::GetEmpty();
+}
+
+void UAnimDetailsProxyBase::SetKey(const IPropertyHandle& KeyedPropertyHandle)
+{
+	using namespace UE::ControlRigEditor;
+
+	const UAnimDetailsProxyManager* ProxyManager = GetTypedOuter<UAnimDetailsProxyManager>();
+	const TSharedPtr<ISequencer> Sequencer = ProxyManager ? ProxyManager->GetSequencer() : nullptr;
+	if (!Sequencer.IsValid() || !Sequencer->GetFocusedMovieSceneSequence())
+	{
+		return;
+	}
+
+	UControlRig* ControlRig = GetControlRig();
+	FRigControlElement* ControlElement = GetControlElement();
+
+	if (ControlRig && ControlElement)
+	{
+		if (ControlRig->GetHierarchy()->Contains(FRigElementKey(ControlElement->GetKey().Name, ERigElementType::Control)))
+		{
+			const FName PropertyName = KeyedPropertyHandle.GetProperty()->GetFName();
+
+			FRigControlModifiedContext Context;
+			Context.SetKey = EControlRigSetKey::Always;
+			Context.KeyMask = (uint32)GetChannelToKeyFromPropertyName(PropertyName);
+
+			SetControlRigElementValueFromCurrent(ControlRig, ControlElement, Context);
+
+			FRigControlModifiedContext NotifyDrivenContext; // always key ever
+			NotifyDrivenContext.SetKey = EControlRigSetKey::Always;
+
+			FControlRigEditMode::NotifyDrivenControls(ControlRig, ControlElement->GetKey(), NotifyDrivenContext);
+		}
+	}
+	else if (UMovieScenePropertyTrack* PropertyTrack = Cast<UMovieScenePropertyTrack>(SequencerItem.GetMovieSceneTrack()))
+	{
+		const FName PropertyName = KeyedPropertyHandle.GetProperty()->GetFName();
+		const EControlRigContextChannelToKey ChannelToKey = GetChannelToKeyFromPropertyName(PropertyName);
+
+		KeyUtils::KeyTrack(Sequencer, this, PropertyTrack, ChannelToKey);
+	}
+}
+
+EPropertyKeyedStatus UAnimDetailsProxyBase::GetPropertyKeyedStatus(const IPropertyHandle& PropertyHandle) const
+{
+	using namespace UE::ControlRigEditor;
+
+	EPropertyKeyedStatus KeyedStatus = EPropertyKeyedStatus::NotKeyed;
+
+	const UAnimDetailsProxyManager* ProxyManager = GetTypedOuter<UAnimDetailsProxyManager>();
+	const TSharedPtr<ISequencer> Sequencer = ProxyManager ? ProxyManager->GetSequencer() : nullptr;
+	if (!Sequencer.IsValid() || !Sequencer->GetFocusedMovieSceneSequence())
+	{
+		return KeyedStatus;
+	}
+
+	const UControlRig* ControlRig = GetControlRig();
+	const URigHierarchy* Hierarchy = ControlRig ? ControlRig->GetHierarchy() : nullptr;
+	const FRigElementKey ControlElementKey = GetControlElementKey();
+
+	const TRange<FFrameNumber> FrameRange = TRange<FFrameNumber>(Sequencer->GetLocalTime().Time.FrameNumber);
+	const FName PropertyName = PropertyHandle.GetProperty()->GetFName();
+	const EControlRigContextChannelToKey ChannelToKey = GetChannelToKeyFromPropertyName(PropertyName);
+
+	UMovieSceneControlRigParameterTrack* ControlRigParameterTrack = ControlRig ? FControlRigSequencerHelpers::FindControlRigTrack(Sequencer->GetFocusedMovieSceneSequence(), ControlRig) : nullptr;
+	if (ControlRig && ControlElementKey.IsValid() && ControlRigParameterTrack)
+	{
+		const EPropertyKeyedStatus NewKeyedStatus = KeyUtils::GetKeyedStatusInTrack(ControlRig, ControlElementKey.Name, ControlRigParameterTrack, FrameRange, ChannelToKey);
+		KeyedStatus = FMath::Max(KeyedStatus, NewKeyedStatus);
+	}
+	else if (UMovieScenePropertyTrack* MovieScenePropertyTrack = Cast<UMovieScenePropertyTrack>(SequencerItem.GetMovieSceneTrack()))
+	{
+		int32 MaxNumIndices = 1;
+		if (IsA<UAnimDetailsProxyTransform>())
+		{
+			MaxNumIndices = 9;
+		}
+		else if (IsA<UAnimDetailsProxyVector2D>())
+		{
+			MaxNumIndices = 2;
+		}
+
+		const EPropertyKeyedStatus NewKeyedStatus = KeyUtils::GetKeyedStatusInTrack(MovieScenePropertyTrack, FrameRange, ChannelToKey, MaxNumIndices);
+		KeyedStatus = FMath::Max(KeyedStatus, NewKeyedStatus);
+	}
+
+	return KeyedStatus;
+}
+
+FName UAnimDetailsProxyBase::GetDetailRowID() const
+{
+	if (bIsIndividual)
+	{
+		const FRigControlElement* ControlElement = GetControlElement();
+		if (ControlElement)
+		{
+			return ControlElement->GetKey().Name;
+		}
+		else if (UMovieSceneTrack* MovieSceneTrack = SequencerItem.GetMovieSceneTrack())
+		{
+			return *MovieSceneTrack->GetPathName();
+		}
+		else
+		{
+			ensureMsgf(0, TEXT("Getting a detail row ID for an invalid proxy. Returning NAME_None as its ID"));
+			return NAME_None;
+		}
+	}
+	else
+	{
+		return GetClass()->GetFName();
+	}
+}
+
+void UAnimDetailsProxyBase::GetLocalizedPropertyName(const FName& InPropertyName, FText& OutPropertyDisplayName, TOptional<FText>& OutOptionalStructDisplayName) const
+{
+	// Proxies with one member in their struct only show inner properties
+	OutOptionalStructDisplayName.Reset();
+
+	const URigHierarchy* Hierarchy = WeakControlRig.IsValid() ? WeakControlRig->GetHierarchy() : nullptr;
+	const FRigControlElement* ControlElement = WeakControlRig.IsValid() ? GetControlElement() : nullptr;
+	if (Hierarchy && ControlElement)
+	{
+		OutPropertyDisplayName = Hierarchy->GetDisplayNameForUI(ControlElement);
+	}
+	else if (const TSharedPtr<FTrackInstancePropertyBindings>& Binding = SequencerItem.GetBinding())
+	{
+		OutPropertyDisplayName = FText::FromName(SequencerItem.GetBinding()->GetPropertyName());
+	}
+}
+
+FName UAnimDetailsProxyBase::GetPropertyID(const FName& PropertyName) const
+{
+	return *(GetDetailRowID().ToString() + TEXT(".") + PropertyName.ToString());
+}
+
+void UAnimDetailsProxyBase::PostEditChangeChainProperty(FPropertyChangedChainEvent& PropertyChangedEvent)
+{
+	Super::PostEditChangeChainProperty(PropertyChangedEvent);
+
+	if (PropertyChangedEvent.ChangeType == EPropertyChangeType::ToggleEditable) // hack so we can clear the reset cache for this property and not actually send this to our controls
+	{
+		return;
+	}
+
+	if (const FProperty* Property = PropertyChangedEvent.Property)
+	{
+		const FProperty* MemberProperty = PropertyChangedEvent.PropertyChain.GetActiveMemberNode() ? 
+			PropertyChangedEvent.PropertyChain.GetActiveMemberNode()->GetValue() : 
+			nullptr;
+
+		if (PropertyIsOnProxy(Property, MemberProperty))
+		{
+			EControlRigContextChannelToKey ChannelToKeyContext = GetChannelToKeyFromPropertyName(Property->GetFName());		
+			AddControlRigInteractionScope(ChannelToKeyContext, PropertyChangedEvent.ChangeType);
+
+			PropagonateValues();
+
+			if (PropertyChangedEvent.ChangeType != EPropertyChangeType::Interactive)
+			{
+				InteractionScopes.Reset();
+			}
+
+			UpdateProxyValues();
+		}
+	}
+}
+
+bool UAnimDetailsProxyBase::Modify(bool bAlwaysMarkDirty)
+{
+	// IPropertyHandle::SetPerObjectValues which the multi edit util uses always modifies the object, 
+	// hence we avoid modificiation by testing for interactive changes here.
+	if (!UE::ControlRigEditor::FAnimDetailsMultiEditUtil::Get().IsInteractive())
+	{
+		return Super::Modify(bAlwaysMarkDirty);
+	}
+
+	return true;
+}
+
+void UAnimDetailsProxyBase::AddControlRigInteractionScope(EControlRigContextChannelToKey ChannelsToKey, EPropertyChangeType::Type ChangeType)
+{
+	if (ChangeType == EPropertyChangeType::Interactive || ChangeType == EPropertyChangeType::ValueSet)
+	{
+		EControlRigInteractionType InteractionType = EControlRigInteractionType::None;
+		if (EnumHasAnyFlags(ChannelsToKey, EControlRigContextChannelToKey::TranslationX) ||
+			EnumHasAnyFlags(ChannelsToKey, EControlRigContextChannelToKey::TranslationY) ||
+			EnumHasAnyFlags(ChannelsToKey, EControlRigContextChannelToKey::TranslationZ))
+		{
+			EnumAddFlags(InteractionType, EControlRigInteractionType::Translate);
+		}
+		if (EnumHasAnyFlags(ChannelsToKey, EControlRigContextChannelToKey::RotationX) ||
+			EnumHasAnyFlags(ChannelsToKey, EControlRigContextChannelToKey::RotationY) ||
+			EnumHasAnyFlags(ChannelsToKey, EControlRigContextChannelToKey::RotationZ))
+		{
+			EnumAddFlags(InteractionType, EControlRigInteractionType::Rotate);
+		}
+		if (EnumHasAnyFlags(ChannelsToKey, EControlRigContextChannelToKey::ScaleX) ||
+			EnumHasAnyFlags(ChannelsToKey, EControlRigContextChannelToKey::ScaleY) ||
+			EnumHasAnyFlags(ChannelsToKey, EControlRigContextChannelToKey::ScaleZ))
+		{
+			EnumAddFlags(InteractionType, EControlRigInteractionType::Scale);
+		}
+
+		UControlRig* ControlRig = GetControlRig();
+		FRigControlElement* ControlElement = GetControlElement();
+
+		if (ControlRig && ControlElement && !InteractionScopes.Contains(ControlElement))
+		{
+			const TSharedRef<FControlRigInteractionScope> InteractionScope = MakeShared<FControlRigInteractionScope>(ControlRig, ControlElement->GetKey(), InteractionType);
+			InteractionScopes.Add(ControlElement, InteractionScope);
+		}
+	}
+}
diff --git a/Engine/Plugins/Animation/ControlRig/Source/ControlRigEditor/Private/AnimDetails/Proxies/AnimDetailsProxyBool.cpp b/Engine/Plugins/Animation/ControlRig/Source/ControlRigEditor/Private/AnimDetails/Proxies/AnimDetailsProxyBool.cpp
new file mode 100644
--- /dev/null
+++ b/Engine/Plugins/Animation/ControlRig/Source/ControlRigEditor/Private/AnimDetails/Proxies/AnimDetailsProxyBool.cpp
@@ -0,0 +1,134 @@
+// Copyright Epic Games, Inc. All Rights Reserved.
+
+#include "AnimDetailsProxyBool.h"
+
+#include "ControlRig.h"
+#include "DetailLayoutBuilder.h"
+#include "MovieSceneCommonHelpers.h"
+
+#include UE_INLINE_GENERATED_CPP_BY_NAME(AnimDetailsProxyBool)
+
+FName UAnimDetailsProxyBool::GetCategoryName() const
+{
+	return "Bool";
+}
+
+void UAnimDetailsProxyBool::UpdatePropertyDisplayNames(IDetailLayoutBuilder& DetailBuilder)
+{
+	TSharedPtr<IPropertyHandle> ValuePropertyHandle = DetailBuilder.GetProperty(GET_MEMBER_NAME_CHECKED(UAnimDetailsProxyBool, Bool), GetClass());
+	if (ValuePropertyHandle.IsValid())
+	{
+		ValuePropertyHandle->SetPropertyDisplayName(GetDisplayNameText());
+	}
+
+	if (bIsIndividual)
+	{
+		const FString PropertyPath = FString::Printf(TEXT("%s.%s"), GET_MEMBER_NAME_STRING_CHECKED(UAnimDetailsProxyBool, Bool), GET_MEMBER_NAME_STRING_CHECKED(FAnimDetailsBool, Bool));
+		ValuePropertyHandle = DetailBuilder.GetProperty(FName(*PropertyPath), GetClass());
+		if (ValuePropertyHandle.IsValid())
+		{
+			ValuePropertyHandle->SetPropertyDisplayName(GetDisplayNameText());
+		}
+	}
+	else
+	{
+		const FString PropertyPath = FString::Printf(TEXT("%s.%s"), GET_MEMBER_NAME_STRING_CHECKED(UAnimDetailsProxyBool, Bool), GET_MEMBER_NAME_STRING_CHECKED(FAnimDetailsBool, Bool));
+		ValuePropertyHandle = DetailBuilder.GetProperty(FName(*PropertyPath), GetClass());
+		if (ValuePropertyHandle.IsValid() && SequencerItem.IsValid())
+		{
+			ValuePropertyHandle->SetPropertyDisplayName(FText::FromName(SequencerItem.GetBinding()->GetPropertyName()));
+		}
+	}
+}
+
+TArray<FName> UAnimDetailsProxyBool::GetPropertyNames() const
+{
+	return
+	{
+		GET_MEMBER_NAME_CHECKED(FAnimDetailsBool, Bool)
+	};
+}
+
+bool UAnimDetailsProxyBool::PropertyIsOnProxy(const FProperty* Property, const FProperty* MemberProperty)
+{
+	return ((Property && Property->GetFName() == GET_MEMBER_NAME_CHECKED(FAnimDetailsBool, Bool)) ||
+		(MemberProperty && MemberProperty->GetFName() == GET_MEMBER_NAME_CHECKED(UAnimDetailsProxyBool, Bool)));
+}
+
+void UAnimDetailsProxyBool::UpdateProxyValues()
+{
+	bool Value = false;
+
+	UControlRig* ControlRig = GetControlRig();
+	FRigControlElement* ControlElement = GetControlElement();
+	if (ControlRig && ControlElement)
+	{
+		if (ControlElement->Settings.ControlType == ERigControlType::Bool)
+		{
+			FRigControlValue ControlValue = ControlRig->GetControlValue(ControlElement, ERigControlValueType::Current);
+			Value = ControlValue.Get<bool>();
+		}
+	}
+	else if (SequencerItem.IsValid())
+	{
+		const TOptional<bool> OptionalValue = SequencerItem.GetBinding()->GetOptionalValue<bool>(*SequencerItem.GetBoundObject());
+		Value = OptionalValue.IsSet() ? OptionalValue.GetValue() : false;
+	}
+
+	const FAnimDetailsBool AnimDetailProxyBool(Value);
+
+	const FName PropName("Bool");
+	FTrackInstancePropertyBindings Binding(PropName, PropName.ToString());
+	Binding.CallFunction<FAnimDetailsBool>(*this, AnimDetailProxyBool);
+}
+
+EControlRigContextChannelToKey UAnimDetailsProxyBool::GetChannelToKeyFromPropertyName(const FName& PropertyName) const
+{
+	if (PropertyName == GET_MEMBER_NAME_CHECKED(FAnimDetailsBool, Bool))
+	{
+		return EControlRigContextChannelToKey::TranslationX;
+	}
+
+	return EControlRigContextChannelToKey::AllTransform;
+}
+
+EControlRigContextChannelToKey UAnimDetailsProxyBool::GetChannelToKeyFromChannelName(const FString& InChannelName) const
+{
+	if (InChannelName == TEXT("Bool"))
+	{
+		return EControlRigContextChannelToKey::TranslationX;
+	}
+
+	const UControlRig* ControlRig = GetControlRig();
+	const FRigControlElement* ControlElement = GetControlElement();
+	if (ControlElement && ControlElement->GetDisplayName() == InChannelName)
+	{
+		return EControlRigContextChannelToKey::TranslationX;
+	}
+
+	return EControlRigContextChannelToKey::AllTransform;
+}
+
+void UAnimDetailsProxyBool::SetControlRigElementValueFromCurrent(UControlRig* ControlRig, FRigControlElement* ControlElement, const FRigControlModifiedContext& Context)
+{
+	if (ControlElement && 
+		ControlRig && 
+		ControlElement->Settings.ControlType == ERigControlType::Bool)
+	{
+		const bool Value = Bool.Bool;
+
+		constexpr bool bNotify = true;
+		constexpr bool bSetupUndo = false;
+		ControlRig->SetControlValue<bool>(ControlElement->GetKey().Name, Value, bNotify, Context, bSetupUndo);
+
+		ControlRig->Evaluate_AnyThread();
+	}
+}
+
+void UAnimDetailsProxyBool::SetBindingValueFromCurrent(UObject* InObject, const TSharedPtr<FTrackInstancePropertyBindings>& Binding, const FRigControlModifiedContext& Context, bool bInteractive)
+{
+	if (InObject && Binding.IsValid())
+	{
+		Binding->SetCurrentValue<bool>(*InObject, Bool.Bool);
+	}
+}
diff --git a/Engine/Plugins/Animation/ControlRig/Source/ControlRigEditor/Private/AnimDetails/Proxies/AnimDetailsProxyEnum.cpp b/Engine/Plugins/Animation/ControlRig/Source/ControlRigEditor/Private/AnimDetails/Proxies/AnimDetailsProxyEnum.cpp
new file mode 100644
--- /dev/null
+++ b/Engine/Plugins/Animation/ControlRig/Source/ControlRigEditor/Private/AnimDetails/Proxies/AnimDetailsProxyEnum.cpp
@@ -0,0 +1,146 @@
+// Copyright Epic Games, Inc. All Rights Reserved.
+
+#include "AnimDetailsProxyEnum.h"
+
+#include "ControlRig.h"
+#include "DetailLayoutBuilder.h"
+#include "MovieSceneCommonHelpers.h"
+
+#include UE_INLINE_GENERATED_CPP_BY_NAME(AnimDetailsProxyEnum)
+
+FName UAnimDetailsProxyEnum::GetCategoryName() const
+{
+	return "Enum";
+}
+
+FName UAnimDetailsProxyEnum::GetDetailRowID() const
+{
+	if (!bIsIndividual)
+	{
+		return Enum.EnumType ? Enum.EnumType->GetFName() : NAME_None;
+	}
+
+	return UAnimDetailsProxyBase::GetDetailRowID();
+}
+
+void UAnimDetailsProxyEnum::UpdatePropertyDisplayNames(IDetailLayoutBuilder& DetailBuilder)
+{
+	TSharedPtr<IPropertyHandle> ValuePropertyHandle = DetailBuilder.GetProperty(GET_MEMBER_NAME_CHECKED(UAnimDetailsProxyEnum, Enum), GetClass());
+	if (ValuePropertyHandle.IsValid())
+	{
+		ValuePropertyHandle->SetPropertyDisplayName(GetDisplayNameText());
+	}
+
+	if (SequencerItem.IsValid())
+	{
+		const FString PropertyPath = FString::Printf(TEXT("%s.%s"), GET_MEMBER_NAME_STRING_CHECKED(UAnimDetailsProxyEnum, Enum), GET_MEMBER_NAME_STRING_CHECKED(FAnimDetailsEnum, EnumIndex));
+		ValuePropertyHandle = DetailBuilder.GetProperty(FName(*PropertyPath), GetClass());
+		if (ValuePropertyHandle.IsValid())
+		{
+			ValuePropertyHandle->SetPropertyDisplayName(FText::FromName(SequencerItem.GetBinding()->GetPropertyName()));
+		}
+	}
+	else if (bIsIndividual)
+	{
+		const FString PropertyPath = FString::Printf(TEXT("%s.%s"), GET_MEMBER_NAME_STRING_CHECKED(UAnimDetailsProxyEnum, Enum), GET_MEMBER_NAME_STRING_CHECKED(FAnimDetailsEnum, EnumIndex));
+		ValuePropertyHandle = DetailBuilder.GetProperty(FName(*PropertyPath), GetClass());
+		if (ValuePropertyHandle.IsValid())
+		{
+			ValuePropertyHandle->SetPropertyDisplayName(GetDisplayNameText());
+		}
+	}
+}
+
+TArray<FName> UAnimDetailsProxyEnum::GetPropertyNames() const
+{
+	return
+	{
+		GET_MEMBER_NAME_CHECKED(FAnimDetailsEnum, EnumIndex)
+	};
+}
+
+
+bool UAnimDetailsProxyEnum::PropertyIsOnProxy(const FProperty* Property, const FProperty* MemberProperty)
+{
+	return ((Property && Property->GetFName() == GET_MEMBER_NAME_CHECKED(FAnimDetailsEnum, EnumIndex)) ||
+		(MemberProperty && MemberProperty->GetFName() == GET_MEMBER_NAME_CHECKED(UAnimDetailsProxyEnum, Enum)));
+}
+
+void UAnimDetailsProxyEnum::UpdateProxyValues()
+{
+	UControlRig* ControlRig = GetControlRig();
+	FRigControlElement* ControlElement = GetControlElement();
+	if (!ControlRig || !ControlElement)
+	{
+		return;
+	}
+
+	TObjectPtr<UEnum> EnumType = nullptr;
+
+	int32 Value = 0;
+	if (ControlElement->Settings.ControlType == ERigControlType::Integer &&
+		ControlElement->Settings.ControlEnum)
+	{
+		EnumType = ControlElement->Settings.ControlEnum;
+		const FRigControlValue ControlValue = ControlRig->GetControlValue(ControlElement, ERigControlValueType::Current);
+
+		Value = ControlValue.Get<int32>();
+	}
+
+	if (EnumType)
+	{
+		FAnimDetailsEnum EnumValue;
+		EnumValue.EnumType = EnumType;
+		EnumValue.EnumIndex = Value;
+
+		const FName PropertyName("Enum");
+		FTrackInstancePropertyBindings Binding(PropertyName, PropertyName.ToString());
+		Binding.CallFunction<FAnimDetailsEnum>(*this, EnumValue);
+	}
+}
+
+EControlRigContextChannelToKey UAnimDetailsProxyEnum::GetChannelToKeyFromPropertyName(const FName& PropertyName) const
+{
+	if (PropertyName == GET_MEMBER_NAME_CHECKED(FAnimDetailsEnum, EnumIndex))
+	{
+		return EControlRigContextChannelToKey::TranslationX;
+	}
+
+	return EControlRigContextChannelToKey::AllTransform;
+}
+
+EControlRigContextChannelToKey UAnimDetailsProxyEnum::GetChannelToKeyFromChannelName(const FString& InChannelName) const
+{
+	if (InChannelName == TEXT("EnumIndex"))
+	{
+		return EControlRigContextChannelToKey::TranslationX;
+	}
+
+	const UControlRig* ControlRig = GetControlRig();
+	const FRigElementKey ElementKey = GetControlElementKey();
+
+	const FRigControlElement* ControlElement = ControlRig ? ControlRig->FindControl(ElementKey.Name) : nullptr;
+	if (ControlElement && ControlElement->GetDisplayName() == InChannelName)
+	{
+		return EControlRigContextChannelToKey::TranslationX;
+	}
+
+	return EControlRigContextChannelToKey::AllTransform;
+}
+
+void UAnimDetailsProxyEnum::SetControlRigElementValueFromCurrent(UControlRig* ControlRig, FRigControlElement* ControlElement, const FRigControlModifiedContext& Context)
+{
+	if (ControlRig && 
+		ControlElement &&
+		ControlElement->Settings.ControlType == ERigControlType::Integer &&
+		ControlElement->Settings.ControlEnum)
+	{
+		int32 Val = Enum.EnumIndex;
+
+		constexpr bool bNotify = true;
+		constexpr bool bSetupUndo = false;
+		ControlRig->SetControlValue<int32>(ControlElement->GetKey().Name, Val, bNotify, Context, bSetupUndo);
+
+		ControlRig->Evaluate_AnyThread();
+	}
+}
diff --git a/Engine/Plugins/Animation/ControlRig/Source/ControlRigEditor/Private/AnimDetails/Proxies/AnimDetailsProxyFloat.cpp b/Engine/Plugins/Animation/ControlRig/Source/ControlRigEditor/Private/AnimDetails/Proxies/AnimDetailsProxyFloat.cpp
new file mode 100644
--- /dev/null
+++ b/Engine/Plugins/Animation/ControlRig/Source/ControlRigEditor/Private/AnimDetails/Proxies/AnimDetailsProxyFloat.cpp
@@ -0,0 +1,165 @@
+// Copyright Epic Games, Inc. All Rights Reserved.
+
+#include "AnimDetailsProxyFloat.h"
+
+#include "ControlRig.h"
+#include "DetailLayoutBuilder.h"
+#include "MovieSceneCommonHelpers.h"
+
+#include UE_INLINE_GENERATED_CPP_BY_NAME(AnimDetailsProxyFloat)
+
+FName UAnimDetailsProxyFloat::GetCategoryName() const
+{
+	return "Float";
+}
+
+void UAnimDetailsProxyFloat::UpdatePropertyDisplayNames(IDetailLayoutBuilder& DetailBuilder)
+{
+	TSharedPtr<IPropertyHandle> ValuePropertyHandle = DetailBuilder.GetProperty(GET_MEMBER_NAME_CHECKED(UAnimDetailsProxyFloat, Float), GetClass());
+	if (ValuePropertyHandle.IsValid())
+	{
+		ValuePropertyHandle->SetPropertyDisplayName(GetDisplayNameText());
+	}
+
+	if (bIsIndividual)
+	{
+		const FString PropertyPath = FString::Printf(TEXT("%s.%s"), GET_MEMBER_NAME_STRING_CHECKED(UAnimDetailsProxyFloat, Float), GET_MEMBER_NAME_STRING_CHECKED(FAnimDetailsFloat, Float));
+		ValuePropertyHandle = DetailBuilder.GetProperty(FName(*PropertyPath), GetClass());
+		if (ValuePropertyHandle.IsValid())
+		{
+			ValuePropertyHandle->SetPropertyDisplayName(GetDisplayNameText());
+		}
+	}
+	else if (SequencerItem.IsValid())
+	{
+		const FString PropertyPath = FString::Printf(TEXT("%s.%s"), GET_MEMBER_NAME_STRING_CHECKED(UAnimDetailsProxyFloat, Float), GET_MEMBER_NAME_STRING_CHECKED(FAnimDetailsFloat, Float));
+		ValuePropertyHandle = DetailBuilder.GetProperty(FName(*PropertyPath), GetClass());
+		if (ValuePropertyHandle.IsValid())
+		{
+			ValuePropertyHandle->SetPropertyDisplayName(FText::FromName(SequencerItem.GetBinding()->GetPropertyName()));
+		}
+	}
+}
+
+TArray<FName> UAnimDetailsProxyFloat::GetPropertyNames() const
+{
+	return
+	{
+		GET_MEMBER_NAME_CHECKED(FAnimDetailsFloat, Float)
+	};
+}
+
+bool UAnimDetailsProxyFloat::PropertyIsOnProxy(const FProperty* Property, const FProperty* MemberProperty)
+{
+	return ((Property && Property->GetFName() == GET_MEMBER_NAME_CHECKED(FAnimDetailsFloat, Float)) ||
+		(MemberProperty && MemberProperty->GetFName() == GET_MEMBER_NAME_CHECKED(UAnimDetailsProxyFloat, Float)));
+}
+
+void UAnimDetailsProxyFloat::UpdateProxyValues()
+{
+	UControlRig* ControlRig = GetControlRig();
+	FRigControlElement* ControlElement = GetControlElement();
+
+	double Value = 0.f;
+	if (ControlRig && ControlElement)
+	{
+		if (ControlElement->Settings.ControlType == ERigControlType::Float ||
+			(ControlElement->Settings.ControlType == ERigControlType::ScaleFloat))
+		{
+			FRigControlValue ControlValue = ControlRig->GetControlValue(ControlElement, ERigControlValueType::Current);
+			Value = ControlValue.Get<float>();
+		}
+	}
+	else if (SequencerItem.IsValid())
+	{
+		if (FProperty* Property = SequencerItem.GetProperty())
+		{
+			if (Property->IsA(FDoubleProperty::StaticClass()))
+			{
+				const TOptional<double> OptionalValue = SequencerItem.GetBinding()->GetOptionalValue<double>(*SequencerItem.GetBoundObject());
+				if (OptionalValue.IsSet())
+				{
+					Value = OptionalValue.GetValue();
+				}
+			}
+			else if (Property->IsA(FFloatProperty::StaticClass()))
+			{
+				const TOptional<float> OptionalValue = SequencerItem.GetBinding()->GetOptionalValue<float>(*SequencerItem.GetBoundObject());
+				if (OptionalValue.IsSet())
+				{
+					Value = OptionalValue.GetValue();
+				}
+			}
+		}
+	}
+
+	const FAnimDetailsFloat AnimDetailProxyFloat(Value);
+
+	const FName PropName("Float");
+	FTrackInstancePropertyBindings Binding(PropName, PropName.ToString());
+	Binding.CallFunction<FAnimDetailsFloat>(*this, AnimDetailProxyFloat);
+}
+
+EControlRigContextChannelToKey UAnimDetailsProxyFloat::GetChannelToKeyFromPropertyName(const FName& PropertyName) const
+{
+	if (PropertyName == GET_MEMBER_NAME_CHECKED(FAnimDetailsFloat, Float))
+	{
+		return EControlRigContextChannelToKey::TranslationX;
+	}
+
+	return EControlRigContextChannelToKey::AllTransform;
+}
+
+EControlRigContextChannelToKey UAnimDetailsProxyFloat::GetChannelToKeyFromChannelName(const FString& InChannelName) const
+{
+	if (InChannelName == TEXT("Float"))
+	{
+		return EControlRigContextChannelToKey::TranslationX;
+	}
+
+	const UControlRig* ControlRig = GetControlRig();
+	const FRigElementKey ElementKey = GetControlElementKey();
+
+	const FRigControlElement* ControlElement = ControlRig ? ControlRig->FindControl(ElementKey.Name) : nullptr;
+	if (ControlElement && ControlElement->GetDisplayName() == InChannelName)
+	{
+		return EControlRigContextChannelToKey::TranslationX;
+	}
+
+	return EControlRigContextChannelToKey::AllTransform;
+}
+
+void UAnimDetailsProxyFloat::SetControlRigElementValueFromCurrent(UControlRig* ControlRig, FRigControlElement* ControlElement, const FRigControlModifiedContext& Context)
+{
+	if (ControlRig && 
+		ControlElement &&
+		(ControlElement->Settings.ControlType == ERigControlType::Float || ControlElement->Settings.ControlType == ERigControlType::ScaleFloat))
+	{
+		float Value = Float.Float;
+
+		constexpr bool bNotify = true;
+		constexpr bool bSetupUndo = false;
+		ControlRig->SetControlValue<float>(ControlElement->GetKey().Name, Value, bNotify, Context, bSetupUndo);
+
+		ControlRig->Evaluate_AnyThread();
+	}
+}
+
+void UAnimDetailsProxyFloat::SetBindingValueFromCurrent(UObject* InObject, const TSharedPtr<FTrackInstancePropertyBindings>& Binding, const FRigControlModifiedContext& Context, bool bInteractive)
+{
+	if (InObject && Binding.IsValid())
+	{
+		if (FProperty* Property = Binding->GetProperty((*(InObject))))
+		{
+			if (Property->IsA(FDoubleProperty::StaticClass()))
+			{
+				Binding->SetCurrentValue<double>(*InObject, Float.Float);
+			}
+			else if (Property->IsA(FFloatProperty::StaticClass()))
+			{
+				float FVal = (float)Float.Float;
+				Binding->SetCurrentValue<float>(*InObject, FVal);
+			}
+		}
+	}
+}
diff --git a/Engine/Plugins/Animation/ControlRig/Source/ControlRigEditor/Private/AnimDetails/Proxies/AnimDetailsProxyFloat.h b/Engine/Plugins/Animation/ControlRig/Source/ControlRigEditor/Private/AnimDetails/Proxies/AnimDetailsProxyFloat.h
new file mode 100644
--- /dev/null
+++ b/Engine/Plugins/Animation/ControlRig/Source/ControlRigEditor/Private/AnimDetails/Proxies/AnimDetailsProxyFloat.h
@@ -0,0 +1,46 @@
+// Copyright Epic Games, Inc. All Rights Reserved.
+
+#pragma once
+
+#include "AnimDetailsProxyBase.h"
+
+#include "AnimDetailsProxyFloat.generated.h"
+
+/** 
+ * A floating point value in anim details  
+ * 
+ * Note, control rig uses 'float' controls so we call this float though it's a 
+ * double internally, so we can use same for non-control rig parameters
+ */
+USTRUCT(BlueprintType)
+struct FAnimDetailsFloat
+{
+	GENERATED_BODY()
+
+	UPROPERTY(EditAnywhere, Interp, Category = "Float", meta = (SliderExponent = "1.0"))
+	double Float = 0.0;
+};
+
+/** Handles an floating point property bound in sequencer, and the related control if the bound object uses a control rig */
+UCLASS(EditInlineNew, CollapseCategories)
+class UAnimDetailsProxyFloat 
+	: public UAnimDetailsProxyBase
+{
+	GENERATED_BODY()
+
+public:
+	//~ Begin UAnimDetailsProxyBase interface
+	virtual FName GetCategoryName() const override;
+	virtual void UpdatePropertyDisplayNames(IDetailLayoutBuilder& DetailBuilder);
+	virtual TArray<FName> GetPropertyNames() const override;
+	virtual bool PropertyIsOnProxy(const FProperty* Property, const FProperty* MemberProperty) override;
+	virtual void UpdateProxyValues() override;
+	virtual EControlRigContextChannelToKey GetChannelToKeyFromPropertyName(const FName& InPropertyName) const override;
+	virtual EControlRigContextChannelToKey GetChannelToKeyFromChannelName(const FString& InChannelName) const override;
+	virtual void SetControlRigElementValueFromCurrent(UControlRig* ControlRig, FRigControlElement* ControlElement, const FRigControlModifiedContext& Context) override;
+	virtual void SetBindingValueFromCurrent(UObject* InObject, const TSharedPtr<FTrackInstancePropertyBindings>& Binding, const FRigControlModifiedContext& Context, bool bInteractive = false) override;
+	//~ End UAnimDetailsProxyBase interface
+
+	UPROPERTY(EditAnywhere, Interp, Category = "Float", meta = (Delta = "0.05", SliderExponent = "1", LinearDeltaSensitivity = "1"))
+	FAnimDetailsFloat Float;
+};
diff --git a/Engine/Plugins/Animation/ControlRig/Source/ControlRigEditor/Private/AnimDetails/Proxies/AnimDetailsProxyInteger.cpp b/Engine/Plugins/Animation/ControlRig/Source/ControlRigEditor/Private/AnimDetails/Proxies/AnimDetailsProxyInteger.cpp
new file mode 100644
--- /dev/null
+++ b/Engine/Plugins/Animation/ControlRig/Source/ControlRigEditor/Private/AnimDetails/Proxies/AnimDetailsProxyInteger.cpp
@@ -0,0 +1,142 @@
+// Copyright Epic Games, Inc. All Rights Reserved.
+
+#include "AnimDetailsProxyInteger.h"
+
+#include "ControlRig.h"
+#include "DetailLayoutBuilder.h"
+#include "MovieSceneCommonHelpers.h"
+
+#include UE_INLINE_GENERATED_CPP_BY_NAME(AnimDetailsProxyInteger)
+
+FName UAnimDetailsProxyInteger::GetCategoryName() const
+{
+	return "Integer";
+}
+
+void UAnimDetailsProxyInteger::UpdatePropertyDisplayNames(IDetailLayoutBuilder& DetailBuilder)
+{
+	TSharedPtr<IPropertyHandle> ValuePropertyHandle = DetailBuilder.GetProperty(GET_MEMBER_NAME_CHECKED(UAnimDetailsProxyInteger, Integer), GetClass());
+	if (ValuePropertyHandle.IsValid())
+	{
+		ValuePropertyHandle->SetPropertyDisplayName(GetDisplayNameText());
+	}
+
+	if (SequencerItem.IsValid())
+	{
+		const FString PropertyPath = FString::Printf(TEXT("%s.%s"), GET_MEMBER_NAME_STRING_CHECKED(UAnimDetailsProxyInteger, Integer), GET_MEMBER_NAME_STRING_CHECKED(FAnimDetailsInteger, Integer));
+		ValuePropertyHandle = DetailBuilder.GetProperty(FName(*PropertyPath), GetClass());
+		if (ValuePropertyHandle.IsValid())
+		{
+			ValuePropertyHandle->SetPropertyDisplayName(FText::FromName(SequencerItem.GetBinding()->GetPropertyName()));
+		}
+	}
+	else if (bIsIndividual)
+	{
+		const FString PropertyPath = FString::Printf(TEXT("%s.%s"), GET_MEMBER_NAME_STRING_CHECKED(UAnimDetailsProxyInteger, Integer), GET_MEMBER_NAME_STRING_CHECKED(FAnimDetailsInteger, Integer));
+		ValuePropertyHandle = DetailBuilder.GetProperty(FName(*PropertyPath), GetClass());
+		if (ValuePropertyHandle.IsValid())
+		{
+			ValuePropertyHandle->SetPropertyDisplayName(GetDisplayNameText());
+		}
+	}
+}
+
+TArray<FName> UAnimDetailsProxyInteger::GetPropertyNames() const
+{
+	return 
+	{ 
+		GET_MEMBER_NAME_CHECKED(UAnimDetailsProxyInteger, Integer)
+	};
+}
+
+bool UAnimDetailsProxyInteger::PropertyIsOnProxy(const FProperty* Property, const FProperty* MemberProperty)
+{
+	return ((Property && Property->GetFName() == GET_MEMBER_NAME_CHECKED(FAnimDetailsInteger, Integer)) ||
+		(MemberProperty && MemberProperty->GetFName() == GET_MEMBER_NAME_CHECKED(UAnimDetailsProxyInteger, Integer)));
+}
+
+void UAnimDetailsProxyInteger::UpdateProxyValues()
+{
+	int64 Value = 0;
+
+	UControlRig* ControlRig = GetControlRig();
+	FRigControlElement* ControlElement = GetControlElement();
+	if (ControlRig && ControlElement)
+	{
+		if (ControlElement->Settings.ControlType == ERigControlType::Integer)
+		{
+			FRigControlValue ControlValue = ControlRig->GetControlValue(ControlElement, ERigControlValueType::Current);
+			Value = ControlValue.Get<int32>();
+		}
+	}
+	else if (SequencerItem.IsValid())
+	{
+		if (const TOptional<int64> OptionalValue = SequencerItem.GetBinding()->GetOptionalValue<int64>(*SequencerItem.GetBoundObject()))
+		{
+			if (OptionalValue.IsSet())
+			{
+				Value = OptionalValue.GetValue();
+			}
+		}
+	}
+
+	const FAnimDetailsInteger AnimDetailProxyInteger(Value);
+
+	const FName PropName("Integer");
+	FTrackInstancePropertyBindings Binding(PropName, PropName.ToString());
+	Binding.CallFunction<FAnimDetailsInteger>(*this, AnimDetailProxyInteger);
+}
+
+EControlRigContextChannelToKey UAnimDetailsProxyInteger::GetChannelToKeyFromPropertyName(const FName& PropertyName) const
+{
+	if (PropertyName == GET_MEMBER_NAME_CHECKED(UAnimDetailsProxyInteger, Integer))
+	{
+		return EControlRigContextChannelToKey::TranslationX;
+	}
+
+	return EControlRigContextChannelToKey::AllTransform;
+}
+
+EControlRigContextChannelToKey UAnimDetailsProxyInteger::GetChannelToKeyFromChannelName(const FString& InChannelName) const
+{
+	if (InChannelName == TEXT("Integer"))
+	{
+		return EControlRigContextChannelToKey::TranslationX;
+	}
+
+	const UControlRig* ControlRig = GetControlRig();
+	const FRigElementKey ElementKey = GetControlElementKey();
+
+	const FRigControlElement* ControlElement = ControlRig ? ControlRig->FindControl(ElementKey.Name) : nullptr;
+	if (ControlElement && ControlElement->GetDisplayName() == InChannelName)
+	{
+		return EControlRigContextChannelToKey::TranslationX;
+	}
+
+	return EControlRigContextChannelToKey::AllTransform;
+}
+
+void UAnimDetailsProxyInteger::SetControlRigElementValueFromCurrent(UControlRig* ControlRig, FRigControlElement* ControlElement, const FRigControlModifiedContext& Context)
+{
+	if (ControlRig && 
+		ControlElement &&		
+		ControlElement->Settings.ControlType == ERigControlType::Integer &&
+		!ControlElement->Settings.ControlEnum)
+	{
+		int32 Val = (int32)Integer.Integer;
+
+		constexpr bool bNotify = true;
+		constexpr bool bSetupUndo = false;
+		ControlRig->SetControlValue<int32>(ControlElement->GetKey().Name, Val, bNotify, Context, bSetupUndo);
+
+		ControlRig->Evaluate_AnyThread();
+	}
+}
+
+void UAnimDetailsProxyInteger::SetBindingValueFromCurrent(UObject* InObject, const TSharedPtr<FTrackInstancePropertyBindings>& Binding, const FRigControlModifiedContext& Context, bool bInteractive)
+{
+	if (InObject && Binding.IsValid())
+	{
+		Binding->SetCurrentValue<int64>(*InObject, Integer.Integer);
+	}
+}
\ No newline at end of file
diff --git a/Engine/Plugins/Animation/ControlRig/Source/ControlRigEditor/Private/AnimDetails/Proxies/AnimDetailsProxyInteger.h b/Engine/Plugins/Animation/ControlRig/Source/ControlRigEditor/Private/AnimDetails/Proxies/AnimDetailsProxyInteger.h
new file mode 100644
--- /dev/null
+++ b/Engine/Plugins/Animation/ControlRig/Source/ControlRigEditor/Private/AnimDetails/Proxies/AnimDetailsProxyInteger.h
@@ -0,0 +1,41 @@
+// Copyright Epic Games, Inc. All Rights Reserved.
+
+#pragma once
+
+#include "AnimDetailsProxyBase.h"
+
+#include "AnimDetailsProxyInteger.generated.h"
+
+/** An integer value in anim details */
+USTRUCT(BlueprintType)
+struct FAnimDetailsInteger
+{
+	GENERATED_BODY()
+
+	UPROPERTY(EditAnywhere, Interp, Category = "Integer")
+	int64 Integer = 0;
+};
+
+/** Handles an integer property bound in sequencer, and the related control if the bound object uses a control rig */
+UCLASS(EditInlineNew, CollapseCategories)
+class UAnimDetailsProxyInteger 
+	: public UAnimDetailsProxyBase
+{
+	GENERATED_BODY()
+
+public:
+	//~ Begin UAnimDetailsProxyBase interface
+	virtual FName GetCategoryName() const override;
+	virtual void UpdatePropertyDisplayNames(IDetailLayoutBuilder& DetailBuilder);
+	virtual TArray<FName> GetPropertyNames() const override;
+	virtual bool PropertyIsOnProxy(const FProperty* Property, const FProperty* MemberProperty) override;
+	virtual void UpdateProxyValues() override;
+	virtual EControlRigContextChannelToKey GetChannelToKeyFromPropertyName(const FName& InPropertyName) const override;
+	virtual EControlRigContextChannelToKey GetChannelToKeyFromChannelName(const FString& InChannelName) const override;
+	virtual void SetControlRigElementValueFromCurrent(UControlRig* ControlRig, FRigControlElement* ControlElement, const FRigControlModifiedContext& Context) override;
+	virtual void SetBindingValueFromCurrent(UObject* InObject, const TSharedPtr<FTrackInstancePropertyBindings>& Binding, const FRigControlModifiedContext& Context, bool bInteractive = false) override;
+	//~ End UAnimDetailsProxyBase interface
+
+	UPROPERTY(EditAnywhere, Interp, Category = "Integer")
+	FAnimDetailsInteger Integer;
+};
diff --git a/Engine/Plugins/Animation/ControlRig/Source/ControlRigEditor/Private/ControlRigEditorModule.cpp b/Engine/Plugins/Animation/ControlRig/Source/ControlRigEditor/Private/ControlRigEditorModule.cpp
--- a/Engine/Plugins/Animation/ControlRig/Source/ControlRigEditor/Private/ControlRigEditorModule.cpp
+++ b/Engine/Plugins/Animation/ControlRig/Source/ControlRigEditor/Private/ControlRigEditorModule.cpp
@@ -1,6 +1,19 @@
 // Copyright Epic Games, Inc. All Rights Reserved.
 
 #include "ControlRigEditorModule.h"
+
+#include "AnimDetails/Customizations/AnimDetailsProxyDetails.h"
+#include "AnimDetails/Customizations/AnimDetailsValueCustomization.h"
+#include "AnimDetails/Customizations/AnimDetailsValueEnumCustomization.h"
+#include "AnimDetails/Proxies/AnimDetailsProxyBool.h"
+#include "AnimDetails/Proxies/AnimDetailsProxyEnum.h"
+#include "AnimDetails/Proxies/AnimDetailsProxyFloat.h"
+#include "AnimDetails/Proxies/AnimDetailsProxyInteger.h"
+#include "AnimDetails/Proxies/AnimDetailsProxyLocation.h"
+#include "AnimDetails/Proxies/AnimDetailsProxyRotation.h"
+#include "AnimDetails/Proxies/AnimDetailsProxyScale.h"
+#include "AnimDetails/Proxies/AnimDetailsProxyTransform.h"
+#include "AnimDetails/Proxies/AnimDetailsProxyVector2D.h"
 #include "BlueprintActionDatabaseRegistrar.h"
 #include "BlueprintNodeSpawner.h"
 #include "PropertyEditorModule.h"
@@ -21,7 +34,6 @@
 #include "EditorModeManager.h"
 #include "Sequencer/MovieSceneControlRigParameterSection.h"
 #include "Sequencer/MovieSceneControlRigSectionDetailsCustomization.h"
-#include "EditMode/ControlsProxyDetailCustomization.h"
 #include "EditMode/ControlRigEditModeCommands.h"
 #include "Framework/MultiBox/MultiBoxBuilder.h"
 #include "ToolMenus.h"
@@ -56,7 +68,6 @@
 #include "ControlRigBlueprint.h"
 #include "Dialogs/Dialogs.h"
 #include "Settings/ControlRigSettings.h"
-#include "EditMode/ControlRigControlsProxy.h"
 #include "IPersonaToolkit.h"
 #include "LevelSequence.h"
 #include "AnimSequenceLevelSequenceLink.h"
@@ -115,6 +126,8 @@ DEFINE_LOG_CATEGORY(LogControlRigEditor);
 
 void FControlRigEditorModule::StartupModule()
 {
+	using namespace UE::ControlRigEditor;
+
 	FControlRigEditModeCommands::Register();
 	FControlRigEditorCommands::Register();
 	FControlRigHierarchyCommands::Register();
@@ -141,122 +154,50 @@ void FControlRigEditorModule::StartupModule()
 
 	ClassesToUnregisterOnShutdown.Add(UControlRig::StaticClass()->GetFName());
 
-	// proxies
 	const FName PropertyEditorModuleName("PropertyEditor");
 	FPropertyEditorModule& PropertyModule = FModuleManager::GetModuleChecked<FPropertyEditorModule>(PropertyEditorModuleName);
 
-	ClassesToUnregisterOnShutdown.Add(UAnimDetailControlsProxyFloat::StaticClass()->GetFName());
-	PropertyEditorModule.RegisterCustomClassLayout(UAnimDetailControlsProxyFloat::StaticClass()->GetFName(), FOnGetDetailCustomizationInstance::CreateLambda([]() {
-		FName CategoryName = FName(TEXT("Float"));
-		return FAnimDetailProxyDetails::MakeInstance(CategoryName);
-		}));
-
-	{	
-		TSharedRef<FPropertySection> Section = RegisterPropertySection(PropertyModule, "AnimDetailControlsProxyFloat", "General", LOCTEXT("General", "General"));
-		Section->AddCategory("Attributes");
-	}
-
-	ClassesToUnregisterOnShutdown.Add(UAnimDetailControlsProxyInteger::StaticClass()->GetFName());
-	PropertyEditorModule.RegisterCustomClassLayout(UAnimDetailControlsProxyInteger::StaticClass()->GetFName(), FOnGetDetailCustomizationInstance::CreateLambda([]() {
-		FName CategoryName = FName(TEXT("Integer"));
-		return FAnimDetailProxyDetails::MakeInstance(CategoryName);
-		}));
-
-	{
-		TSharedRef<FPropertySection> Section = RegisterPropertySection(PropertyModule, "AnimDetailControlsProxyInteger", "General", LOCTEXT("General", "General"));
-		Section->AddCategory("Attributes");
-	}
-
-	ClassesToUnregisterOnShutdown.Add(UAnimDetailControlsProxyBool::StaticClass()->GetFName());
-	PropertyEditorModule.RegisterCustomClassLayout(UAnimDetailControlsProxyBool::StaticClass()->GetFName(), FOnGetDetailCustomizationInstance::CreateLambda([]() {
-		FName CategoryName = FName(TEXT("Bool"));
-		return FAnimDetailProxyDetails::MakeInstance(CategoryName);
-		}));
-
-	{
-		TSharedRef<FPropertySection> Section = RegisterPropertySection(PropertyModule, "AnimDetailControlsProxyBool", "General", LOCTEXT("General", "General"));
-		Section->AddCategory("Attributes");
-	}
-
-	ClassesToUnregisterOnShutdown.Add(UAnimDetailControlsProxyTransform::StaticClass()->GetFName());
-	PropertyEditorModule.RegisterCustomClassLayout(UAnimDetailControlsProxyTransform::StaticClass()->GetFName(), FOnGetDetailCustomizationInstance::CreateLambda([]() {
-		FName CategoryName = FName(TEXT("Transform"));
-		return FAnimDetailProxyDetails::MakeInstance(CategoryName);
-		}));
-
-	{
-		TSharedRef<FPropertySection> Section = RegisterPropertySection(PropertyModule, "AnimDetailControlsProxyTransform", "General", LOCTEXT("General", "General"));
-		Section->AddCategory("Attributes");
-	}
-
-	ClassesToUnregisterOnShutdown.Add(UAnimDetailControlsProxyLocation::StaticClass()->GetFName());
-	PropertyEditorModule.RegisterCustomClassLayout(UAnimDetailControlsProxyLocation::StaticClass()->GetFName(), FOnGetDetailCustomizationInstance::CreateLambda([]() {
-		FName CategoryName = FName(TEXT("Location"));
-		return FAnimDetailProxyDetails::MakeInstance(CategoryName);
-		}));
-
-	{
-		TSharedRef<FPropertySection> Section = RegisterPropertySection(PropertyModule, "AnimDetailControlsProxyLocation", "General", LOCTEXT("General", "General"));
-		Section->AddCategory("Attributes");
-	}
-
-	ClassesToUnregisterOnShutdown.Add(UAnimDetailControlsProxyRotation::StaticClass()->GetFName());
-	PropertyEditorModule.RegisterCustomClassLayout(UAnimDetailControlsProxyRotation::StaticClass()->GetFName(), FOnGetDetailCustomizationInstance::CreateLambda([]() {
-		FName CategoryName = FName(TEXT("Rotation"));
-		return FAnimDetailProxyDetails::MakeInstance(CategoryName);
-		}));
-
+	// Register the custom class layout common to classes that inherit AnimDetailsProxyBase
+	for (TObjectIterator<UClass> ClassIterator; ClassIterator; ++ClassIterator)
 	{
-		TSharedRef<FPropertySection> Section = RegisterPropertySection(PropertyModule, "AnimDetailControlsProxyRotation", "General", LOCTEXT("General", "General"));
-		Section->AddCategory("Attributes");
-	}
+		if (ClassIterator->IsChildOf(UAnimDetailsProxyBase::StaticClass()))
+		{
+			const FName ClassName = ClassIterator->GetFName();
+			if (ClassName == UAnimDetailsProxyBase::StaticClass()->GetFName())
+			{
+				continue;
+			}
 
-	ClassesToUnregisterOnShutdown.Add(UAnimDetailControlsProxyScale::StaticClass()->GetFName());
-	PropertyEditorModule.RegisterCustomClassLayout(UAnimDetailControlsProxyScale::StaticClass()->GetFName(), FOnGetDetailCustomizationInstance::CreateLambda([]() {
-		FName CategoryName = FName(TEXT("Scale"));
-		return FAnimDetailProxyDetails::MakeInstance(CategoryName);
-		}));
-	{
-		TSharedRef<FPropertySection> Section = RegisterPropertySection(PropertyModule, "AnimDetailControlsProxyScale", "General", LOCTEXT("General", "General"));
-		Section->AddCategory("Attributes");
+			ClassesToUnregisterOnShutdown.Add(ClassName);
+			PropertyEditorModule.RegisterCustomClassLayout(ClassName, FOnGetDetailCustomizationInstance::CreateStatic(&FAnimDetailsProxyDetails::MakeInstance));
+		}
 	}
 
-	ClassesToUnregisterOnShutdown.Add(UAnimDetailControlsProxyVector2D::StaticClass()->GetFName());
-	PropertyEditorModule.RegisterCustomClassLayout(UAnimDetailControlsProxyVector2D::StaticClass()->GetFName(), FOnGetDetailCustomizationInstance::CreateLambda([]() {
-		FName CategoryName = FName(TEXT("Vector2D"));
-		return FAnimDetailProxyDetails::MakeInstance(CategoryName);
-		}));
-
-	{
-		TSharedRef<FPropertySection> Section = RegisterPropertySection(PropertyModule, "AnimDetailControlsProxyVector2D", "General", LOCTEXT("General", "General"));
-		Section->AddCategory("Attributes");
-	}
-	
 	// same as ClassesToUnregisterOnShutdown but for properties, there is none right now
 	PropertiesToUnregisterOnShutdown.Reset();
 	PropertiesToUnregisterOnShutdown.Add(FRigVMCompileSettings::StaticStruct()->GetFName());
 	PropertyEditorModule.RegisterCustomPropertyTypeLayout(PropertiesToUnregisterOnShutdown.Last(), FOnGetPropertyTypeCustomizationInstance::CreateStatic(&FRigVMCompileSettingsDetails::MakeInstance));
 
-	PropertiesToUnregisterOnShutdown.Add(FAnimDetailProxyFloat::StaticStruct()->GetFName());
-	PropertyEditorModule.RegisterCustomPropertyTypeLayout(PropertiesToUnregisterOnShutdown.Last(), FOnGetPropertyTypeCustomizationInstance::CreateStatic(&FAnimDetailValueCustomization::MakeInstance));
+	PropertiesToUnregisterOnShutdown.Add(FAnimDetailsFloat::StaticStruct()->GetFName());
+	PropertyEditorModule.RegisterCustomPropertyTypeLayout(PropertiesToUnregisterOnShutdown.Last(), FOnGetPropertyTypeCustomizationInstance::CreateStatic(&FAnimDetailsValueCustomization::MakeInstance));
 
-	PropertiesToUnregisterOnShutdown.Add(FAnimDetailProxyInteger::StaticStruct()->GetFName());
-	PropertyEditorModule.RegisterCustomPropertyTypeLayout(PropertiesToUnregisterOnShutdown.Last(), FOnGetPropertyTypeCustomizationInstance::CreateStatic(&FAnimDetailValueCustomization::MakeInstance));
+	PropertiesToUnregisterOnShutdown.Add(FAnimDetailsInteger::StaticStruct()->GetFName());
+	PropertyEditorModule.RegisterCustomPropertyTypeLayout(PropertiesToUnregisterOnShutdown.Last(), FOnGetPropertyTypeCustomizationInstance::CreateStatic(&FAnimDetailsValueCustomization::MakeInstance));
 
-	PropertiesToUnregisterOnShutdown.Add(FAnimDetailProxyBool::StaticStruct()->GetFName());
-	PropertyEditorModule.RegisterCustomPropertyTypeLayout(PropertiesToUnregisterOnShutdown.Last(), FOnGetPropertyTypeCustomizationInstance::CreateStatic(&FAnimDetailValueCustomization::MakeInstance));
+	PropertiesToUnregisterOnShutdown.Add(FAnimDetailsBool::StaticStruct()->GetFName());
+	PropertyEditorModule.RegisterCustomPropertyTypeLayout(PropertiesToUnregisterOnShutdown.Last(), FOnGetPropertyTypeCustomizationInstance::CreateStatic(&FAnimDetailsValueCustomization::MakeInstance));
 
-	PropertiesToUnregisterOnShutdown.Add(FAnimDetailProxyVector2D::StaticStruct()->GetFName());
-	PropertyEditorModule.RegisterCustomPropertyTypeLayout(PropertiesToUnregisterOnShutdown.Last(), FOnGetPropertyTypeCustomizationInstance::CreateStatic(&FAnimDetailValueCustomization::MakeInstance));
+	PropertiesToUnregisterOnShutdown.Add(FAnimDetailsVector2D::StaticStruct()->GetFName());
+	PropertyEditorModule.RegisterCustomPropertyTypeLayout(PropertiesToUnregisterOnShutdown.Last(), FOnGetPropertyTypeCustomizationInstance::CreateStatic(&FAnimDetailsValueCustomization::MakeInstance));
 
-	PropertiesToUnregisterOnShutdown.Add(FAnimDetailProxyLocation::StaticStruct()->GetFName());
-	PropertyEditorModule.RegisterCustomPropertyTypeLayout(PropertiesToUnregisterOnShutdown.Last(), FOnGetPropertyTypeCustomizationInstance::CreateStatic(&FAnimDetailValueCustomization::MakeInstance));
+	PropertiesToUnregisterOnShutdown.Add(FAnimDetailsLocation::StaticStruct()->GetFName());
+	PropertyEditorModule.RegisterCustomPropertyTypeLayout(PropertiesToUnregisterOnShutdown.Last(), FOnGetPropertyTypeCustomizationInstance::CreateStatic(&FAnimDetailsValueCustomization::MakeInstance));
 
-	PropertiesToUnregisterOnShutdown.Add(FAnimDetailProxyRotation::StaticStruct()->GetFName());
-	PropertyEditorModule.RegisterCustomPropertyTypeLayout(PropertiesToUnregisterOnShutdown.Last(), FOnGetPropertyTypeCustomizationInstance::CreateStatic(&FAnimDetailValueCustomization::MakeInstance));
+	PropertiesToUnregisterOnShutdown.Add(FAnimDetailsRotation::StaticStruct()->GetFName());
+	PropertyEditorModule.RegisterCustomPropertyTypeLayout(PropertiesToUnregisterOnShutdown.Last(), FOnGetPropertyTypeCustomizationInstance::CreateStatic(&FAnimDetailsValueCustomization::MakeInstance));
 
-	PropertiesToUnregisterOnShutdown.Add(FAnimDetailProxyScale::StaticStruct()->GetFName());
-	PropertyEditorModule.RegisterCustomPropertyTypeLayout(PropertiesToUnregisterOnShutdown.Last(), FOnGetPropertyTypeCustomizationInstance::CreateStatic(&FAnimDetailValueCustomization::MakeInstance));
+	PropertiesToUnregisterOnShutdown.Add(FAnimDetailsScale::StaticStruct()->GetFName());
+	PropertyEditorModule.RegisterCustomPropertyTypeLayout(PropertiesToUnregisterOnShutdown.Last(), FOnGetPropertyTypeCustomizationInstance::CreateStatic(&FAnimDetailsValueCustomization::MakeInstance));
 	
 	PropertiesToUnregisterOnShutdown.Add(FRigVMPythonSettings::StaticStruct()->GetFName());
 	PropertyEditorModule.RegisterCustomPropertyTypeLayout(PropertiesToUnregisterOnShutdown.Last(), FOnGetPropertyTypeCustomizationInstance::CreateStatic(&FControlRigPythonLogDetails::MakeInstance));
@@ -264,8 +205,8 @@ void FControlRigEditorModule::StartupModule()
 	PropertiesToUnregisterOnShutdown.Add(FRigVMDrawContainer::StaticStruct()->GetFName());
 	PropertyEditorModule.RegisterCustomPropertyTypeLayout(PropertiesToUnregisterOnShutdown.Last(), FOnGetPropertyTypeCustomizationInstance::CreateStatic(&FControlRigDrawContainerDetails::MakeInstance));
 
-	PropertiesToUnregisterOnShutdown.Add(FControlRigEnumControlProxyValue::StaticStruct()->GetFName());
-	PropertyEditorModule.RegisterCustomPropertyTypeLayout(PropertiesToUnregisterOnShutdown.Last(), FOnGetPropertyTypeCustomizationInstance::CreateStatic(&FControlRigEnumControlProxyValueDetails::MakeInstance));
+	PropertiesToUnregisterOnShutdown.Add(FAnimDetailsEnum::StaticStruct()->GetFName());
+	PropertyEditorModule.RegisterCustomPropertyTypeLayout(PropertiesToUnregisterOnShutdown.Last(), FOnGetPropertyTypeCustomizationInstance::CreateStatic(&FAnimDetailsValueEnumCustomization::MakeInstance));
 
 	PropertiesToUnregisterOnShutdown.Add(FRigElementKey::StaticStruct()->GetFName());
 	PropertyEditorModule.RegisterCustomPropertyTypeLayout(PropertiesToUnregisterOnShutdown.Last(), FOnGetPropertyTypeCustomizationInstance::CreateStatic(&FRigElementKeyDetails::MakeInstance));
