Anim In Engine: Improved Anim Details

diff --git a/Engine/Plugins/Animation/ControlRig/Source/ControlRigEditor/Private/EditMode/ControlRigEditMode.cpp b/Engine/Plugins/Animation/ControlRig/Source/ControlRigEditor/Private/EditMode/ControlRigEditMode.cpp
--- a/Engine/Plugins/Animation/ControlRig/Source/ControlRigEditor/Private/EditMode/ControlRigEditMode.cpp
+++ b/Engine/Plugins/Animation/ControlRig/Source/ControlRigEditor/Private/EditMode/ControlRigEditMode.cpp
@@ -4,6 +4,17 @@
 
 #include "AnimationCoreLibrary.h"
 #include "AnimationEditorPreviewActor.h"
+#include "AnimDetails/AnimDetailsProxyManager.h"
+#include "AnimDetails/AnimDetailsSelection.h"
+#include "AnimDetails/Proxies/AnimDetailsProxyBool.h"
+#include "AnimDetails/Proxies/AnimDetailsProxyEnum.h"
+#include "AnimDetails/Proxies/AnimDetailsProxyFloat.h"
+#include "AnimDetails/Proxies/AnimDetailsProxyInteger.h"
+#include "AnimDetails/Proxies/AnimDetailsProxyLocation.h"
+#include "AnimDetails/Proxies/AnimDetailsProxyRotation.h"
+#include "AnimDetails/Proxies/AnimDetailsProxyScale.h"
+#include "AnimDetails/Proxies/AnimDetailsProxyTransform.h"
+#include "AnimDetails/Proxies/AnimDetailsProxyVector2D.h"
 #include "Components/StaticMeshComponent.h"
 #include "EditMode/ControlRigEditModeToolkit.h"
 #include "Toolkits/ToolkitManager.h"
@@ -11,7 +22,6 @@
 #include "ControlRig.h"
 #include "HitProxies.h"
 #include "EditMode/ControlRigEditModeSettings.h"
-#include "EditMode/SControlRigDetails.h"
 #include "EditMode/SControlRigOutliner.h"
 #include "ISequencer.h"
 #include "MVVM/ViewModels/SequencerEditorViewModel.h"
@@ -35,7 +45,6 @@
 #include "ControlRigBlueprint.h"
 #include "ControlRigGizmoActor.h"
 #include "SEditorViewport.h"
-#include "EditMode/AnimDetailsProxy.h"
 #include "ScopedTransaction.h"
 #include "Rigs/AdditiveControlRig.h"
 #include "Rigs/FKControlRig.h"
@@ -187,8 +196,8 @@ FControlRigEditMode::FControlRigEditMode()
 	, bShowControlsAsOverlay(false)
 	, bIsConstructionEventRunning(false)
 {
-	ControlProxy = NewObject<UControlRigDetailPanelControlProxies>(GetTransientPackage(), NAME_None);
-	ControlProxy->SetFlags(RF_Transactional);
+	AnimDetailsProxyManager = NewObject<UAnimDetailsProxyManager>(GetTransientPackage(), NAME_None, RF_Transactional);
+
 	StoredPose = NewObject<UControlRigPoseAsset>(GetTransientPackage(), NAME_None);
 	DetailKeyFrameCache = MakeShared<FDetailKeyFrameCacheAndHandler>();
 
@@ -283,7 +292,7 @@ bool FControlRigEditMode::SetSequencer(TWeakPtr<ISequencer> InSequencer)
 			
 			SetSequencerDelegates(WeakSequencer);
 			
-			ControlProxy->SetSequencer(WeakSequencer);
+			AnimDetailsProxyManager->NotifySequencerChanged();
 
 			{
 				TSharedRef<SSequencer> SequencerWidget = StaticCastSharedRef<SSequencer>(Sequencer->GetSequencerWidget());
@@ -292,10 +301,6 @@ bool FControlRigEditMode::SetSequencer(TWeakPtr<ISequencer> InSequencer)
 		}
 		
 		SetObjects_Internal();
-		if (FControlRigEditModeToolkit::Details.IsValid())
-		{
-			FControlRigEditModeToolkit::Details->SetEditMode(*this);
-		}
 		if (FControlRigEditModeToolkit::Outliner.IsValid())
 		{
 			FControlRigEditModeToolkit::Outliner->SetEditMode(*this);
@@ -609,9 +614,6 @@ void FControlRigEditMode::Exit()
 	ModeManager->OnWidgetModeChanged().RemoveAll(this);
 	ModeManager->OnCoordSystemChanged().RemoveAll(this);
 
-	//clear proxies
-	ControlProxy->RemoveAllProxies();
-
 	//make sure the widget is reset
 	ResetControlShapeSize();
 
@@ -679,18 +681,6 @@ void FControlRigEditMode::Tick(FEditorViewportClient* ViewportClient, float Delt
 						{
 							ShapeActor->SetSelected(true);
 						}
-
-						if (!bAreEditingControlRigDirectly)
-						{
-							FRigControlElement* ControlElement = ControlRig->FindControl(SelectedKey.Name);
-							if (ControlElement)
-							{
-								if (!ControlRig->IsCurveControl(ControlElement))
-								{
-									ControlProxy->FindOrAddProxy(ControlRig, ControlElement);
-								}
-							}
-						}
 					}
 				}
 			}
@@ -703,6 +693,8 @@ void FControlRigEditMode::Tick(FEditorViewportClient* ViewportClient, float Delt
         {
         	DetailKeyFrameCache->ResetCachedData();
         }
+
+		OnControlRigShapeActorsRecreatedDelegate.Broadcast();
 	}
 
 	{
@@ -2324,10 +2316,9 @@ void FControlRigEditMode::AddReferencedObjects( FReferenceCollector& Collector )
 			Collector.AddReferencedObject(ShapeActor);
 		}
 	}
-	if (ControlProxy)
-	{
-		Collector.AddReferencedObject(ControlProxy);
-	}
+	
+	Collector.AddReferencedObject(AnimDetailsProxyManager);
+	
 	if (StoredPose)
 	{
 		Collector.AddReferencedObject(StoredPose);
@@ -4613,15 +4604,13 @@ void FControlRigEditMode::CreateShapeActors(UControlRig* InControlRig)
 
 	if (!AreEditingControlRigDirectly())
 	{
-		if (ControlProxy)
-		{
-			ControlProxy->RecreateAllProxies(InControlRig);
-		}
 		if (DetailKeyFrameCache)
 		{
 			DetailKeyFrameCache->ResetCachedData();
 		}
 	}
+
+	OnControlRigShapeActorsRecreatedDelegate.Broadcast();
 }
 
 FControlRigEditMode* FControlRigEditMode::GetEditModeFromWorldContext(UWorld* InWorldContext)
@@ -4857,8 +4846,10 @@ void FControlRigEditMode::OnHierarchyModified(ERigHierarchyNotification InNotif,
 				// are correctly selected in sequencer to match what the proxy control is driving.
 				if (FRigControlElement* ControlElement = InHierarchy->Find<FRigControlElement>(InElement->GetKey()))
 				{
-					UControlRig* ControlRig = InHierarchy->GetTypedOuter<UControlRig>();
-					if(ControlProxy->IsSelected(ControlRig, ControlElement))
+					const UControlRig* ControlRig = InHierarchy->GetTypedOuter<UControlRig>();
+					const UAnimDetailsSelection* Selection = AnimDetailsProxyManager->GetAnimDetailsSelection();
+
+					if (Selection && Selection->IsControlElementSelected(ControlRig, ControlElement))
 					{
 						// reselect the control - to affect the details panel / sequencer
 						if(URigHierarchyController* Controller = InHierarchy->GetController())
@@ -4917,76 +4908,38 @@ void FControlRigEditMode::OnHierarchyModified(ERigHierarchyNotification InNotif,
 					if (Key.Type == ERigElementType::Control)
 					{
 						FScopedTransaction ScopedTransaction(LOCTEXT("SelectControlTransaction", "Select Control"), !AreEditingControlRigDirectly() && !GIsTransacting);
-						if (!AreEditingControlRigDirectly())
-						{
-							ControlProxy->Modify();
-						}
-						
 						AControlRigShapeActor* ShapeActor = GetControlShapeFromControlName(ControlRig,Key.Name);
 						if (ShapeActor)
 						{
 							ShapeActor->SetSelected(bSelected);
-
 						}
-						if (!AreEditingControlRigDirectly())
+
+						const FRigControlElement* ControlElement = ControlRig->GetHierarchy()->Find<FRigControlElement>(Key);
+						if (!AreEditingControlRigDirectly() &&
+							ControlElement &&
+							ControlElement->CanDriveControls())
 						{
-							if (FRigControlElement* ControlElement = ControlRig->GetHierarchy()->Find<FRigControlElement>(Key))
-							{
-								ControlProxy->SelectProxy(ControlRig, ControlElement, bSelected);
+							const UControlRigEditModeSettings* Settings = GetSettings();
 
-								if(ControlElement->CanDriveControls())
+							const TArray<FRigElementKey>& DrivenKeys = ControlElement->Settings.DrivenControls;
+							for(const FRigElementKey& DrivenKey : DrivenKeys)
+							{
+								if (FRigControlElement* DrivenControl = ControlRig->GetHierarchy()->Find<FRigControlElement>(DrivenKey))
 								{
-									const UControlRigEditModeSettings* Settings = GetSettings();
-
-									const TArray<FRigElementKey>& DrivenKeys = ControlElement->Settings.DrivenControls;
-									for(const FRigElementKey& DrivenKey : DrivenKeys)
+									if (AControlRigShapeActor* DrivenShapeActor = GetControlShapeFromControlName(ControlRig,DrivenControl->GetFName()))
 									{
-										if (FRigControlElement* DrivenControl = ControlRig->GetHierarchy()->Find<FRigControlElement>(DrivenKey))
+										if(bSelected)
 										{
-											ControlProxy->SelectProxy(ControlRig, DrivenControl, bSelected);
-
-											if (AControlRigShapeActor* DrivenShapeActor = GetControlShapeFromControlName(ControlRig,DrivenControl->GetFName()))
-											{
-												if(bSelected)
-												{
-													DrivenShapeActor->OverrideColor = Settings->DrivenControlColor;
-												}
-												else
-												{
-													DrivenShapeActor->OverrideColor = FLinearColor(0, 0, 0, 0);
-												}
-											}
+											DrivenShapeActor->OverrideColor = Settings->DrivenControlColor;
 										}
-									}
-
-									ControlRig->GetHierarchy()->ForEach<FRigControlElement>(
-										[this, ControlRig, ControlElement, DrivenKeys, bSelected](FRigControlElement* AnimationChannelControl) -> bool
+										else
 										{
-											if(AnimationChannelControl->IsAnimationChannel())
-											{
-												if(const FRigControlElement* ParentControlElement =
-													Cast<FRigControlElement>(ControlRig->GetHierarchy()->GetFirstParent(AnimationChannelControl)))
-												{
-													if(DrivenKeys.Contains(ParentControlElement->GetKey()) ||
-														ParentControlElement->GetKey() == ControlElement->GetKey())
-													{
-														ControlProxy->SelectProxy(ControlRig, AnimationChannelControl, bSelected);
-														return true;
-													}
-												}
-												if(AnimationChannelControl->Settings.Customization.AvailableSpaces.FindByKey(ControlElement->GetKey()))
-												{
-													ControlProxy->SelectProxy(ControlRig, AnimationChannelControl, bSelected);
-													return true;
-												}
-											}
-											return true;
+											DrivenShapeActor->OverrideColor = FLinearColor(0, 0, 0, 0);
 										}
-									);
+									}
 								}
 							}
 						}
-
 					}
 					bSelectionChanged = true;
 		
@@ -5074,13 +5027,16 @@ void FControlRigEditMode::OnHierarchyModified_AnyThread(ERigHierarchyNotificatio
 }
 
 void FControlRigEditMode::OnControlModified(UControlRig* Subject, FRigControlElement* InControlElement, const FRigControlModifiedContext& Context)
-{
-	// do not propagate the change to the proxy when playing for performances reasons
-	if (!bSequencerPlaying)
-	{
-		//this makes sure the details panel ui get's updated, don't remove
-		const bool bModify = Context.SetKey != EControlRigSetKey::Never;
-		ControlProxy->ProxyChanged(Subject, InControlElement, bModify);
+{			
+	// This makes sure the details panel ui get's updated, don't remove.
+	// This may be called from other threads, but only calls on the game thread are not relevant to update the anim details.
+	if (IsInGameThread() && AnimDetailsProxyManager)
+	{	
+		// do not propagate the change to the anim details when playing for performances reasons
+		if (!bSequencerPlaying)
+		{
+			AnimDetailsProxyManager->RequestUpdateProxyValues();
+		}
 	}
 
 	bPivotsNeedUpdate = true;
@@ -6744,101 +6700,87 @@ FInteractionDependencyCache& FControlRigEditMode::GetInteractionDependencies(UCo
 
 bool FDetailKeyFrameCacheAndHandler::IsPropertyKeyable(const UClass* InObjectClass, const IPropertyHandle& InPropertyHandle) const
 {
+	const FProperty* Property = InPropertyHandle.GetProperty();
+	if (!InObjectClass || !Property)
+	{
+		return false;
+	}
+	const FName PropertyName = Property->GetFName();
+
 	TArray<UObject*> OuterObjects;
 	InPropertyHandle.GetOuterObjects(OuterObjects);
-	if (OuterObjects.Num() == 1)
+	for (UObject* OuterObject : OuterObjects)
 	{
-		if (UControlRigControlsProxy* Proxy = Cast< UControlRigControlsProxy>(OuterObjects[0]))
+		if (UAnimDetailsProxyBase* Proxy = Cast<UAnimDetailsProxyBase>(OuterObject))
 		{
-			for (const TPair<TWeakObjectPtr<UControlRig>, FControlRigProxyItem>& Items : Proxy->ControlRigItems)
+			UControlRig* ControlRig = Proxy->GetControlRig();
+			FRigControlElement* ControlElement = Proxy->GetControlElement();
+
+			if (ControlRig && ControlElement)
 			{
-				if (UControlRig* ControlRig = Items.Value.ControlRig.Get())
+				if (!ControlRig->GetHierarchy()->IsAnimatable(ControlElement))
 				{
-					for (const FName& CName : Items.Value.ControlElements)
-					{
-						if (FRigControlElement* ControlElement = Items.Value.GetControlElement(CName))
-						{
-							if (!ControlRig->GetHierarchy()->IsAnimatable(ControlElement))
-							{
-								return false;
-							}
-						}
-					}
+					return false;
 				}
 			}
 		}
 	}
-	
 
-	if (InObjectClass
-		&& InObjectClass->IsChildOf(UAnimDetailControlsProxyTransform::StaticClass())
-		&& InObjectClass->IsChildOf(UAnimDetailControlsProxyLocation::StaticClass())
-		&& InObjectClass->IsChildOf(UAnimDetailControlsProxyRotation::StaticClass())
-		&& InObjectClass->IsChildOf(UAnimDetailControlsProxyScale::StaticClass())
-		&& InObjectClass->IsChildOf(UAnimDetailControlsProxyVector2D::StaticClass())
-		&& InObjectClass->IsChildOf(UAnimDetailControlsProxyFloat::StaticClass())
-		&& InObjectClass->IsChildOf(UAnimDetailControlsProxyBool::StaticClass())
-		&& InObjectClass->IsChildOf(UAnimDetailControlsProxyInteger::StaticClass())
-		)
+	if (InObjectClass->IsChildOf(UAnimLayer::StaticClass()))
 	{
 		return true;
 	}
 
-	if (InObjectClass && InObjectClass->IsChildOf(UAnimLayer::StaticClass()))
+	if (InObjectClass->IsChildOf(UAnimDetailsProxyTransform::StaticClass()) &&
+		(PropertyName == GET_MEMBER_NAME_CHECKED(UAnimDetailsProxyTransform, Location) ||
+		PropertyName == GET_MEMBER_NAME_CHECKED(UAnimDetailsProxyTransform, Rotation) ||
+		PropertyName == GET_MEMBER_NAME_CHECKED(UAnimDetailsProxyTransform, Scale)))
 	{
 		return true;
 	}
 
-	if ((InObjectClass && InObjectClass->IsChildOf(UAnimDetailControlsProxyTransform::StaticClass()) && InPropertyHandle.GetProperty())
-		&& (InPropertyHandle.GetProperty()->GetFName() == GET_MEMBER_NAME_CHECKED(UAnimDetailControlsProxyTransform, Location) ||
-		InPropertyHandle.GetProperty()->GetFName() == GET_MEMBER_NAME_CHECKED(UAnimDetailControlsProxyTransform, Rotation) ||
-		InPropertyHandle.GetProperty()->GetFName() == GET_MEMBER_NAME_CHECKED(UAnimDetailControlsProxyTransform, Scale)))
+	if (InObjectClass->IsChildOf(UAnimDetailsProxyLocation::StaticClass()) &&
+		PropertyName == GET_MEMBER_NAME_CHECKED(UAnimDetailsProxyLocation, Location))
 	{
 		return true;
 	}
 
-
-	if (InObjectClass && InObjectClass->IsChildOf(UAnimDetailControlsProxyLocation::StaticClass()) && InPropertyHandle.GetProperty()
-		&& InPropertyHandle.GetProperty()->GetFName() == GET_MEMBER_NAME_CHECKED(UAnimDetailControlsProxyLocation, Location))
+	if (InObjectClass->IsChildOf(UAnimDetailsProxyRotation::StaticClass()) &&
+		PropertyName == GET_MEMBER_NAME_CHECKED(UAnimDetailsProxyRotation, Rotation))
 	{
 		return true;
 	}
 
-	if (InObjectClass && InObjectClass->IsChildOf(UAnimDetailControlsProxyRotation::StaticClass()) && InPropertyHandle.GetProperty()
-		&& InPropertyHandle.GetProperty()->GetFName() == GET_MEMBER_NAME_CHECKED(UAnimDetailControlsProxyRotation, Rotation))
+	if (InObjectClass->IsChildOf(UAnimDetailsProxyScale::StaticClass()) &&
+		PropertyName == GET_MEMBER_NAME_CHECKED(UAnimDetailsProxyScale, Scale))
 	{
 		return true;
 	}
 
-	if (InObjectClass && InObjectClass->IsChildOf(UAnimDetailControlsProxyScale::StaticClass()) && InPropertyHandle.GetProperty()
-		&& InPropertyHandle.GetProperty()->GetFName() == GET_MEMBER_NAME_CHECKED(UAnimDetailControlsProxyScale, Scale))
+	if (InObjectClass->IsChildOf(UAnimDetailsProxyVector2D::StaticClass()) &&
+		PropertyName == GET_MEMBER_NAME_CHECKED(UAnimDetailsProxyVector2D, Vector2D))
 	{
 		return true;
 	}
 
-	if (InObjectClass && InObjectClass->IsChildOf(UAnimDetailControlsProxyVector2D::StaticClass()) && InPropertyHandle.GetProperty()
-		&& InPropertyHandle.GetProperty()->GetFName() == GET_MEMBER_NAME_CHECKED(UAnimDetailControlsProxyVector2D, Vector2D))
+	if (InObjectClass->IsChildOf(UAnimDetailsProxyInteger::StaticClass()) && 
+		PropertyName == GET_MEMBER_NAME_CHECKED(UAnimDetailsProxyInteger, Integer))
 	{
 		return true;
 	}
 
-	if (InObjectClass && InObjectClass->IsChildOf(UAnimDetailControlsProxyInteger::StaticClass()) && InPropertyHandle.GetProperty()
-		&& InPropertyHandle.GetProperty()->GetFName() == GET_MEMBER_NAME_CHECKED(UAnimDetailControlsProxyInteger, Integer))
+	if (InObjectClass->IsChildOf(UAnimDetailsProxyBool::StaticClass()) && 
+		PropertyName == GET_MEMBER_NAME_CHECKED(UAnimDetailsProxyBool, Bool))
 	{
 		return true;
 	}
 
-	if (InObjectClass && InObjectClass->IsChildOf(UAnimDetailControlsProxyBool::StaticClass()) && InPropertyHandle.GetProperty()
-		&& InPropertyHandle.GetProperty()->GetFName() == GET_MEMBER_NAME_CHECKED(UAnimDetailControlsProxyBool, Bool))
+	if (InObjectClass->IsChildOf(UAnimDetailsProxyFloat::StaticClass()) &&
+		PropertyName == GET_MEMBER_NAME_CHECKED(UAnimDetailsProxyFloat, Float))
 	{
 		return true;
 	}
 
-	if (InObjectClass && InObjectClass->IsChildOf(UAnimDetailControlsProxyFloat::StaticClass()) && InPropertyHandle.GetProperty()
-		&& InPropertyHandle.GetProperty()->GetFName() == GET_MEMBER_NAME_CHECKED(UAnimDetailControlsProxyFloat, Float))
-	{
-		return true;
-	}
 	FCanKeyPropertyParams CanKeyPropertyParams(InObjectClass, InPropertyHandle);
 	if (WeakSequencer.IsValid() && WeakSequencer.Pin()->CanKeyProperty(CanKeyPropertyParams))
 	{
@@ -6881,7 +6823,7 @@ bool FDetailKeyFrameCacheAndHandler::IsPropertyAnimated(const IPropertyHandle& P
 
 void FDetailKeyFrameCacheAndHandler::OnKeyPropertyClicked(const IPropertyHandle& KeyedPropertyHandle)
 {
-	if (WeakSequencer.IsValid() && !WeakSequencer.Pin()->IsAllowedToChange())
+	if (!WeakSequencer.IsValid() || !WeakSequencer.Pin()->IsAllowedToChange())
 	{
 		return;
 	}
@@ -6894,9 +6836,9 @@ void FDetailKeyFrameCacheAndHandler::OnKeyPropertyClicked(const IPropertyHandle&
 	{
 		if (Object)
 		{
-			if (UControlRigControlsProxy* Proxy = Cast< UControlRigControlsProxy>(Object))
+			if (UAnimDetailsProxyBase* Proxy = Cast<UAnimDetailsProxyBase>(Object))
 			{
-				Proxy->SetKey(SequencerPtr, KeyedPropertyHandle);
+				Proxy->SetKey(KeyedPropertyHandle);
 			}
 			else if (UAnimLayer* AnimLayer = (Object->GetTypedOuter<UAnimLayer>()))
 			{
@@ -6944,9 +6886,9 @@ EPropertyKeyedStatus FDetailKeyFrameCacheAndHandler::GetPropertyKeyedStatus(cons
 	{
 		if (Object)
 		{
-			if (UControlRigControlsProxy* Proxy = Cast< UControlRigControlsProxy>(Object))
+			if (UAnimDetailsProxyBase* Proxy = Cast<UAnimDetailsProxyBase>(Object))
 			{
-				KeyedStatus = Proxy->GetPropertyKeyedStatus(SequencerPtr, PropertyHandle);
+				KeyedStatus = Proxy->GetPropertyKeyedStatus(PropertyHandle);
 			}
 			else if (UAnimLayer* AnimLayer = (Object->GetTypedOuter<UAnimLayer>()))
 			{
@@ -7025,11 +6967,12 @@ void FDetailKeyFrameCacheAndHandler::UpdateIfDirty()
 	{
 		if (FMovieSceneConstraintChannelHelper::bDoNotCompensate == false) //if compensating don't reset this.
 		{
-			UControlRigDetailPanelControlProxies* ControlProxies = EditMode ? EditMode->GetDetailProxies() : nullptr;
-			if (ControlProxies)
+			UAnimDetailsProxyManager* ProxyManager = EditMode ? EditMode->GetAnimDetailsProxyManager() : nullptr;
+			if (ProxyManager)
 			{
-				ControlProxies->ValuesChanged();
+				ProxyManager->RequestUpdateProxyValues();
 			}
+
 			bValuesDirty = false;
 		}
 	}
diff --git a/Engine/Plugins/Animation/ControlRig/Source/ControlRigEditor/Private/EditMode/ControlRigEditModeToolkit.cpp b/Engine/Plugins/Animation/ControlRig/Source/ControlRigEditor/Private/EditMode/ControlRigEditModeToolkit.cpp
--- a/Engine/Plugins/Animation/ControlRig/Source/ControlRigEditor/Private/EditMode/ControlRigEditModeToolkit.cpp
+++ b/Engine/Plugins/Animation/ControlRig/Source/ControlRigEditor/Private/EditMode/ControlRigEditModeToolkit.cpp
@@ -5,6 +5,7 @@
 */
 #include "EditMode/ControlRigEditModeToolkit.h"
 
+#include "AnimDetails/Views/SAnimDetailsView.h"
 #include "Widgets/DeclarativeSyntaxSupport.h"
 #include "EditorModes.h"
 #include "Toolkits/BaseToolkit.h"
@@ -18,16 +19,22 @@
 #include "Toolkits/AssetEditorModeUILayer.h"
 #include "Widgets/Docking/SDockTab.h"
 #include "EditMode/ControlRigEditModeSettings.h"
-#include "EditMode/SControlRigDetails.h"
 #include "EditMode/SControlRigOutliner.h"
 #include "EditMode/SControlRigSpacePicker.h"
 #include "LevelEditor.h"
+#include "PropertyEditorModule.h"
 #include "SLevelViewport.h"
 #include "Sequencer/AnimLayers/SAnimLayers.h"
 #include "Sequencer/AnimLayers/AnimLayers.h"
 #include "Tween/SControlRigTweenWidget.h"
 #include "Tools/MotionTrailOptions.h"
 
+class SAnimDetailsView;
+namespace UE::ControlRigEditor
+{
+	class SAnimDetailsView;
+}
+
 #define LOCTEXT_NAMESPACE "FControlRigEditModeToolkit"
 
 namespace 
@@ -51,7 +58,7 @@ const FName FControlRigEditModeToolkit::OutlinerTabName = FName(TEXT("ControlRig
 const FName FControlRigEditModeToolkit::DetailsTabName = FName(TEXT("ControlRigDetailsTab"));
 const FName FControlRigEditModeToolkit::SpacePickerTabName = FName(TEXT("ControlRigSpacePicker"));
 
-TSharedPtr<SControlRigDetails> FControlRigEditModeToolkit::Details = nullptr;
+TSharedPtr<UE::ControlRigEditor::SAnimDetailsView> FControlRigEditModeToolkit::Details = nullptr;
 TSharedPtr<SControlRigOutliner> FControlRigEditModeToolkit::Outliner = nullptr;
 
 FControlRigEditModeToolkit::~FControlRigEditModeToolkit()
@@ -235,7 +242,7 @@ TSharedRef<SDockTab> SpawnDetailsTab(const FSpawnTabArgs& Args, FControlRigEditM
 {
 	return SNew(SDockTab)
 		[
-			SAssignNew(FControlRigEditModeToolkit::Details,SControlRigDetails, *InEditorMode)
+			SAssignNew(FControlRigEditModeToolkit::Details, UE::ControlRigEditor::SAnimDetailsView)
 		];
 }
 
diff --git a/Engine/Plugins/Animation/ControlRig/Source/ControlRigEditor/Private/EditMode/ControlRigEditModeToolkit.h b/Engine/Plugins/Animation/ControlRig/Source/ControlRigEditor/Private/EditMode/ControlRigEditModeToolkit.h
--- a/Engine/Plugins/Animation/ControlRig/Source/ControlRigEditor/Private/EditMode/ControlRigEditModeToolkit.h
+++ b/Engine/Plugins/Animation/ControlRig/Source/ControlRigEditor/Private/EditMode/ControlRigEditModeToolkit.h
@@ -12,9 +12,13 @@
 #include "EditMode/ControlRigEditMode.h"
 #include "Tween/ControlRigTweenModel.h"
 
-class SControlRigDetails;
+class SAnimDetailsView;
 class SControlRigOutliner;
-namespace UE::ControlRigEditor { class SControlRigTweenWidget; }
+namespace UE::ControlRigEditor 
+{ 
+	class SAnimDetailsView;
+	class SControlRigTweenWidget; 
+}
 
 class FControlRigEditModeToolkit : public FModeToolkit
 {
@@ -68,7 +72,7 @@ public:
 	static const FName OutlinerTabName;
 	static const FName SpacePickerTabName;
 
-	static TSharedPtr<SControlRigDetails> Details;
+	static TSharedPtr<UE::ControlRigEditor::SAnimDetailsView> Details;
 	static TSharedPtr<SControlRigOutliner> Outliner;
 
 protected:
diff --git a/Engine/Plugins/Animation/ControlRig/Source/ControlRigEditor/Private/EditMode/SControlRigEditModeTools.cpp b/Engine/Plugins/Animation/ControlRig/Source/ControlRigEditor/Private/EditMode/SControlRigEditModeTools.cpp
--- a/Engine/Plugins/Animation/ControlRig/Source/ControlRigEditor/Private/EditMode/SControlRigEditModeTools.cpp
+++ b/Engine/Plugins/Animation/ControlRig/Source/ControlRigEditor/Private/EditMode/SControlRigEditModeTools.cpp
@@ -1,7 +1,8 @@
 // Copyright Epic Games, Inc. All Rights Reserved.
 
 #include "EditMode/SControlRigEditModeTools.h"
-#include "EditMode/ControlRigControlsProxy.h"
+
+#include "AnimDetails/Proxies/AnimDetailsProxyTransform.h"
 #include "PropertyEditorModule.h"
 #include "IDetailsView.h"
 #include "ISequencer.h"
@@ -39,7 +40,6 @@
 #include "ScopedTransaction.h"
 #include "Transform/TransformConstraint.h"
 #include "EditMode/ControlRigEditModeToolkit.h"
-#include "EditMode/SControlRigDetails.h"
 #include "Editor/Constraints/SConstraintsWidget.h"
 
 #define LOCTEXT_NAMESPACE "ControlRigEditModeTools"
@@ -150,7 +150,6 @@ void SControlRigEditModeTools::Construct(const FArguments& InArgs, TSharedPtr<FC
 	ModeTools = InEditMode.GetModeManager();
 
 	SettingsDetailsView = FModuleManager::GetModuleChecked<FPropertyEditorModule>("PropertyEditor").CreateDetailView(DetailsViewArgs);
-	SettingsDetailsView->SetGenericLayoutDetailsDelegate(FOnGetDetailCustomizationInstance::CreateStatic(&FControlRigEditModeGenericDetails::MakeInstance, ModeTools));
 
 #if USE_LOCAL_DETAILS
 	ControlEulerTransformDetailsView = FModuleManager::GetModuleChecked<FPropertyEditorModule>("PropertyEditor").CreateDetailView(DetailsViewArgs);
@@ -659,14 +658,14 @@ void SControlRigEditModeTools::SetSequencer(TWeakPtr<ISequencer> InSequencer)
 
 bool SControlRigEditModeTools::IsPropertyKeyable(const UClass* InObjectClass, const IPropertyHandle& InPropertyHandle) const
 {
-	if (InObjectClass && InObjectClass->IsChildOf(UAnimDetailControlsProxyTransform::StaticClass()))
+	if (InObjectClass && InObjectClass->IsChildOf(UAnimDetailsProxyTransform::StaticClass()))
 	{
 		return true;
 	}
-	if (InObjectClass && InObjectClass->IsChildOf(UAnimDetailControlsProxyTransform::StaticClass()) && InPropertyHandle.GetProperty()
-		&& (InPropertyHandle.GetProperty()->GetFName() == GET_MEMBER_NAME_CHECKED(UAnimDetailControlsProxyTransform, Location) ||
-		InPropertyHandle.GetProperty()->GetFName() == GET_MEMBER_NAME_CHECKED(UAnimDetailControlsProxyTransform, Rotation) ||
-		InPropertyHandle.GetProperty()->GetFName() == GET_MEMBER_NAME_CHECKED(UAnimDetailControlsProxyTransform, Scale)))
+	if (InObjectClass && InObjectClass->IsChildOf(UAnimDetailsProxyTransform::StaticClass()) && InPropertyHandle.GetProperty()
+		&& (InPropertyHandle.GetProperty()->GetFName() == GET_MEMBER_NAME_CHECKED(UAnimDetailsProxyTransform, Location) ||
+		InPropertyHandle.GetProperty()->GetFName() == GET_MEMBER_NAME_CHECKED(UAnimDetailsProxyTransform, Rotation) ||
+		InPropertyHandle.GetProperty()->GetFName() == GET_MEMBER_NAME_CHECKED(UAnimDetailsProxyTransform, Scale)))
 	{
 		return true;
 	}
@@ -715,7 +714,7 @@ bool SControlRigEditModeTools::IsPropertyAnimated(const IPropertyHandle& Propert
 
 void SControlRigEditModeTools::OnKeyPropertyClicked(const IPropertyHandle& KeyedPropertyHandle)
 {
-	if (WeakSequencer.IsValid() && !WeakSequencer.Pin()->IsAllowedToChange())
+	if (!WeakSequencer.IsValid() || !WeakSequencer.Pin()->IsAllowedToChange())
 	{
 		return;
 	}
@@ -726,10 +725,10 @@ void SControlRigEditModeTools::OnKeyPropertyClicked(const IPropertyHandle& Keyed
 
 	for (UObject *Object : Objects)
 	{
-		UControlRigControlsProxy* Proxy = Cast< UControlRigControlsProxy>(Object);
+		UAnimDetailsProxyBase* Proxy = Cast<UAnimDetailsProxyBase>(Object);
 		if (Proxy)
 	{
-			Proxy->SetKey(SequencerPtr,KeyedPropertyHandle);
+			Proxy->SetKey(KeyedPropertyHandle);
 		}
 	}
 }
diff --git a/Engine/Plugins/Animation/ControlRig/Source/ControlRigEditor/Private/EditMode/SControlRigOutliner.cpp b/Engine/Plugins/Animation/ControlRig/Source/ControlRigEditor/Private/EditMode/SControlRigOutliner.cpp
--- a/Engine/Plugins/Animation/ControlRig/Source/ControlRigEditor/Private/EditMode/SControlRigOutliner.cpp
+++ b/Engine/Plugins/Animation/ControlRig/Source/ControlRigEditor/Private/EditMode/SControlRigOutliner.cpp
@@ -21,7 +21,6 @@
 #include "MVVM/Selection/Selection.h"
 #include "Editor.h"
 #include "Widgets/Text/STextBlock.h"
-#include "EditMode/ControlRigControlsProxy.h"
 #include "TimerManager.h"
 #include "Editor/SRigHierarchyTreeView.h"
 #include "Rigs/RigHierarchyController.h"
diff --git a/Engine/Plugins/Animation/ControlRig/Source/ControlRigEditor/Public/EditMode/ControlRigEditMode.h b/Engine/Plugins/Animation/ControlRig/Source/ControlRigEditor/Public/EditMode/ControlRigEditMode.h
--- a/Engine/Plugins/Animation/ControlRig/Source/ControlRigEditor/Public/EditMode/ControlRigEditMode.h
+++ b/Engine/Plugins/Animation/ControlRig/Source/ControlRigEditor/Public/EditMode/ControlRigEditMode.h
@@ -25,6 +25,7 @@
 
 #include "ControlRigEditMode.generated.h"
 
+class UAnimDetailsProxyManager;
 class UTickableConstraint;
 class FEditorViewportClient;
 class FRigVMEditorBase;
@@ -237,6 +238,9 @@ public:
 	/** Broadcasts a notification when a gizmo manipulation has started */
 	FControlRigEditModeInteractionStartedEvent& OnGizmoInteractionStarted() { return OnGizmoInteractionStartedDelegate; }
 
+	/** Returns a delegate broadcast when control rig shape actors were recreated */
+	FSimpleMulticastDelegate& OnControlRigShapeActorsRecreated() { return OnControlRigShapeActorsRecreatedDelegate; }
+
 	/** Broadcasts a notification when a gizmo manipulation has ended */
 	FControlRigEditModeInteractionEndedEvent& OnGizmoInteractionEnded() { return OnGizmoInteractionEndedDelegate; }
 
@@ -474,6 +478,9 @@ protected:
 	/** Broadcasts a notification when a gizmo manipulation has been updated (before the update actually takes place) */
 	FControlRigEditModeInteractionUpdatedEvent OnGizmoInteractionUpdatedDelegate;
 
+	/** Broadcasts a notification when a control rig shape actors were recreated */
+	FSimpleMulticastDelegate OnControlRigShapeActorsRecreatedDelegate;
+
 	/** GetSelectedRigElements, if InControlRig is nullptr get the first one */
 	TArray<FRigElementKey> GetSelectedRigElements() const;
 	static TArray<FRigElementKey> GetSelectedRigElements(UControlRig* InControlRig);
@@ -491,7 +498,9 @@ protected:
 
 	/** Shape actors */
 	TMap<UControlRig*,TArray<TObjectPtr<AControlRigShapeActor>>> ControlRigShapeActors;
-	TObjectPtr<UControlRigDetailPanelControlProxies> ControlProxy;
+
+	/** Manager for anim details proxies */
+	TObjectPtr<UAnimDetailsProxyManager> AnimDetailsProxyManager;
 
 	/** Utility functions for UI/Some other viewport manipulation*/
 	bool IsControlSelected(const bool bUseShapes = false) const;
@@ -533,7 +542,7 @@ public:
 	TArray<const UControlRig*> GetControlRigsArray(bool bIsVisible) const;
 
 	/** Get the detail proxies control rig*/
-	UControlRigDetailPanelControlProxies* GetDetailProxies() { return ControlProxy; }
+	UAnimDetailsProxyManager* GetAnimDetailsProxyManager() const { return AnimDetailsProxyManager; }
 
 	/** Get Sequencer Driving This*/
 	TWeakPtr<ISequencer> GetWeakSequencer() { return WeakSequencer; }
@@ -756,7 +765,6 @@ private:
 	friend class FControlRigLegacyEditor;
 #endif
 	friend class FControlRigEditor;
-	friend class FControlRigEditModeGenericDetails;
 	friend class UControlRigEditModeDelegateHelper;
 	friend class SControlRigEditModeTools;
 
