Anim in Engine: Fix compile UnrealEditor RTTI Linux, undefined symbol: UE::ControlRigEditor::SAnimDetailsValueNumeric<double>::~SAnimDetailsValueNumeric()

diff --git a/Engine/Plugins/Animation/ControlRig/Source/ControlRigEditor/Private/AnimDetails/Widgets/SAnimDetailsValueNumeric.cpp b/Engine/Plugins/Animation/ControlRig/Source/ControlRigEditor/Private/AnimDetails/Widgets/SAnimDetailsValueNumeric.cpp
--- a/Engine/Plugins/Animation/ControlRig/Source/ControlRigEditor/Private/AnimDetails/Widgets/SAnimDetailsValueNumeric.cpp
+++ b/Engine/Plugins/Animation/ControlRig/Source/ControlRigEditor/Private/AnimDetails/Widgets/SAnimDetailsValueNumeric.cpp
@@ -172,6 +172,9 @@ namespace UE::ControlRigEditor
 		FAnimDetailsMultiEditUtil::Get().Leave(WeakPropertyHandle);
 	}
 
+	template SAnimDetailsValueNumeric<double>::~SAnimDetailsValueNumeric();
+	template SAnimDetailsValueNumeric<int64>::~SAnimDetailsValueNumeric();
+
 	template <typename NumericType>
 	void SAnimDetailsValueNumeric<NumericType>::Construct(const FArguments& InArgs, const TSharedRef<IPropertyHandle>& InPropertyHandle)
 	{
