[Sequencer & Control Rig] Added "Recreate Control Rig with New Settings"

diff --git a/Engine/Plugins/Animation/ControlRig/Source/ControlRig/Private/Sequencer/MovieSceneControlRigParameterTrack.cpp b/Engine/Plugins/Animation/ControlRig/Source/ControlRig/Private/Sequencer/MovieSceneControlRigParameterTrack.cpp
--- a/Engine/Plugins/Animation/ControlRig/Source/ControlRig/Private/Sequencer/MovieSceneControlRigParameterTrack.cpp
+++ b/Engine/Plugins/Animation/ControlRig/Source/ControlRig/Private/Sequencer/MovieSceneControlRigParameterTrack.cpp
@@ -530,6 +530,7 @@ void UMovieSceneControlRigParameterTrack::ReconstructControlRig()
 	if (ControlRig && !ControlRig->HasAnyFlags(RF_NeedLoad | RF_NeedPostLoad | RF_NeedInitialization))
 	{
 		ControlRig->ConditionalPostLoad();
+		ApplyControlRigSettingsOverrides();
 		ControlRig->Initialize();
 		for (UMovieSceneSection* Section : Sections)
 		{
@@ -845,6 +846,43 @@ TOptional<EEulerRotationOrder> UMovieSceneControlRigParameterTrack::GetControlRo
 	return Order;
 }
 
+void UMovieSceneControlRigParameterTrack::UpdateAndApplyControlRigSettingsOverrides(const FInstancedPropertyBag& InNewOverrides)
+{
+	Modify();
+	ControlRigSettingsOverrides = InNewOverrides;
+	ApplyControlRigSettingsOverrides();
+}
+
+void UMovieSceneControlRigParameterTrack::ApplyControlRigSettingsOverrides()
+{
+	if (!ControlRig)
+	{
+		return;
+	}
+
+	if (!ControlRigSettingsOverrides.IsValid())
+	{
+		return;
+	}
+	
+	for (const FPropertyBagPropertyDesc& PropertyDesc : ControlRigSettingsOverrides.GetPropertyBagStruct()->GetPropertyDescs())
+	{
+		FRigVMExternalVariable Variable = ControlRig->GetPublicVariableByName(PropertyDesc.Name);
+		if (Variable.IsValid())
+		{
+			if (Variable.Property->GetClass() == PropertyDesc.CachedProperty->GetClass())
+			{
+				// copy from property bag to control rig
+				void * TargetAddress = Variable.Property->ContainerPtrToValuePtr<void>(ControlRig);
+				void const* SourceAddress = ControlRigSettingsOverrides.GetValue().GetMemory() + PropertyDesc.CachedProperty->GetOffset_ForInternal();
+				Variable.Property->CopyCompleteValue(TargetAddress, SourceAddress);
+			}
+		}
+	}
+	
+}
+
+
 TArray<FName> UMovieSceneControlRigParameterTrack::GetControlsWithDifferentRotationOrders() const
 {
 	TArray<FName> Names;
diff --git a/Engine/Plugins/Animation/ControlRig/Source/ControlRig/Public/Sequencer/MovieSceneControlRigParameterTrack.h b/Engine/Plugins/Animation/ControlRig/Source/ControlRig/Public/Sequencer/MovieSceneControlRigParameterTrack.h
--- a/Engine/Plugins/Animation/ControlRig/Source/ControlRig/Public/Sequencer/MovieSceneControlRigParameterTrack.h
+++ b/Engine/Plugins/Animation/ControlRig/Source/ControlRig/Public/Sequencer/MovieSceneControlRigParameterTrack.h
@@ -168,6 +168,9 @@ public:
 	UMovieSceneControlRigParameterSection::FSpaceChannelAddedEvent& SpaceChannelAdded() { return OnSpaceChannelAdded; }
 	IMovieSceneConstrainedSection::FConstraintChannelAddedEvent& ConstraintChannelAdded() { return OnConstraintChannelAdded; }
 
+	CONTROLRIG_API void UpdateAndApplyControlRigSettingsOverrides(const FInstancedPropertyBag& InNewOverrides);
+	CONTROLRIG_API void ApplyControlRigSettingsOverrides();
+	
 public:
 	CONTROLRIG_API TArray<FName> GetControlsWithDifferentRotationOrders() const;
 	CONTROLRIG_API void ResetControlsToSettingsRotationOrder(const TArray<FName>& Names,
@@ -227,6 +230,11 @@ private:
 	UPROPERTY()
 	int32 PriorityOrder;
 
+	// Stores control rig public variable overrides, will switch to the same override system that rig module
+	// uses in the future
+	UPROPERTY()
+	FInstancedPropertyBag ControlRigSettingsOverrides;
+
 public:
 	static CONTROLRIG_API FColor AbsoluteRigTrackColor;
 	static CONTROLRIG_API FColor LayeredRigTrackColor;
diff --git a/Engine/Plugins/Animation/ControlRig/Source/ControlRigEditor/Private/Sequencer/ControlRigParameterTrackEditor.cpp b/Engine/Plugins/Animation/ControlRig/Source/ControlRigEditor/Private/Sequencer/ControlRigParameterTrackEditor.cpp
--- a/Engine/Plugins/Animation/ControlRig/Source/ControlRigEditor/Private/Sequencer/ControlRigParameterTrackEditor.cpp
+++ b/Engine/Plugins/Animation/ControlRig/Source/ControlRigEditor/Private/Sequencer/ControlRigParameterTrackEditor.cpp
@@ -112,6 +112,7 @@
 #include "CurveModel.h"
 #include "Sequencer/SLoadAnimToControlRig.h"
 #include "FrontendFilterBase.h"
+#include "InstancedPropertyBagStructureDataProvider.h"
 #include "Editor/RigVMEditorTools.h"
 
 #define LOCTEXT_NAMESPACE "FControlRigParameterTrackEditor"
@@ -621,6 +622,11 @@ void FControlRigParameterTrackEditor::ObjectImplicitlyRemoved(UObject* InObject)
 	{
 		UnbindControlRig(ControlRig);
 	}
+
+	if (RecreateRigOperator.IsValid() && RecreateRigOperator->IsEditingObject(InObject))
+	{
+		RecreateRigOperator->Abort();
+	}
 }
 
 void FControlRigParameterTrackEditor::OnRelease()
@@ -1489,6 +1495,253 @@ void FControlRigParameterTrackEditor::ConvertIsLayered(UMovieSceneControlRigPara
 	UControlRigSequencerEditorLibrary::SetControlRigLayeredMode(Track, bSetAdditive);
 }
 
+void FControlRigParameterTrackEditor::RecreateControlRigWithNewSettings(UMovieSceneControlRigParameterTrack* InTrack)
+{
+	CreateAndShowRigSettingsWindow(InTrack);
+}
+
+bool FControlRigParameterTrackEditor::CanRecreateControlRigWithNewSettings(UMovieSceneControlRigParameterTrack* InTrack)
+{
+	UControlRig* ControlRig = InTrack->GetControlRig();
+	if (!ensure(ControlRig))
+	{
+		return false;
+	}
+
+	if (ControlRig->GetPublicVariables().Num() == 0)
+	{
+		// no public variable means the rig is not configurable
+		return false;
+	}
+
+	return true;
+}
+
+
+void FControlRigParameterTrackEditor::CreateAndShowRigSettingsWindow(UMovieSceneControlRigParameterTrack* Track)
+{
+	RecreateRigOperator = MakeShared<FControlRigParameterTrackEditor::FRecreateRigOperator>();
+
+	RecreateRigOperator->Start(SharedThis(this), Track);
+}
+
+void FControlRigParameterTrackEditor::ResetRecreateRigOperatorIfNeeded(TSharedRef<FControlRigParameterTrackEditor::FRecreateRigOperator> InRequestingOperator)
+{
+	if (InRequestingOperator == RecreateRigOperator)
+	{
+		RecreateRigOperator.Reset();
+	}
+}
+
+
+FControlRigParameterTrackEditor::FRecreateRigOperator::~FRecreateRigOperator()
+{
+	if (WeakDetailsView.IsValid())
+	{
+		WeakDetailsView.Pin()->GetOnFinishedChangingPropertiesDelegate().RemoveAll(this);
+	}
+	
+	if (WeakControlRig.IsValid())
+	{
+		WeakControlRig->OnPostForwardsSolve_AnyThread().RemoveAll(this);
+	}
+
+}
+
+void FControlRigParameterTrackEditor::FRecreateRigOperator::Start(const TSharedRef<FControlRigParameterTrackEditor>& InTrackEditor,
+	UMovieSceneControlRigParameterTrack* InTrack)
+{
+	UControlRig* ControlRig = InTrack->GetControlRig();
+
+	if (!ControlRig)
+	{
+		return;
+	}
+	
+	TArray<FPropertyBagPropertyDesc> BagProperties;
+	TArray<const FProperty*> SourceProperties;
+	TArray<FRigVMExternalVariable> PublicVariables = ControlRig->GetPublicVariables();
+	for (const FRigVMExternalVariable& Variable : PublicVariables)
+	{
+		const FProperty* SourceProperty = ControlRig->GetClass()->FindPropertyByName(Variable.Name);
+		SourceProperties.Add(SourceProperty);
+		
+		BagProperties.Add(FPropertyBagPropertyDesc(Variable.Name , SourceProperty));
+	}
+
+	if (BagProperties.Num() == 0)
+	{
+		return;
+	}
+
+	WeakTrackEditor = InTrackEditor;
+	WeakTrack = InTrack;
+	WeakControlRig = ControlRig;
+	SettingsForNewControlRig.AddProperties(BagProperties);
+
+	FDetailsViewArgs DetailsViewArgs;
+	DetailsViewArgs.NameAreaSettings = FDetailsViewArgs::HideNameArea;
+	FStructureDetailsViewArgs StructureDetailsViewArgs;
+	
+	TSharedRef<IStructureDetailsView> DetailsView=
+		FModuleManager::GetModuleChecked<FPropertyEditorModule>("PropertyEditor")
+		.CreateStructureProviderDetailView(
+			DetailsViewArgs,
+			StructureDetailsViewArgs,
+			MakeShared<FInstancePropertyBagStructureDataProvider>(SettingsForNewControlRig));
+
+	DetailsView->GetOnFinishedChangingPropertiesDelegate().AddSP( this,
+		&FControlRigParameterTrackEditor::FRecreateRigOperator::OnRigSettingsChanged);
+
+	WeakDetailsView = DetailsView;
+	
+	ControlRig->OnPostForwardsSolve_AnyThread().AddSP(
+		this,
+		&FControlRigParameterTrackEditor::FRecreateRigOperator::OnPostControlRigForwardSolve_AnyThread
+		);
+
+	
+	TSharedRef<SWindow> Window = SNew(SWindow)
+		.Title( LOCTEXT("RecreateControlRigWithNewSettingsSettingsWindowTitle", "Recreate Control Rig With New Settings") )
+		.bDragAnywhere(true)
+		.Type(EWindowType::Normal)
+		.IsTopmostWindow(true)
+		.SizingRule( ESizingRule::Autosized )
+		.FocusWhenFirstShown(true)
+		.ActivationPolicy(EWindowActivationPolicy::FirstShown)
+		[
+			DetailsView->GetWidget().ToSharedRef()	
+		];
+	
+		
+	FSlateApplication::Get().AddWindow( Window, true);
+	Window->SetOnWindowClosed(FOnWindowClosed::CreateSP(this, &FControlRigParameterTrackEditor::FRecreateRigOperator::OnRigSettingsWindowClosed));
+	
+	WeakWindow = Window;
+}
+
+bool FControlRigParameterTrackEditor::FRecreateRigOperator::IsValid()
+{
+	if (!WeakTrackEditor.IsValid())
+	{
+		return false;
+	}
+	if (!WeakTrack.IsValid())
+	{
+		return false;
+	}
+	
+	UControlRig* ControlRig = WeakTrack.Get()->GetControlRig();
+
+	if (!ControlRig)
+	{
+		return false;
+	}
+
+	if (ControlRig != WeakControlRig)
+	{
+		return false;
+	}
+
+	return true;
+}
+
+void FControlRigParameterTrackEditor::FRecreateRigOperator::Abort()
+{
+	if (WeakWindow.IsValid())
+	{
+		WeakWindow.Pin()->RequestDestroyWindow();
+		WeakWindow.Reset();
+	}
+}
+
+void FControlRigParameterTrackEditor::FRecreateRigOperator::OnRigSettingsWindowClosed(const TSharedRef<SWindow>& Window)
+{
+	if (WeakTrackEditor.IsValid())
+	{
+		WeakTrackEditor.Pin()->ResetRecreateRigOperatorIfNeeded(SharedThis(this));
+	}
+}
+
+void FControlRigParameterTrackEditor::FRecreateRigOperator::RefreshSettingsFromControlRig()
+{
+	if (!IsValid())
+	{
+		Abort();
+		return;
+	}
+	
+	UControlRig* ControlRig = WeakTrack.Get()->GetControlRig();
+
+	TArray<const FProperty*> SourceProperties;
+	TArray<FRigVMExternalVariable> PublicVariables = ControlRig->GetPublicVariables();
+	for (const FRigVMExternalVariable& Variable : PublicVariables)
+	{
+		const FProperty* SourceProperty = ControlRig->GetClass()->FindPropertyByName(Variable.Name);
+		SourceProperties.Add(SourceProperty);
+	}
+
+	for (const FProperty* SourceProperty : SourceProperties)
+	{
+		SettingsForNewControlRig.SetValue(SourceProperty->GetFName(), SourceProperty, ControlRig); 
+	}	
+}
+
+void FControlRigParameterTrackEditor::FRecreateRigOperator::OnPostControlRigForwardSolve_AnyThread(UControlRig*, const FName&)
+{
+	FFunctionGraphTask::CreateAndDispatchWhenReady(
+		[WeakThis = TWeakPtr<FControlRigParameterTrackEditor::FRecreateRigOperator>(SharedThis(this))]()
+	{
+		if (WeakThis.IsValid())
+		{
+			WeakThis.Pin()->RefreshSettingsFromControlRig();
+		}
+	}, TStatId(), nullptr, ENamedThreads::GameThread);
+}
+
+void FControlRigParameterTrackEditor::FRecreateRigOperator::OnRigSettingsChanged(const FPropertyChangedEvent& InEvent)
+{
+	if (!IsValid())
+	{
+		Abort();
+		return;
+	}
+
+	TSharedPtr<FControlRigParameterTrackEditor> TrackEditor = WeakTrackEditor.Pin();
+	UMovieSceneControlRigParameterTrack* Track = WeakTrack.Get();
+	UControlRig* ControlRig = Track->GetControlRig();
+			
+	FScopedTransaction Transaction(LOCTEXT("RecreateControlRigWithNewSettings_Transaction", "Recreated Control Rig with New Settings"));
+	
+	Track->Modify();
+	ControlRig->Modify();
+
+	Track->UpdateAndApplyControlRigSettingsOverrides(SettingsForNewControlRig);
+	ControlRig->Initialize();
+	ControlRig->Evaluate_AnyThread();
+
+	for (UMovieSceneSection* Section : Track->GetAllSections())
+	{
+		if (Section)
+		{
+			UMovieSceneControlRigParameterSection* CRSection = Cast<UMovieSceneControlRigParameterSection>(Section);
+			if (CRSection)
+			{
+				Section->Modify();
+				CRSection->ClearAllParameters();
+				CRSection->RecreateWithThisControlRig(CRSection->GetControlRig(), true);
+			}
+		}
+	}
+
+	TrackEditor->GetSequencer()->NotifyMovieSceneDataChanged(EMovieSceneDataChangeType::MovieSceneStructureItemsChanged);
+}
+
+bool FControlRigParameterTrackEditor::FRecreateRigOperator::IsEditingObject(UObject* InRigToCheck)
+{
+	return WeakControlRig == InRigToCheck;
+}
+
 void FControlRigParameterTrackEditor::BuildObjectBindingTrackMenu(FMenuBuilder& MenuBuilder, const TArray<FGuid>& ObjectBindings, const UClass* ObjectClass)
 {
 	if(!ObjectClass)
@@ -5048,6 +5301,18 @@ void FControlRigParameterTrackEditor::BuildTrackContextMenu(FMenuBuilder& MenuBu
 					   NAME_None,
 					   EUserInterfaceActionType::ToggleButton);
 			}
+
+
+			MenuBuilder.AddMenuEntry(
+				   LOCTEXT("RecreateControlRigWithNewSettingsSettings", "Recreate Control Rig With New Settings"),
+				   LOCTEXT("RecreateControlRigWithNewSettingsSettingsToolTip", "Recreate Control Rig With New Settings"),
+				   FSlateIcon(),
+				   FUIAction(
+					   FExecuteAction::CreateRaw(this, &FControlRigParameterTrackEditor::RecreateControlRigWithNewSettings, Track),
+					   FCanExecuteAction::CreateRaw(this, &FControlRigParameterTrackEditor::CanRecreateControlRigWithNewSettings, Track)
+				   ),
+				   NAME_None,
+				   EUserInterfaceActionType::Button);
 		}
 		MenuBuilder.EndSection();
 	}
diff --git a/Engine/Plugins/Animation/ControlRig/Source/ControlRigEditor/Private/Sequencer/ControlRigParameterTrackEditor.h b/Engine/Plugins/Animation/ControlRig/Source/ControlRigEditor/Private/Sequencer/ControlRigParameterTrackEditor.h
--- a/Engine/Plugins/Animation/ControlRig/Source/ControlRigEditor/Private/Sequencer/ControlRigParameterTrackEditor.h
+++ b/Engine/Plugins/Animation/ControlRig/Source/ControlRigEditor/Private/Sequencer/ControlRigParameterTrackEditor.h
@@ -260,6 +260,47 @@ private:
 	/** Convert an absolute control rit to a layered control rig */
 	void ConvertIsLayered(UMovieSceneControlRigParameterTrack* Track);
 
+	
+	/** Change rig settings and recreates the rig using the new settings */
+	void RecreateControlRigWithNewSettings(UMovieSceneControlRigParameterTrack* Track);
+	
+	bool CanRecreateControlRigWithNewSettings(UMovieSceneControlRigParameterTrack* Track);
+
+	void CreateAndShowRigSettingsWindow(UMovieSceneControlRigParameterTrack* Track);
+	
+private:
+
+	struct FRecreateRigOperator : TSharedFromThis<FRecreateRigOperator>
+	{
+		FInstancedPropertyBag SettingsForNewControlRig;
+		
+		TWeakPtr<FControlRigParameterTrackEditor> WeakTrackEditor;
+		TWeakObjectPtr<UMovieSceneControlRigParameterTrack> WeakTrack;
+		TWeakObjectPtr<UControlRig> WeakControlRig;
+		
+		TWeakPtr<SWindow> WeakWindow;
+		TWeakPtr<IStructureDetailsView> WeakDetailsView;
+
+		void Start(
+			const TSharedRef<FControlRigParameterTrackEditor>& InTrackEditor,
+			UMovieSceneControlRigParameterTrack* InTrack
+			);
+		
+		~FRecreateRigOperator();
+		
+		bool IsValid();
+		void Abort();
+		void OnRigSettingsWindowClosed(const TSharedRef<SWindow>& Window);
+		void RefreshSettingsFromControlRig();
+		void OnPostControlRigForwardSolve_AnyThread(UControlRig*, const FName&);
+		void OnRigSettingsChanged( const FPropertyChangedEvent& InEvent);
+		bool IsEditingObject(UObject* InRigToCheck);
+	};
+
+	TSharedPtr<FRecreateRigOperator> RecreateRigOperator;
+	
+	void ResetRecreateRigOperatorIfNeeded(TSharedRef<FControlRigParameterTrackEditor::FRecreateRigOperator> InRequestingOperator);
+public:
 	/** Set Up EditMode for Specified Control Rig*/
 	void SetUpEditModeIfNeeded(UControlRig* ControlRig);
 
