Control Rig: Fix serialization of override values

#rnx

diff --git a/Engine/Plugins/Animation/ControlRig/Source/ControlRig/Private/ControlRigOverride.cpp b/Engine/Plugins/Animation/ControlRig/Source/ControlRig/Private/ControlRigOverride.cpp
--- a/Engine/Plugins/Animation/ControlRig/Source/ControlRig/Private/ControlRigOverride.cpp
+++ b/Engine/Plugins/Animation/ControlRig/Source/ControlRig/Private/ControlRigOverride.cpp
@@ -291,7 +291,7 @@ bool FControlRigOverrideValue::Serialize(FArchive& Ar)
 		if(const FProperty* LeafProperty = GetLeafProperty())
 		{
 			void *DataPtr = AllocateDataIfRequired();
-			LeafProperty->SerializeItem(FStructuredArchiveFromArchive(Ar).GetSlot(), DataPtr, DataPtr);
+			LeafProperty->SerializeItem(FStructuredArchiveFromArchive(Ar).GetSlot(), DataPtr, nullptr);
 
 			const int64 ArchivePositionAfterData = Ar.GetArchiveState().Tell();
 			int32 DataToSkip = ArchivePositionAfterData - ArchivePositionBeforeData;
