Anim in Engine: Fix some Attributes could no longer be selected in Anim Details

#rnx

diff --git a/Engine/Plugins/Animation/ControlRig/Source/ControlRigEditor/Private/AnimDetails/Customizations/AnimDetailsProxyDetails.cpp b/Engine/Plugins/Animation/ControlRig/Source/ControlRigEditor/Private/AnimDetails/Customizations/AnimDetailsProxyDetails.cpp
--- a/Engine/Plugins/Animation/ControlRig/Source/ControlRigEditor/Private/AnimDetails/Customizations/AnimDetailsProxyDetails.cpp
+++ b/Engine/Plugins/Animation/ControlRig/Source/ControlRigEditor/Private/AnimDetails/Customizations/AnimDetailsProxyDetails.cpp
@@ -90,57 +90,42 @@ namespace UE::ControlRigEditor
 		IDetailCategoryBuilder& DefaultCategory = DetailBuilder.EditCategory(CategoryName);
 		TArray<TSharedRef<IPropertyHandle>> PropertyHandles;
 		DefaultCategory.GetDefaultProperties(PropertyHandles);
-		
-		constexpr bool bForAdvanced = false;
-		constexpr bool bStartExpanded = true;
-		IDetailGroup* IndividualPropertyGroupPtr = bIsIndividual ? &DefaultCategory.AddGroup(CategoryName, DisplayNameText, bForAdvanced, bStartExpanded) : nullptr;
-		
+		DefaultCategory.SetCategoryVisibility(false);
+
+		IDetailGroup* PropertyGroupPtr = nullptr;
 		for (const TSharedRef<IPropertyHandle>& PropertyHandle : PropertyHandles)
-		{
+		{	
+			FProperty* Property = PropertyHandle->GetProperty();
+			if (!Property)
+			{
+				continue;
+			}
+
+			Property->SetPropertyFlags(CPF_TextExportTransient); // Hack to turn off shift copy/paste
+
 			uint32 NumChildren;
 			PropertyHandle->GetNumChildren(NumChildren);
 
 			if (bIsIndividual && NumChildren < 2)
 			{
-				if (const TSharedPtr<IPropertyHandle>& ChildHandle = PropertyHandle->GetChildHandle(0))
-				{
-					// Hack to turn off shift copy/paste
-					ChildHandle->GetProperty()->SetPropertyFlags(CPF_TextExportTransient); 
-
-					DefaultCategory.SetCategoryVisibility(false);
-
-					TArray<UObject*> OuterObjects;
-					ChildHandle->GetOuterObjects(OuterObjects);
-
-					// Show the name of the control instead of the property. 
-					// Otherwise the proxy property name would be e.g. "Bool" when multi editing two bool controls (the default details view behaviour)
-					const FName PropertyName = ChildHandle->GetProperty()->GetFName();
-
-					FText DisplayName;
-					TOptional<FText> StructDisplayName;
-					FirstProxy->GetLocalizedPropertyName(PropertyName, DisplayName, StructDisplayName);
-
-					NoCategory.AddProperty(ChildHandle)
-						.DisplayName(DisplayName);
-
-				}
+				NoCategory.AddProperty(PropertyHandle);
+				NoCategory.InitiallyCollapsed(false);
 			}
-			else
+			else if (bIsIndividual)
 			{
-				// Add properties with more than one child property, e.g. vectors.
-				// Note, we don't need a selection border here. Selection for structs with more than one 
-				// property is handled in their customization, see FAnimDetailValueCustomization.			
-				if (IndividualPropertyGroupPtr)
+				if (!PropertyGroupPtr)
 				{
-					IndividualPropertyGroupPtr->AddPropertyRow(PropertyHandle);
+					constexpr bool bForAdvanced = false;
+					constexpr bool bStartExpanded = true;
+					PropertyGroupPtr = &NoCategory.AddGroup(CategoryName, DisplayNameText, bForAdvanced, bStartExpanded);
 				}
-				else
-				{
-					DefaultCategory.SetCategoryVisibility(false);
 
-					NoCategory.AddProperty(PropertyHandle);
-					NoCategory.InitiallyCollapsed(false);
-				}
+				PropertyGroupPtr->AddPropertyRow(PropertyHandle);
+			}			
+			else
+			{
+				NoCategory.AddProperty(PropertyHandle);
+				NoCategory.InitiallyCollapsed(false);
 			}
 		}
 	}
diff --git a/Engine/Plugins/Animation/ControlRig/Source/ControlRigEditor/Private/AnimDetails/Customizations/AnimDetailsValueCustomization.cpp b/Engine/Plugins/Animation/ControlRig/Source/ControlRigEditor/Private/AnimDetails/Customizations/AnimDetailsValueCustomization.cpp
--- a/Engine/Plugins/Animation/ControlRig/Source/ControlRigEditor/Private/AnimDetails/Customizations/AnimDetailsValueCustomization.cpp
+++ b/Engine/Plugins/Animation/ControlRig/Source/ControlRigEditor/Private/AnimDetails/Customizations/AnimDetailsValueCustomization.cpp
@@ -103,9 +103,11 @@ namespace UE::ControlRigEditor
 
 	void FAnimDetailsValueCustomization::MakeHeaderRow(TSharedRef<IPropertyHandle>& InStructPropertyHandle, FDetailWidgetRow& InRow)
 	{
-		if (IsStructPropertyHiddenByFilter(InStructPropertyHandle))
+		constexpr TCHAR ShowOnlyInnerPropertiesMetaDataName[] = TEXT("ShowOnlyInnerProperties");
+		bool bShowHeader = !InStructPropertyHandle->HasMetaData(ShowOnlyInnerPropertiesMetaDataName);
+		if (!bShowHeader || 
+			IsStructPropertyHiddenByFilter(InStructPropertyHandle))
 		{
-			InStructPropertyHandle->MarkHiddenByCustomization();
 			return;
 		}
 
@@ -143,7 +145,13 @@ namespace UE::ControlRigEditor
 		for (int32 ChildIndex = 0; ChildIndex < SortedChildHandles.Num(); ++ChildIndex)
 		{
 			TSharedRef<IPropertyHandle> ChildHandle = SortedChildHandles[ChildIndex];
-			ChildHandle->GetProperty()->SetPropertyFlags(CPF_TextExportTransient); // Hack to turn off shift copy/paste
+			FProperty* Property = ChildHandle->GetProperty();
+			if (!Property)
+			{
+				continue;
+			}
+
+			Property->SetPropertyFlags(CPF_TextExportTransient); // Hack to turn off shift copy/paste
 
 			const TSharedRef<SWidget> ChildWidget = MakeChildWidget(InStructPropertyHandle, ChildHandle);
 
@@ -207,7 +215,13 @@ namespace UE::ControlRigEditor
 		TSharedRef<IPropertyHandle>& StructurePropertyHandle,
 		TSharedRef<IPropertyHandle>& PropertyHandle)
 	{
-		const FName PropertyName = PropertyHandle->GetProperty()->GetFName();
+		const FProperty* Property = PropertyHandle->GetProperty();
+		if (!Property)
+		{
+			return SNullWidget::NullWidget;
+		}
+
+		const FName PropertyName = Property->GetFName();
 		const FFieldClass* PropertyClass = PropertyHandle->GetPropertyClass();
 		const FLinearColor LabelColor = GetColorFromProperty(PropertyName);
 		const TSharedRef<SWidget> Label = SAnimDetailsValueNumeric<double>::BuildNarrowColorLabel(LabelColor);
diff --git a/Engine/Plugins/Animation/ControlRig/Source/ControlRigEditor/Private/AnimDetails/Proxies/AnimDetailsProxyBool.h b/Engine/Plugins/Animation/ControlRig/Source/ControlRigEditor/Private/AnimDetails/Proxies/AnimDetailsProxyBool.h
--- a/Engine/Plugins/Animation/ControlRig/Source/ControlRigEditor/Private/AnimDetails/Proxies/AnimDetailsProxyBool.h
+++ b/Engine/Plugins/Animation/ControlRig/Source/ControlRigEditor/Private/AnimDetails/Proxies/AnimDetailsProxyBool.h
@@ -36,6 +36,6 @@ public:
 	virtual void SetBindingValueFromCurrent(UObject* InObject, const TSharedPtr<FTrackInstancePropertyBindings>& Binding, const FRigControlModifiedContext& Context, bool bInteractive = false) override;
 	//~ End UAnimDetailsProxyBase interface
 
-	UPROPERTY(EditAnywhere, Interp, Category = "Bool")
+	UPROPERTY(EditAnywhere, Interp, Meta = (ShowOnlyInnerProperties), Category = "Bool")
 	FAnimDetailsBool Bool;
 };
diff --git a/Engine/Plugins/Animation/ControlRig/Source/ControlRigEditor/Private/AnimDetails/Proxies/AnimDetailsProxyEnum.h b/Engine/Plugins/Animation/ControlRig/Source/ControlRigEditor/Private/AnimDetails/Proxies/AnimDetailsProxyEnum.h
--- a/Engine/Plugins/Animation/ControlRig/Source/ControlRigEditor/Private/AnimDetails/Proxies/AnimDetailsProxyEnum.h
+++ b/Engine/Plugins/Animation/ControlRig/Source/ControlRigEditor/Private/AnimDetails/Proxies/AnimDetailsProxyEnum.h
@@ -45,6 +45,6 @@ public:
 	virtual void SetControlRigElementValueFromCurrent(UControlRig* ControlRig, FRigControlElement* ControlElement, const FRigControlModifiedContext& Context) override;
 	//~ End UAnimDetailsProxyBase interface
 
-	UPROPERTY(EditAnywhere, Interp, Category = "Enum")
+	UPROPERTY(EditAnywhere, Interp, Meta = (ShowOnlyInnerProperties), Category = "Enum")
 	FAnimDetailsEnum Enum;
 };
diff --git a/Engine/Plugins/Animation/ControlRig/Source/ControlRigEditor/Private/AnimDetails/Proxies/AnimDetailsProxyFloat.h b/Engine/Plugins/Animation/ControlRig/Source/ControlRigEditor/Private/AnimDetails/Proxies/AnimDetailsProxyFloat.h
--- a/Engine/Plugins/Animation/ControlRig/Source/ControlRigEditor/Private/AnimDetails/Proxies/AnimDetailsProxyFloat.h
+++ b/Engine/Plugins/Animation/ControlRig/Source/ControlRigEditor/Private/AnimDetails/Proxies/AnimDetailsProxyFloat.h
@@ -41,6 +41,6 @@ public:
 	virtual void SetBindingValueFromCurrent(UObject* InObject, const TSharedPtr<FTrackInstancePropertyBindings>& Binding, const FRigControlModifiedContext& Context, bool bInteractive = false) override;
 	//~ End UAnimDetailsProxyBase interface
 
-	UPROPERTY(EditAnywhere, Interp, Category = "Float", meta = (Delta = "0.05", SliderExponent = "1", LinearDeltaSensitivity = "1"))
+	UPROPERTY(EditAnywhere, Interp, Category = "Float", meta = (ShowOnlyInnerProperties, Delta = "0.05", SliderExponent = "1", LinearDeltaSensitivity = "1"))
 	FAnimDetailsFloat Float;
 };
diff --git a/Engine/Plugins/Animation/ControlRig/Source/ControlRigEditor/Private/AnimDetails/Proxies/AnimDetailsProxyInteger.h b/Engine/Plugins/Animation/ControlRig/Source/ControlRigEditor/Private/AnimDetails/Proxies/AnimDetailsProxyInteger.h
--- a/Engine/Plugins/Animation/ControlRig/Source/ControlRigEditor/Private/AnimDetails/Proxies/AnimDetailsProxyInteger.h
+++ b/Engine/Plugins/Animation/ControlRig/Source/ControlRigEditor/Private/AnimDetails/Proxies/AnimDetailsProxyInteger.h
@@ -36,6 +36,6 @@ public:
 	virtual void SetBindingValueFromCurrent(UObject* InObject, const TSharedPtr<FTrackInstancePropertyBindings>& Binding, const FRigControlModifiedContext& Context, bool bInteractive = false) override;
 	//~ End UAnimDetailsProxyBase interface
 
-	UPROPERTY(EditAnywhere, Interp, Category = "Integer")
+	UPROPERTY(EditAnywhere, Interp, Meta = (ShowOnlyInnerProperties), Category = "Integer")
 	FAnimDetailsInteger Integer;
 };
