Control Rig: Avoid tick dependency spam

#rnx

diff --git a/Engine/Plugins/Animation/ControlRig/Source/ControlRig/Private/ControlRigComponent.cpp b/Engine/Plugins/Animation/ControlRig/Source/ControlRig/Private/ControlRigComponent.cpp
--- a/Engine/Plugins/Animation/ControlRig/Source/ControlRig/Private/ControlRigComponent.cpp
+++ b/Engine/Plugins/Animation/ControlRig/Source/ControlRig/Private/ControlRigComponent.cpp
@@ -371,6 +371,8 @@ void UControlRigComponent::Initialize()
 	ValidateMappingData();
 
 #if WITH_EDITOR
+	TickDependencyErrorMessages.Reset();
+
 	if (bUpdateInEditor)
 	{
 		FEditorScriptExecutionGuard AllowScripts;
@@ -402,7 +404,9 @@ void UControlRigComponent::Initialize()
 	// also need to make sure all components driving the control rig component tick before
 	// the control rig component ticks
 	TMap<TObjectPtr<USceneComponent>, EControlRigComponentMapDirection> MappedComponents;
+#if WITH_EDITOR
 	TSet<TObjectPtr<USceneComponent>> MappedComponentsWithErrors;
+#endif
 	for (const FControlRigComponentMappedElement& MappedElement : MappedElements)
 	{
 		if (MappedElement.SceneComponent)
@@ -412,7 +416,9 @@ void UControlRigComponent::Initialize()
 				if (*Direction != MappedElement.Direction)
 				{
 					// elements from the same component should not be mapped to both input and output
+#if WITH_EDITOR
 					MappedComponentsWithErrors.Add(MappedElement.SceneComponent);
+#endif
 					continue;
 				}
 			}
@@ -446,7 +452,11 @@ void UControlRigComponent::Initialize()
 				// only throw the error message if the user has picked a particular component to set up the dependency for
 				if(!MappedElement.ComponentReference.ComponentProperty.IsNone() || MappedElement.ComponentReference.OtherActor.IsValid())
 				{
-					UE_LOG(LogControlRig, Warning, TEXT("%s"), *TickDependencyErrorMessage);
+					if(!TickDependencyErrorMessages.Contains(TickDependencyErrorMessage))
+					{
+						UE_LOG(LogControlRig, Warning, TEXT("%s"), *TickDependencyErrorMessage);
+						TickDependencyErrorMessages.Add(TickDependencyErrorMessage);
+					}
 				}
 			}
 #endif
@@ -462,6 +472,7 @@ void UControlRigComponent::Initialize()
 		}
 	}
 
+#if WITH_EDITOR
 	for (const TObjectPtr<USceneComponent>& ErrorComponent : MappedComponentsWithErrors)
 	{
 		FString Message = FString::Printf(
@@ -469,6 +480,7 @@ void UControlRigComponent::Initialize()
 					" because it creates ambiguity when inferring tick order."), *ErrorComponent.GetName());
 		UE_LOG(LogControlRig, Warning, TEXT("%s: %s"), *GetPathName(), *Message);
 	}
+#endif
 }
 
 void UControlRigComponent::Update(float DeltaTime)
diff --git a/Engine/Plugins/Animation/ControlRig/Source/ControlRig/Public/ControlRigComponent.h b/Engine/Plugins/Animation/ControlRig/Source/ControlRig/Public/ControlRigComponent.h
--- a/Engine/Plugins/Animation/ControlRig/Source/ControlRig/Public/ControlRigComponent.h
+++ b/Engine/Plugins/Animation/ControlRig/Source/ControlRig/Public/ControlRigComponent.h
@@ -665,6 +665,7 @@ private:
 
 #if WITH_EDITOR
 	static TMap<FString, TSharedPtr<SNotificationItem>> EditorNotifications;
+	TSet<FString> TickDependencyErrorMessages;
 #endif
 
 	UPROPERTY(transient)
