Control Rig: Improve tooltip for outdated modules

#rnx

diff --git a/Engine/Plugins/Animation/ControlRig/Source/ControlRigEditor/Private/Editor/SModularRigTreeView.cpp b/Engine/Plugins/Animation/ControlRig/Source/ControlRigEditor/Private/Editor/SModularRigTreeView.cpp
--- a/Engine/Plugins/Animation/ControlRig/Source/ControlRigEditor/Private/Editor/SModularRigTreeView.cpp
+++ b/Engine/Plugins/Animation/ControlRig/Source/ControlRigEditor/Private/Editor/SModularRigTreeView.cpp
@@ -286,12 +286,14 @@ TSharedRef<SWidget> SModularRigModelItem::GenerateWidgetForColumn(const FName& C
 								if (Tag.bMarksSubjectAsInvalid)
 								{
 									ToolTip.Add(FString::Printf(TEXT("%s: %s"), *Tag.Label, *Tag.ToolTip.ToString()));
+									ToolTip.Add(TEXT("Right click on the module to swap it to a newer variant.")); 
 								}
 							}
 						}
 					}
 				}
 			}
+
 			return FText::FromString(FString::Join(ToolTip, TEXT("\n")));
 		})
 		.Image_Lambda([this]() -> const FSlateBrush*
