Control Rig: Module details to allow to sync content browser to module asset

#rnx

diff --git a/Engine/Plugins/Animation/ControlRig/Source/ControlRigEditor/Private/ControlRigModuleDetails.cpp b/Engine/Plugins/Animation/ControlRig/Source/ControlRigEditor/Private/ControlRigModuleDetails.cpp
--- a/Engine/Plugins/Animation/ControlRig/Source/ControlRigEditor/Private/ControlRigModuleDetails.cpp
+++ b/Engine/Plugins/Animation/ControlRig/Source/ControlRigEditor/Private/ControlRigModuleDetails.cpp
@@ -31,6 +31,10 @@
 #include "Widgets/SRigVMVariantTagWidget.h"
 #include "Algo/Sort.h"
 #include "Editor/SRigConnectorTargetWidget.h"
+#include "AssetRegistry/IAssetRegistry.h"
+#include "AssetRegistry/AssetRegistryModule.h"
+#include "IContentBrowserSingleton.h"
+#include "ContentBrowserModule.h"
 
 #define LOCTEXT_NAMESPACE "ControlRigModuleDetails"
 
@@ -194,10 +198,12 @@ void FRigModuleInstanceDetails::CustomizeDetails(IDetailLayoutBuilder& DetailBui
 		]
 		.ValueContent()
 		[
-			SNew(SInlineEditableTextBlock)
-			.Font(IDetailLayoutBuilder::GetDetailFont())
+			SNew(SButton)
+			.ContentPadding(0.0f)
+			.ButtonStyle(FAppStyle::Get(), "NoBorder")
+			.OnClicked(this, &FRigModuleInstanceDetails::HandleOpenRigModuleAsset)
+			.Cursor(EMouseCursor::Default)
 			.Text(this, &FRigModuleInstanceDetails::GetRigClassPath)
-			.IsEnabled(true)
 		];
 
 		TWeakPtr<FRigModuleInstanceDetails> WeakThisPtr = SharedThis(this);
@@ -824,6 +830,28 @@ FText FRigModuleInstanceDetails::GetRigClassPath() const
 	return FText();
 }
 
+FReply FRigModuleInstanceDetails::HandleOpenRigModuleAsset() const
+{
+	if(PerModuleInfos.Num() == 1)
+	{
+		if (const FRigModuleInstance* Module = PerModuleInfos[0].GetModule())
+		{
+			if (const UControlRig* ModuleRig = Module->GetRig())
+			{
+				if(const UObject* Blueprint = ModuleRig->GetClass()->ClassGeneratedBy)
+				{
+					const FSoftObjectPath SoftObjectPath(Blueprint);
+					TArray<FAssetData> AssetData;
+					AssetData.Add(FModuleManager::LoadModuleChecked<FAssetRegistryModule>("AssetRegistry").Get().GetAssetByObjectPath(SoftObjectPath.GetWithoutSubPath()));
+					FModuleManager::LoadModuleChecked<FContentBrowserModule>("ContentBrowser").Get().SyncBrowserToAssets(AssetData);
+					return FReply::Handled();
+				}
+			}
+		}
+	}
+	return FReply::Unhandled();
+}
+
 TArray<FRigModuleConnector> FRigModuleInstanceDetails::GetConnectors() const
 {
 	if(PerModuleInfos.Num() > 1)
diff --git a/Engine/Plugins/Animation/ControlRig/Source/ControlRigEditor/Private/ControlRigModuleDetails.h b/Engine/Plugins/Animation/ControlRig/Source/ControlRigEditor/Private/ControlRigModuleDetails.h
--- a/Engine/Plugins/Animation/ControlRig/Source/ControlRigEditor/Private/ControlRigModuleDetails.h
+++ b/Engine/Plugins/Animation/ControlRig/Source/ControlRigEditor/Private/ControlRigModuleDetails.h
@@ -44,6 +44,7 @@ public:
 	void SetName(const FText& InValue, ETextCommit::Type InCommitType, const TSharedRef<IPropertyUtilities> PropertyUtilities);
 	bool OnVerifyNameChanged(const FText& InText, FText& OutErrorMessage);
 	FText GetRigClassPath() const;
+	FReply HandleOpenRigModuleAsset() const;
 	TArray<FRigModuleConnector> GetConnectors() const;
 	FRigElementKeyRedirector GetConnections() const;
 
