Control Rig: Sort metadata keys alphabetically for user interface

#rnx

diff --git a/Engine/Plugins/Animation/ControlRig/Source/ControlRigEditor/Private/Graph/ControlRigGraphPanelPinFactory.cpp b/Engine/Plugins/Animation/ControlRig/Source/ControlRigEditor/Private/Graph/ControlRigGraphPanelPinFactory.cpp
--- a/Engine/Plugins/Animation/ControlRig/Source/ControlRigEditor/Private/Graph/ControlRigGraphPanelPinFactory.cpp
+++ b/Engine/Plugins/Animation/ControlRig/Source/ControlRigEditor/Private/Graph/ControlRigGraphPanelPinFactory.cpp
@@ -504,7 +504,13 @@ TSharedPtr<SGraphPin> FControlRigGraphPanelPinFactory::CreatePin_Internal(UEdGra
 
 										MetadataNames.Names->Sort([](const TSharedPtr<FRigVMStringWithTag>& A, const TSharedPtr<FRigVMStringWithTag>& B)
 										{
-											return A.Get() > B.Get();
+											const FRigVMStringWithTag* StringA = A.Get();
+											const FRigVMStringWithTag* StringB = B.Get();
+											if(StringA && StringB)
+											{
+												return (*StringA) < (*StringB);
+											}
+											return false;
 										});
 										MetadataNames.Names->Insert(MakeShareable(new FRigVMStringWithTag(FName(NAME_None).ToString())), 0);
 
@@ -592,7 +598,13 @@ TSharedPtr<SGraphPin> FControlRigGraphPanelPinFactory::CreatePin_Internal(UEdGra
 										}
 										MetadataTagNames.Names->Sort([](const TSharedPtr<FRigVMStringWithTag>& A, const TSharedPtr<FRigVMStringWithTag>& B)
 										{
-											return A.Get() > B.Get();
+											const FRigVMStringWithTag* StringA = A.Get();
+											const FRigVMStringWithTag* StringB = B.Get();
+											if(StringA && StringB)
+											{
+												return (*StringA) < (*StringB);
+											}
+											return false;
 										});
 										MetadataTagNames.Names->Insert(MakeShareable(new FRigVMStringWithTag(FName(NAME_None).ToString())), 0);
 
