Anim in Engine: Fix bool properties cannot be toggled in Anim Details

#rnx

diff --git a/Engine/Plugins/Animation/ControlRig/Source/ControlRigEditor/Private/AnimDetails/Widgets/SAnimDetailsValueBoolean.cpp b/Engine/Plugins/Animation/ControlRig/Source/ControlRigEditor/Private/AnimDetails/Widgets/SAnimDetailsValueBoolean.cpp
--- a/Engine/Plugins/Animation/ControlRig/Source/ControlRigEditor/Private/AnimDetails/Widgets/SAnimDetailsValueBoolean.cpp
+++ b/Engine/Plugins/Animation/ControlRig/Source/ControlRigEditor/Private/AnimDetails/Widgets/SAnimDetailsValueBoolean.cpp
@@ -27,6 +27,7 @@ namespace UE::ControlRigEditor
 			return;
 		}
 
+		WeakProxyManager = ProxyManager;
 		WeakPropertyHandle = InPropertyHandle;
 
 		ChildSlot
