Anim in Engine: Fix Control Rig Snapper in subsequences, removing the evaluate in isolation. Need to think about this further.
#trivial

diff --git a/Engine/Plugins/Animation/ControlRig/Source/ControlRigEditor/Private/Tools/ControlRigSnapper.cpp b/Engine/Plugins/Animation/ControlRig/Source/ControlRigEditor/Private/Tools/ControlRigSnapper.cpp
--- a/Engine/Plugins/Animation/ControlRig/Source/ControlRigEditor/Private/Tools/ControlRigSnapper.cpp
+++ b/Engine/Plugins/Animation/ControlRig/Source/ControlRigEditor/Private/Tools/ControlRigSnapper.cpp
@@ -141,7 +141,7 @@ static bool LocalGetControlRigControlTransforms(IMovieScenePlayer* Player, const
 		{
 			const FFrameNumber& FrameNumber = Frames[Index];
 			double DeltaTime = 0.0;
-			if (CurrentFrame.IsSet() == false || CurrentFrame.GetValue() != FrameNumber)
+			if ((Index != 0) || (CurrentFrame.IsSet() == false || CurrentFrame.GetValue() != FrameNumber))
 			{
 				FFrameTime GlobalTime = LocalToRootTransform.TryTransformTime(FrameNumber).Get(FrameNumber); //player evals in root time so need to go back to it.
 
@@ -216,7 +216,7 @@ struct FGuidAndActor
 		{
 			return false;
 		}
-
+		TransformSection->Modify();
 		// fill parent transforms
 		TArray<FTransform> ParentWorldTransforms; 
 		AActor* ParentActor = Actor->GetAttachParentActor();
@@ -556,13 +556,10 @@ bool FControlRigSnapper::SnapIt(FFrameNumber StartFrame, FFrameNumber EndFrame,c
 		}
 
 		//if in subsquence and we want to, turn on ShouldEvaluateSubSequencesInIsolation, use the following object for that
-
-		UE::Sequencer::FSequencerBakingSetupRestore RestoreBaking(Sequencer);
+		//disabling this for nowUE::Sequencer::FSequencerBakingSetupRestore RestoreBaking(Sequencer);
 
 		FScopedTransaction ScopedTransaction(LOCTEXT("SnapAnimation", "Snap Animation"));
 	
-		MovieScene->Modify();
-
 		//need to make sure binding is setup
 		Sequencer->ForceEvaluate();
 		
