Control Rig: Fix crash when switching rig to module
#rnx

diff --git a/Engine/Plugins/Animation/ControlRig/Source/ControlRig/Private/Units/Execution/RigUnit_DynamicHierarchy.cpp b/Engine/Plugins/Animation/ControlRig/Source/ControlRig/Private/Units/Execution/RigUnit_DynamicHierarchy.cpp
--- a/Engine/Plugins/Animation/ControlRig/Source/ControlRig/Private/Units/Execution/RigUnit_DynamicHierarchy.cpp
+++ b/Engine/Plugins/Animation/ControlRig/Source/ControlRig/Private/Units/Execution/RigUnit_DynamicHierarchy.cpp
@@ -922,15 +922,48 @@ FRigUnit_HierarchyAddControlRotator_Execute()
 
 void FRigUnit_HierarchyAddControlTransform_LimitSettings::ConfigureFrom(const FRigControlElement* InControlElement, const FRigControlSettings& InSettings)
 {
-	LimitTranslationX = InSettings.LimitEnabled[0];
-	LimitTranslationY = InSettings.LimitEnabled[1];
-	LimitTranslationZ = InSettings.LimitEnabled[2];
-	LimitPitch = InSettings.LimitEnabled[3];
-	LimitYaw = InSettings.LimitEnabled[4];
-	LimitRoll = InSettings.LimitEnabled[5];
-	LimitScaleX = InSettings.LimitEnabled[6];
-	LimitScaleY = InSettings.LimitEnabled[7];
-	LimitScaleZ = InSettings.LimitEnabled[8];
+	switch (InSettings.ControlType)
+	{
+		case ERigControlType::EulerTransform:
+		case ERigControlType::Transform:
+		{
+			LimitScaleX = InSettings.LimitEnabled[6];
+			LimitScaleY = InSettings.LimitEnabled[7];
+			LimitScaleZ = InSettings.LimitEnabled[8];
+			// fall through
+		}
+		case ERigControlType::TransformNoScale:
+		{
+			LimitPitch = InSettings.LimitEnabled[3];
+			LimitYaw = InSettings.LimitEnabled[4];
+			LimitRoll = InSettings.LimitEnabled[5];
+			// fall through
+		}
+		case ERigControlType::Position:
+		case ERigControlType::Scale:
+		case ERigControlType::Rotator:
+		{
+			LimitTranslationZ = InSettings.LimitEnabled[2];
+			// fall through
+		}
+		case ERigControlType::Vector2D:
+		{
+			LimitTranslationY = InSettings.LimitEnabled[1];
+			// fall through
+		}
+		case ERigControlType::Integer:
+		case ERigControlType::Float:
+		case ERigControlType::ScaleFloat:
+		{
+			LimitTranslationX = InSettings.LimitEnabled[0];
+			// fall through
+		}
+		case ERigControlType::Bool:
+		default:
+		{
+		}
+	}
+
 	MinValue = InSettings.MinimumValue.Get<FRigControlValue::FEulerTransform_Float>().ToTransform();
 	MaxValue = InSettings.MaximumValue.Get<FRigControlValue::FEulerTransform_Float>().ToTransform();
 	bDrawLimits = InSettings.bDrawLimits;
@@ -939,15 +972,47 @@ void FRigUnit_HierarchyAddControlTransform_LimitSettings::ConfigureFrom(const FR
 void FRigUnit_HierarchyAddControlTransform_LimitSettings::Configure(FRigControlSettings& OutSettings) const
 {
 	OutSettings.SetupLimitArrayForType(false, false, false);
-	OutSettings.LimitEnabled[0] = LimitTranslationX;
-	OutSettings.LimitEnabled[1] = LimitTranslationY;
-	OutSettings.LimitEnabled[2] = LimitTranslationZ;
-	OutSettings.LimitEnabled[3] = LimitPitch;
-	OutSettings.LimitEnabled[4] = LimitYaw;
-	OutSettings.LimitEnabled[5] = LimitRoll;
-	OutSettings.LimitEnabled[6] = LimitScaleX;
-	OutSettings.LimitEnabled[7] = LimitScaleY;
-	OutSettings.LimitEnabled[8] = LimitScaleZ;
+
+	switch (OutSettings.ControlType)
+	{
+	case ERigControlType::EulerTransform:
+	case ERigControlType::Transform:
+		{
+			OutSettings.LimitEnabled[6] = LimitScaleX;
+			OutSettings.LimitEnabled[7] = LimitScaleY;
+			OutSettings.LimitEnabled[8] = LimitScaleZ;
+			// fall through
+		}
+	case ERigControlType::TransformNoScale:
+		{
+			OutSettings.LimitEnabled[3] = LimitPitch;
+			OutSettings.LimitEnabled[4] = LimitYaw;
+			OutSettings.LimitEnabled[5] = LimitRoll;
+			// fall through
+		}
+	case ERigControlType::Position:
+	case ERigControlType::Scale:
+	case ERigControlType::Rotator:
+		{
+			OutSettings.LimitEnabled[2] = LimitTranslationZ;
+			// fall through
+		}
+	case ERigControlType::Vector2D:
+		{
+			OutSettings.LimitEnabled[1] = LimitTranslationY;
+			// fall through
+		}
+	case ERigControlType::Integer:
+	case ERigControlType::Float:
+	case ERigControlType::ScaleFloat:
+		{
+			OutSettings.LimitEnabled[0] = LimitTranslationX;
+		}
+	case ERigControlType::Bool:
+	default:
+		{
+		}
+	}
 	OutSettings.MinimumValue = FRigControlValue::Make<FRigControlValue::FEulerTransform_Float>(MinValue);
 	OutSettings.MaximumValue = FRigControlValue::Make<FRigControlValue::FEulerTransform_Float>(MaxValue);
 	OutSettings.bDrawLimits = bDrawLimits;
