Control Rig: Overrides Support for Sequencer

#rnx

diff --git a/Engine/Plugins/Animation/ControlRig/Source/ControlRig/Private/ControlRig.cpp b/Engine/Plugins/Animation/ControlRig/Source/ControlRig/Private/ControlRig.cpp
--- a/Engine/Plugins/Animation/ControlRig/Source/ControlRig/Private/ControlRig.cpp
+++ b/Engine/Plugins/Animation/ControlRig/Source/ControlRig/Private/ControlRig.cpp
@@ -171,6 +171,156 @@ void UControlRig::BeginDestroy()
 	TRACE_OBJECT_LIFETIME_END(this);
 }
 
+int32 UControlRig::NumOverrideAssets() const
+{
+	return OverrideAssets.Num();
+}
+
+const UControlRigOverrideAsset* UControlRig::GetLastOverrideAsset() const
+{
+	return OverrideAssets.IsEmpty() ? nullptr : OverrideAssets.Last();
+}
+
+UControlRigOverrideAsset* UControlRig::GetLastOverrideAsset()
+{
+	return OverrideAssets.IsEmpty() ? nullptr : OverrideAssets.Last();
+}
+
+const UControlRigOverrideAsset* UControlRig::GetOverrideAsset(int32 InIndex) const
+{
+	return OverrideAssets.IsValidIndex(InIndex) ? OverrideAssets[InIndex] : nullptr;
+}
+
+UControlRigOverrideAsset* UControlRig::GetOverrideAsset(int32 InIndex)
+{
+	return OverrideAssets.IsValidIndex(InIndex) ? OverrideAssets[InIndex] : nullptr;
+}
+
+int32 UControlRig::LinkOverrideAsset(UControlRigOverrideAsset* InAsset)
+{
+	if (!InAsset)
+	{
+		return INDEX_NONE;
+	}
+
+	if (OverrideAssets.Contains(InAsset))
+	{
+		return INDEX_NONE;
+	}
+	
+	const int32 Index = OverrideAssets.Add(InAsset);
+	InAsset->OnChanged().AddUObject(this, &UControlRig::HandleOverrideChanged);
+	OverrideAssetsChangedDelegate.Broadcast(this);
+	RequestConstruction();
+	return Index;
+}
+
+bool UControlRig::UnlinkOverrideAsset(UControlRigOverrideAsset* InAsset)
+{
+	const int32 NumRemoved = OverrideAssets.Remove(const_cast<UControlRigOverrideAsset*>(InAsset));
+	if (NumRemoved == 0)
+	{
+		return false;
+	}
+	
+	if (InAsset)
+	{
+		InAsset->OnChanged().RemoveAll(this);
+	}
+		
+	OverrideAssetsChangedDelegate.Broadcast(this);
+	RequestConstruction();
+	
+	return true;
+}
+
+bool UControlRig::UnlinkAllOverrideAssets()
+{
+	if (OverrideAssets.IsEmpty())
+	{
+		return false;
+	}
+
+	for (const TObjectPtr<UControlRigOverrideAsset>& OverrideAsset : OverrideAssets)
+	{
+		OverrideAsset->OnChanged().RemoveAll(this);
+	}
+	OverrideAssets.Reset();
+	OverrideAssetsChangedDelegate.Broadcast(this);
+	RequestConstruction();
+	
+	return true;
+}
+
+bool UControlRig::IsLinkedToOverrideAsset(const UControlRigOverrideAsset* InOverrideAsset) const
+{
+	return OverrideAssets.Contains(InOverrideAsset);
+}
+
+void UControlRig::ApplyElementOverrides()
+{
+	if (!CVarControlRigEnableOverrides.GetValueOnAnyThread())
+	{
+		return;
+	}
+	
+	URigHierarchy* Hierarchy = GetHierarchy();
+	if(!Hierarchy)
+	{
+		return;
+	}
+	
+	TSet<FRigControlElement*> ControlsChanged;
+	
+	for (const UControlRigOverrideAsset* OverrideAsset : OverrideAssets)
+	{
+		if (!OverrideAsset)
+		{
+			continue;
+		}
+
+		for (const FControlRigOverrideValue& Override : OverrideAsset->Overrides)
+		{
+			const FRigElementKey Key(Override.GetSubjectKey(), ERigElementType::Control);
+			FRigControlElement* ControlElement = Hierarchy->Find<FRigControlElement>(Key);
+			if(!ControlElement)
+			{
+				continue;
+			}
+			
+			const FProperty* ShapeTransformProperty =
+				FRigControlSettings::StaticStruct()->FindPropertyByName(GET_MEMBER_NAME_CHECKED(FRigControlSettings, ShapeTransform));
+			const bool bIsTransformShapeOverride = Override.ContainsProperty(ShapeTransformProperty);
+			if (bIsTransformShapeOverride)
+			{
+				ControlElement->Settings.ShapeTransform = Hierarchy->GetControlShapeTransform(ControlElement, ERigTransformType::InitialLocal);
+			}
+
+			// in the future let's avoid this check know on a container level
+			// if this container represents changes to the control
+			if (!Override.CopyToSubject(ControlElement, FRigControlElement::StaticStruct()))
+			{
+				continue;
+			}
+			ControlsChanged.Add(ControlElement);
+
+			if (bIsTransformShapeOverride)
+			{
+				Hierarchy->SetControlShapeTransformByIndex(ControlElement->GetIndex(), ControlElement->Settings.ShapeTransform, true);
+				Hierarchy->SetControlShapeTransformByIndex(ControlElement->GetIndex(), ControlElement->Settings.ShapeTransform, false);
+			}
+		}
+	}
+}
+
+void UControlRig::HandleOverrideChanged(const UControlRigOverrideAsset* InOverrideAsset)
+{
+	if (!bSuspendOverrideAssetChangedDelegate&& InOverrideAsset && OverrideAssets.Contains(InOverrideAsset))
+	{
+		RequestConstruction();
+	}
+}
+
 UWorld* UControlRig::GetWorld() const
 {
 	if (ObjectBinding.IsValid())
@@ -1082,6 +1232,7 @@ bool UControlRig::Execute(const FName& InEventName)
 
 				if(!bPostConstructionEventInQueue)
 				{
+					ApplyElementOverrides();
 					EmitPostConstructionEventFinished();
 				}
 
@@ -1193,6 +1344,7 @@ bool UControlRig::Execute(const FName& InEventName)
 
 		if (bIsPostConstructionEvent)
 		{
+			ApplyElementOverrides();
 			EmitPostConstructionEventFinished();
 		}
 
@@ -1898,6 +2050,11 @@ void UControlRig::Serialize(FArchive& Ar)
 	{
 		PhysicsSimulations.Reset();
 	}
+
+	if(Ar.IsObjectReferenceCollector())
+	{
+		Ar << OverrideAssets;
+	}
 }
 
 void UControlRig::PostLoad()
diff --git a/Engine/Plugins/Animation/ControlRig/Source/ControlRig/Private/ControlRigOverride.cpp b/Engine/Plugins/Animation/ControlRig/Source/ControlRig/Private/ControlRigOverride.cpp
--- a/Engine/Plugins/Animation/ControlRig/Source/ControlRig/Private/ControlRigOverride.cpp
+++ b/Engine/Plugins/Animation/ControlRig/Source/ControlRig/Private/ControlRigOverride.cpp
@@ -6,160 +6,64 @@
 
 #if WITH_EDITOR
 #include "PropertyPath.h"
+#include "AssetToolsModule.h"
+#include "Misc/TransactionObjectEvent.h"
 #endif
+#include "UObject/UObjectIterator.h"
 
 #include UE_INLINE_GENERATED_CPP_BY_NAME(ControlRigOverride)
 
+TAutoConsoleVariable<bool> CVarControlRigEnableOverrides(TEXT("ControlRig.Overrides"), false, TEXT("Enables overrides for use in Control Rig"),
+	FConsoleVariableDelegate::CreateLambda([](IConsoleVariable* InVariable)
+{
+	for (TObjectIterator<UControlRigOverrideAsset> It(RF_ClassDefaultObject, true, EInternalObjectFlags::Garbage); It; ++It)
+	{
+		UControlRigOverrideAsset* Override = *It;
+		if(IsValid(Override))
+		{
+			Override->BroadcastChanged();
+		}
+	}
+}));
+
 ////////////////////////////////////////////////////////////////////////////////
 // FControlRigOverrideValue
 ////////////////////////////////////////////////////////////////////////////////
 
 FControlRigOverrideValue::FControlRigOverrideValue(const FString& InPath, const UObject* InSubject)
+	: SubjectKey(InSubject ? InSubject->GetFName() : NAME_None)
 {
 	if(InSubject)
 	{
-		*this = FControlRigOverrideValue(InPath, InSubject->GetClass(), InSubject, InSubject->GetFName());
+		if(SetPropertiesFromPath(InPath, InSubject->GetClass()))
+		{
+			(void)SetFromSubject(InSubject);
+		}
 	}
 }
 
-FControlRigOverrideValue::FControlRigOverrideValue(const FString& InPath, const UStruct* InOwnerStruct, const void* InSubjectPtr, const FName& InSubjectName)
+FControlRigOverrideValue::FControlRigOverrideValue(const FString& InPath, const UStruct* InOwnerStruct, const void* InSubjectPtr, const FName& InSubjectKey)
+	: SubjectKey(InSubjectKey)
 {
 	ensure(InOwnerStruct);
 	ensureMsgf(InPath.Contains(TEXT(".")) == false, TEXT("Paths for overrides should use the -> separator."));
-	
-	Path = InPath;
-	SubjectName = InSubjectName;
-
-	FString PropertyNameString;
-	FString ArrayIndexString;
-	bool bIsParsingPropertyName = true;
-	const UStruct* OwnerStruct = InOwnerStruct;
-
-	auto ProcessPropertyName = [this, &PropertyNameString, &OwnerStruct]() -> bool
-	{
-		FString StringToProcess;
-		Swap(StringToProcess, PropertyNameString);
-		if(!StringToProcess.IsEmpty())
-		{
-			if(FProperty* Property = FindProperty(OwnerStruct, StringToProcess))
-			{
-				FProperty* PropertyForPath = Property;
-				
-				if(const FArrayProperty* ArrayProperty = CastField<FArrayProperty>(Property))
-				{
-					Property = ArrayProperty->Inner;
-
-					if(Properties.IsEmpty())
-					{
-						PropertyForPath = Property;
-					}
-				}
-				if(const FStructProperty* StructProperty = CastField<FStructProperty>(Property))
-				{
-					OwnerStruct = StructProperty->Struct;
-				}
-
-				Properties.Add({PropertyForPath, INDEX_NONE});
-
-				return true;
-			}
-		}
-		return false;
-	};
-	
-	auto ProcessArrayIndex = [this, &ArrayIndexString]() -> bool
-	{
-		FString StringToProcess;
-		Swap(StringToProcess, ArrayIndexString);
-		if(!StringToProcess.IsEmpty())
-		{
-			Properties.Last().ArrayIndex = FCString::Atoi(*StringToProcess);
-			return true;
-		}
-		return false;
-	};
 
-	// Parse the provided Property path. Path may look like
-	// Settings[2]->Color->R resulting in three segments, the
-	// first one with an array index != INDEX_NONE
-	int32 CharIndex = 0;
-	while(CharIndex < InPath.Len())
+	if(SetPropertiesFromPath(InPath, InOwnerStruct))
 	{
-		// potentially start parsing an array index
-		if(InPath[CharIndex] == TEXT('['))
+		if(InSubjectPtr)
 		{
-			if(!ProcessPropertyName())
-			{
-				Reset();
-				return;
-			}
-			bIsParsingPropertyName = false;
-			CharIndex++;
-			continue;;
+			(void)SetFromSubject(InSubjectPtr);
 		}
-		
-		// potentially end parsing an array index
-		if(InPath[CharIndex] == TEXT(']'))
-		{
-			if(!ProcessArrayIndex())
-			{
-				Reset();
-				return;
-			}
-			CharIndex++;
-			continue;;
-		}
-
-		// given a path separator parse the property name and store it.
-		// if bIsParsingPropertyName == false we already processed the property name
-		// and have been parsing an array index
-		if(InPath.Mid(CharIndex, PathSeparatorLength).Equals(PathSeparator, ESearchCase::CaseSensitive))
-		{
-			if(bIsParsingPropertyName)
-			{
-				if(!ProcessPropertyName())
-				{
-					Reset();
-					return;
-				}
-			}
-			bIsParsingPropertyName = true;
-			CharIndex += PathSeparatorLength;
-			continue;
-		}
-
-		if(bIsParsingPropertyName)
-		{
-			PropertyNameString += InPath[CharIndex];
-		}
-		else
-		{
-			ArrayIndexString += InPath[CharIndex];
-		}
-		CharIndex++;
-	}
-
-	if(!PropertyNameString.IsEmpty() && !ProcessPropertyName())
-	{
-		Reset();
-		return;
-	}
-	if(!ArrayIndexString.IsEmpty() && !ProcessArrayIndex())
-	{
-		Reset();
-		return;
-	}
-
-	if(InSubjectPtr)
-	{
-		SetFromSubject(InSubjectPtr);
 	}
 }
 
-FControlRigOverrideValue::FControlRigOverrideValue(const FString& InPath, const UStruct* InOwnerStruct, const FString& InValueAsString, const FName& InSubjectName)
+FControlRigOverrideValue::FControlRigOverrideValue(const FString& InPath, const UStruct* InOwnerStruct, const FString& InValueAsString, const FName& InSubjectKey)
+	: SubjectKey(InSubjectKey)
 {
-	*this = FControlRigOverrideValue(InPath, InOwnerStruct, nullptr, InSubjectName);
-	SetFromString(InValueAsString);
+	if(SetPropertiesFromPath(InPath, InOwnerStruct))
+	{
+		SetFromString(InValueAsString);
+	}
 }
 
 FControlRigOverrideValue::FControlRigOverrideValue(const FControlRigOverrideValue& InOther)
@@ -176,11 +80,12 @@ FControlRigOverrideValue& FControlRigOverrideValue::operator=(const FControlRigO
 {
 	Reset();
 	Path = InOther.Path;
-	SubjectName = InOther.SubjectName;
+	SubjectKey = InOther.SubjectKey;
 	CachedStringValue = InOther.CachedStringValue;
 	Properties = InOther.Properties;
+	Hash = InOther.Hash;
 
-	if(const FProperty* LeafProperty = GetLeafProperty())
+	if(GetLeafProperty())
 	{
 		if(const void* OtherDataPtr = InOther.GetData())
 		{
@@ -200,7 +105,7 @@ bool FControlRigOverrideValue::Serialize(FArchive& Ar)
 		Reset();
 	}
 
-	Ar << SubjectName;
+	Ar << SubjectKey;
 	
 	int32 NumProperties = 0;
 
@@ -285,6 +190,8 @@ bool FControlRigOverrideValue::Serialize(FArchive& Ar)
 			// we can't deserialize the data - skip it
 			Ar.Seek(ArchivePositionBeforeData + OffsetForData);
 		}
+
+		UpdateHash();
 	}
 	else
 	{
@@ -316,12 +223,13 @@ void FControlRigOverrideValue::Reset()
 	FreeDataIfRequired();
 	CachedStringValue.Reset();
 	Properties.Reset();
+	Hash = 0;
 }
 
-const FName& FControlRigOverrideValue::GetSubjectName() const
+const FName& FControlRigOverrideValue::GetSubjectKey() const
 {
 	static const FName InvalidName = FName(NAME_None);
-	return SubjectName.Get(InvalidName);
+	return SubjectKey.Get(InvalidName);
 }
 
 const FString& FControlRigOverrideValue::GetPath() const
@@ -434,17 +342,7 @@ FProperty* FControlRigOverrideValue::FindProperty(const UStruct* InStruct, const
 
 uint32 GetTypeHash(const FControlRigOverrideValue& InOverride)
 {
-	uint32 Hash = GetTypeHash(InOverride.SubjectName);
-	if(InOverride.IsValid())
-	{
-		for(const FControlRigOverrideValue::FPropertyInfo& PropertyInfo : InOverride.Properties)
-		{
-			Hash = HashCombine(Hash, GetTypeHash(PropertyInfo.Property));
-			Hash = HashCombine(Hash, GetTypeHash(PropertyInfo.ArrayIndex));
-		}
-		Hash = HashCombine(Hash, GetTypeHash(InOverride.ToString()));
-	}
-	return Hash;
+	return InOverride.Hash;
 }
 
 const FProperty* FControlRigOverrideValue::GetLeafProperty() const
@@ -478,6 +376,40 @@ const FProperty* FControlRigOverrideValue::GetLeafProperty() const
 	return nullptr;
 }
 
+int32 FControlRigOverrideValue::GetNumProperties() const
+{
+	return Properties.Num();
+}
+
+const FProperty* FControlRigOverrideValue::GetProperty(int32 InIndex) const
+{
+	if(Properties[InIndex].Property.IsValid())
+	{
+		return Properties[InIndex].Property.Get();
+	}
+	return nullptr;
+}
+
+int32 FControlRigOverrideValue::GetArrayIndex(int32 InIndex) const
+{
+	return Properties[InIndex].ArrayIndex;
+}
+
+bool FControlRigOverrideValue::ContainsProperty(const FProperty* InProperty) const
+{
+	for(int32 Index = 0; Index < Properties.Num(); Index++)
+	{
+		if(Properties[Index].Property.IsValid())
+		{
+			if(Properties[Index].Property.Get() == InProperty)
+			{
+				return true;
+			}
+		}
+	}
+	return false;
+}
+
 void* FControlRigOverrideValue::GetData()
 {
 	if(DataArray.IsEmpty())
@@ -496,6 +428,132 @@ const void* FControlRigOverrideValue::GetData() const
 	return DataArray.GetData();
 }
 
+bool FControlRigOverrideValue::SetPropertiesFromPath(const FString& InPath, const UStruct* InOwnerStruct)
+{
+	Path = InPath;
+
+	FString PropertyNameString;
+	FString ArrayIndexString;
+	bool bIsParsingPropertyName = true;
+	const UStruct* OwnerStruct = InOwnerStruct;
+
+	auto ProcessPropertyName = [this, &PropertyNameString, &OwnerStruct]() -> bool
+	{
+		FString StringToProcess;
+		Swap(StringToProcess, PropertyNameString);
+		if(!StringToProcess.IsEmpty())
+		{
+			if(FProperty* Property = FindProperty(OwnerStruct, StringToProcess))
+			{
+				FProperty* PropertyForPath = Property;
+				
+				if(const FArrayProperty* ArrayProperty = CastField<FArrayProperty>(Property))
+				{
+					Property = ArrayProperty->Inner;
+
+					if(Properties.IsEmpty())
+					{
+						PropertyForPath = Property;
+					}
+				}
+				if(const FStructProperty* StructProperty = CastField<FStructProperty>(Property))
+				{
+					OwnerStruct = StructProperty->Struct;
+				}
+
+				Properties.Add({PropertyForPath, INDEX_NONE});
+
+				return true;
+			}
+		}
+		return false;
+	};
+	
+	auto ProcessArrayIndex = [this, &ArrayIndexString]() -> bool
+	{
+		FString StringToProcess;
+		Swap(StringToProcess, ArrayIndexString);
+		if(!StringToProcess.IsEmpty())
+		{
+			Properties.Last().ArrayIndex = FCString::Atoi(*StringToProcess);
+			return true;
+		}
+		return false;
+	};
+
+	// Parse the provided Property path. Path may look like
+	// Settings[2]->Color->R resulting in three segments, the
+	// first one with an array index != INDEX_NONE
+	int32 CharIndex = 0;
+	while(CharIndex < InPath.Len())
+	{
+		// potentially start parsing an array index
+		if(InPath[CharIndex] == TEXT('['))
+		{
+			if(!ProcessPropertyName())
+			{
+				Reset();
+				return false;
+			}
+			bIsParsingPropertyName = false;
+			CharIndex++;
+			continue;;
+		}
+		
+		// potentially end parsing an array index
+		if(InPath[CharIndex] == TEXT(']'))
+		{
+			if(!ProcessArrayIndex())
+			{
+				Reset();
+				return false;
+			}
+			CharIndex++;
+			continue;;
+		}
+
+		// given a path separator parse the property name and store it.
+		// if bIsParsingPropertyName == false we already processed the property name
+		// and have been parsing an array index
+		if(InPath.Mid(CharIndex, PathSeparatorLength).Equals(PathSeparator, ESearchCase::CaseSensitive))
+		{
+			if(bIsParsingPropertyName)
+			{
+				if(!ProcessPropertyName())
+				{
+					Reset();
+					return false;
+				}
+			}
+			bIsParsingPropertyName = true;
+			CharIndex += PathSeparatorLength;
+			continue;
+		}
+
+		if(bIsParsingPropertyName)
+		{
+			PropertyNameString += InPath[CharIndex];
+		}
+		else
+		{
+			ArrayIndexString += InPath[CharIndex];
+		}
+		CharIndex++;
+	}
+
+	if(!PropertyNameString.IsEmpty() && !ProcessPropertyName())
+	{
+		Reset();
+		return false;
+	}
+	if(!ArrayIndexString.IsEmpty() && !ProcessArrayIndex())
+	{
+		Reset();
+		return false;
+	}
+	return true;
+}
+
 void* FControlRigOverrideValue::AllocateDataIfRequired()
 {
 	return AllocateDataIfRequired(DataArray);
@@ -572,6 +630,20 @@ void FControlRigOverrideValue::CopyValue(void* InDestPtr, const void* InSourcePt
 	}
 }
 
+void FControlRigOverrideValue::UpdateHash()
+{
+	Hash = GetTypeHash(SubjectKey);
+	if(IsValid())
+	{
+		for(const FPropertyInfo& PropertyInfo : Properties)
+		{
+			Hash = HashCombine(Hash, GetTypeHash(PropertyInfo.Property));
+			Hash = HashCombine(Hash, GetTypeHash(PropertyInfo.ArrayIndex));
+		}
+		Hash = HashCombine(Hash, GetTypeHash(CachedStringValue));
+	}
+}
+
 const FString& FControlRigOverrideValue::ToString() const
 {
 	static const FString EmptyPath;
@@ -596,6 +668,7 @@ bool FControlRigOverrideValue::SetFromString(const FString& InValue)
 				Swap(DataArray, TempDataArray);
 				
 				CachedStringValue = InValue;
+				UpdateHash();
 				return true;
 			}
 			FreeDataIfRequired(TempDataArray);
@@ -707,9 +780,9 @@ bool FControlRigOverrideValue::SetFromString(const FString& InValue)
 				}
 			}
 
-			if(SubjectName.IsSet())
+			if(SubjectKey.IsSet())
 			{
-				UE_LOG(LogControlRig, Warning, TEXT("Unable to import default value '%s' for override '%s' on subject '%s'."), *InValue, *GetPath(), *SubjectName.GetValue().ToString());
+				UE_LOG(LogControlRig, Warning, TEXT("Unable to import default value '%s' for override '%s' on subject '%s'."), *InValue, *GetPath(), *SubjectKey.GetValue().ToString());
 			}
 			else
 			{
@@ -722,29 +795,10 @@ bool FControlRigOverrideValue::SetFromString(const FString& InValue)
 
 bool FControlRigOverrideValue::operator==(const FControlRigOverrideValue& InOtherValue) const
 {
-	if(!IsValid() || !InOtherValue.IsValid())
-	{
-		return false;
-	}
-	if(GetSubjectName() != InOtherValue.GetSubjectName())
+	if(Hash != InOtherValue.Hash)
 	{
 		return false;
 	}
-	if(Properties.Num() != InOtherValue.Properties.Num())
-	{
-		return false;
-	}
-	for(int32 PropertyIndex = 0; PropertyIndex < Properties.Num(); PropertyIndex++)
-	{
-		if(Properties[PropertyIndex].Property != InOtherValue.Properties[PropertyIndex].Property)
-		{
-			return false;
-		}
-		if(Properties[PropertyIndex].ArrayIndex != InOtherValue.Properties[PropertyIndex].ArrayIndex)
-		{
-			return false;
-		}
-	}
 	return IdenticalValue(InOtherValue.GetData());
 }
 
@@ -771,7 +825,7 @@ bool FControlRigOverrideValue::IdenticalValueInSubject(const void* InSubjectPtr)
 	return IdenticalValue(SubjectPtrToValuePtr(InSubjectPtr, false));
 }
 
-void FControlRigOverrideValue::CopyToSubject(void* InSubjectPtr, const UStruct* InSubjectStruct) const
+bool FControlRigOverrideValue::CopyToSubject(void* InSubjectPtr, const UStruct* InSubjectStruct) const
 {
 	if(InSubjectStruct)
 	{
@@ -779,7 +833,7 @@ void FControlRigOverrideValue::CopyToSubject(void* InSubjectPtr, const UStruct*
 		{
 			if(RootProperty->GetOwnerStruct() != InSubjectStruct)
 			{
-				return;
+				return false;
 			}
 		}
 	}
@@ -789,11 +843,14 @@ void FControlRigOverrideValue::CopyToSubject(void* InSubjectPtr, const UStruct*
 		if(void* DestPtr = SubjectPtrToValuePtr(InSubjectPtr, true))
 		{
 			CopyValue(DestPtr, SourcePtr);
+			return true;
 		}
 	}
+
+	return false;
 }
 
-void FControlRigOverrideValue::SetFromSubject(const void* InSubjectPtr, const UStruct* InSubjectStruct)
+bool FControlRigOverrideValue::SetFromSubject(const void* InSubjectPtr, const UStruct* InSubjectStruct)
 {
 	if(InSubjectStruct)
 	{
@@ -801,13 +858,11 @@ void FControlRigOverrideValue::SetFromSubject(const void* InSubjectPtr, const US
 		{
 			if(RootProperty->GetOwnerStruct() != InSubjectStruct)
 			{
-				return;
+				return false;
 			}
 		}
 	}
 	
-	CachedStringValue.Reset();
-
 	if(const FProperty* LeafProperty = GetLeafProperty())
 	{
 		void* DestPtr = AllocateDataIfRequired();
@@ -819,16 +874,20 @@ void FControlRigOverrideValue::SetFromSubject(const void* InSubjectPtr, const US
 				CopyValue(DestPtr, SourcePtr);
 				LeafProperty->ExportTextItem_Direct(ValueAsString, DestPtr, DestPtr, nullptr, PPF_None, nullptr);
 				CachedStringValue = ValueAsString;
+				UpdateHash();
+				return true;
 			}
 		}
 	}
+
+	return false;
 }
 
 void FControlRigOverrideValue::CopyToUObject(UObject* InSubjectPtr) const
 {
 	if(InSubjectPtr)
 	{
-		CopyToSubject(InSubjectPtr, InSubjectPtr->GetClass());
+		(void)CopyToSubject(InSubjectPtr, InSubjectPtr->GetClass());
 	}
 }
 
@@ -836,7 +895,7 @@ void FControlRigOverrideValue::SetFromUObject(const UObject* InSubjectPtr)
 {
 	if(InSubjectPtr)
 	{
-		SetFromSubject(InSubjectPtr, InSubjectPtr->GetClass());
+		(void)SetFromSubject(InSubjectPtr, InSubjectPtr->GetClass());
 	}
 }
 
@@ -866,14 +925,14 @@ void FControlRigOverrideContainer::Reserve(int32 InNum)
 	HashIndexLookup.Reserve(InNum);
 }
 
-bool FControlRigOverrideContainer::UsesSubjectName() const
+bool FControlRigOverrideContainer::UsesKeyForSubject() const
 {
-	return bUsesSubjectName;
+	return bUsesKeyForSubject;
 }
 
-void FControlRigOverrideContainer::SetUsesSubjectName(bool InUsesSubjectName)
+void FControlRigOverrideContainer::SetUsesKeyForSubject(bool InUsesKeyForSubject)
 {
-	bUsesSubjectName = InUsesSubjectName;
+	bUsesKeyForSubject = InUsesKeyForSubject;
 	RebuildLookup();
 	InvalidateCache();
 }
@@ -888,17 +947,17 @@ int32 FControlRigOverrideContainer::Add(const FControlRigOverrideValue& InValue)
 	const FString& Path = InValue.GetPath();
 	if(!Path.IsEmpty())
 	{
-		if(!Contains(Path, InValue.GetSubjectName()))
+		if(!Contains(Path, InValue.GetSubjectKey()))
 		{
 			// don't allow to add overrides for child paths of
 			// existing parent paths
-			if(ContainsParentPathOf(Path, InValue.GetSubjectName()))
+			if(ContainsParentPathOf(Path, InValue.GetSubjectKey()))
 			{
 				return INDEX_NONE;
 			}
 			
 			// remove child paths when adding parent overrides
-			if(ContainsChildPathOf(Path, InValue.GetSubjectName()))
+			if(ContainsChildPathOf(Path, InValue.GetSubjectKey()))
 			{
 				TArray<int32> IndicesToRemove;
 				for(int32 Index = 0; Index < Values.Num(); Index++)
@@ -922,13 +981,10 @@ int32 FControlRigOverrideContainer::Add(const FControlRigOverrideValue& InValue)
 			}
 			
 			const int32 Index = Values.Add(InValue);
-			const uint32 Hash = GetLookupHash(Path, InValue.GetSubjectName());
+			const uint32 Hash = GetLookupHash(Path, InValue.GetSubjectKey());
 			HashIndexLookup.Add(Hash, Index);
+			SubjectIndexLookup.FindOrAdd(InValue.GetSubjectKey()).Add(Index);
 
-			if(bUsesSubjectName)
-			{
-				SubjectIndexLookup.FindOrAdd(InValue.GetSubjectName()).Add(Index);
-			}
 			InvalidateCache();
 			return Index;
 		}
@@ -941,7 +997,7 @@ const FControlRigOverrideValue* FControlRigOverrideContainer::FindOrAdd(const FC
 	const FString& Path = InValue.GetPath();
 	if(!Path.IsEmpty())
 	{
-		if(FControlRigOverrideValue* ExistingValue = Find(Path, InValue.GetSubjectName()))
+		if(FControlRigOverrideValue* ExistingValue = Find(Path, InValue.GetSubjectKey()))
 		{
 			*ExistingValue = InValue;
 			return ExistingValue;
@@ -958,12 +1014,12 @@ const FControlRigOverrideValue* FControlRigOverrideContainer::FindOrAdd(const FC
 
 bool FControlRigOverrideContainer::Remove(const FControlRigOverrideValue& InValue)
 {
-	return Remove(InValue.GetPath(), InValue.GetSubjectName());
+	return Remove(InValue.GetPath(), InValue.GetSubjectKey());
 }
 
-bool FControlRigOverrideContainer::Remove(const FString& InPath, const FName& InSubjectName)
+bool FControlRigOverrideContainer::Remove(const FString& InPath, const FName& InSubjectKey)
 {
-	const int32 Index = GetIndex(InPath, InSubjectName);
+	const int32 Index = GetIndex(InPath, InSubjectKey);
 	if(IsValidIndex(Index))
 	{
 		Values.RemoveAt(Index);
@@ -974,15 +1030,15 @@ bool FControlRigOverrideContainer::Remove(const FString& InPath, const FName& In
 	return false;
 }
 
-bool FControlRigOverrideContainer::RemoveAll(const FName& InSubjectName)
+bool FControlRigOverrideContainer::RemoveAll(const FName& InSubjectKey)
 {
 	if(IsEmpty())
 	{
 		return false;
 	}
-	if(bUsesSubjectName)
+	if(!InSubjectKey.IsNone())
 	{
-		if(const TArray<int32>* IndicesPtr = GetIndicesForSubject(InSubjectName))
+		if(const TArray<int32>* IndicesPtr = GetIndicesForSubject(InSubjectKey))
 		{
 			for(int32 Index = IndicesPtr->Num() - 1; Index >= 0; Index--)
 			{
@@ -1001,9 +1057,9 @@ bool FControlRigOverrideContainer::RemoveAll(const FName& InSubjectName)
 	return true;
 }
 
-int32 FControlRigOverrideContainer::GetIndex(const FString& InPath, const FName& InSubjectName) const
+int32 FControlRigOverrideContainer::GetIndex(const FString& InPath, const FName& InSubjectKey) const
 {
-	const uint32 Hash = GetLookupHash(InPath, InSubjectName);
+	const uint32 Hash = GetLookupHash(InPath, InSubjectKey);
 	if(const int32* Index = HashIndexLookup.Find(Hash))
 	{
 		return *Index;
@@ -1011,14 +1067,14 @@ int32 FControlRigOverrideContainer::GetIndex(const FString& InPath, const FName&
 	return INDEX_NONE;
 }
 
-const TArray<int32>* FControlRigOverrideContainer::GetIndicesForSubject(const FName& InSubjectName) const
+const TArray<int32>* FControlRigOverrideContainer::GetIndicesForSubject(const FName& InSubjectKey) const
 {
-	return SubjectIndexLookup.Find(InSubjectName);
+	return SubjectIndexLookup.Find(InSubjectKey);
 }
 
-const FControlRigOverrideValue* FControlRigOverrideContainer::Find(const FString& InPath, const FName& InSubjectName) const
+const FControlRigOverrideValue* FControlRigOverrideContainer::Find(const FString& InPath, const FName& InSubjectKey) const
 {
-	const int32 Index = GetIndex(InPath, InSubjectName);
+	const int32 Index = GetIndex(InPath, InSubjectKey);
 	if(IsValidIndex(Index))
 	{
 		return &Values[Index];
@@ -1026,9 +1082,9 @@ const FControlRigOverrideValue* FControlRigOverrideContainer::Find(const FString
 	return nullptr;
 }
 
-FControlRigOverrideValue* FControlRigOverrideContainer::Find(const FString& InPath, const FName& InSubjectName)
+FControlRigOverrideValue* FControlRigOverrideContainer::Find(const FString& InPath, const FName& InSubjectKey)
 {
-	const int32 Index = GetIndex(InPath, InSubjectName);
+	const int32 Index = GetIndex(InPath, InSubjectKey);
 	if(IsValidIndex(Index))
 	{
 		return &Values[Index];
@@ -1036,52 +1092,59 @@ FControlRigOverrideValue* FControlRigOverrideContainer::Find(const FString& InPa
 	return nullptr;
 }
 
-const FControlRigOverrideValue& FControlRigOverrideContainer::FindChecked(const FString& InPath, const FName& InSubjectName) const
+const FControlRigOverrideValue& FControlRigOverrideContainer::FindChecked(const FString& InPath, const FName& InSubjectKey) const
 {
-	const int32 Index = GetIndex(InPath, InSubjectName);
+	const int32 Index = GetIndex(InPath, InSubjectKey);
 	check(IsValidIndex(Index));
 	return Values[Index];
 }
 
-FControlRigOverrideValue& FControlRigOverrideContainer::FindChecked(const FString& InPath, const FName& InSubjectName)
+FControlRigOverrideValue& FControlRigOverrideContainer::FindChecked(const FString& InPath, const FName& InSubjectKey)
 {
-	const int32 Index = GetIndex(InPath, InSubjectName);
+	const int32 Index = GetIndex(InPath, InSubjectKey);
 	check(IsValidIndex(Index));
 	return Values[Index];
 }
 
-bool FControlRigOverrideContainer::Contains(const FString& InPath, const FName& InSubjectName) const
+TArray<FName> FControlRigOverrideContainer::GenerateSubjectArray() const
+{
+	TArray<FName> Keys;
+	SubjectIndexLookup.GenerateKeyArray(Keys);
+	return Keys;
+}
+
+bool FControlRigOverrideContainer::Contains(const FString& InPath, const FName& InSubjectKey) const
 {
-	return IsValidIndex(GetIndex(InPath, InSubjectName));
+	return IsValidIndex(GetIndex(InPath, InSubjectKey));
 }
 
-bool FControlRigOverrideContainer::ContainsParentPathOf(const FString& InChildPath, const FName& InSubjectName) const
+bool FControlRigOverrideContainer::ContainsParentPathOf(const FString& InChildPath, const FName& InSubjectKey) const
 {
 	if(InChildPath.IsEmpty())
 	{
 		return false;
 	}
 
-	const uint32 Hash = GetLookupHash(InChildPath, InSubjectName);
+	const uint32 Hash = GetLookupHash(InChildPath, InSubjectKey);
 	if(const bool* CachedResult = ContainsParentPathCache.Find(Hash))
 	{
 		return *CachedResult;
 	}
 
 	const FString ParentPath = GetParentPath(InChildPath);
-	const bool bResult = !ParentPath.IsEmpty() && (Contains(ParentPath, InSubjectName) || ContainsParentPathOf(ParentPath, InSubjectName));
+	const bool bResult = !ParentPath.IsEmpty() && (Contains(ParentPath, InSubjectKey) || ContainsParentPathOf(ParentPath, InSubjectKey));
 	ContainsParentPathCache.Add(Hash, bResult);
 	return bResult;
 }
 
-bool FControlRigOverrideContainer::ContainsChildPathOf(const FString& InParentPath, const FName& InSubjectName) const
+bool FControlRigOverrideContainer::ContainsChildPathOf(const FString& InParentPath, const FName& InSubjectKey) const
 {
 	if(InParentPath.IsEmpty())
 	{
 		return false;
 	}
 	
-	const uint32 Hash = GetLookupHash(InParentPath, InSubjectName);
+	const uint32 Hash = GetLookupHash(InParentPath, InSubjectKey);
 	if(const bool* CachedResult = ContainsChildPathCache.Find(Hash))
 	{
 		return *CachedResult;
@@ -1089,9 +1152,9 @@ bool FControlRigOverrideContainer::ContainsChildPathOf(const FString& InParentPa
 
 	for(const FControlRigOverrideValue& Value : Values)
 	{
-		if(bUsesSubjectName)
+		if(bUsesKeyForSubject)
 		{
-			if(Value.GetSubjectName() != InSubjectName)
+			if(Value.GetSubjectKey() != InSubjectKey)
 			{
 				continue;
 			}
@@ -1108,48 +1171,65 @@ bool FControlRigOverrideContainer::ContainsChildPathOf(const FString& InParentPa
 	return false;
 }
 
+bool FControlRigOverrideContainer::ContainsAnyPathForSubject(const FName& InSubjectKey) const
+{
+	return SubjectIndexLookup.Contains(InSubjectKey);
+}
+
+bool FControlRigOverrideContainer::ContainsPathForAnySubject(const FString& InPath) const
+{
+	for(const FControlRigOverrideValue& Override : Values)
+	{
+		if(Override.GetPath() == InPath)
+		{
+			return true;
+		}
+	}
+	return false;
+}
+
 bool FControlRigOverrideContainer::Contains(const FControlRigOverrideValue& InOverrideValue) const
 {
-	return Contains(InOverrideValue.GetPath(), InOverrideValue.GetSubjectName());
+	return Contains(InOverrideValue.GetPath(), InOverrideValue.GetSubjectKey());
 }
 
 bool FControlRigOverrideContainer::ContainsParentPathOf(const FControlRigOverrideValue& InOverrideValue) const
 {
-	return ContainsParentPathOf(InOverrideValue.GetPath(), InOverrideValue.GetSubjectName());
+	return ContainsParentPathOf(InOverrideValue.GetPath(), InOverrideValue.GetSubjectKey());
 }
 
 bool FControlRigOverrideContainer::ContainsChildPathOf(const FControlRigOverrideValue& InOverrideValue) const
 {
-	return ContainsChildPathOf(InOverrideValue.GetPath(), InOverrideValue.GetSubjectName());
+	return ContainsChildPathOf(InOverrideValue.GetPath(), InOverrideValue.GetSubjectKey());
 }
 
-void FControlRigOverrideContainer::CopyToSubject(void* InSubjectPtr, const UStruct* InSubjectStruct, const FName& InSubjectName) const
+void FControlRigOverrideContainer::CopyToSubject(void* InSubjectPtr, const UStruct* InSubjectStruct, const FName& InSubjectKey) const
 {
 	for(const FControlRigOverrideValue& Override : Values)
 	{
-		if(bUsesSubjectName)
+		if(bUsesKeyForSubject)
 		{
-			if(Override.GetSubjectName() != InSubjectName)
+			if(Override.GetSubjectKey() != InSubjectKey)
 			{
 				continue;
 			}
 		}
-		Override.CopyToSubject(InSubjectPtr, InSubjectStruct);
+		(void)Override.CopyToSubject(InSubjectPtr, InSubjectStruct);
 	}
 }
 
-void FControlRigOverrideContainer::SetFromSubject(const void* InSubjectPtr, const UStruct* InSubjectStruct, const FName& InSubjectName)
+void FControlRigOverrideContainer::SetFromSubject(const void* InSubjectPtr, const UStruct* InSubjectStruct, const FName& InSubjectKey)
 {
 	for(FControlRigOverrideValue& Override : Values)
 	{
-		if(bUsesSubjectName)
+		if(bUsesKeyForSubject)
 		{
-			if(Override.GetSubjectName() != InSubjectName)
+			if(Override.GetSubjectKey() != InSubjectKey)
 			{
 				continue;
 			}
 		}
-		Override.SetFromSubject(InSubjectPtr, InSubjectStruct);
+		(void)Override.SetFromSubject(InSubjectPtr, InSubjectStruct);
 	}
 }
 
@@ -1165,7 +1245,7 @@ void FControlRigOverrideContainer::SetFromUObject(const UObject* InSubjectPtr)
 
 bool FControlRigOverrideContainer::Serialize(FArchive& Ar)
 {
-	Ar << bUsesSubjectName;
+	Ar << bUsesKeyForSubject;
 	Ar << Values;
 	
 	if(Ar.IsLoading())
@@ -1173,7 +1253,6 @@ bool FControlRigOverrideContainer::Serialize(FArchive& Ar)
 		RebuildLookup();
 		InvalidateCache();
 	}
-
 	return true;
 }
 
@@ -1190,13 +1269,9 @@ void FControlRigOverrideContainer::RebuildLookup()
 			const FString& Path = Values[Index].GetPath();
 			if(!Path.IsEmpty())
 			{
-				const uint32 Hash = GetLookupHash(Path, Values[Index].GetSubjectName());
+				const uint32 Hash = GetLookupHash(Path, Values[Index].GetSubjectKey());
 				HashIndexLookup.Add(Hash, Index);
-
-				if(bUsesSubjectName)
-				{
-					SubjectIndexLookup.FindOrAdd(Values[Index].GetSubjectName()).Add(Index);
-				}
+				SubjectIndexLookup.FindOrAdd(Values[Index].GetSubjectKey()).Add(Index);
 			}
 		}
 	}
@@ -1208,6 +1283,107 @@ void FControlRigOverrideContainer::InvalidateCache() const
 	ContainsChildPathCache.Reset();
 }
 
+#if WITH_EDITOR
+void UControlRigOverrideAsset::PostTransacted(const FTransactionObjectEvent& TransactionEvent)
+{
+	UObject::PostTransacted(TransactionEvent);
+
+	if(TransactionEvent.GetEventType() == ETransactionObjectEventType::UndoRedo)
+	{
+		for (TObjectIterator<UControlRig> It(RF_ClassDefaultObject, true, EInternalObjectFlags::Garbage); It; ++It)
+		{
+			UControlRig* ControlRig = *It;
+			if(IsValid(ControlRig))
+			{
+				if(ControlRig->IsLinkedToOverrideAsset(this))
+				{
+					ControlRig->RequestConstruction();
+				}
+			}
+		}
+
+	}
+}
+#endif
+
+UControlRigOverrideAsset* UControlRigOverrideAsset::CreateOverrideAsset(const FString& InLongName)
+{
+	UControlRigOverrideAsset* Asset = nullptr;
+#if WITH_EDITOR
+	const FAssetToolsModule& AssetToolsModule = FModuleManager::LoadModuleChecked<FAssetToolsModule>("AssetTools");
+	FString UniquePackageName;
+	FString UniqueAssetName;
+	AssetToolsModule.Get().CreateUniqueAssetName(InLongName, TEXT(""), UniquePackageName, UniqueAssetName);
+	if (UniquePackageName.EndsWith(UniqueAssetName))
+	{
+		UniquePackageName = UniquePackageName.LeftChop(UniqueAssetName.Len() + 1);
+	}
+	Asset = Cast<UControlRigOverrideAsset>(AssetToolsModule.Get().CreateAsset(*UniqueAssetName, *UniquePackageName, UControlRigOverrideAsset::StaticClass(), nullptr));
+	check(Asset);
+#endif
+	return Asset;
+}
+
+UControlRigOverrideAsset* UControlRigOverrideAsset::CreateOverrideAssetInDeveloperFolder(const UObject* InSubject)
+{
+	UControlRigOverrideAsset* Asset = nullptr;
+#if WITH_EDITOR
+	FString DeveloperPathWithoutSlash = FPackageName::FilenameToLongPackageName(FPaths::GameUserDeveloperDir()).LeftChop(1);
+	FString ShortName = InSubject ? InSubject->GetClass()->GetOutermost()->GetName() : TEXT("Default");
+	ShortName = FPaths::GetBaseFilename(ShortName) + TEXT("_Override");
+	const FString LongName = FPaths::Combine(DeveloperPathWithoutSlash, TEXT("Overrides"), ShortName);
+	Asset = CreateOverrideAsset(LongName);
+#endif
+	return Asset;
+}
+
+void UControlRigOverrideAsset::BroadcastChanged()
+{
+	OverrideChangedDelegate.Broadcast(this);
+}
+
+void FControlRigOverrideContainer::AddStructReferencedObjects(FReferenceCollector& Collector)
+{
+	for(const FControlRigOverrideValue& Override : Values)
+	{
+		if(const FProperty* RootProperty = Override.GetRootProperty())
+		{
+			if(UClass* Class = RootProperty->GetOwnerClass())
+			{
+				TWeakObjectPtr<UClass> WeakClass(Class);
+				Collector.AddReferencedObject(WeakClass);
+			}
+			else if(UStruct* Struct = RootProperty->GetOwnerStruct())
+			{
+				TWeakObjectPtr<UStruct> WeakStruct(Struct);
+				Collector.AddReferencedObject(WeakStruct);
+			}
+		}
+	}
+}
+
+bool FControlRigOverrideContainer::Identical(const FControlRigOverrideContainer* Other, uint32 PortFlags) const
+{
+	if(Other == nullptr)
+	{
+		return false;
+	}
+	if(Num() != Other->Num())
+	{
+		return false;
+	}
+	for(int32 Index = 0; Index < Num(); Index++)
+	{
+		const FControlRigOverrideValue& Value = operator[](Index);
+		const FControlRigOverrideValue& OtherValue = Other->operator[](Index);
+		if(!Value.Identical(OtherValue))
+		{
+			return false;
+		}
+	}
+	return true;
+}
+
 FString FControlRigOverrideContainer::GetParentPath(const FString& InChildPath)
 {
 	int32 LastSeparatorIndex = InChildPath.Find(FControlRigOverrideValue::PathSeparator, ESearchCase::CaseSensitive, ESearchDir::FromEnd);
@@ -1271,7 +1447,7 @@ TSharedPtr<FPropertyPath> FControlRigOverrideValue::ToPropertyPath() const
 
 uint32 GetTypeHash(const FControlRigOverrideContainer& InContainer)
 {
-	uint32 Hash = 0;
+	uint32 Hash = GetTypeHash(InContainer.bUsesKeyForSubject);
 	for(const FControlRigOverrideValue& Override : InContainer.Values)
 	{
 		Hash = HashCombine(Hash, GetTypeHash(Override));
diff --git a/Engine/Plugins/Animation/ControlRig/Source/ControlRig/Private/ModularRigModel.cpp b/Engine/Plugins/Animation/ControlRig/Source/ControlRig/Private/ModularRigModel.cpp
--- a/Engine/Plugins/Animation/ControlRig/Source/ControlRig/Private/ModularRigModel.cpp
+++ b/Engine/Plugins/Animation/ControlRig/Source/ControlRig/Private/ModularRigModel.cpp
@@ -109,7 +109,7 @@ void FRigModuleReference::PatchModelsOnLoad()
 		if(!ConfigValues_DEPRECATED.IsEmpty())
 		{
 			ConfigOverrides.Reset();
-			ConfigOverrides.SetUsesSubjectName(false);
+			ConfigOverrides.SetUsesKeyForSubject(false);
 			ConfigOverrides.Reserve(ConfigValues_DEPRECATED.Num());
 			for(const TPair<FName, FString>& Pair : ConfigValues_DEPRECATED)
 			{
diff --git a/Engine/Plugins/Animation/ControlRig/Source/ControlRig/Private/Rigs/RigHierarchy.cpp b/Engine/Plugins/Animation/ControlRig/Source/ControlRig/Private/Rigs/RigHierarchy.cpp
--- a/Engine/Plugins/Animation/ControlRig/Source/ControlRig/Private/Rigs/RigHierarchy.cpp
+++ b/Engine/Plugins/Animation/ControlRig/Source/ControlRig/Private/Rigs/RigHierarchy.cpp
@@ -2946,7 +2946,37 @@ FText URigHierarchy::GetDisplayNameForUI(const FRigBaseElement* InElement, EElem
 
 #if WITH_EDITOR
 	
-	const FName& DisplayName = InElement->GetDisplayName();
+	const FName& OriginalDisplayName = InElement->GetDisplayName();
+	FName DisplayName = OriginalDisplayName;
+	
+	if(CVarControlRigEnableOverrides.GetValueOnAnyThread())
+	{
+		if(const FRigControlElement* ControlElement = Cast<FRigControlElement>(InElement))
+		{
+			if(UControlRig* ControlRig = GetTypedOuter<UControlRig>())
+			{
+				static const FString DisplayNamePath = TEXT("Settings->DisplayName");
+				for(const UControlRigOverrideAsset* OverrideAsset : ControlRig->OverrideAssets)
+				{
+					if(OverrideAsset)
+					{
+						if(const FControlRigOverrideValue* Override = OverrideAsset->Overrides.Find(DisplayNamePath, ControlElement->GetFName()))
+						{
+							ensureMsgf(Override->GetLeafProperty()->IsA<FNameProperty>(), TEXT("Encountered an override property that's not a FName."));
+							if(const FName* DisplayNameFromOverride = Override->GetData<FName>())
+							{
+								if(!DisplayNameFromOverride->IsNone())
+								{
+									DisplayName = *DisplayNameFromOverride;
+								}
+							}
+						}
+					}
+				}
+			}
+		}
+	}
+	
 	FString DisplayNameString = DisplayName.ToString();
 	
 	const FRigHierarchyModulePath ModulePath(DisplayNameString);
diff --git a/Engine/Plugins/Animation/ControlRig/Source/ControlRig/Private/Rigs/RigHierarchyElements.cpp b/Engine/Plugins/Animation/ControlRig/Source/ControlRig/Private/Rigs/RigHierarchyElements.cpp
--- a/Engine/Plugins/Animation/ControlRig/Source/ControlRig/Private/Rigs/RigHierarchyElements.cpp
+++ b/Engine/Plugins/Animation/ControlRig/Source/ControlRig/Private/Rigs/RigHierarchyElements.cpp
@@ -1389,7 +1389,8 @@ FRigControlSettings::FRigControlSettings()
 , bGroupWithParentControl(false)
 , bRestrictSpaceSwitching(false) 
 , PreferredRotationOrder(FRigPreferredEulerAngles::DefaultRotationOrder)
-, bUsePreferredRotationOrder(false) 
+, bUsePreferredRotationOrder(false)
+, ShapeTransform(FTransform::Identity)
 {
 	// rely on the default provided by the shape definition
 	ShapeName = FControlRigShapeDefinition().ShapeName; 
@@ -1642,6 +1643,7 @@ void FRigControlSettings::Load(FArchive& Ar)
 		bUsePreferredRotationOrder = false;
 	}
 
+	ShapeTransform = FTransform::Identity;
 }
 
 uint32 GetTypeHash(const FRigControlSettings& Settings)
diff --git a/Engine/Plugins/Animation/ControlRig/Source/ControlRig/Private/Sequencer/MovieSceneControlRigParameterSection.cpp b/Engine/Plugins/Animation/ControlRig/Source/ControlRig/Private/Sequencer/MovieSceneControlRigParameterSection.cpp
--- a/Engine/Plugins/Animation/ControlRig/Source/ControlRig/Private/Sequencer/MovieSceneControlRigParameterSection.cpp
+++ b/Engine/Plugins/Animation/ControlRig/Source/ControlRig/Private/Sequencer/MovieSceneControlRigParameterSection.cpp
@@ -31,10 +31,13 @@
 #include "MovieSceneTracksComponentTypes.h"
 #include "Sequencer/MovieSceneControlRigSystem.h"
 #include "ControlRigObjectBinding.h"
+#include "UObject/UObjectIterator.h"
+#include "ControlRigOverride.h"
 
 #if WITH_EDITOR
 #include "Misc/MessageDialog.h"
 #include "Settings/ControlRigSettings.h"
+#include "Misc/TransactionObjectEvent.h"
 #endif
 
 #include UE_INLINE_GENERATED_CPP_BY_NAME(MovieSceneControlRigParameterSection)
@@ -1077,8 +1080,87 @@ void UMovieSceneControlRigParameterSection::PostDuplicate(bool bDuplicateForPIE)
 {
 	Super::PostDuplicate(bDuplicateForPIE);	
 }
+
+void UMovieSceneControlRigParameterSection::PostTransacted(const FTransactionObjectEvent& TransactionEvent)
+{
+	Super::PostTransacted(TransactionEvent);
+
+	if(TransactionEvent.GetEventType() == ETransactionObjectEventType::UndoRedo)
+	{
+		const FName OverrideAssetsName = GET_MEMBER_NAME_CHECKED(UMovieSceneControlRigParameterSection, OverrideAssets); 
+		if(TransactionEvent.GetChangedProperties().Contains(OverrideAssetsName))
+		{
+			if(ControlRig)
+			{
+				ControlRig->UnlinkAllOverrideAssets();
+				for(const TSoftObjectPtr<UControlRigOverrideAsset>& OverrideAssetPtr : OverrideAssets)
+				{
+					if(UControlRigOverrideAsset* OverrideAsset = OverrideAssetPtr.LoadSynchronous())
+					{
+						ControlRig->LinkOverrideAsset(OverrideAsset);
+					}
+				}
+				ControlRig->RequestConstruction();
+			}
+
+			UpdateOverrideAssetDelegates();
+		}
+	}
+}
+
 #endif
 
+void UMovieSceneControlRigParameterSection::HandleOverrideAssetsChanged(UControlRig* InControlRig)
+{
+#if WITH_EDITORONLY_DATA
+	if(bSuspendOverrideAssetSync)
+	{
+		return;
+	}
+	if(InControlRig != ControlRig)
+	{
+		return;
+	}
+	Modify();
+	OverrideAssets.Reset();
+	for(int32 Index=0;Index<InControlRig->NumOverrideAssets();Index++)
+	{
+		OverrideAssets.AddUnique(InControlRig->GetOverrideAsset(Index));
+	}
+	UpdateOverrideAssetDelegates();
+	ReconstructChannelProxy();
+#endif
+}
+
+void UMovieSceneControlRigParameterSection::UpdateOverrideAssetDelegates()
+{
+	for (TObjectIterator<UControlRigOverrideAsset> It(RF_ClassDefaultObject, true, EInternalObjectFlags::Garbage); It; ++It)
+	{
+		It->OnChanged().RemoveAll(this);
+	}
+#if WITH_EDITORONLY_DATA
+	for(const TSoftObjectPtr<UControlRigOverrideAsset>& OverrideAssetPtr : OverrideAssets)
+	{
+		if(UControlRigOverrideAsset* OverrideAsset = OverrideAssetPtr.Get())
+		{
+			OverrideAsset->OnChanged().AddUObject(this, &UMovieSceneControlRigParameterSection::HandleOverrideAssetChanged);
+		}
+	}
+#endif
+}
+
+void UMovieSceneControlRigParameterSection::HandleOverrideAssetChanged(const UControlRigOverrideAsset* InOverrideAsset)
+{
+	if(InOverrideAsset)
+	{
+		static const FString DisplayNameString = TEXT("Settings->DisplayName");
+		if(InOverrideAsset->Overrides.ContainsPathForAnySubject(DisplayNameString))
+		{
+			ReconstructChannelProxy();
+		}
+	}
+}
+
 void UMovieSceneControlRigParameterSection::PostEditImport()
 {
 	Super::PostEditImport();
@@ -3050,8 +3132,45 @@ void UMovieSceneControlRigParameterSection::RecreateWithThisControlRig(UControlR
 
 void UMovieSceneControlRigParameterSection::SetControlRig(UControlRig* InControlRig)
 {
+	if(ControlRig)
+	{
+		ControlRig->OnOverrideAssetsChanged().Remove(OnOverrideAssetsChangedHandle);
+#if WITH_EDITORONLY_DATA
+		for(const TSoftObjectPtr<UControlRigOverrideAsset>& OverrideAssetPtr : OverrideAssets)
+		{
+			if(UControlRigOverrideAsset* OverrideAsset = OverrideAssetPtr.Get())
+			{
+				ControlRig->UnlinkOverrideAsset(OverrideAsset);
+			}
+		}
+#endif
+	}
+	
 	ControlRig = InControlRig;
 	ControlRigClass = ControlRig ? ControlRig->GetClass() : nullptr;
+
+	if(ControlRig)
+	{
+#if WITH_EDITORONLY_DATA
+		for(const TSoftObjectPtr<UControlRigOverrideAsset>& OverrideAssetPtr : OverrideAssets)
+		{
+			if(UControlRigOverrideAsset* OverrideAsset = OverrideAssetPtr.LoadSynchronous())
+			{
+				ControlRig->LinkOverrideAsset(OverrideAsset);
+			}
+		}
+#endif
+		TWeakObjectPtr<UMovieSceneControlRigParameterSection> WeakThis(this);
+		OnOverrideAssetsChangedHandle = ControlRig->OnOverrideAssetsChanged().AddLambda(
+			[WeakThis](UControlRig* InControlRig)
+			{
+				if(WeakThis.IsValid())
+				{
+					WeakThis->HandleOverrideAssetsChanged(InControlRig);
+				}
+			}
+		);
+	}
 }
 
 void UMovieSceneControlRigParameterSection::ChangeControlRotationOrder(const FName& InControlName, const  TOptional<EEulerRotationOrder>& OldOrder,
diff --git a/Engine/Plugins/Animation/ControlRig/Source/ControlRig/Private/Units/Execution/RigUnit_DynamicHierarchy.cpp b/Engine/Plugins/Animation/ControlRig/Source/ControlRig/Private/Units/Execution/RigUnit_DynamicHierarchy.cpp
--- a/Engine/Plugins/Animation/ControlRig/Source/ControlRig/Private/Units/Execution/RigUnit_DynamicHierarchy.cpp
+++ b/Engine/Plugins/Animation/ControlRig/Source/ControlRig/Private/Units/Execution/RigUnit_DynamicHierarchy.cpp
@@ -504,6 +504,7 @@ void FRigUnit_HierarchyAddControl_ShapeSettings::Configure(FRigControlSettings&
 	OutSettings.bShapeVisible = bVisible;
 	OutSettings.ShapeName = Name;
 	OutSettings.ShapeColor = Color;
+	OutSettings.ShapeTransform = Transform;
 }
 
 void FRigUnit_HierarchyAddControl_ProxySettings::ConfigureFrom(const FRigControlElement* InControlElement, const FRigControlSettings& InSettings)
diff --git a/Engine/Plugins/Animation/ControlRig/Source/ControlRig/Public/ControlRig.h b/Engine/Plugins/Animation/ControlRig/Source/ControlRig/Public/ControlRig.h
--- a/Engine/Plugins/Animation/ControlRig/Source/ControlRig/Public/ControlRig.h
+++ b/Engine/Plugins/Animation/ControlRig/Source/ControlRig/Public/ControlRig.h
@@ -7,6 +7,7 @@
 #include "Templates/SubclassOf.h"
 #include "ControlRigDefines.h"
 #include "ControlRigGizmoLibrary.h"
+#include "ControlRigOverride.h"
 #include "Rigs/RigHierarchy.h"
 #include "Units/RigUnitContext.h"
 #include "Animation/NodeMappingProviderInterface.h"
@@ -21,6 +22,7 @@
 #include "Animation/MeshDeformerInstance.h"
 #include "Animation/MeshDeformerProducer.h"
 #include "Rigs/RigModuleDefines.h"
+#include "UObject/OverriddenPropertySet.h"
 
 #if WITH_EDITOR
 #include "RigVMModel/RigVMPin.h"
@@ -43,11 +45,13 @@ class AActor;
 class UTransformableControlHandle;
 class UControlRigReplay;
 class UControlRig;
+class UControlRigOverrideAsset;
 
 struct FReferenceSkeleton;
 struct FRigUnit;
 struct FRigControl;
 struct FRigPhysicsSimulationBase;
+struct FControlRigOverrideContainer;
 
 CONTROLRIG_API DECLARE_LOG_CATEGORY_EXTERN(LogControlRig, Log, All);
 
@@ -110,6 +114,8 @@ DECLARE_DYNAMIC_MULTICAST_DELEGATE_TwoParams(FControlRigInteractionEndedEvent, U
 DECLARE_DYNAMIC_MULTICAST_DELEGATE_ThreeParams(FControlRigInteractionUpdatedEvent, UControlRig*, ControlRig, TArray<FRigElementKey>, InteractionKeys, FControlRigInteractionTransformContext, TransformContext);
 #endif
 
+DECLARE_MULTICAST_DELEGATE_OneParam(FControlRigOverrideAssetsChanged, UControlRig*);
+
 /** Runs logic for mapping input data to transforms (the "Rig") */
 UCLASS(MinimalAPI, Blueprintable, Abstract, editinlinenew)
 class UControlRig : public URigVMHost, public INodeMappingProviderInterface, public IRigHierarchyProvider, public IMeshDeformerProducer
@@ -282,10 +288,41 @@ public:
 	//~ Begin IMeshDeformerProducer interface
 	virtual FMeshDeformerBeginDestroyEvent& OnBeginDestroy() override {return BeginDestroyEvent;};
 	//~ End IMeshDeformerProducer interface
+
+	CONTROLRIG_API int32 NumOverrideAssets() const;
 	
+	/**
+	 * The active override asset is the last one linked to this control rig,
+	 * which can be used to record changes.
+	 * It needs to be the last once since they get applied in order, and if two
+	 * override assets override the same property the last one is the one that will win.
+	 */
+	CONTROLRIG_API const UControlRigOverrideAsset* GetLastOverrideAsset() const;
+	CONTROLRIG_API UControlRigOverrideAsset* GetLastOverrideAsset();
+	CONTROLRIG_API const UControlRigOverrideAsset* GetOverrideAsset(int32 InIndex) const;
+	CONTROLRIG_API UControlRigOverrideAsset* GetOverrideAsset(int32 InIndex);
+	CONTROLRIG_API int32 LinkOverrideAsset(UControlRigOverrideAsset* InOverrideAsset);
+	CONTROLRIG_API bool UnlinkOverrideAsset(UControlRigOverrideAsset* InContainer);
+	CONTROLRIG_API bool UnlinkAllOverrideAssets();
+	CONTROLRIG_API bool IsLinkedToOverrideAsset(const UControlRigOverrideAsset* InOverrideAsset) const;
+	CONTROLRIG_API FControlRigOverrideAssetsChanged& OnOverrideAssetsChanged() { return OverrideAssetsChangedDelegate; }
+	CONTROLRIG_API void SetSuspendOverrideAssetChangedDelegate(bool bSuspended)
+	{
+		bSuspendOverrideAssetChangedDelegate = bSuspended;
+	}
+
 private:
+
+	void ApplyElementOverrides();
+	void HandleOverrideChanged(const UControlRigOverrideAsset* InOverrideAsset);
+
 	/** Broadcasts a notification just before the controlrig is destroyed. */
 	FMeshDeformerBeginDestroyEvent BeginDestroyEvent;
+
+	UPROPERTY(Transient)
+	TArray<TObjectPtr<UControlRigOverrideAsset>> OverrideAssets;
+
+	bool bSuspendOverrideAssetChangedDelegate = false;
 	
 public:	
 	
@@ -385,6 +422,7 @@ public:
 	CONTROLRIG_API FVector GetControlSpecifiedEulerAngle(const FRigControlElement* InControlElement, bool bIsInitial = false) const;
 
 	CONTROLRIG_API virtual const TArray<TSoftObjectPtr<UControlRigShapeLibrary>>& GetShapeLibraries() const;
+	CONTROLRIG_API const TMap<FString, FString>& GetShapeLibraryNameMap() const { return ShapeLibraryNameMap; }
 	CONTROLRIG_API virtual void CreateRigControlsForCurveContainer();
 	CONTROLRIG_API virtual void GetControlsInOrder(TArray<FRigControlElement*>& SortedControls) const;
 
@@ -749,6 +787,8 @@ protected:
 
 	TArray<FRigElementKey> QueuedModifiedControls;
 
+	FControlRigOverrideAssetsChanged OverrideAssetsChangedDelegate;
+
 private:
 
 #if WITH_EDITORONLY_DATA
@@ -1017,6 +1057,7 @@ private:
 	friend class UModularRigController;
 	friend struct FControlRigHierarchyMappings;
 	friend class UControlRigReplay;
+	friend class FControlRigEditModeShapeSettingsDetails;
 };
 
 class FControlRigBracketScope
diff --git a/Engine/Plugins/Animation/ControlRig/Source/ControlRig/Public/ControlRigOverride.h b/Engine/Plugins/Animation/ControlRig/Source/ControlRig/Public/ControlRigOverride.h
--- a/Engine/Plugins/Animation/ControlRig/Source/ControlRig/Public/ControlRigOverride.h
+++ b/Engine/Plugins/Animation/ControlRig/Source/ControlRig/Public/ControlRigOverride.h
@@ -6,6 +6,7 @@
 #include "Hash/CityHash.h"
 #include "UObject/WeakFieldPtr.h"
 #include "RigVMStringUtils.h"
+#include "HAL/IConsoleManager.h"
 #include "ControlRigOverride.generated.h"
 
 #define UE_API CONTROLRIG_API
@@ -14,6 +15,8 @@
 class FPropertyPath;
 #endif
 
+extern CONTROLRIG_API TAutoConsoleVariable<bool> CVarControlRigEnableOverrides;
+
 /**
  * A single value used to represent an override on a subject.
  * The value serialized the data using binary serialization based on the last
@@ -33,17 +36,17 @@ public:
 	static inline constexpr int32 PathSeparatorLength = 2;
 	static inline constexpr TCHAR ArraySeparator = TEXT('[');
 
-	FControlRigOverrideValue()
+	UE_API FControlRigOverrideValue()
 	{}
 
 	// constructor given the path and the subject (the instance representing the memory)
 	UE_API FControlRigOverrideValue(const FString& InPath, const UObject* InSubject);
 
 	// constructor given the path, the owning structure as well as the container
-	UE_API FControlRigOverrideValue(const FString& InPath, const UStruct* InOwnerStruct, const void* InSubjectPtr, const FName& InSubjectName = NAME_None);
+	UE_API FControlRigOverrideValue(const FString& InPath, const UStruct* InOwnerStruct, const void* InSubjectPtr, const FName& InSubjectKey = NAME_None);
 
 	// constructor given the path, the owning structure as well as the value as string
-	UE_API FControlRigOverrideValue(const FString& InPath, const UStruct* InOwnerStruct, const FString& InValueAsString, const FName& InSubjectName = NAME_None);
+	UE_API FControlRigOverrideValue(const FString& InPath, const UStruct* InOwnerStruct, const FString& InValueAsString, const FName& InSubjectKey = NAME_None);
 
 	UE_API FControlRigOverrideValue(const FControlRigOverrideValue& InOther);
 
@@ -67,7 +70,7 @@ public:
 	UE_API void Reset();
 
 	// returns the name of the subject
-	UE_API const FName& GetSubjectName() const;
+	UE_API const FName& GetSubjectKey() const;
 
 	// returns the path this value represents
 	UE_API const FString& GetPath() const;
@@ -82,14 +85,14 @@ public:
 	template<typename T>
 	T* GetData()
 	{
-		return static_cast<T*>(GetData());
+		return reinterpret_cast<T*>(GetData());
 	}
 
 	// returns the raw data of this value
 	template<typename T>
 	const T* GetData() const
 	{
-		return static_cast<T*>(GetData());
+		return reinterpret_cast<const T*>(GetData());
 	}
 
 	// returns the value of this override as a string (or an empty string)  
@@ -99,10 +102,10 @@ public:
 	UE_API bool SetFromString(const FString& InValue);
 
 	// copies the override value onto a container
-	UE_API void CopyToSubject(void* InSubjectPtr, const UStruct* InSubjectStruct = nullptr) const;
+	[[nodiscard]] UE_API bool CopyToSubject(void* InSubjectPtr, const UStruct* InSubjectStruct = nullptr) const;
 
 	// copies the value from a container into this override value
-	UE_API void SetFromSubject(const void* InSubjectPtr, const UStruct* InSubjectStruct = nullptr);
+	[[nodiscard]] UE_API bool SetFromSubject(const void* InSubjectPtr, const UStruct* InSubjectStruct = nullptr);
 
 	// copies the override value onto a container
 	UE_API void CopyToUObject(UObject* InSubjectPtr) const;
@@ -132,6 +135,17 @@ public:
 	// returns the leaf property represented by this value (potentially leaf == root property)
 	UE_API const FProperty* GetLeafProperty() const;
 
+	// number of properties in the override
+	int32 GetNumProperties() const;
+
+	// returns a given property
+	const FProperty* GetProperty(int32 InIndex) const;
+
+	// returns the array index for a given property
+	int32 GetArrayIndex(int32 InIndex) const;
+
+	bool ContainsProperty(const FProperty* InProperty) const;
+
 	/**
 	 * converts a memory pointer from the subject to the leaf property.
 	 * @param bResizeArrays If set to true the subject memory will be resized to fit the property chain
@@ -141,15 +155,17 @@ public:
 	// helper method to look up a property under a structure
 	static UE_API FProperty* FindProperty(const UStruct* InStruct, const FString& InNameOrDisplayName);
 
-	CONTROLRIG_API friend uint32 GetTypeHash(const FControlRigOverrideValue& InOverride);
+	UE_API friend uint32 GetTypeHash(const FControlRigOverrideValue& InOverride);
 
 private:
 
+	UE_API bool SetPropertiesFromPath(const FString& InPath, const UStruct* InOwnerStruct);
 	UE_API void* AllocateDataIfRequired();
 	UE_API void FreeDataIfRequired();
 	UE_API void* AllocateDataIfRequired(TArray<uint8, TAlignedHeapAllocator<16>>& InOutDataArray) const;
 	UE_API void FreeDataIfRequired(TArray<uint8, TAlignedHeapAllocator<16>>& InOutDataArray) const;
 	UE_API void CopyValue(void* InDestPtr, const void* InSourcePtr) const;
+	UE_API void UpdateHash();
 
 	template<typename T>
 	bool InitFromLegacyString(const FString& InLegacyString, const T& InDefault)
@@ -160,6 +176,7 @@ private:
 			void* DataPtr = AllocateDataIfRequired(); 
 			CopyValue(DataPtr, &Value);
 			CachedStringValue = InLegacyString;
+			UpdateHash();
 			return true;
 		}
 		return false;
@@ -167,7 +184,8 @@ private:
 
 	TOptional<FString> Path;
 	TOptional<FString> CachedStringValue;
-	TOptional<FName> SubjectName;
+	TOptional<FName> SubjectKey;
+	uint32 Hash;
 
 	struct FPropertyInfo
 	{
@@ -191,7 +209,7 @@ struct TStructOpsTypeTraits<FControlRigOverrideValue> : public TStructOpsTypeTra
 
 /**
  * A container to represent a set of override values for one or more subjects.
- * The bUsesSubjectName setting will control if the subject name is respected when
+ * The bUsesKeyForSubject setting will control if the subject name is respected when
  * adding / looking up / applying overrides.
  * This data-structure is not thread-safe for writing.
  */
@@ -203,7 +221,7 @@ struct FControlRigOverrideContainer
 public:
 
 	FControlRigOverrideContainer()
-		: bUsesSubjectName(true)
+		: bUsesKeyForSubject(true)
 	{}
 
 	bool IsEmpty() const { return Values.IsEmpty(); }
@@ -219,18 +237,18 @@ public:
 	UE_API void Reset();
 	UE_API void Empty();
 	UE_API void Reserve(int32 InNum);
-	UE_API bool UsesSubjectName() const;
-	UE_API void SetUsesSubjectName(bool InUsesSubjectName);
+	UE_API bool UsesKeyForSubject() const;
+	UE_API void SetUsesKeyForSubject(bool InUsesKeyForSubject);
 
 	UE_API int32 Add(const FControlRigOverrideValue& InValue);
 	UE_API const FControlRigOverrideValue* FindOrAdd(const FControlRigOverrideValue& InValue);
 	UE_API bool Remove(const FControlRigOverrideValue& InValue);
-	UE_API bool Remove(const FString& InPath, const FName& InSubjectName = NAME_None);
-	UE_API bool RemoveAll(const FName& InSubjectName = NAME_None);
+	UE_API bool Remove(const FString& InPath, const FName& InSubjectKey = NAME_None);
+	UE_API bool RemoveAll(const FName& InSubjectKey = NAME_None);
 
-	int32 Emplace(const FString& InPath, const UStruct* InOwnerStruct, const void* InSubjectPtr, const FName& InSubjectName = NAME_None)
+	int32 Emplace(const FString& InPath, const UStruct* InOwnerStruct, const void* InSubjectPtr, const FName& InSubjectKey = NAME_None)
 	{
-		const FControlRigOverrideValue Value(InPath, InOwnerStruct, InSubjectPtr, InSubjectName);
+		const FControlRigOverrideValue Value(InPath, InOwnerStruct, InSubjectPtr, InSubjectKey);
 		return Add(Value);
 	}
 
@@ -240,32 +258,36 @@ public:
 		return Add(Value);
 	}
 	
-	int32 Emplace(const FString& InPath, const UStruct* InOwnerStruct, const FString& InValueAsString, const FName& InSubjectName = NAME_None)
+	int32 Emplace(const FString& InPath, const UStruct* InOwnerStruct, const FString& InValueAsString, const FName& InSubjectKey = NAME_None)
 	{
-		const FControlRigOverrideValue Value(InPath, InOwnerStruct, InValueAsString, InSubjectName);
+		const FControlRigOverrideValue Value(InPath, InOwnerStruct, InValueAsString, InSubjectKey);
 		return Add(Value);
 	}
 
-	UE_API int32 GetIndex(const FString& InPath, const FName& InSubjectName = NAME_None) const;
-	UE_API const TArray<int32>* GetIndicesForSubject(const FName& InSubjectName) const;
-	UE_API const FControlRigOverrideValue* Find(const FString& InPath, const FName& InSubjectName = NAME_None) const;
-	UE_API FControlRigOverrideValue* Find(const FString& InPath, const FName& InSubjectName = NAME_None);
-	UE_API const FControlRigOverrideValue& FindChecked(const FString& InPath, const FName& InSubjectName = NAME_None) const;
-	UE_API FControlRigOverrideValue& FindChecked(const FString& InPath, const FName& InSubjectName = NAME_None);
+	UE_API int32 GetIndex(const FString& InPath, const FName& InSubjectKey = NAME_None) const;
+	UE_API const TArray<int32>* GetIndicesForSubject(const FName& InSubjectKey) const;
+	UE_API const FControlRigOverrideValue* Find(const FString& InPath, const FName& InSubjectKey = NAME_None) const;
+	UE_API FControlRigOverrideValue* Find(const FString& InPath, const FName& InSubjectKey = NAME_None);
+	UE_API const FControlRigOverrideValue& FindChecked(const FString& InPath, const FName& InSubjectKey = NAME_None) const;
+	UE_API FControlRigOverrideValue& FindChecked(const FString& InPath, const FName& InSubjectKey = NAME_None);
+
+	UE_API TArray<FName> GenerateSubjectArray() const;
 
-	UE_API bool Contains(const FString& InPath, const FName& InSubjectName = NAME_None) const;
-	UE_API bool ContainsParentPathOf(const FString& InChildPath, const FName& InSubjectName = NAME_None) const;
-	UE_API bool ContainsChildPathOf(const FString& InParentPath, const FName& InSubjectName = NAME_None) const;
+	UE_API bool Contains(const FString& InPath, const FName& InSubjectKey = NAME_None) const;
+	UE_API bool ContainsParentPathOf(const FString& InChildPath, const FName& InSubjectKey = NAME_None) const;
+	UE_API bool ContainsChildPathOf(const FString& InParentPath, const FName& InSubjectKey = NAME_None) const;
+	UE_API bool ContainsAnyPathForSubject(const FName& InSubjectKey) const;
+	UE_API bool ContainsPathForAnySubject(const FString& InPath) const;
 
 	UE_API bool Contains(const FControlRigOverrideValue& InOverrideValue) const;
 	UE_API bool ContainsParentPathOf(const FControlRigOverrideValue& InOverrideValue) const;
 	UE_API bool ContainsChildPathOf(const FControlRigOverrideValue& InOverrideValue) const;
 
 	// copies the override values onto a subject.
-	UE_API void CopyToSubject(void* InSubjectPtr, const UStruct* InSubjectStruct = nullptr, const FName& InSubjectName = NAME_None) const;
+	UE_API void CopyToSubject(void* InSubjectPtr, const UStruct* InSubjectStruct = nullptr, const FName& InSubjectKey = NAME_None) const;
 
 	// copies the value from the subject into this container
-	UE_API void SetFromSubject(const void* InSubjectPtr, const UStruct* InSubjectStruct = nullptr, const FName& InSubjectName = NAME_None);
+	UE_API void SetFromSubject(const void* InSubjectPtr, const UStruct* InSubjectStruct = nullptr, const FName& InSubjectKey = NAME_None);
 
 	// copies the override values onto a subject.
 	UE_API void CopyToUObject(UObject* InSubjectPtr) const;
@@ -274,12 +296,17 @@ public:
 	UE_API void SetFromUObject(const UObject* InSubjectPtr);
 
 	UE_API bool Serialize(FArchive& Ar);
+
 	friend FArchive& operator<<(FArchive& Ar, FControlRigOverrideContainer& P)
 	{
 		P.Serialize(Ar);
 		return Ar;
 	}
 
+	void AddStructReferencedObjects(FReferenceCollector& Collector);
+
+	bool Identical(const FControlRigOverrideContainer* Other, uint32 PortFlags) const;
+
 	// returns a parent path of a given child path, separators used are -> and [,
 	// so parent of 'Values[1]->Flag' is 'Values[1]',
 	// and parent of 'Values[1]' is 'Values'.
@@ -288,29 +315,25 @@ public:
 	// returns true if a given child path is a child of a given parent path 
 	static UE_API bool IsChildPathOf(const FString& InChildPath, const FString& InParentPath);
 
-	CONTROLRIG_API friend uint32 GetTypeHash(const FControlRigOverrideContainer& InContainer);
+	UE_API friend uint32 GetTypeHash(const FControlRigOverrideContainer& InContainer);
 
 private:
 
 	UE_API void RebuildLookup();
 	UE_API void InvalidateCache() const;
 	
-	FORCEINLINE uint32 GetLookupHash(const FString& InPath, const FName& InSubjectName) const
+	FORCEINLINE uint32 GetLookupHash(const FString& InPath, const FName& InSubjectKey) const
 	{
 		const uint32 PathHash = CityHash32(
 				(const char*)InPath.GetCharArray().GetData(),
 				InPath.GetCharArray().GetTypeSize() *
 				InPath.GetCharArray().Num());
-		if(!bUsesSubjectName)
-		{
-			return PathHash;
-		}
-		return HashCombineFast(GetTypeHash(InSubjectName), PathHash);
+		return HashCombineFast(GetTypeHash(InSubjectKey), PathHash);
 	}
 
 	TArray<FControlRigOverrideValue> Values;
 
-	bool bUsesSubjectName;
+	bool bUsesKeyForSubject;
 	TMap<uint32, int32> HashIndexLookup;
 	TMap<FName, TArray<int32>> SubjectIndexLookup;
 	mutable TMap<uint32, bool> ContainsParentPathCache; 
@@ -323,8 +346,106 @@ struct TStructOpsTypeTraits<FControlRigOverrideContainer> : public TStructOpsTyp
 	enum 
 	{
 		WithSerializer = true,
+		WithIdentical = true,
+		WithAddStructReferencedObjects = true,
 	};
 	static constexpr EPropertyObjectReferenceType WithSerializerObjectReferences = EPropertyObjectReferenceType::None;
 };
 
+class UControlRigOverrideAsset;
+DECLARE_EVENT_OneParam(UControlRigOverrideAsset, FControlRigOverrideChanged, const UControlRigOverrideAsset*);
+
+UCLASS(MinimalAPI, BlueprintType)
+class UControlRigOverrideAsset : public UObject
+{
+public:
+	GENERATED_BODY()
+
+	UPROPERTY(EditAnywhere, Category="Overrides")
+	FControlRigOverrideContainer Overrides;
+
+#if WITH_EDITOR
+	UE_API virtual void PostTransacted(const FTransactionObjectEvent& TransactionEvent) override;
+#endif
+
+	UE_API static UControlRigOverrideAsset* CreateOverrideAsset(const FString& InLongName);
+	UE_API static UControlRigOverrideAsset* CreateOverrideAssetInDeveloperFolder(const UObject* InSubject);
+
+	UE_API FControlRigOverrideChanged& OnChanged() { return OverrideChangedDelegate; }
+	UE_API void BroadcastChanged();
+
+private:
+	FControlRigOverrideChanged OverrideChangedDelegate;
+};
+
+template<typename T>
+class TControlRigOverrideHandle : public TSharedFromThis<TControlRigOverrideHandle<T>>
+{
+public:
+	
+	TControlRigOverrideHandle(UControlRigOverrideAsset* InOverrideAsset, int32 InIndex)
+		: WeakOverrideAsset(InOverrideAsset)
+		, OverrideIndex(InIndex)
+		, LeafProperty(nullptr)
+	{
+		if(UControlRigOverrideAsset* OverrideAsset = GetOverrideAsset())
+		{
+			if(OverrideAsset->Overrides.IsValidIndex(OverrideIndex))
+			{
+				LeafProperty = OverrideAsset->Overrides[OverrideIndex].GetLeafProperty();
+			}
+		}
+	}
+
+	bool IsValid() const
+	{
+		if(UControlRigOverrideAsset* OverrideAsset = GetOverrideAsset())
+		{
+			if(OverrideAsset->Overrides.IsValidIndex(OverrideIndex))
+			{
+				return OverrideAsset->Overrides[OverrideIndex].GetLeafProperty() == LeafProperty;
+			}
+		}
+		return false;
+	}
+
+	UControlRigOverrideAsset* GetOverrideAsset() const
+	{
+		if(WeakOverrideAsset.IsValid())
+		{
+			return WeakOverrideAsset.Get();
+		}
+		return nullptr;
+	}
+
+	const FProperty* GetLeafProperty() const
+	{
+		return LeafProperty;
+	}
+
+	const T* GetData() const
+	{
+		if(!IsValid())
+		{
+			return nullptr;
+		}
+		return GetOverrideAsset()->Overrides[OverrideIndex].template GetData<T>();
+	}
+
+	T* GetData()
+    {
+    	if(!IsValid())
+    	{
+    		return nullptr;
+    	}
+    	return GetOverrideAsset()->Overrides[OverrideIndex].template GetData<T>();
+    }
+
+private:
+
+	TWeakObjectPtr<UControlRigOverrideAsset> WeakOverrideAsset;
+	int32 OverrideIndex;
+	const FProperty* LeafProperty;
+};;
+
 #undef UE_API
diff --git a/Engine/Plugins/Animation/ControlRig/Source/ControlRig/Public/ModularRigModel.h b/Engine/Plugins/Animation/ControlRig/Source/ControlRig/Public/ModularRigModel.h
--- a/Engine/Plugins/Animation/ControlRig/Source/ControlRig/Public/ModularRigModel.h
+++ b/Engine/Plugins/Animation/ControlRig/Source/ControlRig/Public/ModularRigModel.h
@@ -62,7 +62,7 @@ public:
 		, PreviousName(NAME_None)
 		, PreviousParentName(NAME_None)
 	{
-		ConfigOverrides.SetUsesSubjectName(false);
+		ConfigOverrides.SetUsesKeyForSubject(false);
 	}
 	
 	FRigModuleReference(const FName& InName, TSubclassOf<UControlRig> InClass, const FName& InParentModuleName, const FModularRigModel* InModel)
@@ -74,7 +74,7 @@ public:
 		, PreviousParentName(NAME_None)
 		, Model(InModel)
 	{
-		ConfigOverrides.SetUsesSubjectName(false);
+		ConfigOverrides.SetUsesKeyForSubject(false);
 	}
 
 	UPROPERTY()
diff --git a/Engine/Plugins/Animation/ControlRig/Source/ControlRig/Public/Rigs/RigHierarchyElements.h b/Engine/Plugins/Animation/ControlRig/Source/ControlRig/Public/Rigs/RigHierarchyElements.h
--- a/Engine/Plugins/Animation/ControlRig/Source/ControlRig/Public/Rigs/RigHierarchyElements.h
+++ b/Engine/Plugins/Animation/ControlRig/Source/ControlRig/Public/Rigs/RigHierarchyElements.h
@@ -1600,6 +1600,12 @@ struct FRigControlSettings
 	UPROPERTY(meta=(DeprecatedProperty, DeprecationMessage = "Use animation_type or shape_visible instead."))
 	bool bShapeEnabled_DEPRECATED = true;
 #endif
+
+	/**
+	 * Transient storage for overrides when changing the shape transform
+	 */
+	UPROPERTY(BlueprintReadOnly, Transient, Category = Animation)
+	FTransform ShapeTransform;
 	
 	/** Applies the limits expressed by these settings to a value */
 	void ApplyLimits(FRigControlValue& InOutValue) const
diff --git a/Engine/Plugins/Animation/ControlRig/Source/ControlRig/Public/Sequencer/MovieSceneControlRigParameterSection.h b/Engine/Plugins/Animation/ControlRig/Source/ControlRig/Public/Sequencer/MovieSceneControlRigParameterSection.h
--- a/Engine/Plugins/Animation/ControlRig/Source/ControlRig/Public/Sequencer/MovieSceneControlRigParameterSection.h
+++ b/Engine/Plugins/Animation/ControlRig/Source/ControlRig/Public/Sequencer/MovieSceneControlRigParameterSection.h
@@ -335,6 +335,12 @@ protected:
 	UPROPERTY()
 	TArray<FConstraintAndActiveChannel> ConstraintsChannels;
 
+#if WITH_EDITORONLY_DATA
+	UPROPERTY(EditAnywhere, Category=Customization)
+	TArray<TSoftObjectPtr<UControlRigOverrideAsset>> OverrideAssets;
+	bool bSuspendOverrideAssetSync = false;
+#endif
+
 public:
 
 	UE_API UMovieSceneControlRigParameterSection();
@@ -652,8 +658,14 @@ protected:
 
 #if WITH_EDITOR
 	UE_API virtual void PostDuplicate(bool bDuplicateForPIE) override;
+	UE_API virtual void PostTransacted(const FTransactionObjectEvent& TransactionEvent) override;
 #endif
 
+	void HandleOverrideAssetsChanged(UControlRig* InControlRig);
+	void HandleOverrideAssetChanged(const UControlRigOverrideAsset* InOverrideAsset);
+	UE_API void UpdateOverrideAssetDelegates();
+
+	FDelegateHandle OnOverrideAssetsChangedHandle;
 
 	// When true we do not set a key on the section, since it will be set because we changed the value
 	// We need this because control rig notifications are set on every change even when just changing sequencer time
@@ -727,6 +739,8 @@ private:
 		return IndicesToRemove.Num() > 0;
 	}
 	UE_API void HACK_FixMultipleParamsWithSameName();
+
+	friend class SControlRigEditModeTools;
 };
 
 #undef UE_API
diff --git a/Engine/Plugins/Animation/ControlRig/Source/ControlRigEditor/Private/AnimDetails/Customizations/AnimDetailsOverrideDetails.cpp b/Engine/Plugins/Animation/ControlRig/Source/ControlRigEditor/Private/AnimDetails/Customizations/AnimDetailsOverrideDetails.cpp
new file mode 100644
--- /dev/null
+++ b/Engine/Plugins/Animation/ControlRig/Source/ControlRigEditor/Private/AnimDetails/Customizations/AnimDetailsOverrideDetails.cpp
@@ -0,0 +1,202 @@
+// Copyright Epic Games, Inc. All Rights Reserved.
+
+#include "AnimDetailsOverrideDetails.h"
+
+#include "AnimDetails/AnimDetailsProxyManager.h"
+#include "AnimDetails/AnimDetailsSelection.h"
+#include "AnimDetails/Proxies/AnimDetailsProxyBase.h"
+#include "AnimDetails/Widgets/SAnimDetailsPropertySelectionBorder.h"
+#include "DetailCategoryBuilder.h"
+#include "DetailLayoutBuilder.h"
+#include "DetailWidgetRow.h"
+#include "EditMode/ControlRigEditMode.h"
+#include "EditorModeManager.h"
+#include "IDetailGroup.h"
+#include "IPropertyUtilities.h"
+#include "IDetailChildrenBuilder.h"
+
+#define LOCTEXT_NAMESPACE "AnimDetailsProxyDetails"
+
+namespace UE::ControlRigEditor
+{
+	void FAnimDetailsOverrideDetails::CustomizeHeader(TSharedRef<IPropertyHandle> StructPropertyHandle, FDetailWidgetRow& HeaderRow,
+		IPropertyTypeCustomizationUtils& StructCustomizationUtils)
+	{
+		HeaderRow.NameContent()
+		.MaxDesiredWidth(30)
+		[
+			StructPropertyHandle->CreatePropertyNameWidget()
+		];
+	}
+	
+	void FAnimDetailsOverrideDetails::CustomizeChildren(TSharedRef<IPropertyHandle> StructPropertyHandle,
+		IDetailChildrenBuilder& StructBuilder, IPropertyTypeCustomizationUtils& StructCustomizationUtils)
+	{
+		TArray<TWeakObjectPtr<UObject>> ObjectsBeingCustomized = StructCustomizationUtils.GetPropertyUtilities()->GetSelectedObjects();
+		const UControlRig* ControlRig = nullptr;
+		for(const TWeakObjectPtr<UObject>& ObjectBeingCustomized : ObjectsBeingCustomized)
+		{
+			if(const UAnimDetailsProxyBase* ControlsProxy = Cast<UAnimDetailsProxyBase>(ObjectBeingCustomized.Get()))
+			{
+				ControlRig = ControlsProxy->GetControlRig();
+				if(ControlRig)
+				{
+					break;
+				}
+			}
+		}
+		if(ControlRig == nullptr)
+		{
+			return;
+		}
+		const TMap<FString, FString>& LibraryNameMap = ControlRig->GetShapeLibraryNameMap();
+		TArray<TSoftObjectPtr<UControlRigShapeLibrary>> ShapeLibraries = ControlRig->GetShapeLibraries();
+		for(const TSoftObjectPtr<UControlRigShapeLibrary>& ShapeLibrary : ShapeLibraries)
+		{
+			if(ShapeLibrary.IsNull() || !ShapeLibrary.IsValid())
+			{
+				(void)ShapeLibrary.LoadSynchronous();
+			}
+			if(ShapeLibrary.IsNull() || !ShapeLibrary.IsValid())
+			{
+				continue;
+			}
+			const bool bUseNameSpace = ShapeLibraries.Num() > 1;
+			FString LibraryName = ShapeLibrary->GetName();
+			if(const FString* RemappedName = LibraryNameMap.Find(LibraryName))
+			{
+				LibraryName = *RemappedName;
+			}
+			
+			const FString NameSpace = bUseNameSpace ? LibraryName + TEXT(".") : FString();
+			ShapeNameList.Add(MakeShared<FRigVMStringWithTag>(UControlRigShapeLibrary::GetShapeName(ShapeLibrary.Get(), bUseNameSpace, LibraryNameMap, ShapeLibrary->DefaultShape)));
+			for (const FControlRigShapeDefinition& Shape : ShapeLibrary->Shapes)
+			{
+				ShapeNameList.Add(MakeShared<FRigVMStringWithTag>(UControlRigShapeLibrary::GetShapeName(ShapeLibrary.Get(), bUseNameSpace, LibraryNameMap, Shape)));
+			}
+		}
+		if(ShapeNameList.IsEmpty())
+		{
+			return;
+		}
+		StructBuilder.AddProperty(
+			StructPropertyHandle->GetChildHandle(GET_MEMBER_NAME_CHECKED(FRigUnit_HierarchyAddControl_ShapeSettings, Color)).ToSharedRef());
+		Property = StructPropertyHandle->GetChildHandle(GET_MEMBER_NAME_CHECKED(FRigUnit_HierarchyAddControl_ShapeSettings, Name));
+		if(Property.IsValid())
+		{
+			TSharedPtr<FRigVMStringWithTag> InitialSelected;
+			const FString CurrentShapeName = GetShapeNameListText().ToString();
+			for (TSharedPtr<FRigVMStringWithTag> Item : ShapeNameList)
+			{
+				if (Item->Equals(CurrentShapeName))
+				{
+					InitialSelected = Item;
+				}
+			}
+			IDetailPropertyRow& Row = StructBuilder.AddProperty(Property.ToSharedRef());
+			constexpr bool bShowChildren = true;
+			Row.CustomWidget(bShowChildren)
+			.NameContent()
+			[
+				Property->CreatePropertyNameWidget()
+			]
+			.ValueContent()
+			[
+				SAssignNew(ShapeNameListWidget, SRigVMGraphPinNameListValueWidget)
+				.OptionsSource(&ShapeNameList)
+				.OnGenerateWidget(this, &FAnimDetailsOverrideDetails::MakeShapeNameListItemWidget)
+				.OnSelectionChanged(this, &FAnimDetailsOverrideDetails::OnShapeNameListChanged)
+				.OnComboBoxOpening(this, &FAnimDetailsOverrideDetails::OnShapeNameListComboBox)
+				.InitiallySelectedItem(InitialSelected)
+				.Content()
+				[
+					SNew(STextBlock)
+					.Text(this, &FAnimDetailsOverrideDetails::GetShapeNameListText)
+					.Font(IDetailLayoutBuilder::GetDetailFont())
+				]
+			];
+		}
+		
+		StructBuilder.AddProperty(
+			StructPropertyHandle->GetChildHandle(GET_MEMBER_NAME_CHECKED(FRigUnit_HierarchyAddControl_ShapeSettings, bVisible)).ToSharedRef());
+		if(TSharedPtr<IPropertyHandle> TransformHandle = StructPropertyHandle->GetChildHandle(GET_MEMBER_NAME_CHECKED(FRigUnit_HierarchyAddControl_ShapeSettings, Transform)))
+		{
+			StructBuilder.AddProperty(TransformHandle->GetChildHandle(TEXT("Translation")).ToSharedRef());
+			StructBuilder.AddProperty(TransformHandle->GetChildHandle(TEXT("Rotation")).ToSharedRef());
+			StructBuilder.AddProperty(TransformHandle->GetChildHandle(TEXT("Scale3D")).ToSharedRef());
+		}
+	}
+	
+	TSharedRef<SWidget> FAnimDetailsOverrideDetails::MakeShapeNameListItemWidget(TSharedPtr<FRigVMStringWithTag> InItem)
+	{
+		return SNew(STextBlock)
+			.Text(FText::FromString(InItem->GetStringWithTag()))
+			.Font(IDetailLayoutBuilder::GetDetailFont());
+	}
+	
+	void FAnimDetailsOverrideDetails::OnShapeNameListChanged(TSharedPtr<FRigVMStringWithTag> NewSelection, ESelectInfo::Type SelectInfo)
+	{
+		if (SelectInfo == ESelectInfo::Direct)
+		{
+			return;
+		}
+		
+		const FString& NewShapeNameString = NewSelection->GetString();
+		if(Property)
+		{
+			const FName ShapeName = NewShapeNameString.IsEmpty() ? FName(NAME_None) : FName(*NewShapeNameString);
+			Property->SetValue(ShapeName);
+		}
+	}
+	
+	void FAnimDetailsOverrideDetails::OnShapeNameListComboBox()
+	{
+		const FString ShapeNameListText = GetShapeNameListText().ToString();
+		const TSharedPtr<FRigVMStringWithTag>* CurrentlySelectedItem =
+			ShapeNameList.FindByPredicate([ShapeNameListText](const TSharedPtr<FRigVMStringWithTag>& InItem)
+			{
+				return ShapeNameListText == InItem->GetString();
+			});
+		
+		if(CurrentlySelectedItem)
+		{
+			ShapeNameListWidget->SetSelectedItem(*CurrentlySelectedItem);
+		}
+	}
+	
+	FText FAnimDetailsOverrideDetails::GetShapeNameListText() const
+	{
+		if(!Property)
+		{
+			return FText();
+		}
+		
+		TOptional<FString> SharedValue;
+		for(int32 Index = 0; Index < Property->GetNumPerObjectValues(); Index++)
+		{
+			FString SingleValue;
+			if(!Property->GetPerObjectValue(Index, SingleValue))
+			{
+				SharedValue.Reset();
+				break;
+			}
+			if(!SharedValue.IsSet())
+			{
+				SharedValue = SingleValue;
+			}
+			else if(SharedValue.GetValue() != SingleValue)
+			{
+				SharedValue.Reset();
+				break;
+			}
+		}
+		
+		if(SharedValue.IsSet())
+		{
+			return FText::FromString(SharedValue.GetValue());
+		}
+		return LOCTEXT("MultipleValues", "Multiple Values");
+	}
+}
+
+#undef LOCTEXT_NAMESPACE
diff --git a/Engine/Plugins/Animation/ControlRig/Source/ControlRigEditor/Private/AnimDetails/Customizations/AnimDetailsOverrideDetails.h b/Engine/Plugins/Animation/ControlRig/Source/ControlRigEditor/Private/AnimDetails/Customizations/AnimDetailsOverrideDetails.h
new file mode 100644
--- /dev/null
+++ b/Engine/Plugins/Animation/ControlRig/Source/ControlRigEditor/Private/AnimDetails/Customizations/AnimDetailsOverrideDetails.h
@@ -0,0 +1,33 @@
+// Copyright Epic Games, Inc. All Rights Reserved.
+
+#pragma once
+
+#include "IDetailCustomization.h"
+#include "IPropertyTypeCustomization.h"
+#include "DetailWidgetRow.h"
+#include "Input/Reply.h"
+#include "Templates/SharedPointer.h"
+#include "Widgets/SRigVMGraphPinNameList.h"
+
+namespace UE::ControlRigEditor
+{
+	/** Details customization for Control Shape settings */
+	class FAnimDetailsOverrideDetails : public IPropertyTypeCustomization
+	{
+	public:
+		static TSharedRef<class IPropertyTypeCustomization> MakeInstance()
+		{
+			return MakeShareable(new FAnimDetailsOverrideDetails());
+		}
+		virtual void CustomizeHeader(TSharedRef<IPropertyHandle> StructPropertyHandle, FDetailWidgetRow& HeaderRow, IPropertyTypeCustomizationUtils& StructCustomizationUtils) override;
+		virtual void CustomizeChildren(TSharedRef<IPropertyHandle> StructPropertyHandle, IDetailChildrenBuilder& StructBuilder, IPropertyTypeCustomizationUtils& StructCustomizationUtils) override;
+	private:
+		TSharedRef<SWidget> MakeShapeNameListItemWidget(TSharedPtr<FRigVMStringWithTag> InItem);
+		void OnShapeNameListChanged(TSharedPtr<FRigVMStringWithTag> NewSelection, ESelectInfo::Type SelectInfo);
+		void OnShapeNameListComboBox();
+		FText GetShapeNameListText() const;
+		TSharedPtr<IPropertyHandle> Property;
+		TArray<TSharedPtr<FRigVMStringWithTag>> ShapeNameList;
+		TSharedPtr<SRigVMGraphPinNameListValueWidget> ShapeNameListWidget;
+	};
+}
diff --git a/Engine/Plugins/Animation/ControlRig/Source/ControlRigEditor/Private/AnimDetails/Proxies/AnimDetailsProxyBase.cpp b/Engine/Plugins/Animation/ControlRig/Source/ControlRigEditor/Private/AnimDetails/Proxies/AnimDetailsProxyBase.cpp
--- a/Engine/Plugins/Animation/ControlRig/Source/ControlRigEditor/Private/AnimDetails/Proxies/AnimDetailsProxyBase.cpp
+++ b/Engine/Plugins/Animation/ControlRig/Source/ControlRigEditor/Private/AnimDetails/Proxies/AnimDetailsProxyBase.cpp
@@ -603,11 +603,16 @@ FRigControlElement* UAnimDetailsProxyBase::GetControlElement() const
 	return const_cast<FRigControlElement*>(ControlElement);
 }
 
-FRigElementKey UAnimDetailsProxyBase::GetControlElementKey() const
+const FRigElementKey& UAnimDetailsProxyBase::GetControlElementKey() const
 {
 	return CachedRigElement.GetResolvedKey();
 }
 
+const FName& UAnimDetailsProxyBase::GetControlName() const
+{
+	return GetControlElementKey().Name;
+}
+
 void UAnimDetailsProxyBase::PropagonateValues()
 {
 	FRigControlModifiedContext Context;
@@ -636,6 +641,11 @@ void UAnimDetailsProxyBase::PropagonateValues()
 
 FText UAnimDetailsProxyBase::GetDisplayNameText() const
 {
+	if(!DisplayName.IsEmpty())
+	{
+		return FText::FromString(DisplayName);
+	}
+	
 	const URigHierarchy* Hierarchy = WeakControlRig.IsValid() ? WeakControlRig->GetHierarchy() : nullptr;
 	const FRigControlElement* ControlElement = WeakControlRig.IsValid() ? GetControlElement() : nullptr;
 	if (Hierarchy && ControlElement)
@@ -786,6 +796,15 @@ void UAnimDetailsProxyBase::GetLocalizedPropertyName(const FName& InPropertyName
 	}
 }
 
+void UAnimDetailsProxyBase::UpdateOverrideableProperties()
+{
+	if(const FRigControlElement* ControlElement = GetControlElement())
+	{
+		DisplayName = ControlElement->GetDisplayName().ToString();
+		Shape.ConfigureFrom(ControlElement, ControlElement->Settings);
+	}
+}
+
 FName UAnimDetailsProxyBase::GetPropertyID(const FName& PropertyName) const
 {
 	return *(GetDetailRowID().ToString() + TEXT(".") + PropertyName.ToString());
diff --git a/Engine/Plugins/Animation/ControlRig/Source/ControlRigEditor/Private/AnimDetails/Proxies/AnimDetailsProxyBase.h b/Engine/Plugins/Animation/ControlRig/Source/ControlRigEditor/Private/AnimDetails/Proxies/AnimDetailsProxyBase.h
--- a/Engine/Plugins/Animation/ControlRig/Source/ControlRigEditor/Private/AnimDetails/Proxies/AnimDetailsProxyBase.h
+++ b/Engine/Plugins/Animation/ControlRig/Source/ControlRigEditor/Private/AnimDetails/Proxies/AnimDetailsProxyBase.h
@@ -5,6 +5,7 @@
 #include "AnimDetailsSequencerProxyItem.h"
 #include "Rigs/RigHierarchyCache.h"
 #include "Misc/Optional.h"
+#include "Units/Execution/RigUnit_DynamicHierarchy.h"
 
 #include "AnimDetailsProxyBase.generated.h"
 
@@ -44,7 +45,10 @@ public:
 	FRigControlElement* GetControlElement() const;
 
 	/** Returns a key to the control element this proxy handles. */
-	FRigElementKey GetControlElementKey() const;
+	const FRigElementKey& GetControlElementKey() const;
+
+	/** Returns the name of the control element his proxy handles */
+	const FName& GetControlName() const;
 
 	/** Returns the sequencer binding, or nullptr if this proxy is not assigned to a seqeuncer object */
 	const FAnimDetailsSequencerProxyItem& GetSequencerItem() { return SequencerItem; }
@@ -98,6 +102,9 @@ public:
 	/** Updates the proxy values from the sequencer binding or the control in the control rig respectively */
 	virtual void UpdateProxyValues() PURE_VIRTUAL(UAnimDetailsProxyBase::UpdateProxyValues, return;)
 
+	/** Updates the proxy shape override properties such as DisplayName and shape */
+	virtual void UpdateOverrideableProperties();
+
 	/** Returns the channel to key flags from the property name */
 	virtual EControlRigContextChannelToKey GetChannelToKeyFromPropertyName(const FName& PropertyName) const PURE_VIRTUAL(UAnimDetailsProxyBase::GetChannelToKeyFromPropertyName, return EControlRigContextChannelToKey::AllTransform;) 
 	
@@ -115,6 +122,14 @@ public:
 	/** The control rig type to use for this proxy. Note this is used even if there is only a sequencer binding but no control rig */
 	ERigControlType Type = ERigControlType::Transform;
 
+	/** An overrideable display name for the control */
+	UPROPERTY(EditAnywhere, Category = "Overrides")
+	FString DisplayName;
+
+	/** Overrideable shape settings for the control */
+	UPROPERTY(EditAnywhere, Category = "Overrides", meta=(ExpandByDefault))
+	FRigUnit_HierarchyAddControl_ShapeSettings Shape;	
+
 protected:
 	/** Sets the control rig element value from the current proxy value */
 	virtual void SetControlRigElementValueFromCurrent(UControlRig* ControlRig, FRigControlElement* ControlElement, const FRigControlModifiedContext& Context) PURE_VIRTUAL(UAnimDetailsProxyBase::SetControlRigElementValueFromCurrent, return;)
diff --git a/Engine/Plugins/Animation/ControlRig/Source/ControlRigEditor/Private/AnimDetails/Views/SAnimDetailsView.cpp b/Engine/Plugins/Animation/ControlRig/Source/ControlRigEditor/Private/AnimDetails/Views/SAnimDetailsView.cpp
--- a/Engine/Plugins/Animation/ControlRig/Source/ControlRigEditor/Private/AnimDetails/Views/SAnimDetailsView.cpp
+++ b/Engine/Plugins/Animation/ControlRig/Source/ControlRigEditor/Private/AnimDetails/Views/SAnimDetailsView.cpp
@@ -4,7 +4,9 @@
 
 #include "AnimDetails/AnimDetailsProxyManager.h"
 #include "AnimDetails/AnimDetailsSelection.h"
+#include "AnimDetails/Customizations/AnimDetailsProxyDetails.h"
 #include "AnimDetails/Customizations/AnimDetailsProxyManagerDetails.h"
+#include "AnimDetails/Proxies/AnimDetailsProxyBase.h"
 #include "AnimDetails/Widgets/SAnimDetailsOptions.h"
 #include "AnimDetails/Widgets/SAnimDetailsSearchBox.h"
 #include "CurveEditor.h"
@@ -15,12 +17,13 @@
 #include "MVVM/ViewModels/SequencerEditorViewModel.h"
 #include "PropertyEditorDelegates.h"
 #include "PropertyEditorModule.h"
+#include "SAnimOverrideDetailsView.h"
 
 namespace UE::ControlRigEditor
 {
 	void SAnimDetailsView::Construct(const FArguments& InArgs)
 	{
-		FControlRigEditMode* EditMode = static_cast<FControlRigEditMode*>(GLevelEditorModeTools().GetActiveMode(FControlRigEditMode::ModeName));
+		const FControlRigEditMode* EditMode = static_cast<FControlRigEditMode*>(GLevelEditorModeTools().GetActiveMode(FControlRigEditMode::ModeName));
 		UAnimDetailsProxyManager* ProxyManager = EditMode ? EditMode->GetAnimDetailsProxyManager() : nullptr;
 		if (!EditMode || !ProxyManager)
 		{
@@ -50,6 +53,7 @@ namespace UE::ControlRigEditor
 		DetailsView->SetGenericLayoutDetailsDelegate(FOnGetDetailCustomizationInstance::CreateStatic(&FAnimDetailProxyManagerDetails::MakeInstance));
 		DetailsView->SetIsPropertyVisibleDelegate(FIsPropertyVisible::CreateSP(this, &SAnimDetailsView::ShouldDisplayProperty));
 		DetailsView->SetIsPropertyReadOnlyDelegate(FIsPropertyReadOnly::CreateSP(this, &SAnimDetailsView::IsReadOnlyProperty));
+		DetailsView->RegisterInstancedCustomPropertyLayout(UAnimDetailsProxyBase::StaticClass(), FOnGetDetailCustomizationInstance::CreateStatic(&FAnimDetailsProxyDetails::MakeInstance));
 
 		DetailsView->SetObjects({ TWeakObjectPtr<UObject>(ProxyManager) });
 		
@@ -85,7 +89,24 @@ namespace UE::ControlRigEditor
 			+ SVerticalBox::Slot()
 			.FillHeight(1.f)
 			[
-				DetailsView
+				SNew(SVerticalBox)
+
+				+ SVerticalBox::Slot()
+				.AutoHeight()
+				[
+					DetailsView
+				]
+
+				+ SVerticalBox::Slot()
+				.AutoHeight()
+				[
+					SNew(SAnimOverrideDetailsView)
+					.OnRequestRefreshDetails(this, &SAnimDetailsView::RefreshDetailsView)
+					.Visibility_Lambda([]()
+					{
+						return CVarControlRigEnableOverrides.GetValueOnAnyThread() ? EVisibility::Visible : EVisibility::Collapsed;
+					})
+				]
 			]
 		];
 
@@ -118,7 +139,7 @@ namespace UE::ControlRigEditor
 
 	void SAnimDetailsView::RefreshDetailsView()
 	{
-		FControlRigEditMode* EditMode = static_cast<FControlRigEditMode*>(GLevelEditorModeTools().GetActiveMode(FControlRigEditMode::ModeName));
+		const FControlRigEditMode* EditMode = static_cast<FControlRigEditMode*>(GLevelEditorModeTools().GetActiveMode(FControlRigEditMode::ModeName));
 		UAnimDetailsProxyManager* ProxyManager = EditMode ? EditMode->GetAnimDetailsProxyManager() : nullptr;
 		if (!EditMode || !ProxyManager)
 		{
@@ -128,10 +149,7 @@ namespace UE::ControlRigEditor
 		const FText SearchText = SearchBox->GetSearchText();
 		ProxyManager->GetAnimDetailsFilter().Update(SearchText, ProxyManager->GetExternalSelection());
 
-		if (WeakDetailsView.IsValid())
-		{
-			WeakDetailsView.Pin()->ForceRefresh();
-		}
+		OnOptionsChanged();
 	}
 
 	void SAnimDetailsView::OnOptionsChanged()
@@ -146,7 +164,7 @@ namespace UE::ControlRigEditor
 	{
 		auto ShouldPropertyBeVisible = [](const FProperty& InProperty)
 			{
-				bool bShow = InProperty.HasAnyPropertyFlags(CPF_Interp) || InProperty.HasMetaData(FRigVMStruct::InputMetaName) || InProperty.HasMetaData(FRigVMStruct::OutputMetaName);
+				const bool bShow = InProperty.HasAnyPropertyFlags(CPF_Interp) || InProperty.HasMetaData(FRigVMStruct::InputMetaName) || InProperty.HasMetaData(FRigVMStruct::OutputMetaName);
 				return bShow;
 			};
 
@@ -183,7 +201,7 @@ namespace UE::ControlRigEditor
 
 		auto ShouldPropertyBeEnabledLambda = [](const FProperty& InProperty)
 			{
-				bool bShouldPropertyBeEnabled = InProperty.HasAnyPropertyFlags(CPF_Interp) || InProperty.HasMetaData(FRigVMStruct::InputMetaName);
+				const bool bShouldPropertyBeEnabled = InProperty.HasAnyPropertyFlags(CPF_Interp) || InProperty.HasMetaData(FRigVMStruct::InputMetaName);
 				return bShouldPropertyBeEnabled;
 			};
 
diff --git a/Engine/Plugins/Animation/ControlRig/Source/ControlRigEditor/Private/AnimDetails/Views/SAnimOverrideDetailsView.cpp b/Engine/Plugins/Animation/ControlRig/Source/ControlRigEditor/Private/AnimDetails/Views/SAnimOverrideDetailsView.cpp
new file mode 100644
--- /dev/null
+++ b/Engine/Plugins/Animation/ControlRig/Source/ControlRigEditor/Private/AnimDetails/Views/SAnimOverrideDetailsView.cpp
@@ -0,0 +1,470 @@
+// Copyright Epic Games, Inc. All Rights Reserved.
+
+#include "SAnimOverrideDetailsView.h"
+
+#include "AnimDetails/AnimDetailsProxyManager.h"
+#include "AnimDetails/Proxies/AnimDetailsProxyBase.h"
+#include "AnimDetails/Customizations/AnimDetailsProxyDetails.h"
+#include "AnimDetails/AnimDetailsSelection.h"
+#include "AnimDetails/Customizations/AnimDetailsProxyManagerDetails.h"
+#include "AnimDetails/Widgets/SAnimDetailsOptions.h"
+#include "AnimDetails/Widgets/SAnimDetailsSearchBox.h"
+#include "CurveEditor.h"
+#include "EditMode/ControlRigEditMode.h"
+#include "EditorModeManager.h"
+#include "ISequencer.h"
+#include "MVVM/CurveEditorExtension.h"
+#include "MVVM/ViewModels/SequencerEditorViewModel.h"
+#include "PropertyEditorDelegates.h"
+#include "PropertyEditorModule.h"
+#include "Overrides/OverrideStatusDetailsObjectFilter.h"
+#include "Rigs/RigHierarchyController.h"
+#include "Widgets/Notifications/SNotificationList.h"
+#include "Framework/Notifications/NotificationManager.h"
+#include "AssetRegistry/IAssetRegistry.h"
+#include "AssetRegistry/AssetRegistryModule.h"
+#include "IContentBrowserSingleton.h"
+#include "ContentBrowserModule.h"
+#include "AnimDetails/Customizations/AnimDetailsOverrideDetails.h"
+
+namespace UE::ControlRigEditor
+{
+	SLATE_IMPLEMENT_WIDGET(SAnimOverrideDetailsView)
+	void SAnimOverrideDetailsView::PrivateRegisterAttributes(FSlateAttributeInitializer& AttributeInitializer)
+	{
+	}
+
+	void SAnimOverrideDetailsView::Construct(const FArguments& InArgs)
+	{
+		const FControlRigEditMode* EditMode = static_cast<FControlRigEditMode*>(GLevelEditorModeTools().GetActiveMode(FControlRigEditMode::ModeName));
+		UAnimDetailsProxyManager* ProxyManager = EditMode ? EditMode->GetAnimDetailsProxyManager() : nullptr;
+		if (!EditMode || !ProxyManager)
+		{
+			return;
+		}
+
+		RequestRefreshDetailsDelegate = InArgs._OnRequestRefreshDetails;
+
+		ProxyPropertyToControl = {
+			{
+				TEXT("DisplayName"),
+				TEXT("Settings->DisplayName"),
+			},
+			{
+				TEXT("Shape->bVisible"),
+				TEXT("Settings->bShapeVisible"),
+			},
+			{
+				TEXT("Shape->Name"),
+				TEXT("Settings->ShapeName"),
+			},
+			{
+				TEXT("Shape->Color"),
+				TEXT("Settings->ShapeColor"),
+			},
+			{
+				TEXT("Shape->Transform"),
+				TEXT("Settings->ShapeTransform"),
+			},
+			{
+				TEXT("Rotation"),
+				TEXT("Settings->ShapeTransform->Rotation"),
+			}
+		};
+		
+		FDetailsViewArgs DetailsViewArgs;
+		{
+			DetailsViewArgs.bAllowSearch = false;
+			DetailsViewArgs.bShowPropertyMatrixButton = false;
+			DetailsViewArgs.bHideSelectionTip = true;
+			DetailsViewArgs.bLockable = false;
+			DetailsViewArgs.bUpdatesFromSelection = false;
+			DetailsViewArgs.bShowOptions = false;
+			DetailsViewArgs.bShowModifiedPropertiesOption = true;
+			DetailsViewArgs.bCustomFilterAreaLocation = false;
+			DetailsViewArgs.bCustomNameAreaLocation = true;
+			DetailsViewArgs.NameAreaSettings = FDetailsViewArgs::HideNameArea;
+			DetailsViewArgs.bAllowMultipleTopLevelObjects = true;
+			DetailsViewArgs.bShowScrollBar = true;
+		}
+		
+		const TSharedPtr<IDetailsView> DetailsView = FModuleManager::GetModuleChecked<FPropertyEditorModule>("PropertyEditor").CreateDetailView(DetailsViewArgs);
+		DetailsView->SetIsPropertyVisibleDelegate(FIsPropertyVisible::CreateSP(this, &SAnimOverrideDetailsView::ShouldDisplayProperty));
+		DetailsView->SetIsPropertyReadOnlyDelegate(FIsPropertyReadOnly::CreateSP(this, &SAnimOverrideDetailsView::IsReadOnlyProperty));
+		DetailsView->RegisterInstancedCustomPropertyTypeLayout(FRigUnit_HierarchyAddControl_ShapeSettings::StaticStruct()->GetFName(), FOnGetPropertyTypeCustomizationInstance::CreateStatic(&FAnimDetailsOverrideDetails::MakeInstance));
+		DetailsView->OnFinishedChangingProperties().AddSP(this, &SAnimOverrideDetailsView::OnFinishedChangingOverride);
+		WeakDetailsView = DetailsView;
+
+		ObjectFilter = FOverrideStatusDetailsViewObjectFilter::Create();
+		ObjectFilter->OnCanMergeObjects().BindSP(this, &SAnimOverrideDetailsView::CanMergeObjects);
+		ObjectFilter->OnCanCreateWidget().BindSP(this, &SAnimOverrideDetailsView::CanCreateWidget);
+		ObjectFilter->OnGetStatus().BindSP(this, &SAnimOverrideDetailsView::GetOverrideStatus);
+		ObjectFilter->OnAddOverride().BindSP(this, &SAnimOverrideDetailsView::OnAddOverride);
+		ObjectFilter->OnClearOverride().BindSP(this, &SAnimOverrideDetailsView::OnClearOverride);
+		DetailsView->SetObjectFilter(ObjectFilter);
+		
+		ChildSlot
+		[
+			DetailsView.ToSharedRef()
+		];
+
+		ProxyManager->GetOnProxiesChanged().AddSP(this, &SAnimOverrideDetailsView::RefreshDetailsView);
+	}
+
+	void SAnimOverrideDetailsView::RefreshDetailsView()
+	{
+		if (WeakDetailsView.IsValid())
+		{
+			const FControlRigEditMode* EditMode = static_cast<FControlRigEditMode*>(GLevelEditorModeTools().GetActiveMode(FControlRigEditMode::ModeName));
+			const UAnimDetailsProxyManager* ProxyManager = EditMode ? EditMode->GetAnimDetailsProxyManager() : nullptr;
+			TArray<UObject*> OverrideableProxies;
+			if (ProxyManager)
+			{
+				const TArray<TObjectPtr<UAnimDetailsProxyBase>>& SelectedProxies = ProxyManager->GetExternalSelection();
+				for(const TObjectPtr<UAnimDetailsProxyBase>& SelectedProxy : SelectedProxies)
+				{
+					SelectedProxy->UpdateOverrideableProperties();
+					OverrideableProxies.Add(SelectedProxy);
+				}
+			}
+			WeakDetailsView.Pin()->SetObjects(OverrideableProxies);
+		}
+	}
+
+	bool SAnimOverrideDetailsView::ShouldDisplayProperty(const FPropertyAndParent& InPropertyAndParent) const
+	{
+		static const FLazyName CategoryName = TEXT("Category");
+		static const FString OverrideCategoryName = TEXT("Overrides");
+		if(InPropertyAndParent.Property.GetMetaData(CategoryName) == OverrideCategoryName)
+		{
+			return true;
+		}
+		if(!InPropertyAndParent.ParentProperties.IsEmpty())
+		{
+			if(InPropertyAndParent.ParentProperties.Last()->GetMetaData(CategoryName) == OverrideCategoryName)
+			{
+				return true;
+			}
+		}
+		return false;
+	}
+
+	bool SAnimOverrideDetailsView::IsReadOnlyProperty(const FPropertyAndParent& InPropertyAndParent) const
+	{
+		// only allow display name editing if we have a single control selected
+		static const FProperty* DisplayNameProperty = UAnimDetailsProxyBase::StaticClass()->FindPropertyByName(GET_MEMBER_NAME_CHECKED(UAnimDetailsProxyBase, DisplayName));
+		if(&InPropertyAndParent.Property == DisplayNameProperty)
+		{
+			return InPropertyAndParent.Objects.Num() > 1;
+		}
+	
+		return false;
+	}
+
+	void SAnimOverrideDetailsView::OnFinishedChangingOverride(const FPropertyChangedEvent& PropertyChangedEvent)
+	{
+		const FProperty* Property = PropertyChangedEvent.Property;
+		if(!ensure(Property))
+		{
+			return;
+		}
+		FString PropertyPathString = Property->GetName();
+		
+		TMap<FString, int32> PropertyNameStack;
+		PropertyChangedEvent.GetArrayIndicesPerObject(0, PropertyNameStack);
+		if (PropertyNameStack.Num() > 0)
+		{
+			TArray<FString> PropertyNames;
+			PropertyNameStack.GetKeys(PropertyNames);
+			for(int32 Index = 0; Index < PropertyNames.Num(); Index++)
+			{
+				const int32 ArrayIndex = PropertyNameStack.FindChecked(PropertyNames[Index]);
+				if(ArrayIndex != INDEX_NONE)
+				{
+					PropertyNames[Index] = FString::Printf(TEXT("%s[%d]"), *PropertyNames[Index], ArrayIndex);
+				}
+			}
+			Algo::Reverse(PropertyNames);
+			PropertyPathString = FString::Join(PropertyNames, TEXT("->"));
+		}
+		MapPropertyFromProxyToControl(PropertyPathString);
+
+		TMap<TWeakObjectPtr<UControlRig>, TArray<FRigElementKey>> PreviousSelection;
+		if(PropertyPathString.Contains(TEXT("DisplayName")))
+		{
+			for(int32 Index = 0; Index < PropertyChangedEvent.GetNumObjectsBeingEdited(); Index++)
+			{
+				if(const UAnimDetailsProxyBase* Proxy = Cast<UAnimDetailsProxyBase>(PropertyChangedEvent.GetObjectBeingEdited(Index)))
+				{
+					if(UControlRig* ControlRig = Proxy->GetControlRig())
+					{
+						if(const URigHierarchy* Hierarchy = ControlRig->GetHierarchy())
+						{
+							PreviousSelection.Add(ControlRig, Hierarchy->GetSelectedKeys());
+						}
+					}
+				}
+			}
+		}
+
+		{
+			FScopedTransaction Transaction(NSLOCTEXT("SAnimOverrideDetailsView", "EditOverrideValue", "Edit Override Value"));
+			for(int32 Index = 0; Index < PropertyChangedEvent.GetNumObjectsBeingEdited(); Index++)
+			{
+				if(const UAnimDetailsProxyBase* Proxy = Cast<UAnimDetailsProxyBase>(PropertyChangedEvent.GetObjectBeingEdited(Index)))
+				{
+					if(UControlRig* ControlRig = Proxy->GetControlRig())
+					{
+						if(URigHierarchy* Hierarchy = ControlRig->GetHierarchy())
+						{
+							if(const FRigControlElement* ControlElement = Hierarchy->Find<FRigControlElement>(Proxy->GetControlElementKey()))
+							{
+								if(const URigHierarchyController* Controller = Hierarchy->GetController())
+								{
+									// copy the settings from proxy to element
+									FRigControlSettings Settings = ControlElement->Settings;
+									const FText DisplayNameText = Proxy->GetDisplayNameText();
+									Settings.DisplayName = DisplayNameText.IsEmpty() ? FName(NAME_None) : FName(*DisplayNameText.ToString());
+									Proxy->Shape.Configure(Settings);
+									(void)Controller->SetControlSettings(ControlElement->GetKey(), Settings);
+									Hierarchy->SetControlShapeTransform(ControlElement->GetKey(), Proxy->Shape.Transform);
+								}
+								if(UControlRigOverrideAsset* DefaultOverrideAsset = GetOrCreateOverrideAsset(ControlRig))
+								{
+									DefaultOverrideAsset->Modify();
+									if(FControlRigOverrideValue* Override = DefaultOverrideAsset->Overrides.Find(PropertyPathString, ControlElement->GetFName()))
+									{
+										(void)Override->SetFromSubject(ControlElement, FRigControlElement::StaticStruct());
+									}
+									else
+									{
+										DefaultOverrideAsset->Overrides.FindOrAdd(FControlRigOverrideValue(PropertyPathString, FRigControlElement::StaticStruct(), ControlElement, ControlElement->GetFName()));
+									}
+									
+									DefaultOverrideAsset->BroadcastChanged();
+								}
+							}
+						}
+					}
+				}
+			}
+		}
+		if(PropertyPathString.Contains(TEXT("DisplayName")))
+		{
+			// refresh the build of this view so that the label of the category is up2date
+			RefreshDetailsView();
+			(void)RequestRefreshDetailsDelegate.ExecuteIfBound();
+		}
+
+		for(const TPair<TWeakObjectPtr<UControlRig>, TArray<FRigElementKey>>& Pair : PreviousSelection)
+		{
+			if(UControlRig* ControlRig = Pair.Key.Get())
+			{
+				if(URigHierarchy* Hierarchy = ControlRig->GetHierarchy())
+				{
+					if(URigHierarchyController* Controller = Hierarchy->GetController())
+					{
+						Controller->SetSelection(Pair.Value);
+					}
+				}
+			}
+		}
+	}
+
+	bool SAnimOverrideDetailsView::CanMergeObjects(const UObject* InObjectA, const UObject* InObjectB) const
+	{
+		if(InObjectA && InObjectB)
+		{
+			return InObjectA->IsA<UAnimDetailsProxyBase>() && (InObjectA->GetClass() == InObjectB->GetClass()); 
+		}
+		return false;
+	}
+
+	bool SAnimOverrideDetailsView::CanCreateWidget(const FOverrideStatusSubject& InSubject) const
+	{
+		return InSubject.Contains<UAnimDetailsProxyBase>();
+	}
+
+	EOverrideWidgetStatus::Type SAnimOverrideDetailsView::GetOverrideStatus(const FOverrideStatusSubject& InSubject) const
+	{
+		FString PropertyPathString = InSubject.GetPropertyPathString();
+		MapPropertyFromProxyToControl(PropertyPathString);
+		return InSubject.GetStatus<UAnimDetailsProxyBase>(
+			[PropertyPathString](const FOverrideStatusObjectHandle<UAnimDetailsProxyBase>& InControlsProxy) -> TOptional<EOverrideWidgetStatus::Type>
+			{
+				const UControlRig* ControlRig = InControlsProxy->GetControlRig();
+				if(ControlRig == nullptr)
+				{
+					return {};
+				}
+				for(int32 AssetIndex = 0; AssetIndex < ControlRig->NumOverrideAssets(); AssetIndex++)
+				{
+					const UControlRigOverrideAsset* Asset = ControlRig->GetOverrideAsset(AssetIndex);
+					check(Asset);
+					if(PropertyPathString.IsEmpty())
+					{
+						if(Asset->Overrides.ContainsAnyPathForSubject(InControlsProxy->GetControlName()))
+						{
+							return EOverrideWidgetStatus::ChangedInside;
+						}
+					}
+					else if(Asset->Overrides.Contains(PropertyPathString, InControlsProxy->GetControlName()))
+					{
+						return EOverrideWidgetStatus::ChangedHere;
+					}
+					else if(Asset->Overrides.ContainsChildPathOf(PropertyPathString, InControlsProxy->GetControlName()))
+					{
+						return EOverrideWidgetStatus::ChangedInside;
+					}
+					else if(Asset->Overrides.ContainsParentPathOf(PropertyPathString, InControlsProxy->GetControlName()))
+					{
+						return EOverrideWidgetStatus::ChangedOutside;
+					}
+				}
+				return {};
+			}
+		).Get(EOverrideWidgetStatus::None);
+	}
+
+	FReply SAnimOverrideDetailsView::OnAddOverride(const FOverrideStatusSubject& InSubject)
+	{
+		FScopedTransaction Transaction(NSLOCTEXT("SAnimOverrideDetailsView", "AddOverride", "Add Override"));
+			
+		FString PropertyPathString = InSubject.GetPropertyPathString();
+		MapPropertyFromProxyToControl(PropertyPathString);
+		TArray<UControlRig*> AffectedControlRigs;
+		TArray<UControlRigOverrideAsset*> AffectedAssets;
+		InSubject.ForEach<UAnimDetailsProxyBase>(
+			[InSubject, PropertyPathString, &AffectedAssets, &AffectedControlRigs](const FOverrideStatusObjectHandle<UAnimDetailsProxyBase>& InControlsProxy)
+			{
+				if(InSubject.HasPropertyPath())
+				{
+					UControlRig* ControlRig = InControlsProxy->GetControlRig();
+					if(ControlRig == nullptr)
+					{
+						return;
+					}
+					if(const FRigControlElement* ControlElement = ControlRig->FindControl(InControlsProxy->GetControlName()))
+					{
+						if(UControlRigOverrideAsset* DefaultOverrideAsset = GetOrCreateOverrideAsset(ControlRig))
+						{
+							const FControlRigOverrideValue Value(PropertyPathString, FRigControlElement::StaticStruct(), ControlElement, ControlElement->GetFName());
+							if(Value.IsValid())
+							{
+								DefaultOverrideAsset->Modify();
+								DefaultOverrideAsset->Overrides.Add(Value);
+								AffectedControlRigs.AddUnique(ControlRig);
+								AffectedAssets.AddUnique(DefaultOverrideAsset);
+							}
+						}
+					}
+				}
+			}
+		);
+		for(UControlRig* AffectedControlRig : AffectedControlRigs)
+		{
+			AffectedControlRig->SetSuspendOverrideAssetChangedDelegate(true);
+		}
+		for(UControlRigOverrideAsset* AffectedAsset : AffectedAssets)
+		{
+			AffectedAsset->BroadcastChanged();
+		}
+		for(UControlRig* AffectedControlRig : AffectedControlRigs)
+		{
+			AffectedControlRig->SetSuspendOverrideAssetChangedDelegate(false);
+		}
+		return FReply::Handled();
+	}
+
+	FReply SAnimOverrideDetailsView::OnClearOverride(const FOverrideStatusSubject& InSubject)
+	{
+		FScopedTransaction Transaction(NSLOCTEXT("SAnimOverrideDetailsView", "ClearOverride", "Clear Override"));
+			
+		FString PropertyPathString = InSubject.GetPropertyPathString();
+		MapPropertyFromProxyToControl(PropertyPathString);
+		TArray<UControlRigOverrideAsset*> AffectedAssets;
+		InSubject.ForEach<UAnimDetailsProxyBase>(
+			[InSubject, PropertyPathString, &AffectedAssets](const FOverrideStatusObjectHandle<UAnimDetailsProxyBase>& InControlsProxy)
+			{
+				if(InSubject.HasPropertyPath())
+				{
+					UControlRig* ControlRig = InControlsProxy->GetControlRig();
+					if(ControlRig == nullptr)
+					{
+						return;
+					}
+					if(UControlRigOverrideAsset* OverrideAsset = GetOrCreateOverrideAsset(ControlRig))
+					{
+						if(OverrideAsset->Overrides.Contains(PropertyPathString, InControlsProxy->GetControlName()))
+						{
+							OverrideAsset->Modify();
+							OverrideAsset->Overrides.Remove(PropertyPathString, InControlsProxy->GetControlName());
+							AffectedAssets.AddUnique(OverrideAsset);
+						}
+					}
+				}
+			}
+		);
+		for(UControlRigOverrideAsset* AffectedAsset : AffectedAssets)
+		{
+			AffectedAsset->BroadcastChanged();
+		}
+		return FReply::Handled();
+	}
+
+	void SAnimOverrideDetailsView::MapPropertyFromProxyToControl(FString& InOutPropertyPath) const
+	{
+		if(const FString* MappedPath = ProxyPropertyToControl.Find(InOutPropertyPath))
+		{
+			InOutPropertyPath= *MappedPath;
+		}
+		else
+		{
+			for(const TPair<FString,FString>& Pair : ProxyPropertyToControl)
+			{
+				if(InOutPropertyPath.StartsWith(Pair.Key))
+				{
+					InOutPropertyPath.ReplaceInline(*Pair.Key, *Pair.Value);
+				}
+			}
+		}
+	}
+
+	UControlRigOverrideAsset* SAnimOverrideDetailsView::GetOrCreateOverrideAsset(UControlRig* InControlRig)
+	{
+		check(InControlRig);
+		if(!CVarControlRigEnableOverrides.GetValueOnAnyThread())
+		{
+			return nullptr;
+		}
+	
+		if(InControlRig->NumOverrideAssets() == 0)
+		{
+			if(UControlRigOverrideAsset* OverrideAsset = UControlRigOverrideAsset::CreateOverrideAssetInDeveloperFolder(InControlRig))
+			{
+				OverrideAsset->Overrides.SetUsesKeyForSubject(false);
+				InControlRig->LinkOverrideAsset(OverrideAsset);
+				const FSoftObjectPath SoftObjectPath(OverrideAsset->GetOutermost());
+				static const FString Message = TEXT("An override asset has been created.");
+				FNotificationInfo Info(FText::FromString(Message));
+				Info.bUseSuccessFailIcons = true;
+				Info.Image = FAppStyle::GetBrush(TEXT("MessageLog.Info"));
+				Info.bFireAndForget = true;
+				Info.bUseThrobber = true;
+				Info.FadeOutDuration = 8.0f;
+				Info.ExpireDuration = Info.FadeOutDuration;
+				Info.Hyperlink = FSimpleDelegate::CreateLambda([SoftObjectPath] {
+					// Select the cloud in Content Browser when the hyperlink is clicked
+					TArray<FAssetData> AssetData;
+					AssetData.Add(FModuleManager::LoadModuleChecked<FAssetRegistryModule>("AssetRegistry").Get().GetAssetByObjectPath(SoftObjectPath.GetWithoutSubPath()));
+					FModuleManager::LoadModuleChecked<FContentBrowserModule>("ContentBrowser").Get().SyncBrowserToAssets(AssetData);
+					});
+				Info.HyperlinkText = FText::FromString(FPaths::GetBaseFilename(SoftObjectPath.ToString()));
+										
+				TSharedPtr<SNotificationItem> NotificationPtr = FSlateNotificationManager::Get().AddNotification(Info);
+				NotificationPtr->SetCompletionState(SNotificationItem::CS_Fail);
+			}
+		}
+		return InControlRig->GetLastOverrideAsset();
+	}
+}
diff --git a/Engine/Plugins/Animation/ControlRig/Source/ControlRigEditor/Private/AnimDetails/Views/SAnimOverrideDetailsView.h b/Engine/Plugins/Animation/ControlRig/Source/ControlRigEditor/Private/AnimDetails/Views/SAnimOverrideDetailsView.h
new file mode 100644
--- /dev/null
+++ b/Engine/Plugins/Animation/ControlRig/Source/ControlRigEditor/Private/AnimDetails/Views/SAnimOverrideDetailsView.h
@@ -0,0 +1,78 @@
+// Copyright Epic Games, Inc. All Rights Reserved.
+
+#pragma once
+
+#include "Widgets/SCompoundWidget.h"
+#include "Overrides/OverrideStatusSubject.h"
+
+class IDetailsView;
+struct FPropertyAndParent;
+class UControlRig;
+class UControlRigOverrideAsset;
+class FOverrideStatusDetailsViewObjectFilter;
+
+namespace UE::ControlRigEditor
+{
+	class SAnimDetailsSearchBox;
+
+	class SAnimOverrideDetailsView
+		: public SCompoundWidget
+	{
+		SLATE_DECLARE_WIDGET(SAnimOverrideDetailsView, SCompoundWidget)
+		
+		SLATE_BEGIN_ARGS(SAnimOverrideDetailsView)
+			{}
+		SLATE_EVENT(FSimpleDelegate, OnRequestRefreshDetails)
+		
+		SLATE_END_ARGS()
+
+		/** Constructs this widget */
+		void Construct(const FArguments& InArgs);
+
+	private:
+		/** Refreshes the details view */
+		void RefreshDetailsView();
+
+		/** Returns ture if the property should be displayed */
+		bool ShouldDisplayProperty(const FPropertyAndParent& InPropertyAndParent) const;
+		
+		/** Returns ture if the property is ready-only */
+		bool IsReadOnlyProperty(const FPropertyAndParent& InPropertyAndParent) const;
+
+		/** Called when editing a property in the override details panel has finished */ 
+		void OnFinishedChangingOverride(const FPropertyChangedEvent& PropertyChangedEvent);
+
+		/** Returns true if two objects can be merged for display on the details panel */
+		bool CanMergeObjects(const UObject* InObjectA, const UObject* InObjectB) const;
+
+		/** Returns true if we can create the override widget for a given subject */
+		bool CanCreateWidget(const FOverrideStatusSubject& InSubject) const;
+
+		/** Returns the override status for a given subject */
+		EOverrideWidgetStatus::Type GetOverrideStatus(const FOverrideStatusSubject& InSubject) const;
+
+		/** React to the user interface request to add an override on a given subject */
+		FReply OnAddOverride(const FOverrideStatusSubject& InSubject);
+
+		/** React to the user interface request to clear all overrides on a given subject */
+		FReply OnClearOverride(const FOverrideStatusSubject& InSubject);
+
+		/** Updates a property path from what it is on the proxy to what it is on the control settings */
+		void MapPropertyFromProxyToControl(FString& InOutPropertyPath) const;
+
+		/** Returns the override asset for a given control rig to record changes to*/
+		static UControlRigOverrideAsset* GetOrCreateOverrideAsset(UControlRig* InControlRig);
+		
+		/** Weak ptr to the details view this widget is displaying */
+		TWeakPtr<IDetailsView> WeakDetailsView;
+
+		/** The object filter used to show the objects in the override details */
+		TSharedPtr<FOverrideStatusDetailsViewObjectFilter> ObjectFilter;
+		
+		/** A map from proxy property path to control settings property path */
+		TMap<FString,FString> ProxyPropertyToControl;
+
+		/** A delegate fired when we request to refresh the details */
+		FSimpleDelegate RequestRefreshDetailsDelegate;
+	};
+}
diff --git a/Engine/Plugins/Animation/ControlRig/Source/ControlRigEditor/Private/ControlRigEditorModule.cpp b/Engine/Plugins/Animation/ControlRig/Source/ControlRigEditor/Private/ControlRigEditorModule.cpp
--- a/Engine/Plugins/Animation/ControlRig/Source/ControlRigEditor/Private/ControlRigEditorModule.cpp
+++ b/Engine/Plugins/Animation/ControlRig/Source/ControlRigEditor/Private/ControlRigEditorModule.cpp
@@ -60,6 +60,7 @@
 #include "ControlRigDrawingDetails.h"
 #include "ControlRigAnimGraphDetails.h"
 #include "ControlRigBlueprintDetails.h"
+#include "ControlRigOverrideDetails.h"
 #include "Subsystems/AssetEditorSubsystem.h"
 #include "Sequencer/ControlRigParameterTrackEditor.h"
 #include "ActorFactories/ActorFactorySkeletalMesh.h"
@@ -157,22 +158,6 @@ void FControlRigEditorModule::StartupModule()
 	const FName PropertyEditorModuleName("PropertyEditor");
 	FPropertyEditorModule& PropertyModule = FModuleManager::GetModuleChecked<FPropertyEditorModule>(PropertyEditorModuleName);
 
-	// Register the custom class layout common to classes that inherit AnimDetailsProxyBase
-	for (TObjectIterator<UClass> ClassIterator; ClassIterator; ++ClassIterator)
-	{
-		if (ClassIterator->IsChildOf(UAnimDetailsProxyBase::StaticClass()))
-		{
-			const FName ClassName = ClassIterator->GetFName();
-			if (ClassName == UAnimDetailsProxyBase::StaticClass()->GetFName())
-			{
-				continue;
-			}
-
-			ClassesToUnregisterOnShutdown.Add(ClassName);
-			PropertyEditorModule.RegisterCustomClassLayout(ClassName, FOnGetDetailCustomizationInstance::CreateStatic(&FAnimDetailsProxyDetails::MakeInstance));
-		}
-	}
-
 	// same as ClassesToUnregisterOnShutdown but for properties, there is none right now
 	PropertiesToUnregisterOnShutdown.Reset();
 	PropertiesToUnregisterOnShutdown.Add(FRigVMCompileSettings::StaticStruct()->GetFName());
@@ -228,6 +213,9 @@ void FControlRigEditorModule::StartupModule()
 
 	FRigBaseElementDetails::RegisterSectionMappings(PropertyEditorModule);
 
+	PropertiesToUnregisterOnShutdown.Add(FControlRigOverrideContainer::StaticStruct()->GetFName());
+	PropertyEditorModule.RegisterCustomPropertyTypeLayout(PropertiesToUnregisterOnShutdown.Last(), FOnGetPropertyTypeCustomizationInstance::CreateStatic(&FControlRigOverrideDetails::MakeInstance));
+
 	// Register asset tools
 	auto RegisterAssetTypeAction = [this](const TSharedRef<IAssetTypeActions>& InAssetTypeAction)
 	{
diff --git a/Engine/Plugins/Animation/ControlRig/Source/ControlRigEditor/Private/ControlRigOverrideDetails.cpp b/Engine/Plugins/Animation/ControlRig/Source/ControlRigEditor/Private/ControlRigOverrideDetails.cpp
new file mode 100644
--- /dev/null
+++ b/Engine/Plugins/Animation/ControlRig/Source/ControlRigEditor/Private/ControlRigOverrideDetails.cpp
@@ -0,0 +1,437 @@
+// Copyright Epic Games, Inc. All Rights Reserved.
+
+#include "ControlRigOverrideDetails.h"
+#include "ControlRigOverride.h"
+#include "PropertyCustomizationHelpers.h"
+#include "IPropertyUtilities.h"
+#include "IDetailChildrenBuilder.h"
+#include "DetailLayoutBuilder.h"
+#include "IDetailGroup.h"
+#include "Widgets/Text/STextBlock.h"
+#include "Widgets/Input/SNumericEntryBox.h"
+#include "Widgets/Input/SEditableText.h"
+#include "Widgets/Input/SVectorInputBox.h"
+#include "Widgets/Input/SRotatorInputBox.h"
+#include "Widgets/Colors/SColorBlock.h"
+
+#define LOCTEXT_NAMESPACE "ControlRigOverrideDetails"
+
+void FControlRigOverrideDetails::CustomizeHeader(TSharedRef<IPropertyHandle> InStructPropertyHandle, FDetailWidgetRow& HeaderRow, IPropertyTypeCustomizationUtils& StructCustomizationUtils)
+{
+	// nothing to do here
+}
+
+void FControlRigOverrideDetails::CustomizeChildren(TSharedRef<IPropertyHandle> InStructPropertyHandle, IDetailChildrenBuilder& StructBuilder, IPropertyTypeCustomizationUtils& StructCustomizationUtils)
+{
+	const TArray<TWeakObjectPtr<UObject>>& SelectedObjects = StructCustomizationUtils.GetPropertyUtilities()->GetSelectedObjects();
+	if(SelectedObjects.Num() != 1)
+	{
+		return;
+	}
+	if(!SelectedObjects[0].IsValid() || SelectedObjects[0].Get() == nullptr)
+	{
+		return;
+	}
+
+	UControlRigOverrideAsset* OverrideAsset = Cast<UControlRigOverrideAsset>(SelectedObjects[0].Get());
+	if(OverrideAsset == nullptr)
+	{
+		return;
+	}
+
+	struct FSortNamesAlphabetically
+	{
+		bool operator()(const FName& A, const FName& B) const
+		{
+			return A.Compare(B) < 0;
+		}
+	};
+
+	TArray<FName> SubjectKeys = OverrideAsset->Overrides.GenerateSubjectArray();
+	SubjectKeys.Sort(FSortNamesAlphabetically());
+
+	for(const FName& SubjectKey : SubjectKeys)
+	{
+		const TArray<int32>* IndicesPtr = OverrideAsset->Overrides.GetIndicesForSubject(SubjectKey);
+		if(IndicesPtr == nullptr)
+		{
+			continue;
+		}
+
+		const TArray<int32>& Indices = *IndicesPtr;
+		if(Indices.IsEmpty())
+		{
+			continue;
+		}
+
+		StructBuilder.AddCustomBuilder(MakeShared<FControlRigOverrideDetailsBuilder>(OverrideAsset, SubjectKey));
+	}
+}
+
+FControlRigOverrideDetailsBuilder::FControlRigOverrideDetailsBuilder(UControlRigOverrideAsset* InOverrideAsset, const FName& InSubjectKey)
+	: WeakOverrideAsset(InOverrideAsset)
+	, SubjectKey(InSubjectKey)
+{
+}
+
+void FControlRigOverrideDetailsBuilder::Tick(float DeltaTime)
+{
+	if(!WeakOverrideAsset.IsValid())
+	{
+		return;
+	}
+	const UControlRigOverrideAsset* OverrideAsset = WeakOverrideAsset.Get();
+	if(OverrideAsset == nullptr)
+	{
+		return;
+	}
+	const uint32 Hash = GetTypeHash(OverrideAsset->Overrides);
+	if(LastHash.Get(UINT32_MAX) != Hash)
+	{
+		(void)OnRebuildChildren.ExecuteIfBound();
+	}
+}
+
+void FControlRigOverrideDetailsBuilder::GenerateHeaderRowContent(FDetailWidgetRow& NodeRow)
+{
+	NodeRow.NameContent()
+	[
+		SNew(STextBlock)
+		.Text(FText::FromName(SubjectKey))
+		.Font(IDetailLayoutBuilder::GetDetailFont())
+	];
+}
+
+void FControlRigOverrideDetailsBuilder::GenerateChildContent(IDetailChildrenBuilder& ChildrenBuilder)
+{
+	if(!WeakOverrideAsset.IsValid())
+	{
+		return;
+	}
+
+	UControlRigOverrideAsset* OverrideAsset = WeakOverrideAsset.Get();
+	if(OverrideAsset == nullptr)
+	{
+		return;
+	}
+
+	const TArray<int32>* IndicesPtr = OverrideAsset->Overrides.GetIndicesForSubject(SubjectKey);
+	if(IndicesPtr == nullptr)
+	{
+		return;
+	}
+
+	const TArray<int32>& Indices = *IndicesPtr;
+	const FString SubjectKeyPrefix = SubjectKey.ToString() + TEXT("|");
+	
+	for(int32 Index : Indices)
+	{
+		if(!OverrideAsset->Overrides.IsValidIndex(Index))
+		{
+			continue;
+		}
+		FControlRigOverrideValue& Override = OverrideAsset->Overrides[Index];
+		if(Override.GetSubjectKey() != SubjectKey)
+		{
+			continue;
+		}
+
+		const FProperty* LeafProperty = Override.GetLeafProperty();
+		if(LeafProperty == nullptr)
+		{
+			continue;
+		}
+		const FStructProperty* StructProperty = CastField<FStructProperty>(LeafProperty);
+		
+		FDetailWidgetRow& Row = ChildrenBuilder.AddCustomRow(FText::FromString(SubjectKeyPrefix + Override.GetPath()));
+
+		// add the default name content for anything but a transform property
+		if(StructProperty == nullptr || StructProperty->Struct != TBaseStructure<FTransform>::Get())
+		{
+			Row.NameContent()
+			[
+				SNew(STextBlock)
+				.Text(FText::FromString(Override.GetPath()))
+				.Font(IDetailLayoutBuilder::GetDetailFont())
+			];
+		}
+
+		// initially overrides will only be offered one a limited set of types
+		//
+		if(CastField<FBoolProperty>(LeafProperty))
+		{
+			TSharedPtr<TControlRigOverrideHandle<bool>> Handle = MakeShared<TControlRigOverrideHandle<bool>>(OverrideAsset, Index);
+			Row.ValueContent()
+			[
+				SNew(SCheckBox)
+				.IsEnabled(false)
+				.IsChecked_Lambda([Handle]() -> ECheckBoxState
+				{
+					if(const bool* Data = Handle->GetData())
+					{
+						return (*Data) ? ECheckBoxState::Checked : ECheckBoxState::Unchecked;
+					}
+					return ECheckBoxState::Undetermined;
+				})
+			];
+		}
+		else if(CastField<FFloatProperty>(LeafProperty))
+		{
+			TSharedPtr<TControlRigOverrideHandle<float>> Handle = MakeShared<TControlRigOverrideHandle<float>>(OverrideAsset, Index);
+			Row.ValueContent()
+			[
+				SNew(SNumericEntryBox<float>)
+				.IsEnabled(false)
+				.Value_Lambda([Handle]() -> TOptional<float>
+				{
+					if(const float* Data = Handle->GetData())
+					{
+						return *Data;
+					}
+					return TOptional<float>();
+				})
+			];
+		}
+		else if(CastField<FDoubleProperty>(LeafProperty))
+		{
+			TSharedPtr<TControlRigOverrideHandle<double>> Handle = MakeShared<TControlRigOverrideHandle<double>>(OverrideAsset, Index);
+			Row.ValueContent()
+			[
+				SNew(SNumericEntryBox<double>)
+				.IsEnabled(false)
+				.Value_Lambda([Handle]() -> TOptional<double>
+				{
+					if(const double* Data = Handle->GetData())
+					{
+						return *Data;
+					}
+					return TOptional<double>();
+				})
+			];
+		}
+		else if(CastField<FNameProperty>(LeafProperty))
+		{
+			TSharedPtr<TControlRigOverrideHandle<FName>> Handle = MakeShared<TControlRigOverrideHandle<FName>>(OverrideAsset, Index);
+			Row.ValueContent()
+			[
+				SNew(SEditableText)
+				.IsEnabled(false)
+				.Text_Lambda([Handle]() -> FText
+				{
+					if(const FName* Data = Handle->GetData())
+					{
+						return FText::FromName(*Data);
+					}
+					return FText();
+				})
+			];
+		}
+		else if(CastField<FStrProperty>(LeafProperty))
+		{
+			TSharedPtr<TControlRigOverrideHandle<FString>> Handle = MakeShared<TControlRigOverrideHandle<FString>>(OverrideAsset, Index);
+			Row.ValueContent()
+			[
+				SNew(SEditableText)
+				.IsEnabled(false)
+				.Text_Lambda([Handle]() -> FText
+				{
+					if(const FString* Data = Handle->GetData())
+					{
+						return FText::FromString(*Data);
+					}
+					return FText();
+				})
+			];
+		}
+		else if(StructProperty)
+		{
+			const UStruct* Struct = StructProperty->Struct;
+
+			if(Struct == TBaseStructure<FVector>::Get())
+			{
+				TSharedPtr<TControlRigOverrideHandle<FVector>> Handle = MakeShared<TControlRigOverrideHandle<FVector>>(OverrideAsset, Index);
+				Row.ValueContent()
+				[
+					SNew(SNumericVectorInputBox<double>)
+					.IsEnabled(false)
+					.Vector_Lambda([Handle]() -> TOptional<FVector>
+					{
+						if(const FVector* Data = Handle->GetData())
+						{
+							return *Data;
+						}
+						return {};
+					})
+				];
+			}
+			else if(Struct == TBaseStructure<FRotator>::Get())
+			{
+				TSharedPtr<TControlRigOverrideHandle<FRotator>> Handle = MakeShared<TControlRigOverrideHandle<FRotator>>(OverrideAsset, Index);
+				Row.ValueContent()
+				[
+					SNew(SNumericRotatorInputBox<double>)
+					.IsEnabled(false)
+					.Pitch_Lambda([Handle]() -> TOptional<double>
+					{
+						if(const FRotator* Data = Handle->GetData())
+						{
+							return Data->Pitch;
+						}
+						return {};
+					})
+					.Yaw_Lambda([Handle]() -> TOptional<double>
+					{
+						if(const FRotator* Data = Handle->GetData())
+						{
+							return Data->Yaw;
+						}
+						return {};
+					})
+					.Roll_Lambda([Handle]() -> TOptional<double>
+					{
+						if(const FRotator* Data = Handle->GetData())
+						{
+							return Data->Roll;
+						}
+						return {};
+					})
+				];
+			}
+			else if(Struct == TBaseStructure<FQuat>::Get())
+			{
+				TSharedPtr<TControlRigOverrideHandle<FQuat>> Handle = MakeShared<TControlRigOverrideHandle<FQuat>>(OverrideAsset, Index);
+				Row.ValueContent()
+				[
+					SNew(SNumericRotatorInputBox<double>)
+					.IsEnabled(false)
+					.Pitch_Lambda([Handle]() -> TOptional<double>
+					{
+						if(const FQuat* Data = Handle->GetData())
+						{
+							return Data->Rotator().Pitch;
+						}
+						return {};
+					})
+					.Yaw_Lambda([Handle]() -> TOptional<double>
+					{
+						if(const FQuat* Data = Handle->GetData())
+						{
+							return Data->Rotator().Yaw;
+						}
+						return {};
+					})
+					.Roll_Lambda([Handle]() -> TOptional<double>
+					{
+						if(const FQuat* Data = Handle->GetData())
+						{
+							return Data->Rotator().Roll;
+						}
+						return {};
+					})
+				];
+			}
+			else if(Struct == TBaseStructure<FTransform>::Get())
+			{
+				TSharedPtr<TControlRigOverrideHandle<FTransform>> Handle = MakeShared<TControlRigOverrideHandle<FTransform>>(OverrideAsset, Index);
+				Row.NameContent()
+				[
+					SNew(STextBlock)
+					.Text(FText::FromString(Override.GetPath()+TEXT("->Location")))
+					.Font(IDetailLayoutBuilder::GetDetailFont())
+				]
+				.ValueContent()
+				[
+					SNew(SNumericVectorInputBox<double>)
+					.IsEnabled(false)
+					.Vector_Lambda([Handle]() -> TOptional<FVector>
+					{
+					if(const FTransform* Data = Handle->GetData())
+					{
+						return Data->GetLocation();
+					}
+					return {};
+					})
+				];
+
+				ChildrenBuilder.AddCustomRow(FText::FromString(SubjectKeyPrefix + Override.GetPath() + TEXT("->Rotation")))
+				.NameContent()
+				[
+					SNew(STextBlock)
+					.Text(FText::FromString(Override.GetPath()+TEXT("->Rotation")))
+					.Font(IDetailLayoutBuilder::GetDetailFont())
+				]
+				.ValueContent()
+				[
+					SNew(SNumericRotatorInputBox<double>)
+					.IsEnabled(false)
+					.Pitch_Lambda([Handle]() -> TOptional<double>
+					{
+						if(const FTransform* Data = Handle->GetData())
+						{
+							return Data->Rotator().Pitch;
+						}
+						return {};
+					})
+					.Yaw_Lambda([Handle]() -> TOptional<double>
+					{
+						if(const FTransform* Data = Handle->GetData())
+						{
+							return Data->Rotator().Yaw;
+						}
+						return {};
+					})
+					.Roll_Lambda([Handle]() -> TOptional<double>
+					{
+						if(const FTransform* Data = Handle->GetData())
+						{
+							return Data->Rotator().Roll;
+						}
+						return {};
+					})
+				];
+
+				ChildrenBuilder.AddCustomRow(FText::FromString(SubjectKeyPrefix + Override.GetPath() + TEXT("->Rotation")))
+				.NameContent()
+				[
+					SNew(STextBlock)
+					.Text(FText::FromString(Override.GetPath()+TEXT("->Scale3D")))
+					.Font(IDetailLayoutBuilder::GetDetailFont())
+				]
+				.ValueContent()
+				[
+					SNew(SNumericVectorInputBox<double>)
+					.IsEnabled(false)
+					.Vector_Lambda([Handle]() -> TOptional<FVector>
+					{
+					if(const FTransform* Data = Handle->GetData())
+					{
+						return Data->GetScale3D();
+					}
+					return {};
+					})
+				];
+			}
+			else if(Struct == TBaseStructure<FLinearColor>::Get())
+			{
+				TSharedPtr<TControlRigOverrideHandle<FLinearColor>> Handle = MakeShared<TControlRigOverrideHandle<FLinearColor>>(OverrideAsset, Index);
+				Row.ValueContent()
+				[
+					SNew(SColorBlock)
+					.IsEnabled(false)
+					.Color_Lambda([Handle]() -> FLinearColor
+					{
+						if(const FLinearColor* Data = Handle->GetData())
+						{
+							return *Data;
+						}
+						return FLinearColor::Black;
+					})
+				];
+			}
+		}
+	}
+
+	LastHash = GetTypeHash(OverrideAsset->Overrides);
+}
+
+#undef LOCTEXT_NAMESPACE
diff --git a/Engine/Plugins/Animation/ControlRig/Source/ControlRigEditor/Private/ControlRigOverrideDetails.h b/Engine/Plugins/Animation/ControlRig/Source/ControlRigEditor/Private/ControlRigOverrideDetails.h
new file mode 100644
--- /dev/null
+++ b/Engine/Plugins/Animation/ControlRig/Source/ControlRigEditor/Private/ControlRigOverrideDetails.h
@@ -0,0 +1,49 @@
+// Copyright Epic Games, Inc. All Rights Reserved.
+
+#pragma once
+
+#include "CoreMinimal.h"
+#include "ControlRigOverride.h"
+#include "IPropertyTypeCustomization.h"
+#include "IDetailCustomNodeBuilder.h"
+#include "IStructureDataProvider.h"
+
+class IPropertyHandle;
+class UControlRigOverrideAsset;
+
+class FControlRigOverrideDetails : public IPropertyTypeCustomization
+{
+public:
+
+	static TSharedRef<IPropertyTypeCustomization> MakeInstance()
+	{
+		return MakeShareable(new FControlRigOverrideDetails);
+	}
+
+	/** IPropertyTypeCustomization interface */
+	virtual void CustomizeHeader(TSharedRef<class IPropertyHandle> InStructPropertyHandle, class FDetailWidgetRow& HeaderRow, IPropertyTypeCustomizationUtils& StructCustomizationUtils) override;
+	virtual void CustomizeChildren(TSharedRef<class IPropertyHandle> InStructPropertyHandle, class IDetailChildrenBuilder& StructBuilder, IPropertyTypeCustomizationUtils& StructCustomizationUtils) override;
+};
+
+class FControlRigOverrideDetailsBuilder
+	: public IDetailCustomNodeBuilder
+	, public TSharedFromThis<FControlRigOverrideDetailsBuilder>
+{
+public:
+	FControlRigOverrideDetailsBuilder(UControlRigOverrideAsset* InOverrideAsset, const FName& InSubjectKey);
+	
+	/** IDetailCustomNodeBuilder interface */
+	virtual void SetOnRebuildChildren( FSimpleDelegate InOnRebuildChildren  ) override { OnRebuildChildren = InOnRebuildChildren; } 
+	virtual bool RequiresTick() const override { return true; }
+	virtual void Tick( float DeltaTime ) override;
+	virtual void GenerateHeaderRowContent( FDetailWidgetRow& NodeRow ) override;
+	virtual void GenerateChildContent( IDetailChildrenBuilder& ChildrenBuilder ) override;
+	virtual FName GetName() const override { return SubjectKey; }
+	virtual bool InitiallyCollapsed() const override { return false; }
+
+private:
+	TWeakObjectPtr<UControlRigOverrideAsset> WeakOverrideAsset;
+	FName SubjectKey;
+	FSimpleDelegate OnRebuildChildren;
+	TOptional<uint32> LastHash;
+};
\ No newline at end of file
diff --git a/Engine/Plugins/Animation/ControlRig/Source/ControlRigEditor/Private/EditMode/ControlRigEditMode.cpp b/Engine/Plugins/Animation/ControlRig/Source/ControlRigEditor/Private/EditMode/ControlRigEditMode.cpp
--- a/Engine/Plugins/Animation/ControlRig/Source/ControlRigEditor/Private/EditMode/ControlRigEditMode.cpp
+++ b/Engine/Plugins/Animation/ControlRig/Source/ControlRigEditor/Private/EditMode/ControlRigEditMode.cpp
@@ -4487,6 +4487,10 @@ void FControlRigEditMode::RecreateControlShapeActors()
 	TArray<UControlRig*> ControlRigsCopy = ControlRigsToRecreate;
 	for (UControlRig* ControlRig : ControlRigsCopy)
 	{
+		if(!IsValid(ControlRig))
+		{
+			continue;
+		}
 		const bool bUpdated = TryUpdatingControlsShapes(ControlRig);
 		if (!bUpdated)
 		{
@@ -4495,7 +4499,14 @@ void FControlRigEditMode::RecreateControlShapeActors()
 		}
 	}
 	RecreateControlShapesRequired = ERecreateControlRigShape::RecreateNone;
-	ControlRigsToRecreate.SetNum(0);
+	ControlRigsToRecreate.Empty();
+
+	// todo:
+	/*if(ControlProxy)
+	{
+		ControlProxy->SyncAllProxies();
+	}
+	//*/
 }
 
 void FControlRigEditMode::CreateShapeActors(UControlRig* InControlRig)
diff --git a/Engine/Plugins/Animation/ControlRig/Source/ControlRigEditor/Private/EditMode/SControlRigEditModeTools.cpp b/Engine/Plugins/Animation/ControlRig/Source/ControlRigEditor/Private/EditMode/SControlRigEditModeTools.cpp
--- a/Engine/Plugins/Animation/ControlRig/Source/ControlRigEditor/Private/EditMode/SControlRigEditModeTools.cpp
+++ b/Engine/Plugins/Animation/ControlRig/Source/ControlRigEditor/Private/EditMode/SControlRigEditModeTools.cpp
@@ -110,6 +110,9 @@ void SControlRigEditModeTools::SetControlRigs(const TArrayView<TWeakObjectPtr<UC
 	}
 	RigOptionsDetailsView->SetObjects(Objects);
 
+	UpdateOverridesDetailsView();
+
+
 #if USE_LOCAL_DETAILS
 	HierarchyTreeView->RefreshTreeView(true);
 #endif
@@ -211,6 +214,10 @@ void SControlRigEditModeTools::Construct(const FArguments& InArgs, TSharedPtr<FC
 	RigOptionsDetailsView->SetKeyframeHandler(SharedThis(this));
 	RigOptionsDetailsView->OnFinishedChangingProperties().AddSP(this, &SControlRigEditModeTools::OnRigOptionFinishedChange);
 
+	OverridesDetailsView = FModuleManager::GetModuleChecked<FPropertyEditorModule>("PropertyEditor").CreateDetailView(DetailsViewArgs);
+	OverridesDetailsView->OnFinishedChangingProperties().AddSP(this, &SControlRigEditModeTools::OnOverrideOptionFinishedChange);
+	OverridesDetailsView->SetIsPropertyVisibleDelegate(FIsPropertyVisible::CreateSP(this, &SControlRigEditModeTools::ShouldShowOverrideProperty));
+
 	DisplaySettings.bShowBones = false;
 	DisplaySettings.bShowControls = true;
 	DisplaySettings.bShowNulls = false;
@@ -507,6 +514,37 @@ void SControlRigEditModeTools::Construct(const FArguments& InArgs, TSharedPtr<FC
 					RigOptionsDetailsView.ToSharedRef()
 				]
 			]
+
+			+SVerticalBox::Slot()
+			.AutoHeight()
+			[
+				SNew(SExpandableArea)
+				.InitiallyCollapsed(true)
+				.Visibility(this, &SControlRigEditModeTools::GetOverridesExpanderVisibility)
+				.AreaTitle(LOCTEXT("OverridesWidget", "Overrides"))
+				.AreaTitleFont(FAppStyle::GetFontStyle("DetailsView.CategoryFontStyle"))
+				.BorderBackgroundColor(FLinearColor(.6f, .6f, .6f))
+				.HeaderContent()
+				[
+					SNew(SHorizontalBox)
+
+					// "Overrides" label
+					+SHorizontalBox::Slot()
+					.AutoWidth()
+					.HAlign(HAlign_Left)
+					.VAlign(VAlign_Center)
+					.Padding(0.f, 0.f, 0.f, 0.f)
+					[
+						SNew(STextBlock)
+						.Text(LOCTEXT("OverridesWidget", "Overrides"))
+						.Font(FCoreStyle::Get().GetFontStyle("ExpandableArea.TitleFont"))
+					]
+				]
+				.BodyContent()
+				[
+					OverridesDetailsView.ToSharedRef()
+				]
+			]
 		]
 	];
 #if USE_LOCAL_DETAILS
@@ -654,6 +692,7 @@ void SControlRigEditModeTools::SetSequencer(TWeakPtr<ISequencer> InSequencer)
 	{
 		ConstraintsEditionWidget->SequencerChanged(InSequencer);
 	}
+	UpdateOverridesDetailsView();
 }
 
 bool SControlRigEditModeTools::IsPropertyKeyable(const UClass* InObjectClass, const IPropertyHandle& InPropertyHandle) const
@@ -1227,7 +1266,7 @@ EVisibility SControlRigEditModeTools::GetRigOptionExpanderVisibility() const
 			}
 		}
 	}
-	return EVisibility::Hidden;
+	return EVisibility::Collapsed;
 }
 
 void SControlRigEditModeTools::OnRigOptionFinishedChange(const FPropertyChangedEvent& PropertyChangedEvent)
@@ -1433,4 +1472,72 @@ FText SControlRigEditModeTools::GetActiveToolMessage() const
 	return  FText();
 }
 
+EVisibility SControlRigEditModeTools::GetOverridesExpanderVisibility() const
+{
+	if(CVarControlRigEnableOverrides.GetValueOnAnyThread())
+	{
+		if(!OverridesDetailsView->GetSelectedObjects().IsEmpty())
+		{
+			return EVisibility::Visible;
+		}
+	}
+	return EVisibility::Collapsed;
+}
+
+void SControlRigEditModeTools::OnOverrideOptionFinishedChange(const FPropertyChangedEvent& PropertyChangedEvent)
+{
+	TArray<TWeakObjectPtr<UObject>> SelectedObjects = OverridesDetailsView->GetSelectedObjects();
+	for(const TWeakObjectPtr<UObject>& SelectedObject : SelectedObjects)
+	{
+		if(UMovieSceneControlRigParameterSection* CRSection = Cast<UMovieSceneControlRigParameterSection>(SelectedObject.Get()))
+		{
+			if(UControlRig* ControlRig = CRSection->GetControlRig())
+			{
+				TGuardValue<bool> _(CRSection->bSuspendOverrideAssetSync, true);
+				ControlRig->UnlinkAllOverrideAssets();
+				for(const TSoftObjectPtr<UControlRigOverrideAsset>& OverrideAssetPtr : CRSection->OverrideAssets)
+				{
+					if(UControlRigOverrideAsset* OverrideAsset = OverrideAssetPtr.LoadSynchronous())
+					{
+						ControlRig->LinkOverrideAsset(OverrideAsset);
+					}
+				}
+				CRSection->UpdateOverrideAssetDelegates();
+				CRSection->ReconstructChannelProxy();
+			}
+		}
+	}
+}
+
+bool SControlRigEditModeTools::ShouldShowOverrideProperty(const FPropertyAndParent& InPropertyAndParent) const
+{
+	static const FProperty* OverrideAssetsProperty = UMovieSceneControlRigParameterSection::StaticClass()->FindPropertyByName(
+		GET_MEMBER_NAME_CHECKED(UMovieSceneControlRigParameterSection, OverrideAssets));
+	return OverrideAssetsProperty == &InPropertyAndParent.Property ||
+		InPropertyAndParent.ParentProperties.Contains(OverrideAssetsProperty);
+}
+
+void SControlRigEditModeTools::UpdateOverridesDetailsView()
+{
+	TArray<TWeakObjectPtr<>> Sections;
+	if(WeakSequencer.IsValid())
+	{
+		ISequencer* Sequencer = WeakSequencer.Pin().Get();
+		for(const TWeakObjectPtr<UControlRig>& WeakControlRig : ControlRigs)
+		{
+			if(const UControlRig* ControlRig = WeakControlRig.Get())
+			{
+				if(!ControlRig->CurrentControlSelection().IsEmpty())
+				{
+					if(UMovieSceneControlRigParameterSection* CRSection = FControlRigSpaceChannelHelpers::GetControlRigSection(Sequencer, ControlRig))
+					{
+						Sections.Emplace(CRSection);
+					}
+				}
+			}
+		}
+	}
+	OverridesDetailsView->SetObjects(Sections);
+}
+
 #undef LOCTEXT_NAMESPACE
diff --git a/Engine/Plugins/Animation/ControlRig/Source/ControlRigEditor/Private/EditMode/SControlRigEditModeTools.h b/Engine/Plugins/Animation/ControlRig/Source/ControlRigEditor/Private/EditMode/SControlRigEditModeTools.h
--- a/Engine/Plugins/Animation/ControlRig/Source/ControlRigEditor/Private/EditMode/SControlRigEditModeTools.h
+++ b/Engine/Plugins/Animation/ControlRig/Source/ControlRigEditor/Private/EditMode/SControlRigEditModeTools.h
@@ -103,6 +103,8 @@ private:
 	TSharedPtr<SExpandableArea> ConstraintPickerExpander = nullptr;
 	TSharedPtr<SConstraintsEditionWidget> ConstraintsEditionWidget = nullptr; 
 
+	TSharedPtr<IDetailsView> OverridesDetailsView;
+
 	/** Display or edit set up for property */
 	static bool ShouldShowPropertyOnDetailCustomization(const struct FPropertyAndParent& InPropertyAndParent);
 	static bool IsReadOnlyPropertyOnDetailCustomization(const struct FPropertyAndParent& InPropertyAndParent);
@@ -131,6 +133,11 @@ private:
 	FText GetShowConstraintsName() const;
 	FText GetShowConstraintsTooltip() const;
 
+	EVisibility GetOverridesExpanderVisibility() const;
+	void OnOverrideOptionFinishedChange(const FPropertyChangedEvent& PropertyChangedEvent);
+	bool ShouldShowOverrideProperty(const FPropertyAndParent& InPropertyAndParent) const;
+	void UpdateOverridesDetailsView();
+	
 private:
 	/** Toolbar functions and windows*/
 	void ToggleEditPivotMode();
@@ -156,5 +163,3 @@ public:
 	FText GetActiveToolMessage() const;
 };
 
-
-
diff --git a/Engine/Plugins/Animation/ControlRig/Source/ControlRigEditor/Private/EditMode/SControlRigOutliner.cpp b/Engine/Plugins/Animation/ControlRig/Source/ControlRigEditor/Private/EditMode/SControlRigOutliner.cpp
--- a/Engine/Plugins/Animation/ControlRig/Source/ControlRigEditor/Private/EditMode/SControlRigOutliner.cpp
+++ b/Engine/Plugins/Animation/ControlRig/Source/ControlRigEditor/Private/EditMode/SControlRigOutliner.cpp
@@ -41,6 +41,7 @@
 #include "SEnumCombo.h"
 #include "Components/SkeletalMeshComponent.h"
 #include "Overrides/SOverrideStatusWidget.h"
+#include "Async/TaskGraphInterfaces.h"
 
 #define LOCTEXT_NAMESPACE "ControlRigOutliner"
 
@@ -136,6 +137,11 @@ FText FMultiRigData::GetDisplayName(const FRigTreeDisplaySettings& InSettings) c
 	return *CachedDisplayName;
 }
 
+void FMultiRigData::InvalidateDisplayName() const
+{
+	CachedDisplayName.Reset();
+}
+
 FText FMultiRigData::GetToolTipText(const FRigTreeDisplaySettings& Settings) const
 {
 	if (IsControlElement())
@@ -266,6 +272,7 @@ void FMultiRigTreeElement::RefreshDisplaySettings(const URigHierarchy* InHierarc
 		TextColor = FSlateColor::UseForeground();
 	}
 
+	Data.InvalidateDisplayName();
 }
 
 bool FMultiRigTreeElement::AreControlsVisible() const
@@ -1525,7 +1532,7 @@ void SMultiRigHierarchyTreeView::SetControlRigs(const TArrayView<TWeakObjectPtr<
 	{
 		if (ControlRig.IsValid())
 		{
-			ControlRigs.Add(ControlRig.Get());
+			ControlRigs.AddUnique(ControlRig);
 		}
 	}
 	
@@ -1920,11 +1927,16 @@ SControlRigOutliner::~SControlRigOutliner()
 		if (ControlRig.IsValid())
 		{
 			ControlRig.Get()->ControlRigBound().RemoveAll(this);
+			ControlRig.Get()->OnPostConstruction_AnyThread().RemoveAll(this);
 			const TSharedPtr<IControlRigObjectBinding> Binding = ControlRig.Get()->GetObjectBinding();
 			if (Binding)
 			{
 				Binding->OnControlRigBind().RemoveAll(this);
 			}
+			if(URigHierarchy* Hierarchy = ControlRig.Get()->GetHierarchy())
+			{
+				Hierarchy->OnModified().RemoveAll(this);
+			}
 		}
 	}
 	BoundControlRigs.SetNum(0);
@@ -1998,6 +2010,39 @@ void SControlRigOutliner::HandleRigVisibilityChanged(TArray<UControlRig*> InCont
 	}
 }
 
+void SControlRigOutliner::HandleHierarchyModified(ERigHierarchyNotification InNotification, URigHierarchy* InHierarchy,
+	const FRigNotificationSubject& InSubject)
+{
+	if(InHierarchy == nullptr)
+	{
+		return;
+	}
+	
+	UControlRig* ControlRig = InHierarchy->GetTypedOuter<UControlRig>();
+	if(ControlRig == nullptr)
+	{
+		return;
+	}
+	
+	if(InNotification == ERigHierarchyNotification::ControlSettingChanged)
+	{
+		if(const FRigControlElement* ControlElement = Cast<FRigControlElement>(InSubject.Element))
+		{
+			const FRigElementKey Key = ControlElement->GetKey();
+			const FMultiRigData Data(ControlRig, Key);
+
+			for (int32 RootIndex = 0; RootIndex < HierarchyTreeView->GetTreeView()->GetRootElements().Num(); ++RootIndex)
+			{
+				TSharedPtr<FMultiRigTreeElement> Found = HierarchyTreeView->GetTreeView()->FindElement(Data, HierarchyTreeView->GetTreeView()->GetRootElements()[RootIndex]);
+				if (Found.IsValid())
+				{
+					Found->RefreshDisplaySettings(InHierarchy, GetDisplaySettings());
+				}
+			}
+		}
+	}
+}
+
 void SControlRigOutliner::HandleSelectionChanged(TSharedPtr<FMultiRigTreeElement> Selection, ESelectInfo::Type SelectInfo)
 {
 	if (bIsChangingRigHierarchy)
@@ -2177,11 +2222,19 @@ void SControlRigOutliner::SetEditMode(FControlRigEditMode& InEditMode)
 					ControlRig.Get()->ControlRigBound().AddRaw(this, &SControlRigOutliner::HandleOnControlRigBound);
 					BoundControlRigs.Add(ControlRig);
 				}
+				if(!ControlRig.Get()->OnPostConstruction_AnyThread().IsBoundToObject(this))
+				{
+					ControlRig.Get()->OnPostConstruction_AnyThread().AddRaw(this, &SControlRigOutliner::HandlePostConstruction);
+				}
 				const TSharedPtr<IControlRigObjectBinding> Binding = ControlRig.Get()->GetObjectBinding();
 				if (Binding && !Binding->OnControlRigBind().IsBoundToObject(this))
 				{
 					Binding->OnControlRigBind().AddRaw(this, &SControlRigOutliner::HandleOnObjectBoundToControlRig);
 				}
+				if(URigHierarchy* Hierarchy = ControlRig.Get()->GetHierarchy())
+				{
+					Hierarchy->OnModified().AddRaw(this, &SControlRigOutliner::HandleHierarchyModified);
+				}
 			}
 		}
 		HierarchyTreeView->GetTreeView()->SetControlRigs(ControlRigs); //will refresh tree
@@ -2205,6 +2258,10 @@ void SControlRigOutliner::HandleControlAdded(UControlRig* ControlRig, bool bIsAd
 			{
 				Binding->OnControlRigBind().AddRaw(this, &SControlRigOutliner::HandleOnObjectBoundToControlRig);
 			}
+			if(!ControlRig->OnPostConstruction_AnyThread().IsBoundToObject(this))
+			{
+				ControlRig->OnPostConstruction_AnyThread().AddRaw(this, &SControlRigOutliner::HandlePostConstruction);
+			}
 		}
 		else
 		{
@@ -2215,6 +2272,11 @@ void SControlRigOutliner::HandleControlAdded(UControlRig* ControlRig, bool bIsAd
 			{
 				Binding->OnControlRigBind().RemoveAll(this);
 			}
+			if(URigHierarchy* Hierarchy = ControlRig->GetHierarchy())
+			{
+				Hierarchy->OnModified().RemoveAll(this);
+			}
+			ControlRig->OnPostConstruction_AnyThread().RemoveAll(this);
 		}
 	}
 	if (FControlRigEditMode* EditMode = static_cast<FControlRigEditMode*>(ModeTools->GetActiveMode(FControlRigEditMode::ModeName)))
@@ -2245,18 +2307,30 @@ void SControlRigOutliner::HandleOnObjectBoundToControlRig(UObject* InObject)
 	//just refresh the views, but do so on next tick since with FK control rig's the controls aren't set up
 	//until AFTER we are bound.
 	TWeakPtr<SControlRigOutliner> WeakPtr = StaticCastSharedRef<SControlRigOutliner>(AsShared()).ToWeakPtr();
-	GEditor->GetTimerManager()->SetTimerForNextTick([WeakPtr]()
+	FFunctionGraphTask::CreateAndDispatchWhenReady([WeakPtr]()
 	{
 		if (WeakPtr.IsValid())
 		{
 			TSharedPtr<SControlRigOutliner> StrongThis = WeakPtr.Pin();
-			if (FControlRigEditMode* EditMode = static_cast<FControlRigEditMode*>(StrongThis->ModeTools->GetActiveMode(FControlRigEditMode::ModeName)))
+			if(StrongThis)
 			{
-				TArrayView<TWeakObjectPtr<UControlRig>> ControlRigs = EditMode->GetControlRigs();
-				StrongThis->HierarchyTreeView->GetTreeView()->SetControlRigs(ControlRigs); //will refresh tree
+				if (FControlRigEditMode* EditMode = static_cast<FControlRigEditMode*>(StrongThis->ModeTools->GetActiveMode(FControlRigEditMode::ModeName)))
+				{
+					TArrayView<TWeakObjectPtr<UControlRig>> ControlRigs = EditMode->GetControlRigs();
+					if(StrongThis->HierarchyTreeView)
+					{
+						StrongThis->HierarchyTreeView->GetTreeView()->SetControlRigs(ControlRigs); //will refresh tree
+					}
+				}
 			}
 		}
-	});
+	}, TStatId(), nullptr, ENamedThreads::GameThread);
+}
+
+void SControlRigOutliner::HandlePostConstruction(UControlRig* InControlRig, const FName& InEventName)
+{
+	// rely on the code above to refresh the views
+	HandleOnObjectBoundToControlRig(nullptr);
 }
 
 #undef LOCTEXT_NAMESPACE
diff --git a/Engine/Plugins/Animation/ControlRig/Source/ControlRigEditor/Private/EditMode/SControlRigOutliner.h b/Engine/Plugins/Animation/ControlRig/Source/ControlRigEditor/Private/EditMode/SControlRigOutliner.h
--- a/Engine/Plugins/Animation/ControlRig/Source/ControlRigEditor/Private/EditMode/SControlRigOutliner.h
+++ b/Engine/Plugins/Animation/ControlRig/Source/ControlRigEditor/Private/EditMode/SControlRigOutliner.h
@@ -104,6 +104,7 @@ public:
 	FMultiRigData(UControlRig* InControlRig, EMultiRigDataType InType, FName InName) : WeakControlRig(InControlRig), Type(InType), Name(InName) {};
 	FText GetName() const;
 	FText GetDisplayName(const FRigTreeDisplaySettings& DisplaySettings) const;
+	void InvalidateDisplayName() const;
 	FText GetToolTipText(const FRigTreeDisplaySettings& Settings) const;
 	bool operator == (const FMultiRigData & Other) const;
 	bool IsValid() const;
@@ -346,10 +347,12 @@ private:
 	virtual void HandleControlAdded(UControlRig* ControlRig, bool bIsAdded) override;
 	virtual void HandleControlSelected(UControlRig* Subject, FRigControlElement* InControl, bool bSelected) override;
 	virtual void HandleRigVisibilityChanged(TArray<UControlRig*> InControlRigs) override;
+	virtual void HandleHierarchyModified(ERigHierarchyNotification InNotification, URigHierarchy* InHierarchy, const FRigNotificationSubject& InSubject);
 
 	//control rig delegates
 	void HandleOnControlRigBound(UControlRig* InControlRig);
 	void HandleOnObjectBoundToControlRig(UObject* InObject);
+	void HandlePostConstruction(UControlRig* InControlRig, const FName& InEventName);
 
 	void OnObjectsReplaced(const TMap<UObject*, UObject*>& OldToNewInstanceMap);
 	void OnSequencerTreeViewChanged(EMovieSceneDataChangeType MovieSceneDataChange);
@@ -365,5 +368,3 @@ private:
 	//set of control rigs we are bound too and need to clear delegates from
 	TArray<TWeakObjectPtr<UControlRig>> BoundControlRigs;
 };
-
-
diff --git a/Engine/Plugins/Animation/ControlRig/Source/ControlRigEditor/Private/Editor/ControlRigEditor.cpp b/Engine/Plugins/Animation/ControlRig/Source/ControlRigEditor/Private/Editor/ControlRigEditor.cpp
--- a/Engine/Plugins/Animation/ControlRig/Source/ControlRigEditor/Private/Editor/ControlRigEditor.cpp
+++ b/Engine/Plugins/Animation/ControlRig/Source/ControlRigEditor/Private/Editor/ControlRigEditor.cpp
@@ -1570,7 +1570,7 @@ void FControlRigBaseEditor::SetDetailViewForRigModules(const TArray<FName> InMod
 		}
 	}
 
-	if(!Objects.IsEmpty() && CVarRigVMEnablePinOverrides.GetValueOnAnyThread())
+	if(!Objects.IsEmpty() && CVarControlRigEnableOverrides.GetValueOnAnyThread())
 	{
 		TSharedPtr<FOverrideStatusDetailsViewObjectFilter> ObjectFilter = FOverrideStatusDetailsViewObjectFilter::Create();
 
@@ -1596,7 +1596,7 @@ void FControlRigBaseEditor::SetDetailViewForRigModules(const TArray<FName> InMod
 		{
 			const FString PropertyPath = InSubject.GetPropertyPathString();
 			return InSubject.GetStatus<UControlRig>(
-				[PropertyPath, RigBlueprint](const UControlRig* InModuleRig) -> TOptional<EOverrideWidgetStatus::Type>
+				[PropertyPath, RigBlueprint](const FOverrideStatusObjectHandle<UControlRig>& InModuleRig) -> TOptional<EOverrideWidgetStatus::Type>
 				{
 					if(const FRigModuleReference* ModuleReference = RigBlueprint->ModularRigModel.FindModule(InModuleRig->GetFName()))
 					{
@@ -1636,11 +1636,11 @@ void FControlRigBaseEditor::SetDetailViewForRigModules(const TArray<FName> InMod
 				RigBlueprint->Modify();
 
 				InSubject.ForEach<UControlRig>(
-					[RigBlueprint, InSubject](const UControlRig* InModuleRig)
+					[RigBlueprint, InSubject](const FOverrideStatusObjectHandle<UControlRig>& InModuleRig)
 					{
 						if(InSubject.HasPropertyPath())
 						{
-							const FControlRigOverrideValue Value(InSubject.GetPropertyPathString(), InModuleRig);
+							const FControlRigOverrideValue Value(InSubject.GetPropertyPathString(), InModuleRig.GetObject());
 							if(Value.IsValid())
 							{
 								RigBlueprint->GetModularRigController()->SetConfigValueInModule(InModuleRig->GetFName(), Value);
@@ -1652,7 +1652,7 @@ void FControlRigBaseEditor::SetDetailViewForRigModules(const TArray<FName> InMod
 							for(const FRigVMExternalVariable& Variable : Variables)
 							{
 								const FString PropertyPath = Variable.Name.ToString();
-								const FControlRigOverrideValue Value(PropertyPath, InModuleRig);
+								const FControlRigOverrideValue Value(PropertyPath, InModuleRig.GetObject());
 								RigBlueprint->GetModularRigController()->SetConfigValueInModule(InModuleRig->GetFName(), Value);
 							}
 						}
@@ -1669,7 +1669,7 @@ void FControlRigBaseEditor::SetDetailViewForRigModules(const TArray<FName> InMod
 				RigBlueprint->Modify();
 
 				InSubject.ForEach<UControlRig>(
-					[RigBlueprint, InSubject](const UControlRig* InModuleRig)
+					[RigBlueprint, InSubject](const FOverrideStatusObjectHandle<UControlRig>& InModuleRig)
 					{
 						if(InSubject.HasPropertyPath())
 						{
@@ -1707,7 +1707,7 @@ void FControlRigBaseEditor::SetDetailViewForRigModules(const TArray<FName> InMod
 				RigBlueprint->Modify();
 
 				InSubject.ForEach<UControlRig>(
-					[RigBlueprint, InSubject](const UControlRig* InModuleRig)
+					[RigBlueprint, InSubject](const FOverrideStatusObjectHandle<UControlRig>& InModuleRig)
 					{
 						if(InSubject.HasPropertyPath())
 						{
@@ -1742,7 +1742,7 @@ void FControlRigBaseEditor::SetDetailViewForRigModules(const TArray<FName> InMod
 			[this, RigBlueprint](const FOverrideStatusSubject& InSubject)
 			{
 				return InSubject.GetCommonValue<bool, UControlRig>(
-					[RigBlueprint, InSubject](const UControlRig* InModuleRig)
+					[RigBlueprint, InSubject](const FOverrideStatusObjectHandle<UControlRig>& InModuleRig)
 					{
 						if(const FRigModuleReference* ModuleReference = RigBlueprint->ModularRigModel.FindModule(InModuleRig->GetFName()))
 						{
diff --git a/Engine/Plugins/Animation/ControlRig/Source/ControlRigEditor/Private/Editor/ControlRigEditorMode.cpp b/Engine/Plugins/Animation/ControlRig/Source/ControlRigEditor/Private/Editor/ControlRigEditorMode.cpp
--- a/Engine/Plugins/Animation/ControlRig/Source/ControlRigEditor/Private/Editor/ControlRigEditorMode.cpp
+++ b/Engine/Plugins/Animation/ControlRig/Source/ControlRigEditor/Private/Editor/ControlRigEditorMode.cpp
@@ -160,10 +160,6 @@ FModularRigLegacyEditorMode::FModularRigLegacyEditorMode(const TSharedRef<FContr
 
 	TabFactories.RegisterFactory(MakeShared<FModularRigModelTabSummoner>(InControlRigEditor));
 	TabFactories.RegisterFactory(MakeShared<FRigModuleAssetBrowserTabSummoner>(InControlRigEditor));
-	if(CVarRigVMEnablePinOverrides.GetValueOnAnyThread())
-	{
-		TabFactories.RegisterFactory(MakeShared<FRigOverrideListTabSummoner>(InControlRigEditor));
-	}
 
 	TabLayout = FTabManager::NewLayout("Standalone_ModularRigEditMode_Layout_v1.3")
 		->AddArea
@@ -380,10 +376,6 @@ FModularRigEditorMode::FModularRigEditorMode(const TSharedRef<FControlRigEditor>
 
 	TabFactories.RegisterFactory(MakeShared<FModularRigModelTabSummoner>(InControlRigEditor));
 	TabFactories.RegisterFactory(MakeShared<FRigModuleAssetBrowserTabSummoner>(InControlRigEditor));
-	if(CVarRigVMEnablePinOverrides.GetValueOnAnyThread())
-	{
-		TabFactories.RegisterFactory(MakeShared<FRigOverrideListTabSummoner>(InControlRigEditor));
-	}
 
 	TabLayout = FTabManager::NewLayout("Standalone_ModularRigEditMode_Layout_v1.3")
 		->AddArea
diff --git a/Engine/Plugins/Runtime/RigVM/Source/RigVMEditor/Private/Editor/RigVMEditor.cpp b/Engine/Plugins/Runtime/RigVM/Source/RigVMEditor/Private/Editor/RigVMEditor.cpp
--- a/Engine/Plugins/Runtime/RigVM/Source/RigVMEditor/Private/Editor/RigVMEditor.cpp
+++ b/Engine/Plugins/Runtime/RigVM/Source/RigVMEditor/Private/Editor/RigVMEditor.cpp
@@ -2957,7 +2957,7 @@ void FRigVMEditorBase::SetDetailObjects(const TArray<UObject*>& InObjects, bool
 
 		ObjectFilter->OnCanCreateWidget().BindLambda([](const FOverrideStatusSubject& InSubject) -> bool
 		{
-			return InSubject.Contains<URigVMDetailsViewWrapperObject>([](const URigVMDetailsViewWrapperObject* InWrapper)
+			return InSubject.Contains<URigVMDetailsViewWrapperObject>([](const FOverrideStatusObjectHandle<URigVMDetailsViewWrapperObject>& InWrapper)
 			{
 				return Cast<URigVMNode>(InWrapper->GetSubject()) != nullptr;
 			});
@@ -2975,7 +2975,7 @@ void FRigVMEditorBase::SetDetailObjects(const TArray<UObject*>& InObjects, bool
 		{
 			const FString PinPath = PropertyPathToPinPath(InSubject);
 			return InSubject.GetStatus<URigVMDetailsViewWrapperObject>(
-				[PinPath](const URigVMDetailsViewWrapperObject* InWrapper) -> TOptional<EOverrideWidgetStatus::Type>
+				[PinPath](const FOverrideStatusObjectHandle<URigVMDetailsViewWrapperObject>& InWrapper) -> TOptional<EOverrideWidgetStatus::Type>
 				{
 					if(const URigVMNode* Node = Cast<URigVMNode>(InWrapper->GetSubject()))
 					{
@@ -3025,7 +3025,7 @@ void FRigVMEditorBase::SetDetailObjects(const TArray<UObject*>& InObjects, bool
 				TArray<FName> NodeNames;
 
 				InSubject.ForEach<URigVMDetailsViewWrapperObject>(
-					[PinPath, &PinPaths, &NodeNames](const URigVMDetailsViewWrapperObject* InWrapper)
+					[PinPath, &PinPaths, &NodeNames](const FOverrideStatusObjectHandle<URigVMDetailsViewWrapperObject>& InWrapper)
 					{
 						if(const URigVMNode* Node = Cast<URigVMNode>(InWrapper->GetSubject()))
 						{
@@ -3070,7 +3070,7 @@ void FRigVMEditorBase::SetDetailObjects(const TArray<UObject*>& InObjects, bool
 				TArray<FName> NodeNames;
 
 				InSubject.ForEach<URigVMDetailsViewWrapperObject>(
-					[PinPath, &PinPaths, &NodeNames](const URigVMDetailsViewWrapperObject* InWrapper)
+					[PinPath, &PinPaths, &NodeNames](const FOverrideStatusObjectHandle<URigVMDetailsViewWrapperObject>& InWrapper)
 					{
 						if(const URigVMNode* Node = Cast<URigVMNode>(InWrapper->GetSubject()))
 						{
@@ -3115,7 +3115,7 @@ void FRigVMEditorBase::SetDetailObjects(const TArray<UObject*>& InObjects, bool
 				TArray<FName> NodeNames;
 
 				InSubject.ForEach<URigVMDetailsViewWrapperObject>(
-					[PinPath, &PinPaths, &NodeNames](const URigVMDetailsViewWrapperObject* InWrapper)
+					[PinPath, &PinPaths, &NodeNames](const FOverrideStatusObjectHandle<URigVMDetailsViewWrapperObject>& InWrapper)
 					{
 						if(const URigVMNode* Node = Cast<URigVMNode>(InWrapper->GetSubject()))
 						{
@@ -3157,7 +3157,7 @@ void FRigVMEditorBase::SetDetailObjects(const TArray<UObject*>& InObjects, bool
 			[this, PropertyPathToPinPath](const FOverrideStatusSubject& InSubject)
 			{
 				const FString PinPath = PropertyPathToPinPath(InSubject);
-				return InSubject.Contains<URigVMDetailsViewWrapperObject>([PinPath](const URigVMDetailsViewWrapperObject* InWrapper)
+				return InSubject.Contains<URigVMDetailsViewWrapperObject>([PinPath](const FOverrideStatusObjectHandle<URigVMDetailsViewWrapperObject>& InWrapper)
 				{
 					if(const URigVMNode* Node = Cast<URigVMNode>(InWrapper->GetSubject()))
 					{
diff --git a/Engine/Source/Editor/AnimationEditorWidgets/Private/Overrides/OverrideStatusDetailsDisplayManager.cpp b/Engine/Source/Editor/AnimationEditorWidgets/Private/Overrides/OverrideStatusDetailsDisplayManager.cpp
--- a/Engine/Source/Editor/AnimationEditorWidgets/Private/Overrides/OverrideStatusDetailsDisplayManager.cpp
+++ b/Engine/Source/Editor/AnimationEditorWidgets/Private/Overrides/OverrideStatusDetailsDisplayManager.cpp
@@ -55,14 +55,14 @@ TSharedPtr<FOverrideStatusDetailsWidgetBuilder> FOverrideStatusDetailsDisplayMan
 	// if we are displaying an overrideable object, create the overrides widget
 	if ( bIsDisplayingOverrideableObject )
 	{
-		TArray<TWeakObjectPtr<const UObject>> Objects;
+		TArray<FOverrideStatusObject> Objects;
 		if(Args.Objects)
 		{
 			for(const TWeakObjectPtr<UObject>& ObjectPtr : *Args.Objects)
 			{
 				if(ObjectPtr.IsValid())
 				{
-					Objects.Add(ObjectPtr);
+					Objects.Add(ObjectPtr.Get() /* todo Keys */);
 				}
 			}
 		}
diff --git a/Engine/Source/Editor/AnimationEditorWidgets/Private/Overrides/OverrideStatusDetailsWidgetBuilder.cpp b/Engine/Source/Editor/AnimationEditorWidgets/Private/Overrides/OverrideStatusDetailsWidgetBuilder.cpp
--- a/Engine/Source/Editor/AnimationEditorWidgets/Private/Overrides/OverrideStatusDetailsWidgetBuilder.cpp
+++ b/Engine/Source/Editor/AnimationEditorWidgets/Private/Overrides/OverrideStatusDetailsWidgetBuilder.cpp
@@ -17,7 +17,7 @@
 
 FOverrideStatusDetailsWidgetBuilder::FOverrideStatusDetailsWidgetBuilder(
 	const TSharedRef<FOverrideStatusDetailsDisplayManager>& InDetailsDisplayManager,
-	const TArray<TWeakObjectPtr<const UObject>>& InObjects,
+	const TArray<FOverrideStatusObject>& InObjects,
 	const TSharedPtr<const FPropertyPath>& InPropertyPath,
 	const FName& InCategory)
 	: FPropertyUpdatedWidgetBuilder()
diff --git a/Engine/Source/Editor/AnimationEditorWidgets/Private/Overrides/OverrideStatusSubject.cpp b/Engine/Source/Editor/AnimationEditorWidgets/Private/Overrides/OverrideStatusSubject.cpp
--- a/Engine/Source/Editor/AnimationEditorWidgets/Private/Overrides/OverrideStatusSubject.cpp
+++ b/Engine/Source/Editor/AnimationEditorWidgets/Private/Overrides/OverrideStatusSubject.cpp
@@ -6,15 +6,59 @@
 
 #define LOCTEXT_NAMESPACE "OverrideStatusSubject"
 
-FOverrideStatusSubject::FOverrideStatusSubject(const TArray<TWeakObjectPtr<const UObject>>& InObjects,
-	const TSharedPtr<const FPropertyPath>& InPropertyPath, const FName& InCategory)
+FOverrideStatusObject::FOverrideStatusObject()
+	: Key(NAME_None)
+{
+}
+
+FOverrideStatusObject::FOverrideStatusObject(const UObject* InObject, const FName& InKey)
+	: WeakObjectPtr(TWeakObjectPtr<const UObject>(InObject))
+	, Key(InKey)
+{
+}
+
+FName FOverrideStatusObject::GetFName() const
+{
+	if(const UObject* Object = GetObject())
+	{
+		return Object->GetFName();
+	}
+	return NAME_None;
+}
+
+bool FOverrideStatusObject::IsValid() const
+{
+	return GetObject() != nullptr;
+}
+
+const UObject* FOverrideStatusObject::GetObject() const
+{
+	if(WeakObjectPtr.IsValid())
+	{
+		return WeakObjectPtr.Get();
+	}
+	return nullptr;
+}
+
+const FName& FOverrideStatusObject::GetKey() const
+{
+	return Key;
+}
+
+bool FOverrideStatusObject::operator==(const FOverrideStatusObject& InOther) const
+{
+	return (GetObject() == InOther.GetObject()) && (GetKey() == InOther.GetKey());
+}
+
+FOverrideStatusSubject::FOverrideStatusSubject(const TArray<FOverrideStatusObject>& InObjects,
+                                               const TSharedPtr<const FPropertyPath>& InPropertyPath, const FName& InCategory)
 	: Objects(InObjects)
 	, PropertyPath(InPropertyPath)
 	, Category(InCategory)
 {
 }
 
-FOverrideStatusSubject::FOverrideStatusSubject(const TWeakObjectPtr<const UObject>& InObject,
+FOverrideStatusSubject::FOverrideStatusSubject(const FOverrideStatusObject& InObject,
 	const TSharedPtr<const FPropertyPath>& InPropertyPath, const FName& InCategory)
 	: Objects({InObject})
 	, PropertyPath(InPropertyPath)
@@ -23,8 +67,8 @@ FOverrideStatusSubject::FOverrideStatusSubject(const TWeakObjectPtr<const UObjec
 }
 
 FOverrideStatusSubject::FOverrideStatusSubject(const UObject* InObject, const TSharedPtr<const FPropertyPath>& InPropertyPath,
-	const FName& InCategory)
-	: Objects({TWeakObjectPtr<const UObject>(InObject)})
+	const FName& InCategory, const FName& InKey)
+	: Objects({FOverrideStatusObject(InObject, InKey)})
 	, PropertyPath(InPropertyPath)
 	, Category(InCategory)
 {
@@ -32,7 +76,7 @@ FOverrideStatusSubject::FOverrideStatusSubject(const UObject* InObject, const TS
 
 bool FOverrideStatusSubject::IsValid() const
 {
-	for(const TWeakObjectPtr<const UObject>& Object : Objects)
+	for(const FOverrideStatusObject& Object : Objects)
 	{
 		if(!Object.IsValid())
 		{
@@ -47,13 +91,9 @@ int32 FOverrideStatusSubject::Num() const
 	return Objects.Num();
 }
 
-const UObject* FOverrideStatusSubject::operator[](int32 InIndex) const
+const FOverrideStatusObject& FOverrideStatusSubject::operator[](int32 InIndex) const
 {
-	if(Objects[InIndex].IsValid())
-	{
-		return Objects[InIndex].Get();
-	}
-	return nullptr;
+	return Objects[InIndex];
 }
 
 bool FOverrideStatusSubject::HasPropertyPath() const
@@ -97,32 +137,26 @@ const FName& FOverrideStatusSubject::GetCategory() const
 	return Category;
 }
 
-bool FOverrideStatusSubject::Contains(const UObject* InObject) const
+bool FOverrideStatusSubject::Contains(const FOverrideStatusObject& InObject) const
 {
-	for(const TWeakObjectPtr<const UObject>& Object : Objects)
+	for(const FOverrideStatusObject& Object : Objects)
 	{
-		if(Object.IsValid())
+		if(Object == InObject)
 		{
-			if(Object.Get() == InObject)
-			{
-				return true;
-			}
+			return true;
 		}
 	}
 	return false;
 }
 
-int32 FOverrideStatusSubject::Find(const UObject* InObject) const
+int32 FOverrideStatusSubject::Find(const FOverrideStatusObject& InObject) const
 {
 	for(int32 Index = 0; Index < Objects.Num(); Index++)
 	{
-		const TWeakObjectPtr<const UObject>& Object = Objects[Index];
-		if(Object.IsValid())
+		const FOverrideStatusObject& Object = Objects[Index];
+		if(Object == InObject)
 		{
-			if(Object.Get() == InObject)
-			{
-				return Index;
-			}
+			return Index;
 		}
 	}
 	return INDEX_NONE;
diff --git a/Engine/Source/Editor/AnimationEditorWidgets/Private/Overrides/SOverrideListWidget.cpp b/Engine/Source/Editor/AnimationEditorWidgets/Private/Overrides/SOverrideListWidget.cpp
--- a/Engine/Source/Editor/AnimationEditorWidgets/Private/Overrides/SOverrideListWidget.cpp
+++ b/Engine/Source/Editor/AnimationEditorWidgets/Private/Overrides/SOverrideListWidget.cpp
@@ -76,7 +76,7 @@ void SOverrideListWidget::Tick(
 				Content = FText::Format(
 					LOCTEXT("SubjectFormat", "{0}\n{1} {2}"),
 					Content,
-					FText::FromName(Subject[0]->GetFName()),
+					FText::FromName(Subject[0].GetFName()),
 					FText::FromString(Subject.GetPropertyPathString())
 				);
 			}
diff --git a/Engine/Source/Editor/AnimationEditorWidgets/Public/Overrides/OverrideStatusDetailsWidgetBuilder.h b/Engine/Source/Editor/AnimationEditorWidgets/Public/Overrides/OverrideStatusDetailsWidgetBuilder.h
--- a/Engine/Source/Editor/AnimationEditorWidgets/Public/Overrides/OverrideStatusDetailsWidgetBuilder.h
+++ b/Engine/Source/Editor/AnimationEditorWidgets/Public/Overrides/OverrideStatusDetailsWidgetBuilder.h
@@ -27,7 +27,7 @@ public:
 	 */
 	FOverrideStatusDetailsWidgetBuilder(
 		const TSharedRef<FOverrideStatusDetailsDisplayManager>& InDetailsDisplayManager,
-		const TArray<TWeakObjectPtr<const UObject>>& InObjects,
+		const TArray<FOverrideStatusObject>& InObjects,
 		const TSharedPtr<const FPropertyPath>& InPropertyPath,
 		const FName& InCategory);
 
diff --git a/Engine/Source/Editor/AnimationEditorWidgets/Public/Overrides/OverrideStatusSubject.h b/Engine/Source/Editor/AnimationEditorWidgets/Public/Overrides/OverrideStatusSubject.h
--- a/Engine/Source/Editor/AnimationEditorWidgets/Public/Overrides/OverrideStatusSubject.h
+++ b/Engine/Source/Editor/AnimationEditorWidgets/Public/Overrides/OverrideStatusSubject.h
@@ -17,6 +17,89 @@ DECLARE_DELEGATE_RetVal_OneParam(FReply, FOverrideStatus_ClearOverride, const FO
 DECLARE_DELEGATE_RetVal_OneParam(FReply, FOverrideStatus_ResetToDefault, const FOverrideStatusSubject&);
 DECLARE_DELEGATE_RetVal_OneParam(bool, FOverrideStatus_ValueDiffersFromDefault, const FOverrideStatusSubject&);
 
+/**
+ * A template handle pointing to given object type
+ */
+template<typename T>
+struct FOverrideStatusObjectHandle
+{
+public:
+
+	FOverrideStatusObjectHandle()
+		: Object(nullptr)
+		, Key(NAME_None)
+	{
+	}
+	
+	FOverrideStatusObjectHandle(const T* InObject, const FName& InKey = NAME_None)
+		: Object(InObject)
+		, Key(InKey)
+	{
+	}
+
+	bool IsValid() const
+	{
+		return Object != nullptr;
+	}
+
+	operator bool() const
+	{
+		return IsValid();
+	}
+
+	const T* GetObject() const
+	{
+		return Object;
+	}
+	
+	const FName& GetKey() const
+	{
+		return Key;
+	}
+
+	bool operator==(const FOverrideStatusObjectHandle& InOther) const
+	{
+		return (Object == InOther.Object) && (Key == InOther.Key);
+	}
+
+	const T* operator->() const
+	{
+		return GetObject();
+	}
+
+private:
+	const T* Object;
+	FName Key;
+};
+
+/**
+ * A single object used within an override widget - the object is identified by a weak object pointer and a potential sub object key.
+ */
+struct ANIMATIONEDITORWIDGETS_API FOverrideStatusObject
+{
+public:
+
+	FOverrideStatusObject();
+	FOverrideStatusObject(const UObject* InObject, const FName& InKey = NAME_None);
+
+	FName GetFName() const;
+	bool IsValid() const;
+	const UObject* GetObject() const;
+	const FName& GetKey() const;
+
+	bool operator==(const FOverrideStatusObject& InOther) const;
+
+	template<typename T>
+	FOverrideStatusObjectHandle<T> GetHandle() const
+	{
+		return FOverrideStatusObjectHandle<T>(Cast<T>(GetObject()), GetKey());
+	}
+
+private:
+	TWeakObjectPtr<const UObject> WeakObjectPtr;
+	FName Key;
+};
+
 /**
  * The subject of an override status (widget).
  * To support multi selection the subject is represented by an array of objects and a property path.
@@ -26,9 +109,9 @@ DECLARE_DELEGATE_RetVal_OneParam(bool, FOverrideStatus_ValueDiffersFromDefault,
 struct ANIMATIONEDITORWIDGETS_API FOverrideStatusSubject
 {
 public:
-	FOverrideStatusSubject(const TArray<TWeakObjectPtr<const UObject>>& InObjects, const TSharedPtr<const FPropertyPath>& InPropertyPath, const FName& InCategory);
-	FOverrideStatusSubject(const TWeakObjectPtr<const UObject>& InObject, const TSharedPtr<const FPropertyPath>& InPropertyPath = nullptr, const FName& InCategory = NAME_None);
-	FOverrideStatusSubject(const UObject* InObject, const TSharedPtr<const FPropertyPath>& InPropertyPath = nullptr, const FName& InCategory = NAME_None);
+	FOverrideStatusSubject(const TArray<FOverrideStatusObject>& InObjects, const TSharedPtr<const FPropertyPath>& InPropertyPath, const FName& InCategory);
+	FOverrideStatusSubject(const FOverrideStatusObject& InObject, const TSharedPtr<const FPropertyPath>& InPropertyPath = nullptr, const FName& InCategory = NAME_None);
+	FOverrideStatusSubject(const UObject* InObject, const TSharedPtr<const FPropertyPath>& InPropertyPath = nullptr, const FName& InCategory = NAME_None, const FName& InKey = NAME_None);
 
 	// returns true if any of the objects within the subject is still valid
 	bool IsValid() const;
@@ -37,11 +120,18 @@ public:
 	int32 Num() const;
 
 	// array access operator 
-	const UObject* operator [](int32 InIndex) const;
+	const FOverrideStatusObject& operator [](int32 InIndex) const;
+
+	// array handle access
+	template<typename T>
+	FOverrideStatusObjectHandle<T> GetHandle(int32 InIndex) const
+	{
+		return operator[](InIndex).GetHandle<T>();
+	}
 
 	// range access
-	TArray<TWeakObjectPtr<const UObject>>::RangedForConstIteratorType begin() const { return Objects.begin(); }
-	TArray<TWeakObjectPtr<const UObject>>::RangedForConstIteratorType end() const { return Objects.end(); }
+	TArray<FOverrideStatusObject>::RangedForConstIteratorType begin() const { return Objects.begin(); }
+	TArray<FOverrideStatusObject>::RangedForConstIteratorType end() const { return Objects.end(); }
 
 	// returns true if this subject contains a property path
 	bool HasPropertyPath() const;
@@ -59,7 +149,7 @@ public:
 	const FName& GetCategory() const;
 
 	// returns true if a given object is part of the subject
-	bool Contains(const UObject* InObject) const;
+	bool Contains(const FOverrideStatusObject& InObject) const;
 
 	// returns true if the subject contains an object of type ObjectType
 	template<typename ObjectType>
@@ -67,7 +157,7 @@ public:
 	{
 		for(int32 Index = 0; Index < Num(); Index++)
 		{
-			if(Cast<ObjectType>(operator[](Index)))
+			if(Cast<ObjectType>(operator[](Index).GetObject()))
 			{
 				return true;
 			}
@@ -77,13 +167,13 @@ public:
 
 	// returns true if the subject contains an object matching the given predicate
 	template<typename ObjectType = UObject>
-	bool Contains(TFunction<bool(const ObjectType*)> InMatchPredicate) const
+	bool Contains(TFunction<bool(const FOverrideStatusObjectHandle<ObjectType>&)> InMatchPredicate) const
 	{
 		for(int32 Index = 0; Index < Num(); Index++)
 		{
-			if(const ObjectType* CastObject = Cast<ObjectType>(operator[](Index)))
+			if(const FOverrideStatusObjectHandle Handle = GetHandle<ObjectType>(Index))
 			{
-				if(InMatchPredicate(CastObject))
+				if(InMatchPredicate(Handle))
 				{
 					return true;
 				}
@@ -93,18 +183,18 @@ public:
 	}
 
 	// returns the index of a given object or INDEX_NONE
-	int32 Find(const UObject* InObject) const;
+	int32 Find(const FOverrideStatusObject& InObject) const;
 
 	// returns the index of a given object or INDEX_NONE,
 	// the match is determined based on the given predicate 
 	template<typename ObjectType = UObject>
-	int32 Find(TFunction<bool(const ObjectType*)> InMatchPredicate) const
+	int32 Find(TFunction<bool(const FOverrideStatusObjectHandle<ObjectType>&)> InMatchPredicate) const
 	{
 		for(int32 Index = 0; Index < Num(); Index++)
 		{
-			if(const ObjectType* CastObject = Cast<ObjectType>(operator[](Index)))
+			if(const FOverrideStatusObjectHandle Handle = GetHandle<ObjectType>(Index))
 			{
-				if(InMatchPredicate(CastObject))
+				if(InMatchPredicate(Handle))
 				{
 					return true;
 				}
@@ -115,13 +205,13 @@ public:
 
 	// executes a given predicate for each object in the subject
 	template<typename ObjectType = UObject>
-	void ForEach(TFunction<void(const ObjectType*)> InPerObjectPredicate) const
+	void ForEach(TFunction<void(const FOverrideStatusObjectHandle<ObjectType>&)> InPerObjectPredicate) const
 	{
 		for(int32 Index = 0; Index < Num(); Index++)
 		{
-			if(const ObjectType* CastObject = Cast<ObjectType>(operator[](Index)))
+			if(const FOverrideStatusObjectHandle Handle = GetHandle<ObjectType>(Index))
 			{
-				InPerObjectPredicate(CastObject);
+				InPerObjectPredicate(Handle);
 			}
 		}
 	}
@@ -129,13 +219,13 @@ public:
 	// returns the first valid value for a given predicate (or an unset TOptional). each object in the subject
 	// is interrogated for using the predicate and can return a TOptional indicating success.
 	template<typename ValueType, typename ObjectType = UObject>
-	TOptional<ValueType> GetFirstValue(TFunction<TOptional<ValueType>(const ObjectType*)> InGetValuePerObjectPredicate) const
+	TOptional<ValueType> GetFirstValue(TFunction<TOptional<ValueType>(const FOverrideStatusObjectHandle<ObjectType>&)> InGetValuePerObjectPredicate) const
 	{
 		for(int32 Index = 0; Index < Num(); Index++)
 		{
-			if(const ObjectType* CastObject = Cast<ObjectType>(operator[](Index)))
+			if(const FOverrideStatusObjectHandle Handle = GetHandle<ObjectType>(Index))
 			{
-				const TOptional<ValueType> SingleValue = InGetValuePerObjectPredicate(CastObject);
+				const TOptional<ValueType> SingleValue = InGetValuePerObjectPredicate(Handle);
 				if(SingleValue.IsSet())
 				{
 					return SingleValue;
@@ -150,14 +240,14 @@ public:
 	// if the value matches for all objects the value will be returned as a set TOptional, otherwise
 	// an unset TOptional will be returned.
 	template<typename ValueType, typename ObjectType = UObject>
-	TOptional<ValueType> GetCommonValue(TFunction<TOptional<ValueType>(const ObjectType*)> InGetValuePerObjectPredicate) const
+	TOptional<ValueType> GetCommonValue(TFunction<TOptional<ValueType>(const FOverrideStatusObjectHandle<ObjectType>&)> InGetValuePerObjectPredicate) const
 	{
 		TOptional<ValueType> Result;
 		for(int32 Index = 0; Index < Num(); Index++)
 		{
-			if(const ObjectType* CastObject = Cast<ObjectType>(operator[](Index)))
+			if(const FOverrideStatusObjectHandle Handle = GetHandle<ObjectType>(Index))
 			{
-				const TOptional<ValueType> SingleValue = InGetValuePerObjectPredicate(CastObject);
+				const TOptional<ValueType> SingleValue = InGetValuePerObjectPredicate(Handle);
 				if(SingleValue.IsSet())
 				{
 					if(Result.IsSet())
@@ -180,14 +270,15 @@ public:
 
 	// returns the status provided by the subject's objects or an empty TOptional if it varies across objects
 	template<typename ObjectType = UObject>
-	TOptional<EOverrideWidgetStatus::Type> GetStatus(TFunction<TOptional<EOverrideWidgetStatus::Type>(const ObjectType*)> InGetValuePerObjectPredicate) const
+	TOptional<EOverrideWidgetStatus::Type> GetStatus(TFunction<TOptional<EOverrideWidgetStatus::Type>(const FOverrideStatusObjectHandle<ObjectType>&)> InGetValuePerObjectPredicate) const
 	{
 		return GetCommonValue<EOverrideWidgetStatus::Type, ObjectType>(InGetValuePerObjectPredicate);
 	}
 
 private:
-	TArray<TWeakObjectPtr<const UObject>> Objects;
+	TArray<FOverrideStatusObject> Objects;
 	TSharedPtr<const FPropertyPath> PropertyPath;
+	FName SubObjectKey;
 	FName Category;
 	
 	mutable TOptional<FString> LastSeparator;
