Control Rig: Performance improvement for details panel in constructoin mode

#rnx

diff --git a/Engine/Plugins/Animation/ControlRig/Source/ControlRigEditor/Private/ControlRigElementDetails.cpp b/Engine/Plugins/Animation/ControlRig/Source/ControlRigEditor/Private/ControlRigElementDetails.cpp
--- a/Engine/Plugins/Animation/ControlRig/Source/ControlRigEditor/Private/ControlRigElementDetails.cpp
+++ b/Engine/Plugins/Animation/ControlRig/Source/ControlRigEditor/Private/ControlRigElementDetails.cpp
@@ -1783,6 +1783,16 @@ void FRigBaseElementDetails::CustomizeMetadata(IDetailLayoutBuilder& DetailBuild
 					const FRigBaseElement* Element = PerElementInfos.Num() == 1 ? PerElementInfos[0].GetElement() : nullptr;
 					if (InKey.Type == ERigElementType::All || (Element && Element->GetKey() == InKey))
 					{
+						if(const URigHierarchy* Hierarchy = PerElementInfos[0].GetHierarchy())
+						{
+							if(const UControlRig* ControlRig = Hierarchy->GetTypedOuter<UControlRig>())
+							{
+								if(ControlRig->IsConstructionModeEnabled())
+								{
+									return;
+								}
+							}
+						}
 						// run the refresh of the user interface on the next tick on the game thread 
 						FFunctionGraphTask::CreateAndDispatchWhenReady([WeakPropertyUtilities]()
 						{
