Control Rig: Performance improvement for details panel in constructoin mode Follow up

#rnx

diff --git a/Engine/Plugins/Animation/ControlRig/Source/ControlRig/Private/ControlRig.cpp b/Engine/Plugins/Animation/ControlRig/Source/ControlRig/Private/ControlRig.cpp
--- a/Engine/Plugins/Animation/ControlRig/Source/ControlRig/Private/ControlRig.cpp
+++ b/Engine/Plugins/Animation/ControlRig/Source/ControlRig/Private/ControlRig.cpp
@@ -2221,7 +2221,8 @@ FRigControlElement* UControlRig::FindControl(const FName& InControlName) const
 
 bool UControlRig::IsConstructionModeEnabled() const
 {
-	return EventQueueToRun.Num() == 1 && EventQueueToRun.Contains(FRigUnit_PrepareForExecution::EventName);
+	return (EventQueueToRun.Num() == 1 && EventQueueToRun.Contains(FRigUnit_PrepareForExecution::EventName)) ||
+		 (EventQueue.Num() == 1 && EventQueue.Contains(FRigUnit_PrepareForExecution::EventName));
 }
 
 FTransform UControlRig::SetupControlFromGlobalTransform(const FName& InControlName, const FTransform& InGlobalTransform)
