Control Rig: Fix crash when redoing bake to control rig
#rnx

diff --git a/Engine/Plugins/Animation/ControlRig/Source/ControlRigEditor/Private/EditMode/ControlRigEditMode.cpp b/Engine/Plugins/Animation/ControlRig/Source/ControlRigEditor/Private/EditMode/ControlRigEditMode.cpp
--- a/Engine/Plugins/Animation/ControlRig/Source/ControlRigEditor/Private/EditMode/ControlRigEditMode.cpp
+++ b/Engine/Plugins/Animation/ControlRig/Source/ControlRigEditor/Private/EditMode/ControlRigEditMode.cpp
@@ -5855,7 +5855,10 @@ void FControlRigEditMode::AddControlRigInternal(UControlRig* InControlRig)
 	{
 		ModularRig->ForEachModule([](FRigModuleInstance* Module)
 		{
-			Module->GetRig()->SetControlsVisible(true);
+			if (UControlRig* Rig = Module->GetRig())
+			{
+				Rig->SetControlsVisible(true);
+			}
 			return true;
 		});
 	}
