Anim In Engine: Can't do space compensenation from nongame thread which can happen when deleting many channels, in which case we probably don't want to compensate anyway

diff --git a/Engine/Plugins/Animation/ControlRig/Source/ControlRigEditor/Private/Sequencer/ControlRigSpaceChannelEditors.cpp b/Engine/Plugins/Animation/ControlRig/Source/ControlRigEditor/Private/Sequencer/ControlRigSpaceChannelEditors.cpp
--- a/Engine/Plugins/Animation/ControlRig/Source/ControlRigEditor/Private/Sequencer/ControlRigSpaceChannelEditors.cpp
+++ b/Engine/Plugins/Animation/ControlRig/Source/ControlRigEditor/Private/Sequencer/ControlRigSpaceChannelEditors.cpp
@@ -671,8 +671,8 @@ static EControlRigContextChannelToKey GetCurrentTransformKeysAtThisTime(UControl
 void  FControlRigSpaceChannelHelpers::SequencerSpaceChannelKeyDeleted(UControlRig* ControlRig, ISequencer* Sequencer, FName ControlName, FMovieSceneControlRigSpaceChannel* Channel, UMovieSceneControlRigParameterSection* SectionToKey,
 	FFrameNumber TimeOfDeletion)
 {
-	if (bDoNotCompensate == true)
-	{
+	if (bDoNotCompensate == true || IsInGameThread() == false) //this may happen in an non game thread via a parallel for when we are deleting many channels
+	{															//if that's the case we assume the user may not want to compensate
 		return;
 	}
 	FMovieSceneControlRigSpaceBaseKey ExistingValue, PreviousValue;
