Anim Layers:  Make Name Column not shrink when shrinking

diff --git a/Engine/Plugins/Animation/ControlRig/Source/ControlRigEditor/Private/Sequencer/AnimLayers/SAnimLayers.cpp b/Engine/Plugins/Animation/ControlRig/Source/ControlRigEditor/Private/Sequencer/AnimLayers/SAnimLayers.cpp
--- a/Engine/Plugins/Animation/ControlRig/Source/ControlRigEditor/Private/Sequencer/AnimLayers/SAnimLayers.cpp
+++ b/Engine/Plugins/Animation/ControlRig/Source/ControlRigEditor/Private/Sequencer/AnimLayers/SAnimLayers.cpp
@@ -579,6 +579,7 @@ public:
 			return SNew(SHorizontalBox)
 				+ SHorizontalBox::Slot()
 				.MaxWidth(6)
+				.MinWidth(6)
 				[
 					SNew(SBox)
 						.WidthOverride(6)
@@ -1007,19 +1008,24 @@ void FAnimLayerSourcesView::CreateSourcesListView()
 		(
 			SNew(SHeaderRow)
 			+ SHeaderRow::Column(AnimLayerSourceListUI::LayerColumnName)
-			.FillWidth(100.f)
+			.FillSized(160.f)
+			.HAlignCell(EHorizontalAlignment::HAlign_Left)
 			.DefaultLabel(LOCTEXT("LayerColumnName", "Name"))
 			+ SHeaderRow::Column(AnimLayerSourceListUI::ActionColumnName)
-			.FillWidth(30.f)
+			.FillWidth(40.f)
+			.HAlignCell(EHorizontalAlignment::HAlign_Center)
 			.DefaultLabel(LOCTEXT("ActionColumnName", "Action"))
 			+ SHeaderRow::Column(AnimLayerSourceListUI::StatusColumnName)
 			.FillWidth(40.f)
+			.HAlignCell(EHorizontalAlignment::HAlign_Center)
 			.DefaultLabel(LOCTEXT("StatusColumnName", "Status"))
 			+ SHeaderRow::Column(AnimLayerSourceListUI::WeightColumnName)
 			.FillWidth(60.f)
+			.HAlignCell(EHorizontalAlignment::HAlign_Center)
 			.DefaultLabel(LOCTEXT("WeightColumnName", "Weight"))
 			+ SHeaderRow::Column(AnimLayerSourceListUI::TypeColumnName)
-			.FillWidth(80.f)
+			.FillSized(80.f)
+			.HAlignCell(EHorizontalAlignment::HAlign_Right)
 			.DefaultLabel(LOCTEXT("TypeColumnName", "Type"))
 		);
 }
