Anim In Engine: Pose Libarary fixes, use alt drag for non additive drags and for additive it goes from -1 to 2.  Blending and additive mode working on mirror.

diff --git a/Engine/Plugins/Animation/ControlRig/Source/ControlRig/Private/Tools/ControlRigPose.cpp b/Engine/Plugins/Animation/ControlRig/Source/ControlRig/Private/Tools/ControlRigPose.cpp
--- a/Engine/Plugins/Animation/ControlRig/Source/ControlRig/Private/Tools/ControlRigPose.cpp
+++ b/Engine/Plugins/Animation/ControlRig/Source/ControlRig/Private/Tools/ControlRigPose.cpp
@@ -234,7 +234,7 @@ void FControlRigControlPose::BlendWithInitialPoses(FControlRigControlPose& Initi
 		if (CopyRigControl)
 		{
 			FRigControlCopy* InitialFound = nullptr;
-			int32* Index = InitialPose.CopyOfControlsNameToIndex.Find(CopyRigControl->Name);
+			int32* Index = InitialPose.CopyOfControlsNameToIndex.Find(ControlElement->GetFName());
 			if (Index)
 			{
 				InitialFound = &(InitialPose.CopyOfControls[*Index]);
diff --git a/Engine/Plugins/Animation/ControlRig/Source/ControlRigEditor/Private/EditMode/SControlRigControlViews.cpp b/Engine/Plugins/Animation/ControlRig/Source/ControlRigEditor/Private/EditMode/SControlRigControlViews.cpp
--- a/Engine/Plugins/Animation/ControlRig/Source/ControlRigEditor/Private/EditMode/SControlRigControlViews.cpp
+++ b/Engine/Plugins/Animation/ControlRig/Source/ControlRigEditor/Private/EditMode/SControlRigControlViews.cpp
@@ -188,6 +188,8 @@ void SControlRigAssetEditableTextBox::OnTextChanged(const FText& InLabel)
 bool SControlRigPoseView::bIsKey = false;
 bool SControlRigPoseView::bIsMirror = false;
 bool SControlRigPoseView::bIsAdditive = false;
+float SControlRigPoseView::bDynamicMin = 0.0;
+float SControlRigPoseView::bDynamicMax = 1.0;
 
 void SControlRigPoseView::Construct(const FArguments& InArgs)
 {
@@ -360,14 +362,48 @@ void SControlRigPoseView::Construct(const FArguments& InArgs)
 								// Only allow spinning if we have a single value
 								.PreventThrottling(true)
 								.Value(this, &SControlRigPoseView::OnGetPoseBlendValue)
-								.ToolTipText(LOCTEXT("BlendTooltip", "Blend between current pose and pose asset."))
+								.ToolTipText(LOCTEXT("BlendTooltip", "Blend between current pose and pose asset. Use Alt to extend values"))
 								.MinValue(-1.0f)
 								.MaxValue(2.0f)
-								.MinSliderValue(-1.0)
-								.MaxSliderValue(2.0f)
+								.MinSliderValue_Lambda([]() ->float
+								{
+									if (SControlRigPoseView::bIsAdditive)
+									{
+										return -1.0;
+									}
+									return SControlRigPoseView::bDynamicMin;
+
+								})
+								.MaxSliderValue_Lambda([]() ->float
+								{
+									if (SControlRigPoseView::bIsAdditive)
+									{
+										return 2.0f;
+									}
+									return SControlRigPoseView::bDynamicMax;
+
+								})
 								.SliderExponent(1)
 								.Delta(0.005f)
 								.MinDesiredWidth(100.0f)
+								.SupportDynamicSliderMinValue(true)
+								.SupportDynamicSliderMaxValue(true)
+								.OnDynamicSliderMinValueChanged_Lambda([](float NewMinSliderValue, TWeakPtr<SWidget> InValueChangedSourceWidget, bool bIsOriginator, bool bUpdateOnlyIfLower)
+								{
+									if (((NewMinSliderValue < SControlRigPoseView::bDynamicMin) && bUpdateOnlyIfLower) || !bUpdateOnlyIfLower)
+									{
+										SControlRigPoseView::bDynamicMin = NewMinSliderValue;
+									}
+
+								})
+								.OnDynamicSliderMaxValueChanged_Lambda([](float NewMaxSliderValue, TWeakPtr<SWidget> InValueChangedSourceWidget, bool bIsOriginator, bool bUpdateOnlyIfHigher)
+								{
+									if (((NewMaxSliderValue > SControlRigPoseView::bDynamicMax) && bUpdateOnlyIfHigher) || !bUpdateOnlyIfHigher)
+									{
+										SControlRigPoseView::bDynamicMax = NewMaxSliderValue;
+									}
+
+								})
 								.OnValueChanged(this, &SControlRigPoseView::OnPoseBlendChanged)
 								.OnValueCommitted(this, &SControlRigPoseView::OnPoseBlendCommited)
 								.OnBeginSliderMovement(this,&SControlRigPoseView::OnBeginSliderMovement)
diff --git a/Engine/Plugins/Animation/ControlRig/Source/ControlRigEditor/Private/EditMode/SControlRigControlViews.h b/Engine/Plugins/Animation/ControlRig/Source/ControlRigEditor/Private/EditMode/SControlRigControlViews.h
--- a/Engine/Plugins/Animation/ControlRig/Source/ControlRigEditor/Private/EditMode/SControlRigControlViews.h
+++ b/Engine/Plugins/Animation/ControlRig/Source/ControlRigEditor/Private/EditMode/SControlRigControlViews.h
@@ -71,6 +71,9 @@ private:
 	static bool bIsKey;
 	static bool bIsMirror;
 	static bool bIsAdditive;
+	static float bDynamicMin;
+	static float bDynamicMax;
+
 	float PoseBlendValue;
 	bool bIsBlending;
 	bool bSliderStartedTransaction;
