Anim in Engine: Fix transactions can turn anim details proxy manager invalid

#rnx

diff --git a/Engine/Plugins/Animation/ControlRig/Source/ControlRigEditor/Private/EditMode/ControlRigEditMode.cpp b/Engine/Plugins/Animation/ControlRig/Source/ControlRigEditor/Private/EditMode/ControlRigEditMode.cpp
--- a/Engine/Plugins/Animation/ControlRig/Source/ControlRigEditor/Private/EditMode/ControlRigEditMode.cpp
+++ b/Engine/Plugins/Animation/ControlRig/Source/ControlRigEditor/Private/EditMode/ControlRigEditMode.cpp
@@ -196,7 +196,7 @@ FControlRigEditMode::FControlRigEditMode()
 	, bShowControlsAsOverlay(false)
 	, bIsConstructionEventRunning(false)
 {
-	AnimDetailsProxyManager = NewObject<UAnimDetailsProxyManager>(GetTransientPackage(), NAME_None, RF_Transactional);
+	AnimDetailsProxyManager = NewObject<UAnimDetailsProxyManager>(GetTransientPackage(), NAME_None);
 
 	StoredPose = NewObject<UControlRigPoseAsset>(GetTransientPackage(), NAME_None);
 	DetailKeyFrameCache = MakeShared<FDetailKeyFrameCacheAndHandler>();
