Control Rig: Fix clearing connections when attempting to create a cyclic connection
#rnx

diff --git a/Engine/Plugins/Animation/ControlRig/Source/ControlRig/Private/ModularRigController.cpp b/Engine/Plugins/Animation/ControlRig/Source/ControlRig/Private/ModularRigController.cpp
--- a/Engine/Plugins/Animation/ControlRig/Source/ControlRig/Private/ModularRigController.cpp
+++ b/Engine/Plugins/Animation/ControlRig/Source/ControlRig/Private/ModularRigController.cpp
@@ -397,7 +397,16 @@ bool UModularRigController::ConnectConnectorToElements(const FRigElementKey& InC
 	}
 #endif
 
-	(void)DisconnectCyclicConnectors();
+	const TArray<FRigElementKey> DisconnectedConnectors = DisconnectCyclicConnectors();
+
+	// If the connection that was attempted generated cycles, try to reestablish the previous connections
+	if (DisconnectedConnectors.Contains(InConnectorKey))
+	{
+		if (TArray<FRigElementKey>* Previous = PreviousConnections.Find(InConnectorKey))
+		{
+			ConnectConnectorToElements(InConnectorKey, *Previous, bSetupUndo, bAutomaticReparenting, bCheckValidConnection);
+		}
+	}
 
 #if WITH_EDITOR
 	TransactionPtr.Reset();
