Control Rig: Fix uses of static inline FLazyName

#rnx

diff --git a/Engine/Plugins/Animation/ControlRig/Source/ControlRig/Private/Rigs/RigHierarchy.cpp b/Engine/Plugins/Animation/ControlRig/Source/ControlRig/Private/Rigs/RigHierarchy.cpp
--- a/Engine/Plugins/Animation/ControlRig/Source/ControlRig/Private/Rigs/RigHierarchy.cpp
+++ b/Engine/Plugins/Animation/ControlRig/Source/ControlRig/Private/Rigs/RigHierarchy.cpp
@@ -102,6 +102,9 @@ bool URigHierarchy::bEnableValidityCheckbyDefault = true;
 bool URigHierarchy::bEnableValidityCheckbyDefault = false;
 #endif
 
+const FLazyName URigHierarchy::DefaultParentKeyLabel = TEXT("Parent");
+const FLazyName URigHierarchy::WorldSpaceKeyLabel = TEXT("World");
+
 URigHierarchy::URigHierarchy()
 : TopologyVersion(0)
 , MetadataVersion(0)
diff --git a/Engine/Plugins/Animation/ControlRig/Source/ControlRig/Public/Rigs/RigHierarchy.h b/Engine/Plugins/Animation/ControlRig/Source/ControlRig/Public/Rigs/RigHierarchy.h
--- a/Engine/Plugins/Animation/ControlRig/Source/ControlRig/Public/Rigs/RigHierarchy.h
+++ b/Engine/Plugins/Animation/ControlRig/Source/ControlRig/Public/Rigs/RigHierarchy.h
@@ -3761,8 +3761,8 @@ public:
 	
 	static CONTROLRIG_API FRigElementKey GetDefaultParentKey();
 	static CONTROLRIG_API FRigElementKey GetWorldSpaceReferenceKey();
-	static inline const FLazyName DefaultParentKeyLabel = TEXT("Parent");
-	static inline const FLazyName WorldSpaceKeyLabel = TEXT("World");
+	static CONTROLRIG_API const FLazyName DefaultParentKeyLabel;
+	static CONTROLRIG_API const FLazyName WorldSpaceKeyLabel;
 
 	/**
 	 * Returns true if an element is parented to another element
