Anim in Engine: Fix ensure condition hit on automated tests in Anim Details

#rnx

diff --git a/Engine/Plugins/Animation/ControlRig/Source/ControlRigEditor/Private/AnimDetails/Proxies/AnimDetailsProxyBase.cpp b/Engine/Plugins/Animation/ControlRig/Source/ControlRigEditor/Private/AnimDetails/Proxies/AnimDetailsProxyBase.cpp
--- a/Engine/Plugins/Animation/ControlRig/Source/ControlRigEditor/Private/AnimDetails/Proxies/AnimDetailsProxyBase.cpp
+++ b/Engine/Plugins/Animation/ControlRig/Source/ControlRigEditor/Private/AnimDetails/Proxies/AnimDetailsProxyBase.cpp
@@ -769,7 +769,10 @@ FName UAnimDetailsProxyBase::GetDetailRowID() const
 		}
 		else
 		{
-			ensureMsgf(0, TEXT("Getting a detail row ID for an invalid proxy. Returning NAME_None as its ID"));
+			// Editor test may create control rigs that are not fully valid. In such cases it's acceptable to return NAME_None here.
+			const bool bRunningCommandlet = IsRunningCommandlet();
+			ensureMsgf(!bRunningCommandlet, TEXT("Getting a detail row ID for an invalid proxy. Returning NAME_None as its ID"));
+		
 			return NAME_None;
 		}
 	}
