Control Rig: Use Anim Details customization only on AnimDetails DetailsView

#rnx

diff --git a/Engine/Plugins/Animation/ControlRig/Source/ControlRigEditor/Private/AnimDetails/Views/SAnimDetailsView.cpp b/Engine/Plugins/Animation/ControlRig/Source/ControlRigEditor/Private/AnimDetails/Views/SAnimDetailsView.cpp
--- a/Engine/Plugins/Animation/ControlRig/Source/ControlRigEditor/Private/AnimDetails/Views/SAnimDetailsView.cpp
+++ b/Engine/Plugins/Animation/ControlRig/Source/ControlRigEditor/Private/AnimDetails/Views/SAnimDetailsView.cpp
@@ -6,7 +6,17 @@
 #include "AnimDetails/AnimDetailsSelection.h"
 #include "AnimDetails/Customizations/AnimDetailsProxyDetails.h"
 #include "AnimDetails/Customizations/AnimDetailsProxyManagerDetails.h"
+#include "AnimDetails/Customizations/AnimDetailsValueCustomization.h"
+#include "AnimDetails/Customizations/AnimDetailsValueEnumCustomization.h"
 #include "AnimDetails/Proxies/AnimDetailsProxyBase.h"
+#include "AnimDetails/Proxies/AnimDetailsProxyBool.h"
+#include "AnimDetails/Proxies/AnimDetailsProxyEnum.h"
+#include "AnimDetails/Proxies/AnimDetailsProxyFloat.h"
+#include "AnimDetails/Proxies/AnimDetailsProxyInteger.h"
+#include "AnimDetails/Proxies/AnimDetailsProxyLocation.h"
+#include "AnimDetails/Proxies/AnimDetailsProxyRotation.h"
+#include "AnimDetails/Proxies/AnimDetailsProxyScale.h"
+#include "AnimDetails/Proxies/AnimDetailsProxyVector2D.h"
 #include "AnimDetails/Widgets/SAnimDetailsOptions.h"
 #include "AnimDetails/Widgets/SAnimDetailsSearchBox.h"
 #include "CurveEditor.h"
@@ -54,6 +64,14 @@ namespace UE::ControlRigEditor
 		DetailsView->SetIsPropertyVisibleDelegate(FIsPropertyVisible::CreateSP(this, &SAnimDetailsView::ShouldDisplayProperty));
 		DetailsView->SetIsPropertyReadOnlyDelegate(FIsPropertyReadOnly::CreateSP(this, &SAnimDetailsView::IsReadOnlyProperty));
 		DetailsView->RegisterInstancedCustomPropertyLayout(UAnimDetailsProxyBase::StaticClass(), FOnGetDetailCustomizationInstance::CreateStatic(&FAnimDetailsProxyDetails::MakeInstance));
+		DetailsView->RegisterInstancedCustomPropertyTypeLayout(FAnimDetailsFloat::StaticStruct()->GetFName(), FOnGetPropertyTypeCustomizationInstance::CreateStatic(&FAnimDetailsValueCustomization::MakeInstance));
+		DetailsView->RegisterInstancedCustomPropertyTypeLayout(FAnimDetailsInteger::StaticStruct()->GetFName(), FOnGetPropertyTypeCustomizationInstance::CreateStatic(&FAnimDetailsValueCustomization::MakeInstance));
+		DetailsView->RegisterInstancedCustomPropertyTypeLayout(FAnimDetailsBool::StaticStruct()->GetFName(), FOnGetPropertyTypeCustomizationInstance::CreateStatic(&FAnimDetailsValueCustomization::MakeInstance));
+		DetailsView->RegisterInstancedCustomPropertyTypeLayout(FAnimDetailsVector2D::StaticStruct()->GetFName(), FOnGetPropertyTypeCustomizationInstance::CreateStatic(&FAnimDetailsValueCustomization::MakeInstance));
+		DetailsView->RegisterInstancedCustomPropertyTypeLayout(FAnimDetailsLocation::StaticStruct()->GetFName(), FOnGetPropertyTypeCustomizationInstance::CreateStatic(&FAnimDetailsValueCustomization::MakeInstance));
+		DetailsView->RegisterInstancedCustomPropertyTypeLayout(FAnimDetailsRotation::StaticStruct()->GetFName(), FOnGetPropertyTypeCustomizationInstance::CreateStatic(&FAnimDetailsValueCustomization::MakeInstance));
+		DetailsView->RegisterInstancedCustomPropertyTypeLayout(FAnimDetailsScale::StaticStruct()->GetFName(), FOnGetPropertyTypeCustomizationInstance::CreateStatic(&FAnimDetailsValueCustomization::MakeInstance));
+		DetailsView->RegisterInstancedCustomPropertyTypeLayout(FAnimDetailsEnum::StaticStruct()->GetFName(), FOnGetPropertyTypeCustomizationInstance::CreateStatic(&FAnimDetailsValueEnumCustomization::MakeInstance));
 
 		DetailsView->SetObjects({ TWeakObjectPtr<UObject>(ProxyManager) });
 		
diff --git a/Engine/Plugins/Animation/ControlRig/Source/ControlRigEditor/Private/ControlRigEditorModule.cpp b/Engine/Plugins/Animation/ControlRig/Source/ControlRigEditor/Private/ControlRigEditorModule.cpp
--- a/Engine/Plugins/Animation/ControlRig/Source/ControlRigEditor/Private/ControlRigEditorModule.cpp
+++ b/Engine/Plugins/Animation/ControlRig/Source/ControlRigEditor/Private/ControlRigEditorModule.cpp
@@ -3,17 +3,6 @@
 #include "ControlRigEditorModule.h"
 
 #include "AnimDetails/Customizations/AnimDetailsProxyDetails.h"
-#include "AnimDetails/Customizations/AnimDetailsValueCustomization.h"
-#include "AnimDetails/Customizations/AnimDetailsValueEnumCustomization.h"
-#include "AnimDetails/Proxies/AnimDetailsProxyBool.h"
-#include "AnimDetails/Proxies/AnimDetailsProxyEnum.h"
-#include "AnimDetails/Proxies/AnimDetailsProxyFloat.h"
-#include "AnimDetails/Proxies/AnimDetailsProxyInteger.h"
-#include "AnimDetails/Proxies/AnimDetailsProxyLocation.h"
-#include "AnimDetails/Proxies/AnimDetailsProxyRotation.h"
-#include "AnimDetails/Proxies/AnimDetailsProxyScale.h"
-#include "AnimDetails/Proxies/AnimDetailsProxyTransform.h"
-#include "AnimDetails/Proxies/AnimDetailsProxyVector2D.h"
 #include "BlueprintActionDatabaseRegistrar.h"
 #include "BlueprintNodeSpawner.h"
 #include "PropertyEditorModule.h"
@@ -163,36 +152,12 @@ void FControlRigEditorModule::StartupModule()
 	PropertiesToUnregisterOnShutdown.Add(FRigVMCompileSettings::StaticStruct()->GetFName());
 	PropertyEditorModule.RegisterCustomPropertyTypeLayout(PropertiesToUnregisterOnShutdown.Last(), FOnGetPropertyTypeCustomizationInstance::CreateStatic(&FRigVMCompileSettingsDetails::MakeInstance));
 
-	PropertiesToUnregisterOnShutdown.Add(FAnimDetailsFloat::StaticStruct()->GetFName());
-	PropertyEditorModule.RegisterCustomPropertyTypeLayout(PropertiesToUnregisterOnShutdown.Last(), FOnGetPropertyTypeCustomizationInstance::CreateStatic(&FAnimDetailsValueCustomization::MakeInstance));
-
-	PropertiesToUnregisterOnShutdown.Add(FAnimDetailsInteger::StaticStruct()->GetFName());
-	PropertyEditorModule.RegisterCustomPropertyTypeLayout(PropertiesToUnregisterOnShutdown.Last(), FOnGetPropertyTypeCustomizationInstance::CreateStatic(&FAnimDetailsValueCustomization::MakeInstance));
-
-	PropertiesToUnregisterOnShutdown.Add(FAnimDetailsBool::StaticStruct()->GetFName());
-	PropertyEditorModule.RegisterCustomPropertyTypeLayout(PropertiesToUnregisterOnShutdown.Last(), FOnGetPropertyTypeCustomizationInstance::CreateStatic(&FAnimDetailsValueCustomization::MakeInstance));
-
-	PropertiesToUnregisterOnShutdown.Add(FAnimDetailsVector2D::StaticStruct()->GetFName());
-	PropertyEditorModule.RegisterCustomPropertyTypeLayout(PropertiesToUnregisterOnShutdown.Last(), FOnGetPropertyTypeCustomizationInstance::CreateStatic(&FAnimDetailsValueCustomization::MakeInstance));
-
-	PropertiesToUnregisterOnShutdown.Add(FAnimDetailsLocation::StaticStruct()->GetFName());
-	PropertyEditorModule.RegisterCustomPropertyTypeLayout(PropertiesToUnregisterOnShutdown.Last(), FOnGetPropertyTypeCustomizationInstance::CreateStatic(&FAnimDetailsValueCustomization::MakeInstance));
-
-	PropertiesToUnregisterOnShutdown.Add(FAnimDetailsRotation::StaticStruct()->GetFName());
-	PropertyEditorModule.RegisterCustomPropertyTypeLayout(PropertiesToUnregisterOnShutdown.Last(), FOnGetPropertyTypeCustomizationInstance::CreateStatic(&FAnimDetailsValueCustomization::MakeInstance));
-
-	PropertiesToUnregisterOnShutdown.Add(FAnimDetailsScale::StaticStruct()->GetFName());
-	PropertyEditorModule.RegisterCustomPropertyTypeLayout(PropertiesToUnregisterOnShutdown.Last(), FOnGetPropertyTypeCustomizationInstance::CreateStatic(&FAnimDetailsValueCustomization::MakeInstance));
-	
 	PropertiesToUnregisterOnShutdown.Add(FRigVMPythonSettings::StaticStruct()->GetFName());
 	PropertyEditorModule.RegisterCustomPropertyTypeLayout(PropertiesToUnregisterOnShutdown.Last(), FOnGetPropertyTypeCustomizationInstance::CreateStatic(&FControlRigPythonLogDetails::MakeInstance));
 
 	PropertiesToUnregisterOnShutdown.Add(FRigVMDrawContainer::StaticStruct()->GetFName());
 	PropertyEditorModule.RegisterCustomPropertyTypeLayout(PropertiesToUnregisterOnShutdown.Last(), FOnGetPropertyTypeCustomizationInstance::CreateStatic(&FControlRigDrawContainerDetails::MakeInstance));
 
-	PropertiesToUnregisterOnShutdown.Add(FAnimDetailsEnum::StaticStruct()->GetFName());
-	PropertyEditorModule.RegisterCustomPropertyTypeLayout(PropertiesToUnregisterOnShutdown.Last(), FOnGetPropertyTypeCustomizationInstance::CreateStatic(&FAnimDetailsValueEnumCustomization::MakeInstance));
-
 	PropertiesToUnregisterOnShutdown.Add(FRigElementKey::StaticStruct()->GetFName());
 	PropertyEditorModule.RegisterCustomPropertyTypeLayout(PropertiesToUnregisterOnShutdown.Last(), FOnGetPropertyTypeCustomizationInstance::CreateStatic(&FRigElementKeyDetails::MakeInstance));
 
