Control Rig: Fix crash when undo/redo adding rigs in sequencer
#rnx

diff --git a/Engine/Plugins/Animation/ControlRig/Source/ControlRig/Private/ModularRig.cpp b/Engine/Plugins/Animation/ControlRig/Source/ControlRig/Private/ModularRig.cpp
--- a/Engine/Plugins/Animation/ControlRig/Source/ControlRig/Private/ModularRig.cpp
+++ b/Engine/Plugins/Animation/ControlRig/Source/ControlRig/Private/ModularRig.cpp
@@ -267,6 +267,7 @@ void UModularRig::Serialize(FArchive& Ar)
 	{
 		ModularRigModel.UpdateCachedChildren();
 		ModularRigModel.Connections.UpdateFromConnectionList();
+		UpdateCachedChildren();
 	}
 }
 
