Control Rig: Arrange controls in modular rigs by module in sequencer
#rnx

diff --git a/Engine/Plugins/Animation/ControlRig/Source/ControlRig/Private/ModularRig.cpp b/Engine/Plugins/Animation/ControlRig/Source/ControlRig/Private/ModularRig.cpp
--- a/Engine/Plugins/Animation/ControlRig/Source/ControlRig/Private/ModularRig.cpp
+++ b/Engine/Plugins/Animation/ControlRig/Source/ControlRig/Private/ModularRig.cpp
@@ -456,6 +456,42 @@ const TArray<FName>& UModularRig::GetSupportedEvents() const
 	return SupportedEvents;
 }
 
+void UModularRig::GetControlsInOrder(TArray<FRigControlElement*>& SortedControls) const
+{
+	SortedControls.Reset();
+
+	if(DynamicHierarchy == nullptr)
+	{
+		return;
+	}
+
+	TMap<FName, TArray<FRigControlElement*>> ControlsByModule;
+
+	DynamicHierarchy->Traverse([&](FRigBaseElement* Element, bool& bContinue)
+	{
+		if (FRigControlElement* Control = Cast<FRigControlElement>(Element))
+		{
+			const FName ModuleName = DynamicHierarchy->GetModuleFName(Control->GetKey());
+			ControlsByModule.FindOrAdd(ModuleName).Add(Control);
+		}
+		bContinue = true;
+	});
+
+	ForEachModule([ControlsByModule, &SortedControls](const FRigModuleInstance* Module) -> bool
+	{
+		if (const TArray<FRigControlElement*>* Controls = ControlsByModule.Find(Module->Name))
+		{
+			SortedControls.Append(*Controls);
+		}
+		return true;
+	});
+	
+	if (TArray<FRigControlElement*>* Controls = ControlsByModule.Find(NAME_None))
+	{
+		SortedControls.Append(*Controls);
+	}
+}
+
 const FModularRigSettings& UModularRig::GetModularRigSettings() const
 {
 	if(HasAnyFlags(RF_ClassDefaultObject))
diff --git a/Engine/Plugins/Animation/ControlRig/Source/ControlRig/Public/ModularRig.h b/Engine/Plugins/Animation/ControlRig/Source/ControlRig/Public/ModularRig.h
--- a/Engine/Plugins/Animation/ControlRig/Source/ControlRig/Public/ModularRig.h
+++ b/Engine/Plugins/Animation/ControlRig/Source/ControlRig/Public/ModularRig.h
@@ -157,6 +157,7 @@ public:
 	virtual FRigElementKeyRedirector GetElementKeyRedirector() const override { return ElementKeyRedirector; }
 	UE_API virtual bool SupportsEvent(const FName& InEventName) const override;
 	UE_API virtual const TArray<FName>& GetSupportedEvents() const override;
+	virtual void GetControlsInOrder(TArray<FRigControlElement*>& SortedControls) const override;
 	// END ControlRig
 
 	UPROPERTY()
diff --git a/Engine/Plugins/Animation/ControlRig/Source/ControlRigEditor/Private/EditMode/SControlRigOutliner.cpp b/Engine/Plugins/Animation/ControlRig/Source/ControlRigEditor/Private/EditMode/SControlRigOutliner.cpp
--- a/Engine/Plugins/Animation/ControlRig/Source/ControlRigEditor/Private/EditMode/SControlRigOutliner.cpp
+++ b/Engine/Plugins/Animation/ControlRig/Source/ControlRigEditor/Private/EditMode/SControlRigOutliner.cpp
@@ -1315,6 +1315,42 @@ void SMultiRigHierarchyTreeView::RefreshTreeView()
 						bContinue = true;
 					});
 				}
+
+				// remove empty modules
+				if (ControlRig->IsModularRig() && Settings.bArrangeByModules)
+				{
+					TArray<FMultiRigData> ToRemove;
+					for (const TPair<FMultiRigData, TSharedPtr<FMultiRigTreeElement>> Pair : ElementMap)
+					{
+						if (Pair.Key.IsModule() && Pair.Value->Children.IsEmpty())
+						{
+							ToRemove.Add(Pair.Key);
+						}
+					}
+				
+					for (int32 i=0; i<ToRemove.Num(); i++)
+					{
+						FMultiRigData Parent;
+						if (FMultiRigData* ParentDataPtr = ParentMap.Find(ToRemove[i]))
+						{
+							Parent = *ParentDataPtr;
+						}
+					
+						RemoveElement(ToRemove[i]);
+
+						// If the parent has no children, remove it
+						if (Parent.IsModule())
+						{
+							if (TSharedPtr<FMultiRigTreeElement>* ParentElement = ElementMap.Find(Parent))
+							{
+								if ((*ParentElement)->Children.IsEmpty())
+								{
+									ToRemove.AddUnique(Parent);
+								}
+							}
+						}
+					}
+				}
 			}
 		}
 	}
