Control Rig: Fixed compilation error
#rnx

diff --git a/Engine/Plugins/Animation/ControlRig/Source/ControlRigEditor/Private/EditMode/SControlRigOutliner.cpp b/Engine/Plugins/Animation/ControlRig/Source/ControlRigEditor/Private/EditMode/SControlRigOutliner.cpp
--- a/Engine/Plugins/Animation/ControlRig/Source/ControlRigEditor/Private/EditMode/SControlRigOutliner.cpp
+++ b/Engine/Plugins/Animation/ControlRig/Source/ControlRigEditor/Private/EditMode/SControlRigOutliner.cpp
@@ -1320,7 +1320,7 @@ void SMultiRigHierarchyTreeView::RefreshTreeView()
 				if (ControlRig->IsModularRig() && Settings.bArrangeByModules)
 				{
 					TArray<FMultiRigData> ToRemove;
-					for (const TPair<FMultiRigData, TSharedPtr<FMultiRigTreeElement>> Pair : ElementMap)
+					for (const TPair<FMultiRigData, TSharedPtr<FMultiRigTreeElement>>& Pair : ElementMap)
 					{
 						if (Pair.Key.IsModule() && Pair.Value->Children.IsEmpty())
 						{
