Control Rig: Avoid switching to rig module continues after closing Change To Rig Module dialog
#rnx

diff --git a/Engine/Plugins/Animation/ControlRig/Source/ControlRigEditor/Private/Editor/ControlRigEditor.cpp b/Engine/Plugins/Animation/ControlRig/Source/ControlRigEditor/Private/Editor/ControlRigEditor.cpp
--- a/Engine/Plugins/Animation/ControlRig/Source/ControlRigEditor/Private/Editor/ControlRigEditor.cpp
+++ b/Engine/Plugins/Animation/ControlRig/Source/ControlRigEditor/Private/Editor/ControlRigEditor.cpp
@@ -862,7 +862,7 @@ void FControlRigBaseEditor::FillToolbarImpl(FToolBarBuilder& ToolbarBuilder, boo
 								{
 									static const FText Message(LOCTEXT("TurnIntoControlRigModuleMessage", "This rig requires some changes to the hierarchy to turn it into a module.\n\nWe'll try to recreate the hierarchy by relying on nodes in the construction event instead.\n\nDo you want to continue?"));
 									EAppReturnType::Type Ret = FMessageDialog::Open(EAppMsgType::YesNo, Message);
-									if(Ret == EAppReturnType::No)
+									if(Ret != EAppReturnType::Yes)
 									{
 										return;
 									}
