Control Rig: Temp Fix too many controls getting added for modular rig, just use default.
#trival

diff --git a/Engine/Plugins/Animation/ControlRig/Source/ControlRig/Private/ModularRig.cpp b/Engine/Plugins/Animation/ControlRig/Source/ControlRig/Private/ModularRig.cpp
--- a/Engine/Plugins/Animation/ControlRig/Source/ControlRig/Private/ModularRig.cpp
+++ b/Engine/Plugins/Animation/ControlRig/Source/ControlRig/Private/ModularRig.cpp
@@ -458,6 +458,8 @@ const TArray<FName>& UModularRig::GetSupportedEvents() const
 
 void UModularRig::GetControlsInOrder(TArray<FRigControlElement*>& SortedControls) const
 {
+	UControlRig::GetControlsInOrder(SortedControls);
+	/* not working for MR_BIPED 
 	SortedControls.Reset();
 
 	if(DynamicHierarchy == nullptr)
@@ -490,6 +492,7 @@ void UModularRig::GetControlsInOrder(TArray<FRigControlElement*>& SortedControls
 	{
 		SortedControls.Append(*Controls);
 	}
+	*/
 }
 
 const FModularRigSettings& UModularRig::GetModularRigSettings() const
