Control Rig Sequencer API: added the ability to set values as if the rig was interacting

diff --git a/Engine/Plugins/Animation/ControlRig/Source/ControlRigEditor/Private/ControlRigSequencerEditorLibrary.cpp b/Engine/Plugins/Animation/ControlRig/Source/ControlRigEditor/Private/ControlRigSequencerEditorLibrary.cpp
--- a/Engine/Plugins/Animation/ControlRig/Source/ControlRigEditor/Private/ControlRigSequencerEditorLibrary.cpp
+++ b/Engine/Plugins/Animation/ControlRig/Source/ControlRigEditor/Private/ControlRigSequencerEditorLibrary.cpp
@@ -1608,30 +1608,105 @@ TArray<float> UControlRigSequencerEditorLibrary::GetLocalControlRigFloats(ULevel
 	return Values;
 }
 
-void UControlRigSequencerEditorLibrary::SetLocalControlRigFloat(ULevelSequence* LevelSequence, UControlRig* ControlRig, FName ControlName, FFrameNumber Frame, float Value, EMovieSceneTimeUnit TimeUnit, bool bSetKey)
+namespace UE::Private
 {
-	if (LevelSequence == nullptr || ControlRig == nullptr)
-	{
-		return;
-	}
-	if (FRigControlElement* Element = ControlRig->FindControl(ControlName))
+	FRigControlModifiedContext GetContext(FFrameNumber InFrame, const UMovieScene* InMovieScene, const EMovieSceneTimeUnit InTimeUnit, const bool bInSetKey)
 	{
-		if (Element->Settings.ControlType != ERigControlType::Float && Element->Settings.ControlType != ERigControlType::ScaleFloat)
+		const FFrameRate TickResolution = InMovieScene->GetTickResolution();
+		if (InTimeUnit == EMovieSceneTimeUnit::DisplayRate)
 		{
-			UE_LOG(LogControlRig, Error, TEXT("Control Rig Wrong Type"));
-			return;
+			InFrame = FFrameRate::TransformTime(FFrameTime(InFrame, 0), InMovieScene->GetDisplayRate(), TickResolution).RoundToFrame();
 		}
+		
+		FRigControlModifiedContext Context;
+		Context.SetKey = bInSetKey ? EControlRigSetKey::Always : EControlRigSetKey::DoNotCare;
+		Context.LocalTime = TickResolution.AsSeconds(FFrameTime(InFrame));
+
+		return MoveTemp(Context);
 	}
-	if (UMovieScene* MovieScene = LevelSequence->GetMovieScene())
+	
+	struct FRigSetterContext
 	{
-		FFrameRate TickResolution = MovieScene->GetTickResolution();
-		if (TimeUnit == EMovieSceneTimeUnit::DisplayRate)
+	public:
+		FRigSetterContext(UControlRig* InControlRig, const FName InControlName)
+			: ControlRig(InControlRig)
+			, ControlName(InControlName)
+		{}
+
+		FRigControlElement* GetControl(const TFunctionRef<bool(const ERigControlType ControlType)>& InTypeFilter) const
 		{
-			Frame = FFrameRate::TransformTime(FFrameTime(Frame, 0), MovieScene->GetDisplayRate(), MovieScene->GetTickResolution()).RoundToFrame();
+			if (!ControlRig)
+			{
+				UE_LOG(LogControlRig, Error, TEXT("ControlRig is null."));
+				return nullptr;
+			}
+			
+			FRigControlElement* Control = ControlRig->FindControl(ControlName);
+			if (!Control)
+			{
+				UE_LOG(LogControlRig, Error, TEXT("Control '%s' not found."), *ControlName.ToString());
+				return nullptr;
+			}
+
+			if (!InTypeFilter(Control->Settings.ControlType))
+			{
+				UE_LOG(LogControlRig, Error, TEXT("Control Rig Wrong Type"));
+				return nullptr;
+			}
+
+			return Control;
 		}
-		FRigControlModifiedContext Context;
-		Context.SetKey = bSetKey ? EControlRigSetKey::Always : EControlRigSetKey::DoNotCare;
-		Context.LocalTime = TickResolution.AsSeconds(FFrameTime(Frame));
+
+		FRigControlElement* GetControl(const ERigControlType InControlType) const
+		{
+			return GetControl([InControlType](const ERigControlType ControlType)
+			{
+				return ControlType == InControlType;
+			});
+		}
+		
+	private:
+		UControlRig* ControlRig = nullptr;
+		FName ControlName = NAME_None;
+	};
+
+	static bool bIsInteracting = false;
+}
+
+void UControlRigSequencerEditorLibrary::SetInteraction(bool bIsInteracting)
+{
+	UE::Private::bIsInteracting = bIsInteracting;
+}
+
+void UControlRigSequencerEditorLibrary::SetLocalControlRigFloat(ULevelSequence* LevelSequence, UControlRig* ControlRig, FName ControlName,
+	FFrameNumber Frame, float Value, EMovieSceneTimeUnit TimeUnit, bool bSetKey)
+{
+	const UMovieScene* MovieScene = LevelSequence ? LevelSequence->GetMovieScene() : nullptr;
+	if (MovieScene == nullptr || ControlRig == nullptr)
+	{
+		return;
+	}
+
+	const UE::Private::FRigSetterContext Setter(ControlRig, ControlName);
+	const FRigControlElement* Control = Setter.GetControl([](const ERigControlType ControlType)
+	{
+		return ControlType == ERigControlType::Float || ControlType == ERigControlType::ScaleFloat;
+	});
+
+	if (!Control)
+	{
+		return;
+	}
+	
+	const FRigControlModifiedContext Context = UE::Private::GetContext(Frame, MovieScene, TimeUnit, bSetKey);
+	if (UE::Private::bIsInteracting)
+	{
+		FControlRigInteractionScope InteractionScope(ControlRig, Control->GetKey());
+		ControlRig->SetControlValue<float>(ControlName, Value, true, Context);
+		ControlRig->Evaluate_AnyThread();
+	}
+	else
+	{
 		ControlRig->SetControlValue<float>(ControlName, Value, true, Context);
 	}
 }
@@ -1639,34 +1714,39 @@ void UControlRigSequencerEditorLibrary::SetLocalControlRigFloat(ULevelSequence*
 void UControlRigSequencerEditorLibrary::SetLocalControlRigFloats(ULevelSequence* LevelSequence, UControlRig* ControlRig, FName ControlName, const TArray<FFrameNumber>& Frames, const TArray<float> Values,
 	EMovieSceneTimeUnit TimeUnit)
 {
-	if (LevelSequence == nullptr || ControlRig == nullptr || (Frames.Num() != Values.Num()))
+	const UMovieScene* MovieScene = LevelSequence ? LevelSequence->GetMovieScene() : nullptr;
+	if (MovieScene == nullptr || ControlRig == nullptr || (Frames.Num() != Values.Num()))
 	{
 		return;
 	}
-	if (FRigControlElement* Element = ControlRig->FindControl(ControlName))
+
+	const UE::Private::FRigSetterContext Setter(ControlRig, ControlName);
+	const FRigControlElement* Control = Setter.GetControl([](const ERigControlType ControlType)
 	{
-		if (Element->Settings.ControlType != ERigControlType::Float && Element->Settings.ControlType != ERigControlType::ScaleFloat)
-		{
-			UE_LOG(LogControlRig, Error, TEXT("Control Rig Wrong Type"));
-			return;
-		}
+		return ControlType == ERigControlType::Float || ControlType == ERigControlType::ScaleFloat;
+	});
+
+	if (!Control)
+	{
+		return;
 	}
-	if (UMovieScene* MovieScene = LevelSequence->GetMovieScene())
+
+	const FFrameRate TickResolution = MovieScene->GetTickResolution();
+	const FFrameRate DisplayRate = MovieScene->GetDisplayRate();
+	
+	FRigControlModifiedContext Context;
+	Context.SetKey = EControlRigSetKey::Always;
+	
+	for (int32 Index = 0; Index < Frames.Num(); ++Index)
 	{
-		FFrameRate TickResolution = MovieScene->GetTickResolution();
-		FRigControlModifiedContext Context;
-		Context.SetKey = EControlRigSetKey::Always;
-		for (int32 Index = 0; Index < Frames.Num(); ++Index)
+		FFrameNumber Frame = Frames[Index];
+		if (TimeUnit == EMovieSceneTimeUnit::DisplayRate)
 		{
-			FFrameNumber Frame = Frames[Index];
-			if (TimeUnit == EMovieSceneTimeUnit::DisplayRate)
-			{
-				Frame = FFrameRate::TransformTime(FFrameTime(Frame, 0), MovieScene->GetDisplayRate(), MovieScene->GetTickResolution()).RoundToFrame();
-			}
-			float  Value = Values[Index];
-			Context.LocalTime = TickResolution.AsSeconds(FFrameTime(Frame));
-			ControlRig->SetControlValue<float>(ControlName, Value, true, Context);
+			Frame = FFrameRate::TransformTime(FFrameTime(Frame, 0), DisplayRate, TickResolution).RoundToFrame();
 		}
+		float  Value = Values[Index];
+		Context.LocalTime = TickResolution.AsSeconds(FFrameTime(Frame));
+		ControlRig->SetControlValue<float>(ControlName, Value, true, Context);
 	}
 }
 
@@ -1724,65 +1804,69 @@ TArray<bool> UControlRigSequencerEditorLibrary::GetLocalControlRigBools(ULevelSe
 	return Values;
 }
 
-void UControlRigSequencerEditorLibrary::SetLocalControlRigBool(ULevelSequence* LevelSequence, UControlRig* ControlRig, FName ControlName, FFrameNumber Frame, bool Value, EMovieSceneTimeUnit TimeUnit, bool bSetKey)
+void UControlRigSequencerEditorLibrary::SetLocalControlRigBool(ULevelSequence* LevelSequence, UControlRig* ControlRig, FName ControlName,
+	FFrameNumber Frame, bool Value, EMovieSceneTimeUnit TimeUnit, bool bSetKey)
 {
-	if (LevelSequence == nullptr || ControlRig == nullptr)
+	const UMovieScene* MovieScene = LevelSequence ? LevelSequence->GetMovieScene() : nullptr;
+	if (!MovieScene || !ControlRig)
 	{
 		return;
 	}
-	if (UMovieScene* MovieScene = LevelSequence->GetMovieScene())
+
+	const UE::Private::FRigSetterContext Setter(ControlRig, ControlName);
+	const FRigControlElement* Control = Setter.GetControl(ERigControlType::Bool);
+	if (!Control)
+	{
+		return;
+	}
+
+	const FRigControlModifiedContext Context = UE::Private::GetContext(Frame, MovieScene, TimeUnit, bSetKey);
+	
+	if (UE::Private::bIsInteracting)
+	{
+		FControlRigInteractionScope InteractionScope(ControlRig, Control->GetKey());
+		ControlRig->SetControlValue<bool>(ControlName, Value, true, Context);
+		ControlRig->Evaluate_AnyThread();
+	}
+	else
 	{
-		if (FRigControlElement* Element = ControlRig->FindControl(ControlName))
-		{
-			if (Element->Settings.ControlType != ERigControlType::Bool)
-			{
-				UE_LOG(LogControlRig, Error, TEXT("Control Rig Wrong Type"));
-				return;
-			}
-		}
-		FFrameRate TickResolution = MovieScene->GetTickResolution();
-		if (TimeUnit == EMovieSceneTimeUnit::DisplayRate)
-		{
-			Frame = FFrameRate::TransformTime(FFrameTime(Frame, 0), MovieScene->GetDisplayRate(), MovieScene->GetTickResolution()).RoundToFrame();
-		}
-		FRigControlModifiedContext Context;
-		Context.SetKey = bSetKey ? EControlRigSetKey::Always : EControlRigSetKey::DoNotCare;
-		Context.LocalTime = TickResolution.AsSeconds(FFrameTime(Frame));
 		ControlRig->SetControlValue<bool>(ControlName, Value, true, Context);
 	}
 }
 
 void UControlRigSequencerEditorLibrary::SetLocalControlRigBools(ULevelSequence* LevelSequence, UControlRig* ControlRig, FName ControlName, const TArray<FFrameNumber>& Frames,
 	const TArray<bool> Values, EMovieSceneTimeUnit TimeUnit)
-{
-	if (LevelSequence == nullptr || ControlRig == nullptr || (Frames.Num() != Values.Num()))
+{	
+	const UMovieScene* MovieScene = LevelSequence ? LevelSequence->GetMovieScene() : nullptr;
+	if (MovieScene == nullptr || ControlRig == nullptr || (Frames.Num() != Values.Num()))
 	{
 		return;
 	}
-	if (UMovieScene* MovieScene = LevelSequence->GetMovieScene())
+
+	const UE::Private::FRigSetterContext Setter(ControlRig, ControlName);
+	const FRigControlElement* Control = Setter.GetControl(ERigControlType::Bool);
+	if (!Control)
 	{
-		if (FRigControlElement* Element = ControlRig->FindControl(ControlName))
-		{
-			if (Element->Settings.ControlType != ERigControlType::Bool)
-			{
-				UE_LOG(LogControlRig, Error, TEXT("Control Rig Wrong Type"));
-				return;
-			}
-		}
-		FFrameRate TickResolution = MovieScene->GetTickResolution();
-		FRigControlModifiedContext Context;
-		Context.SetKey = EControlRigSetKey::Always;
-		for (int32 Index = 0; Index < Frames.Num(); ++Index)
+		return;
+	}
+
+	const FFrameRate TickResolution = MovieScene->GetTickResolution();
+	const FFrameRate DisplayRate = MovieScene->GetDisplayRate();
+	
+	FRigControlModifiedContext Context;
+	Context.SetKey = EControlRigSetKey::Always;
+
+	for (int32 Index = 0; Index < Frames.Num(); ++Index)
+	{
+		FFrameNumber Frame = Frames[Index];
+		if (TimeUnit == EMovieSceneTimeUnit::DisplayRate)
 		{
-			FFrameNumber Frame = Frames[Index];
-			if (TimeUnit == EMovieSceneTimeUnit::DisplayRate)
-			{
-				Frame = FFrameRate::TransformTime(FFrameTime(Frame, 0), MovieScene->GetDisplayRate(), MovieScene->GetTickResolution()).RoundToFrame();
-			}
-			bool Value = Values[Index];
-			Context.LocalTime = TickResolution.AsSeconds(FFrameTime(Frame));
-			ControlRig->SetControlValue<bool>(ControlName, Value, true, Context);
+			Frame = FFrameRate::TransformTime(FFrameTime(Frame, 0), DisplayRate, TickResolution).RoundToFrame();
 		}
+		
+		bool Value = Values[Index];
+		Context.LocalTime = TickResolution.AsSeconds(FFrameTime(Frame));
+		ControlRig->SetControlValue<bool>(ControlName, Value, true, Context);
 	}
 }
 
@@ -1839,30 +1923,32 @@ TArray<int32> UControlRigSequencerEditorLibrary::GetLocalControlRigInts(ULevelSe
 	return Values;
 }
 
-void UControlRigSequencerEditorLibrary::SetLocalControlRigInt(ULevelSequence* LevelSequence, UControlRig* ControlRig, FName ControlName, FFrameNumber Frame, int32 Value, EMovieSceneTimeUnit TimeUnit,bool bSetKey)
+void UControlRigSequencerEditorLibrary::SetLocalControlRigInt(ULevelSequence* LevelSequence, UControlRig* ControlRig, FName ControlName,
+	FFrameNumber Frame, int32 Value, EMovieSceneTimeUnit TimeUnit,bool bSetKey)
 {
-	if (LevelSequence == nullptr || ControlRig == nullptr)
+	const UMovieScene* MovieScene = LevelSequence ? LevelSequence->GetMovieScene() : nullptr;
+	if (!MovieScene || !ControlRig)
 	{
 		return;
 	}
-	if (FRigControlElement* Element = ControlRig->FindControl(ControlName))
+
+	const UE::Private::FRigSetterContext Setter(ControlRig, ControlName);
+	const FRigControlElement* Control = Setter.GetControl(ERigControlType::Integer);
+	if (!Control)
 	{
-		if (Element->Settings.ControlType != ERigControlType::Integer)
-		{
-			UE_LOG(LogControlRig, Error, TEXT("Control Rig Wrong Type"));
-			return;
-		}
+		return;
 	}
-	if (UMovieScene* MovieScene = LevelSequence->GetMovieScene())
+
+	const FRigControlModifiedContext Context = UE::Private::GetContext(Frame, MovieScene, TimeUnit, bSetKey);
+
+	if (UE::Private::bIsInteracting)
+	{
+		FControlRigInteractionScope InteractionScope(ControlRig, Control->GetKey());
+		ControlRig->SetControlValue<int32>(ControlName, Value, true, Context);
+		ControlRig->Evaluate_AnyThread();
+	}
+	else
 	{
-		FFrameRate TickResolution = MovieScene->GetTickResolution();
-		if (TimeUnit == EMovieSceneTimeUnit::DisplayRate)
-		{
-			Frame = FFrameRate::TransformTime(FFrameTime(Frame, 0), MovieScene->GetDisplayRate(), MovieScene->GetTickResolution()).RoundToFrame();
-		}
-		FRigControlModifiedContext Context;
-		Context.SetKey = bSetKey ? EControlRigSetKey::Always : EControlRigSetKey::DoNotCare;
-		Context.LocalTime = TickResolution.AsSeconds(FFrameTime(Frame));
 		ControlRig->SetControlValue<int32>(ControlName, Value, true, Context);
 	}
 }
@@ -1870,34 +1956,35 @@ void UControlRigSequencerEditorLibrary::SetLocalControlRigInt(ULevelSequence* Le
 void UControlRigSequencerEditorLibrary::SetLocalControlRigInts(ULevelSequence* LevelSequence, UControlRig* ControlRig, FName ControlName, const TArray<FFrameNumber>& Frames,
 	const TArray<int32> Values , EMovieSceneTimeUnit TimeUnit)
 {
-	if (LevelSequence == nullptr || ControlRig == nullptr || (Frames.Num() != Values.Num()))
+	const UMovieScene* MovieScene = LevelSequence ? LevelSequence->GetMovieScene() : nullptr;
+	if (MovieScene == nullptr || ControlRig == nullptr || (Frames.Num() != Values.Num()))
 	{
 		return;
 	}
-	if (FRigControlElement* Element = ControlRig->FindControl(ControlName))
+	
+	const UE::Private::FRigSetterContext Setter(ControlRig, ControlName);
+	const FRigControlElement* Control = Setter.GetControl(ERigControlType::Integer);
+	if (!Control)
 	{
-		if (Element->Settings.ControlType != ERigControlType::Integer)
-		{
-			UE_LOG(LogControlRig, Error, TEXT("Control Rig Wrong Type"));
-			return;
-		}
+		return;
 	}
-	if (UMovieScene* MovieScene = LevelSequence->GetMovieScene())
+	
+	const FFrameRate TickResolution = MovieScene->GetTickResolution();
+	const FFrameRate DisplayRate = MovieScene->GetDisplayRate();
+	
+	FRigControlModifiedContext Context;
+	Context.SetKey = EControlRigSetKey::Always;
+	
+	for (int32 Index = 0; Index < Frames.Num(); ++Index)
 	{
-		FFrameRate TickResolution = MovieScene->GetTickResolution();
-		FRigControlModifiedContext Context;
-		Context.SetKey = EControlRigSetKey::Always;
-		for (int32 Index = 0; Index < Frames.Num(); ++Index)
+		FFrameNumber Frame = Frames[Index];
+		if (TimeUnit == EMovieSceneTimeUnit::DisplayRate)
 		{
-			FFrameNumber Frame = Frames[Index];
-			if (TimeUnit == EMovieSceneTimeUnit::DisplayRate)
-			{
-				Frame = FFrameRate::TransformTime(FFrameTime(Frame, 0), MovieScene->GetDisplayRate(), MovieScene->GetTickResolution()).RoundToFrame();
-			}
-			int32 Value = Values[Index];
-			Context.LocalTime = TickResolution.AsSeconds(FFrameTime(Frame));
-			ControlRig->SetControlValue<int32>(ControlName, Value, true, Context);
+			Frame = FFrameRate::TransformTime(FFrameTime(Frame, 0), DisplayRate, TickResolution).RoundToFrame();
 		}
+		int32 Value = Values[Index];
+		Context.LocalTime = TickResolution.AsSeconds(FFrameTime(Frame));
+		ControlRig->SetControlValue<int32>(ControlName, Value, true, Context);
 	}
 }
 
@@ -1957,30 +2044,31 @@ TArray<FVector2D> UControlRigSequencerEditorLibrary::GetLocalControlRigVector2Ds
 
 }
 
-void UControlRigSequencerEditorLibrary::SetLocalControlRigVector2D(ULevelSequence* LevelSequence, UControlRig* ControlRig, FName ControlName, FFrameNumber Frame, FVector2D Value, EMovieSceneTimeUnit TimeUnit,bool bSetKey)
+void UControlRigSequencerEditorLibrary::SetLocalControlRigVector2D(ULevelSequence* LevelSequence, UControlRig* ControlRig, FName ControlName,
+	FFrameNumber Frame, FVector2D Value, EMovieSceneTimeUnit TimeUnit,bool bSetKey)
 {
-	if (LevelSequence == nullptr || ControlRig == nullptr)
+	const UMovieScene* MovieScene = LevelSequence ? LevelSequence->GetMovieScene() : nullptr;
+	if (MovieScene == nullptr || ControlRig == nullptr)
 	{
 		return;
 	}
-	if (FRigControlElement* Element = ControlRig->FindControl(ControlName))
+
+	const UE::Private::FRigSetterContext Setter(ControlRig, ControlName);
+	const FRigControlElement* Control = Setter.GetControl(ERigControlType::Vector2D);
+	if (!Control)
 	{
-		if (Element->Settings.ControlType != ERigControlType::Vector2D)
-		{
-			UE_LOG(LogControlRig, Error, TEXT("Control Rig Wrong Type"));
-			return;
-		}
+		return;
 	}
-	if (UMovieScene* MovieScene = LevelSequence->GetMovieScene())
+
+	const FRigControlModifiedContext Context = UE::Private::GetContext(Frame, MovieScene, TimeUnit, bSetKey);
+	if (UE::Private::bIsInteracting)
+	{
+		FControlRigInteractionScope InteractionScope(ControlRig, Control->GetKey());
+		ControlRig->SetControlValue<FVector3f>(ControlName, FVector3f(Value.X, Value.Y, 0.f), true, Context);
+		ControlRig->Evaluate_AnyThread();
+	}
+	else
 	{
-		FFrameRate TickResolution = MovieScene->GetTickResolution();
-		if (TimeUnit == EMovieSceneTimeUnit::DisplayRate)
-		{
-			Frame = FFrameRate::TransformTime(FFrameTime(Frame, 0), MovieScene->GetDisplayRate(), MovieScene->GetTickResolution()).RoundToFrame();
-		}
-		FRigControlModifiedContext Context;
-		Context.SetKey = bSetKey ? EControlRigSetKey::Always : EControlRigSetKey::DoNotCare;
-		Context.LocalTime = TickResolution.AsSeconds(FFrameTime(Frame));
 		ControlRig->SetControlValue<FVector3f>(ControlName, FVector3f(Value.X, Value.Y, 0.f), true, Context);
 	}
 }
@@ -1988,34 +2076,34 @@ void UControlRigSequencerEditorLibrary::SetLocalControlRigVector2D(ULevelSequenc
 void UControlRigSequencerEditorLibrary::SetLocalControlRigVector2Ds(ULevelSequence* LevelSequence, UControlRig* ControlRig, FName ControlName, const TArray<FFrameNumber>& Frames, 
 	const TArray<FVector2D> Values, EMovieSceneTimeUnit TimeUnit)
 {
-	if (LevelSequence == nullptr || ControlRig == nullptr || (Frames.Num() != Values.Num()))
+	const UMovieScene* MovieScene = LevelSequence ? LevelSequence->GetMovieScene() : nullptr;
+	if (MovieScene == nullptr || ControlRig == nullptr || (Frames.Num() != Values.Num()))
 	{
 		return;
 	}
-	if (FRigControlElement* Element = ControlRig->FindControl(ControlName))
+
+	const UE::Private::FRigSetterContext Setter(ControlRig, ControlName);
+	const FRigControlElement* Control = Setter.GetControl(ERigControlType::Vector2D);
+	if (!Control)
 	{
-		if (Element->Settings.ControlType != ERigControlType::Vector2D)
-		{
-			UE_LOG(LogControlRig, Error, TEXT("Control Rig Wrong Type"));
-			return;
-		}
+		return;
 	}
-	if (UMovieScene* MovieScene = LevelSequence->GetMovieScene())
+	
+	const FFrameRate TickResolution = MovieScene->GetTickResolution();
+	const FFrameRate DisplayRate = MovieScene->GetDisplayRate();
+	
+	FRigControlModifiedContext Context;
+	Context.SetKey = EControlRigSetKey::Always;
+	for (int32 Index = 0; Index < Frames.Num(); ++Index)
 	{
-		FFrameRate TickResolution = MovieScene->GetTickResolution();
-		FRigControlModifiedContext Context;
-		Context.SetKey = EControlRigSetKey::Always;
-		for (int32 Index = 0; Index < Frames.Num(); ++Index)
+		FFrameNumber Frame = Frames[Index];
+		if (TimeUnit == EMovieSceneTimeUnit::DisplayRate)
 		{
-			FFrameNumber Frame = Frames[Index];
-			if (TimeUnit == EMovieSceneTimeUnit::DisplayRate)
-			{
-				Frame = FFrameRate::TransformTime(FFrameTime(Frame, 0), MovieScene->GetDisplayRate(), MovieScene->GetTickResolution()).RoundToFrame();
-			}
-			FVector2D Value = Values[Index];
-			Context.LocalTime = TickResolution.AsSeconds(FFrameTime(Frame));
-			ControlRig->SetControlValue<FVector3f>(ControlName, FVector3f(Value.X, Value.Y, 0.f), true, Context);
+			Frame = FFrameRate::TransformTime(FFrameTime(Frame, 0), DisplayRate, TickResolution).RoundToFrame();
 		}
+		FVector2D Value = Values[Index];
+		Context.LocalTime = TickResolution.AsSeconds(FFrameTime(Frame));
+		ControlRig->SetControlValue<FVector3f>(ControlName, FVector3f(Value.X, Value.Y, 0.f), true, Context);
 	}
 }
 
@@ -2073,30 +2161,32 @@ TArray<FVector> UControlRigSequencerEditorLibrary::GetLocalControlRigPositions(U
 	return Values;
 }
 
-void UControlRigSequencerEditorLibrary::SetLocalControlRigPosition(ULevelSequence* LevelSequence, UControlRig* ControlRig, FName ControlName, FFrameNumber Frame, FVector Value, EMovieSceneTimeUnit TimeUnit,bool bSetKey)
+void UControlRigSequencerEditorLibrary::SetLocalControlRigPosition(ULevelSequence* LevelSequence, UControlRig* ControlRig, FName ControlName,
+	FFrameNumber Frame, FVector Value, EMovieSceneTimeUnit TimeUnit,bool bSetKey)
 {
-	if (LevelSequence == nullptr || ControlRig == nullptr)
+	const UMovieScene* MovieScene = LevelSequence ? LevelSequence->GetMovieScene() : nullptr;
+	if (MovieScene == nullptr || ControlRig == nullptr)
 	{
 		return;
 	}
-	if (FRigControlElement* Element = ControlRig->FindControl(ControlName))
+
+	const UE::Private::FRigSetterContext Setter(ControlRig, ControlName);
+	const FRigControlElement* Control = Setter.GetControl(ERigControlType::Position);
+	if (!Control)
 	{
-		if (Element->Settings.ControlType != ERigControlType::Position)
-		{
-			UE_LOG(LogControlRig, Error, TEXT("Control Rig Wrong Type"));
-			return;
-		}
+		return;
 	}
-	if (UMovieScene* MovieScene = LevelSequence->GetMovieScene())
+	
+	const FRigControlModifiedContext Context = UE::Private::GetContext(Frame, MovieScene, TimeUnit, bSetKey);
+
+	if (UE::Private::bIsInteracting)
+	{
+		FControlRigInteractionScope InteractionScope(ControlRig, Control->GetKey());
+		ControlRig->SetControlValue<FVector3f>(ControlName, (FVector3f)Value, true, Context);
+		ControlRig->Evaluate_AnyThread();
+	}
+	else
 	{
-		FFrameRate TickResolution = MovieScene->GetTickResolution();
-		if (TimeUnit == EMovieSceneTimeUnit::DisplayRate)
-		{
-			Frame = FFrameRate::TransformTime(FFrameTime(Frame, 0), MovieScene->GetDisplayRate(), MovieScene->GetTickResolution()).RoundToFrame();
-		}
-		FRigControlModifiedContext Context;
-		Context.SetKey = bSetKey ? EControlRigSetKey::Always : EControlRigSetKey::DoNotCare;
-		Context.LocalTime = TickResolution.AsSeconds(FFrameTime(Frame));
 		ControlRig->SetControlValue<FVector3f>(ControlName, (FVector3f)Value, true, Context);
 	}
 }
@@ -2104,34 +2194,34 @@ void UControlRigSequencerEditorLibrary::SetLocalControlRigPosition(ULevelSequenc
 void UControlRigSequencerEditorLibrary::SetLocalControlRigPositions(ULevelSequence* LevelSequence, UControlRig* ControlRig, FName ControlName, const TArray<FFrameNumber>& Frames, 
 	const TArray<FVector> Values, EMovieSceneTimeUnit TimeUnit)
 {
-	if (LevelSequence == nullptr || ControlRig == nullptr || (Frames.Num() != Values.Num()))
+	const UMovieScene* MovieScene = LevelSequence ? LevelSequence->GetMovieScene() : nullptr;
+	if (MovieScene == nullptr || ControlRig == nullptr || (Frames.Num() != Values.Num()))
 	{
 		return;
 	}
-	if (FRigControlElement* Element = ControlRig->FindControl(ControlName))
+	
+	const UE::Private::FRigSetterContext Setter(ControlRig, ControlName);
+	const FRigControlElement* Control = Setter.GetControl(ERigControlType::Position);
+	if (!Control)
 	{
-		if (Element->Settings.ControlType != ERigControlType::Position)
-		{
-			UE_LOG(LogControlRig, Error, TEXT("Control Rig Wrong Type"));
-			return;
-		}
+		return;
 	}
-	if (UMovieScene* MovieScene = LevelSequence->GetMovieScene())
+
+	const FFrameRate TickResolution = MovieScene->GetTickResolution();
+	const FFrameRate DisplayRate = MovieScene->GetDisplayRate();
+	
+	FRigControlModifiedContext Context;
+	Context.SetKey = EControlRigSetKey::Always;
+	for (int32 Index = 0; Index < Frames.Num(); ++Index)
 	{
-		FFrameRate TickResolution = MovieScene->GetTickResolution();
-		FRigControlModifiedContext Context;
-		Context.SetKey = EControlRigSetKey::Always;
-		for (int32 Index = 0; Index < Frames.Num(); ++Index)
+		FFrameNumber Frame = Frames[Index];
+		if (TimeUnit == EMovieSceneTimeUnit::DisplayRate)
 		{
-			FFrameNumber Frame = Frames[Index];
-			if (TimeUnit == EMovieSceneTimeUnit::DisplayRate)
-			{
-				Frame = FFrameRate::TransformTime(FFrameTime(Frame, 0), MovieScene->GetDisplayRate(), MovieScene->GetTickResolution()).RoundToFrame();
-			}
-			FVector Value = Values[Index];
-			Context.LocalTime = TickResolution.AsSeconds(FFrameTime(Frame));
-			ControlRig->SetControlValue<FVector3f>(ControlName, (FVector3f)Value, true, Context);
+			Frame = FFrameRate::TransformTime(FFrameTime(Frame, 0), DisplayRate, TickResolution).RoundToFrame();
 		}
+		FVector Value = Values[Index];
+		Context.LocalTime = TickResolution.AsSeconds(FFrameTime(Frame));
+		ControlRig->SetControlValue<FVector3f>(ControlName, (FVector3f)Value, true, Context);
 	}
 }
 
@@ -2189,36 +2279,35 @@ TArray<FRotator> UControlRigSequencerEditorLibrary::GetLocalControlRigRotators(U
 	return Values;
 }
 
-void UControlRigSequencerEditorLibrary::SetLocalControlRigRotator(ULevelSequence* LevelSequence, UControlRig* ControlRig, FName ControlName, FFrameNumber Frame, FRotator Value, EMovieSceneTimeUnit TimeUnit, bool bSetKey)
+void UControlRigSequencerEditorLibrary::SetLocalControlRigRotator(ULevelSequence* LevelSequence, UControlRig* ControlRig, FName ControlName,
+	FFrameNumber Frame, FRotator Value, EMovieSceneTimeUnit TimeUnit, bool bSetKey)
 {
-	if (LevelSequence == nullptr || ControlRig == nullptr)
+	const UMovieScene* MovieScene = LevelSequence ? LevelSequence->GetMovieScene() : nullptr;
+	if (MovieScene == nullptr || ControlRig == nullptr)
 	{
 		return;
 	}
-	FRigControlElement* Element = ControlRig->FindControl(ControlName);
-	if (!Element)
+
+	const UE::Private::FRigSetterContext Setter(ControlRig, ControlName);
+	FRigControlElement* Control = Setter.GetControl(ERigControlType::Rotator);
+	if (!Control)
 	{
-		UE_LOG(LogControlRig, Error, TEXT("Control name not found."));
 		return;
 	}
-	if (Element->Settings.ControlType != ERigControlType::Rotator)
+
+	const FVector EulerAngle(Value.Roll, Value.Pitch, Value.Yaw);
+	ControlRig->GetHierarchy()->SetControlSpecifiedEulerAngle(Control, EulerAngle);
+
+	const FRigControlModifiedContext Context = UE::Private::GetContext(Frame, MovieScene, TimeUnit, bSetKey);
+
+	if (UE::Private::bIsInteracting)
 	{
-		UE_LOG(LogControlRig, Error, TEXT("Control Rig Wrong Type"));
-		return;
+		FControlRigInteractionScope InteractionScope(ControlRig, Control->GetKey());
+		ControlRig->SetControlValue<FVector3f>(ControlName, (FVector3f)Value.Euler(), true, Context);
+		ControlRig->Evaluate_AnyThread();
 	}
-	if (UMovieScene* MovieScene = LevelSequence->GetMovieScene())
+	else
 	{
-		FFrameRate TickResolution = MovieScene->GetTickResolution();
-		if (TimeUnit == EMovieSceneTimeUnit::DisplayRate)
-		{
-			Frame = FFrameRate::TransformTime(FFrameTime(Frame, 0), MovieScene->GetDisplayRate(), MovieScene->GetTickResolution()).RoundToFrame();
-		}
-		FRigControlModifiedContext Context;
-		Context.SetKey = bSetKey ? EControlRigSetKey::Always : EControlRigSetKey::DoNotCare;
-		Context.LocalTime = TickResolution.AsSeconds(FFrameTime(Frame));
-
-		const FVector EulerAngle(Value.Roll, Value.Pitch, Value.Yaw);
-		ControlRig->GetHierarchy()->SetControlSpecifiedEulerAngle(Element, EulerAngle);
 		ControlRig->SetControlValue<FVector3f>(ControlName, (FVector3f)Value.Euler(), true, Context);
 	}
 }
@@ -2226,39 +2315,37 @@ void UControlRigSequencerEditorLibrary::SetLocalControlRigRotator(ULevelSequence
 void UControlRigSequencerEditorLibrary::SetLocalControlRigRotators(ULevelSequence* LevelSequence, UControlRig* ControlRig, FName ControlName, const TArray<FFrameNumber>& Frames,
 	const TArray<FRotator> Values, EMovieSceneTimeUnit TimeUnit)
 {
-	if (LevelSequence == nullptr || ControlRig == nullptr || (Frames.Num() != Values.Num()))
-	{
-		return;
-	}
-	FRigControlElement* Element = ControlRig->FindControl(ControlName);
-	if (!Element)
+	const UMovieScene* MovieScene = LevelSequence ? LevelSequence->GetMovieScene() : nullptr;
+	if (MovieScene == nullptr || ControlRig == nullptr || (Frames.Num() != Values.Num()))
 	{
-		UE_LOG(LogControlRig, Error, TEXT("Control name not found."));
 		return;
 	}
-	if (Element->Settings.ControlType != ERigControlType::Rotator)
+
+	const UE::Private::FRigSetterContext Setter(ControlRig, ControlName);
+	FRigControlElement* Control = Setter.GetControl(ERigControlType::Rotator);
+	if (!Control)
 	{
-		UE_LOG(LogControlRig, Error, TEXT("Control Rig Wrong Type"));
 		return;
 	}
-	if (UMovieScene* MovieScene = LevelSequence->GetMovieScene())
+
+	const FFrameRate TickResolution = MovieScene->GetTickResolution();
+	const FFrameRate DisplayRate = MovieScene->GetDisplayRate();
+	
+	FRigControlModifiedContext Context;
+	Context.SetKey = EControlRigSetKey::Always;
+	for (int32 Index = 0; Index < Frames.Num(); ++Index)
 	{
-		FFrameRate TickResolution = MovieScene->GetTickResolution();
-		FRigControlModifiedContext Context;
-		Context.SetKey = EControlRigSetKey::Always;
-		for (int32 Index = 0; Index < Frames.Num(); ++Index)
+		FFrameNumber Frame = Frames[Index];
+		if (TimeUnit == EMovieSceneTimeUnit::DisplayRate)
 		{
-			FFrameNumber Frame = Frames[Index];
-			if (TimeUnit == EMovieSceneTimeUnit::DisplayRate)
-			{
-				Frame = FFrameRate::TransformTime(FFrameTime(Frame, 0), MovieScene->GetDisplayRate(), MovieScene->GetTickResolution()).RoundToFrame();
-			}
-			const FRotator Value = Values[Index];
-			Context.LocalTime = TickResolution.AsSeconds(FFrameTime(Frame));
-			const FVector EulerAngle(Value.Roll, Value.Pitch, Value.Yaw);
-			ControlRig->GetHierarchy()->SetControlSpecifiedEulerAngle(Element, EulerAngle);
-			ControlRig->SetControlValue<FVector3f>(ControlName, (FVector3f)Value.Euler(), true, Context);
+			Frame = FFrameRate::TransformTime(FFrameTime(Frame, 0), DisplayRate, TickResolution).RoundToFrame();
 		}
+		Context.LocalTime = TickResolution.AsSeconds(FFrameTime(Frame));
+
+		const FRotator Value = Values[Index];
+		const FVector EulerAngle(Value.Roll, Value.Pitch, Value.Yaw);
+		ControlRig->GetHierarchy()->SetControlSpecifiedEulerAngle(Control, EulerAngle);
+		ControlRig->SetControlValue<FVector3f>(ControlName, (FVector3f)Value.Euler(), true, Context);
 	}
 }
 
@@ -2316,30 +2403,32 @@ TArray<FVector>UControlRigSequencerEditorLibrary::GetLocalControlRigScales(ULeve
 	return Values;
 }
 
-void UControlRigSequencerEditorLibrary::SetLocalControlRigScale(ULevelSequence* LevelSequence, UControlRig* ControlRig, FName ControlName, FFrameNumber Frame, FVector Value, EMovieSceneTimeUnit TimeUnit, bool bSetKey)
+void UControlRigSequencerEditorLibrary::SetLocalControlRigScale(ULevelSequence* LevelSequence, UControlRig* ControlRig, FName ControlName,
+	FFrameNumber Frame, FVector Value, EMovieSceneTimeUnit TimeUnit, bool bSetKey)
 {
-	if (LevelSequence == nullptr || ControlRig == nullptr)
+	const UMovieScene* MovieScene = LevelSequence ? LevelSequence->GetMovieScene() : nullptr;
+	if (MovieScene == nullptr || ControlRig == nullptr)
 	{
 		return;
 	}
-	if (FRigControlElement* Element = ControlRig->FindControl(ControlName))
+
+	const UE::Private::FRigSetterContext Setter(ControlRig, ControlName);
+	const FRigControlElement* Control = Setter.GetControl(ERigControlType::Scale);
+	if (!Control)
 	{
-		if (Element->Settings.ControlType != ERigControlType::Scale)
-		{
-			UE_LOG(LogControlRig, Error, TEXT("Control Rig Wrong Type"));
-			return;
-		}
+		return;
 	}
-	if (UMovieScene* MovieScene = LevelSequence->GetMovieScene())
+	
+	const FRigControlModifiedContext Context = UE::Private::GetContext(Frame, MovieScene, TimeUnit, bSetKey);
+	
+	if (UE::Private::bIsInteracting)
+	{
+		FControlRigInteractionScope InteractionScope(ControlRig, Control->GetKey());
+		ControlRig->SetControlValue<FVector3f>(ControlName, (FVector3f)Value, true, Context);
+		ControlRig->Evaluate_AnyThread();
+	}
+	else
 	{
-		FFrameRate TickResolution = MovieScene->GetTickResolution();
-		if (TimeUnit == EMovieSceneTimeUnit::DisplayRate)
-		{
-			Frame = FFrameRate::TransformTime(FFrameTime(Frame, 0), MovieScene->GetDisplayRate(), MovieScene->GetTickResolution()).RoundToFrame();
-		}
-		FRigControlModifiedContext Context;
-		Context.SetKey = bSetKey ? EControlRigSetKey::Always : EControlRigSetKey::DoNotCare;
-		Context.LocalTime = TickResolution.AsSeconds(FFrameTime(Frame));
 		ControlRig->SetControlValue<FVector3f>(ControlName, (FVector3f)Value, true, Context);
 	}
 }
@@ -2347,34 +2436,34 @@ void UControlRigSequencerEditorLibrary::SetLocalControlRigScale(ULevelSequence*
 void UControlRigSequencerEditorLibrary::SetLocalControlRigScales(ULevelSequence* LevelSequence, UControlRig* ControlRig, FName ControlName, const TArray<FFrameNumber>& Frames, 
 	const TArray<FVector> Values, EMovieSceneTimeUnit TimeUnit)
 {
-	if (LevelSequence == nullptr || ControlRig == nullptr || (Frames.Num() != Values.Num()))
+	const UMovieScene* MovieScene = LevelSequence ? LevelSequence->GetMovieScene() : nullptr;
+	if (MovieScene == nullptr || ControlRig == nullptr || (Frames.Num() != Values.Num()))
 	{
 		return;
 	}
-	if (FRigControlElement* Element = ControlRig->FindControl(ControlName))
+	
+	const UE::Private::FRigSetterContext Setter(ControlRig, ControlName);
+	const FRigControlElement* Control = Setter.GetControl(ERigControlType::Scale);
+	if (!Control)
 	{
-		if (Element->Settings.ControlType != ERigControlType::Scale)
-		{
-			UE_LOG(LogControlRig, Error, TEXT("Control Rig Wrong Type"));
-			return;
-		}
+		return;
 	}
-	if (UMovieScene* MovieScene = LevelSequence->GetMovieScene())
+
+	const FFrameRate TickResolution = MovieScene->GetTickResolution();
+	const FFrameRate DisplayRate = MovieScene->GetDisplayRate();
+	
+	FRigControlModifiedContext Context;
+	Context.SetKey = EControlRigSetKey::Always;
+	for (int32 Index = 0; Index < Frames.Num(); ++Index)
 	{
-		FFrameRate TickResolution = MovieScene->GetTickResolution();
-		FRigControlModifiedContext Context;
-		Context.SetKey = EControlRigSetKey::Always;
-		for (int32 Index = 0; Index < Frames.Num(); ++Index)
+		FFrameNumber Frame = Frames[Index];
+		if (TimeUnit == EMovieSceneTimeUnit::DisplayRate)
 		{
-			FFrameNumber Frame = Frames[Index];
-			if (TimeUnit == EMovieSceneTimeUnit::DisplayRate)
-			{
-				Frame = FFrameRate::TransformTime(FFrameTime(Frame, 0), MovieScene->GetDisplayRate(), MovieScene->GetTickResolution()).RoundToFrame();
-			}
-			FVector Value = Values[Index];
-			Context.LocalTime = TickResolution.AsSeconds(FFrameTime(Frame));
-			ControlRig->SetControlValue<FVector3f>(ControlName, (FVector3f)Value, true, Context);
+			Frame = FFrameRate::TransformTime(FFrameTime(Frame, 0), DisplayRate, TickResolution).RoundToFrame();
 		}
+		FVector Value = Values[Index];
+		Context.LocalTime = TickResolution.AsSeconds(FFrameTime(Frame));
+		ControlRig->SetControlValue<FVector3f>(ControlName, (FVector3f)Value, true, Context);
 	}
 }
 
@@ -2435,33 +2524,31 @@ TArray<FEulerTransform> UControlRigSequencerEditorLibrary::GetLocalControlRigEul
 void UControlRigSequencerEditorLibrary::SetLocalControlRigEulerTransform(ULevelSequence* LevelSequence, UControlRig* ControlRig, FName ControlName, FFrameNumber Frame, FEulerTransform Value,
 	EMovieSceneTimeUnit TimeUnit, bool bSetKey)
 {
-	if (LevelSequence == nullptr || ControlRig == nullptr)
+	const UMovieScene* MovieScene = LevelSequence ? LevelSequence->GetMovieScene() : nullptr;
+	if (MovieScene == nullptr || ControlRig == nullptr)
 	{
 		return;
 	}
-	if (FRigControlElement* Element = ControlRig->FindControl(ControlName))
+	
+	const UE::Private::FRigSetterContext Setter(ControlRig, ControlName);
+	FRigControlElement* Control = Setter.GetControl(ERigControlType::EulerTransform);
+	if (!Control)
 	{
-		if (Element->Settings.ControlType != ERigControlType::EulerTransform)
-		{
-			UE_LOG(LogControlRig, Error, TEXT("Control Rig Wrong Type"));
-			return;
-		}
+		return;
 	}
-	if (UMovieScene* MovieScene = LevelSequence->GetMovieScene())
+
+	const FVector EulerAngle(Value.Rotation.Roll, Value.Rotation.Pitch, Value.Rotation.Yaw);
+	ControlRig->GetHierarchy()->SetControlSpecifiedEulerAngle(Control, EulerAngle);
+
+	const FRigControlModifiedContext Context = UE::Private::GetContext(Frame, MovieScene, TimeUnit, bSetKey);
+	if (UE::Private::bIsInteracting)
+	{
+		FControlRigInteractionScope InteractionScope(ControlRig, Control->GetKey());
+		ControlRig->SetControlValue<FRigControlValue::FEulerTransform_Float>(ControlName, Value, true, Context);
+		ControlRig->Evaluate_AnyThread();
+	}
+	else
 	{
-		FFrameRate TickResolution = MovieScene->GetTickResolution();
-		if (TimeUnit == EMovieSceneTimeUnit::DisplayRate)
-		{
-			Frame = FFrameRate::TransformTime(FFrameTime(Frame, 0), MovieScene->GetDisplayRate(), MovieScene->GetTickResolution()).RoundToFrame();
-		}
-		FRigControlModifiedContext Context;
-		Context.SetKey = bSetKey ? EControlRigSetKey::Always : EControlRigSetKey::DoNotCare;
-		Context.LocalTime = TickResolution.AsSeconds(FFrameTime(Frame));
-		if (FRigControlElement* ControlElement = ControlRig->FindControl(ControlName))
-		{
-			FVector EulerAngle(Value.Rotation.Roll, Value.Rotation.Pitch, Value.Rotation.Yaw);
-			ControlRig->GetHierarchy()->SetControlSpecifiedEulerAngle(ControlElement, EulerAngle);
-		}
 		ControlRig->SetControlValue<FRigControlValue::FEulerTransform_Float>(ControlName, Value, true, Context);
 	}
 }
@@ -2469,40 +2556,38 @@ void UControlRigSequencerEditorLibrary::SetLocalControlRigEulerTransform(ULevelS
 void UControlRigSequencerEditorLibrary::SetLocalControlRigEulerTransforms(ULevelSequence* LevelSequence, UControlRig* ControlRig, FName ControlName, const TArray<FFrameNumber>& Frames,
 	const TArray<FEulerTransform> Values, EMovieSceneTimeUnit TimeUnit)
 {
-	if (LevelSequence == nullptr || ControlRig == nullptr || (Frames.Num() != Values.Num()))
+	const UMovieScene* MovieScene = LevelSequence ? LevelSequence->GetMovieScene() : nullptr;
+	if (MovieScene == nullptr || ControlRig == nullptr || (Frames.Num() != Values.Num()))
 	{
 		return;
 	}
-	if (FRigControlElement* Element = ControlRig->FindControl(ControlName))
+
+	const UE::Private::FRigSetterContext Setter(ControlRig, ControlName);
+	FRigControlElement* Control = Setter.GetControl(ERigControlType::EulerTransform);
+	if (!Control)
 	{
-		if (Element->Settings.ControlType != ERigControlType::EulerTransform)
-		{
-			UE_LOG(LogControlRig, Error, TEXT("Control Rig Wrong Type"));
-			return;
-		}
+		return;
 	}
-	if (UMovieScene* MovieScene = LevelSequence->GetMovieScene())
+	
+	URigHierarchy* Hierarchy = ControlRig->GetHierarchy();
+	const FFrameRate TickResolution = MovieScene->GetTickResolution();
+	const FFrameRate DisplayRate = MovieScene->GetDisplayRate();
+	
+	FRigControlModifiedContext Context;
+	Context.SetKey = EControlRigSetKey::Always;
+	for (int32 Index = 0; Index < Frames.Num(); ++Index)
 	{
-		FFrameRate TickResolution = MovieScene->GetTickResolution();
-		FRigControlModifiedContext Context;
-		Context.SetKey = EControlRigSetKey::Always;
-		for (int32 Index = 0; Index < Frames.Num(); ++Index)
+		FFrameNumber Frame = Frames[Index];
+		if (TimeUnit == EMovieSceneTimeUnit::DisplayRate)
 		{
-			FFrameNumber Frame = Frames[Index];
-			if (TimeUnit == EMovieSceneTimeUnit::DisplayRate)
-			{
-				Frame = FFrameRate::TransformTime(FFrameTime(Frame, 0), MovieScene->GetDisplayRate(), MovieScene->GetTickResolution()).RoundToFrame();
-			}
-			FEulerTransform Value = Values[Index];
-			Context.LocalTime = TickResolution.AsSeconds(FFrameTime(Frame));
-
-			if (FRigControlElement* ControlElement = ControlRig->FindControl(ControlName))
-			{
-				FVector EulerAngle(Value.Rotation.Roll, Value.Rotation.Pitch, Value.Rotation.Yaw);
-				ControlRig->GetHierarchy()->SetControlSpecifiedEulerAngle(ControlElement, EulerAngle);
-			}
-			ControlRig->SetControlValue<FRigControlValue::FEulerTransform_Float>(ControlName, Value, true, Context);
+			Frame = FFrameRate::TransformTime(FFrameTime(Frame, 0), DisplayRate, TickResolution).RoundToFrame();
 		}
+		FEulerTransform Value = Values[Index];
+		Context.LocalTime = TickResolution.AsSeconds(FFrameTime(Frame));
+
+		const FVector EulerAngle(Value.Rotation.Roll, Value.Rotation.Pitch, Value.Rotation.Yaw);
+		Hierarchy->SetControlSpecifiedEulerAngle(Control, EulerAngle);
+		ControlRig->SetControlValue<FRigControlValue::FEulerTransform_Float>(ControlName, Value, true, Context);
 	}
 }
 
@@ -2560,31 +2645,31 @@ TArray<FTransformNoScale> UControlRigSequencerEditorLibrary::GetLocalControlRigT
 	return Values;
 }
 
-void UControlRigSequencerEditorLibrary::SetLocalControlRigTransformNoScale(ULevelSequence* LevelSequence, UControlRig* ControlRig, FName ControlName, FFrameNumber Frame, FTransformNoScale Value,
-	EMovieSceneTimeUnit TimeUnit,bool bSetKey)
+void UControlRigSequencerEditorLibrary::SetLocalControlRigTransformNoScale(ULevelSequence* LevelSequence, UControlRig* ControlRig, FName ControlName,
+	FFrameNumber Frame, FTransformNoScale Value, EMovieSceneTimeUnit TimeUnit,bool bSetKey)
 {
-	if (LevelSequence == nullptr || ControlRig == nullptr)
+	const UMovieScene* MovieScene = LevelSequence ? LevelSequence->GetMovieScene() : nullptr;
+	if (MovieScene == nullptr || ControlRig == nullptr)
 	{
 		return;
 	}
-	if (FRigControlElement* Element = ControlRig->FindControl(ControlName))
+	
+	const UE::Private::FRigSetterContext Setter(ControlRig, ControlName);
+	const FRigControlElement* Control = Setter.GetControl(ERigControlType::TransformNoScale);
+	if (!Control)
 	{
-		if (Element->Settings.ControlType != ERigControlType::TransformNoScale)
-		{
-			UE_LOG(LogControlRig, Error, TEXT("Control Rig Wrong Type"));
-			return;
-		}
+		return;
 	}
-	if (UMovieScene* MovieScene = LevelSequence->GetMovieScene())
+
+	const FRigControlModifiedContext Context = UE::Private::GetContext(Frame, MovieScene, TimeUnit, bSetKey);
+	if (UE::Private::bIsInteracting)
+	{
+		FControlRigInteractionScope InteractionScope(ControlRig, Control->GetKey());
+		ControlRig->SetControlValue<FRigControlValue::FTransformNoScale_Float>(ControlName, Value, true, Context);
+		ControlRig->Evaluate_AnyThread();
+	}
+	else
 	{
-		FFrameRate TickResolution = MovieScene->GetTickResolution();
-		if (TimeUnit == EMovieSceneTimeUnit::DisplayRate)
-		{
-			Frame = FFrameRate::TransformTime(FFrameTime(Frame, 0), MovieScene->GetDisplayRate(), MovieScene->GetTickResolution()).RoundToFrame();
-		}
-		FRigControlModifiedContext Context;
-		Context.SetKey = bSetKey ? EControlRigSetKey::Always : EControlRigSetKey::DoNotCare;
-		Context.LocalTime = TickResolution.AsSeconds(FFrameTime(Frame));
 		ControlRig->SetControlValue<FRigControlValue::FTransformNoScale_Float>(ControlName, Value, true, Context);
 	}
 }
@@ -2592,34 +2677,34 @@ void UControlRigSequencerEditorLibrary::SetLocalControlRigTransformNoScale(ULeve
 void UControlRigSequencerEditorLibrary::SetLocalControlRigTransformNoScales(ULevelSequence* LevelSequence, UControlRig* ControlRig, FName ControlName, const TArray<FFrameNumber>& Frames, 
 	const TArray<FTransformNoScale> Values, EMovieSceneTimeUnit TimeUnit)
 {
-	if (LevelSequence == nullptr || ControlRig == nullptr || (Frames.Num() != Values.Num()))
+	const UMovieScene* MovieScene = LevelSequence ? LevelSequence->GetMovieScene() : nullptr;
+	if (MovieScene == nullptr || ControlRig == nullptr || (Frames.Num() != Values.Num()))
 	{
 		return;
 	}
-	if (FRigControlElement* Element = ControlRig->FindControl(ControlName))
+
+	const UE::Private::FRigSetterContext Setter(ControlRig, ControlName);
+	const FRigControlElement* Control = Setter.GetControl(ERigControlType::TransformNoScale);
+	if (!Control)
 	{
-		if (Element->Settings.ControlType != ERigControlType::TransformNoScale)
-		{
-			UE_LOG(LogControlRig, Error, TEXT("Control Rig Wrong Type"));
-			return;
-		}
+		return;
 	}
-	if (UMovieScene* MovieScene = LevelSequence->GetMovieScene())
+
+	const FFrameRate TickResolution = MovieScene->GetTickResolution();
+	const FFrameRate DisplayRate = MovieScene->GetDisplayRate();
+	
+	FRigControlModifiedContext Context;
+	Context.SetKey = EControlRigSetKey::Always;
+	for (int32 Index = 0; Index < Frames.Num(); ++Index)
 	{
-		FFrameRate TickResolution = MovieScene->GetTickResolution();
-		FRigControlModifiedContext Context;
-		Context.SetKey = EControlRigSetKey::Always;
-		for (int32 Index = 0; Index < Frames.Num(); ++Index)
+		FFrameNumber Frame = Frames[Index];
+		if (TimeUnit == EMovieSceneTimeUnit::DisplayRate)
 		{
-			FFrameNumber Frame = Frames[Index];
-			if (TimeUnit == EMovieSceneTimeUnit::DisplayRate)
-			{
-				Frame = FFrameRate::TransformTime(FFrameTime(Frame, 0), MovieScene->GetDisplayRate(), MovieScene->GetTickResolution()).RoundToFrame();
-			}
-			FTransformNoScale Value = Values[Index];
-			Context.LocalTime = TickResolution.AsSeconds(FFrameTime(Frame));
-			ControlRig->SetControlValue<FRigControlValue::FTransformNoScale_Float>(ControlName, Value, true, Context);
+			Frame = FFrameRate::TransformTime(FFrameTime(Frame, 0), DisplayRate, TickResolution).RoundToFrame();
 		}
+		FTransformNoScale Value = Values[Index];
+		Context.LocalTime = TickResolution.AsSeconds(FFrameTime(Frame));
+		ControlRig->SetControlValue<FRigControlValue::FTransformNoScale_Float>(ControlName, Value, true, Context);
 	}
 }
 
@@ -2677,31 +2762,32 @@ TArray<FTransform> UControlRigSequencerEditorLibrary::GetLocalControlRigTransfor
 	return Values;
 }
 
-void UControlRigSequencerEditorLibrary::SetLocalControlRigTransform(ULevelSequence* LevelSequence, UControlRig* ControlRig, FName ControlName, FFrameNumber Frame, FTransform Value,
-	EMovieSceneTimeUnit TimeUnit,bool bSetKey)
+void UControlRigSequencerEditorLibrary::SetLocalControlRigTransform(ULevelSequence* LevelSequence, UControlRig* ControlRig, FName ControlName,
+	FFrameNumber Frame, FTransform Value, EMovieSceneTimeUnit TimeUnit, bool bSetKey)
 {
-	if (LevelSequence == nullptr || ControlRig == nullptr)
+	const UMovieScene* MovieScene = LevelSequence ? LevelSequence->GetMovieScene() : nullptr;
+	if (MovieScene == nullptr || ControlRig == nullptr)
 	{
 		return;
 	}
-	if (FRigControlElement* Element = ControlRig->FindControl(ControlName))
+
+
+	const UE::Private::FRigSetterContext Setter(ControlRig, ControlName);
+	const FRigControlElement* Control = Setter.GetControl(ERigControlType::Transform);
+	if (!Control)
 	{
-		if (Element->Settings.ControlType != ERigControlType::Transform)
-		{
-			UE_LOG(LogControlRig, Error, TEXT("Control Rig Wrong Type"));
-			return;
-		}
+		return;
 	}
-	if (UMovieScene* MovieScene = LevelSequence->GetMovieScene())
+
+	const FRigControlModifiedContext Context = UE::Private::GetContext(Frame, MovieScene, TimeUnit, bSetKey);
+	if (UE::Private::bIsInteracting)
+	{
+		FControlRigInteractionScope InteractionScope(ControlRig, Control->GetKey());
+		ControlRig->SetControlValue<FRigControlValue::FTransform_Float>(ControlName, Value, true, Context);
+		ControlRig->Evaluate_AnyThread();
+	}
+	else
 	{
-		FFrameRate TickResolution = MovieScene->GetTickResolution();
-		if (TimeUnit == EMovieSceneTimeUnit::DisplayRate)
-		{
-			Frame = FFrameRate::TransformTime(FFrameTime(Frame, 0), MovieScene->GetDisplayRate(), MovieScene->GetTickResolution()).RoundToFrame();
-		}
-		FRigControlModifiedContext Context;
-		Context.SetKey = bSetKey ? EControlRigSetKey::Always : EControlRigSetKey::DoNotCare;
-		Context.LocalTime = TickResolution.AsSeconds(FFrameTime(Frame));
 		ControlRig->SetControlValue<FRigControlValue::FTransform_Float>(ControlName, Value, true, Context);
 	}
 }
@@ -2709,34 +2795,34 @@ void UControlRigSequencerEditorLibrary::SetLocalControlRigTransform(ULevelSequen
 void UControlRigSequencerEditorLibrary::SetLocalControlRigTransforms(ULevelSequence* LevelSequence, UControlRig* ControlRig, FName ControlName, const TArray<FFrameNumber>& Frames, 
 	const TArray<FTransform> Values, EMovieSceneTimeUnit TimeUnit)
 {
-	if (LevelSequence == nullptr || ControlRig == nullptr || (Frames.Num() != Values.Num()))
+	const UMovieScene* MovieScene = LevelSequence ? LevelSequence->GetMovieScene() : nullptr;
+	if (MovieScene == nullptr || ControlRig == nullptr || (Frames.Num() != Values.Num()))
 	{
 		return;
 	}
-	if (FRigControlElement* Element = ControlRig->FindControl(ControlName))
+	
+	const UE::Private::FRigSetterContext Setter(ControlRig, ControlName);
+	const FRigControlElement* Control = Setter.GetControl(ERigControlType::Transform);
+	if (!Control)
 	{
-		if (Element->Settings.ControlType != ERigControlType::Transform)
-		{
-			UE_LOG(LogControlRig, Error, TEXT("Control Rig Wrong Type"));
-			return;
-		}
+		return;
 	}
-	if (UMovieScene* MovieScene = LevelSequence->GetMovieScene())
+	
+	const FFrameRate TickResolution = MovieScene->GetTickResolution();
+	const FFrameRate DisplayRate = MovieScene->GetDisplayRate();
+	
+	FRigControlModifiedContext Context;
+	Context.SetKey = EControlRigSetKey::Always;
+	for (int32 Index = 0; Index < Frames.Num(); ++Index)
 	{
-		FFrameRate TickResolution = MovieScene->GetTickResolution();
-		FRigControlModifiedContext Context;
-		Context.SetKey = EControlRigSetKey::Always;
-		for (int32 Index = 0; Index < Frames.Num(); ++Index)
+		FFrameNumber Frame = Frames[Index];
+		if (TimeUnit == EMovieSceneTimeUnit::DisplayRate)
 		{
-			FFrameNumber Frame = Frames[Index];
-			if (TimeUnit == EMovieSceneTimeUnit::DisplayRate)
-			{
-				Frame = FFrameRate::TransformTime(FFrameTime(Frame, 0), MovieScene->GetDisplayRate(), MovieScene->GetTickResolution()).RoundToFrame();
-			}
-			FTransform Value = Values[Index];
-			Context.LocalTime = TickResolution.AsSeconds(FFrameTime(Frame));
-			ControlRig->SetControlValue<FRigControlValue::FTransform_Float>(ControlName, Value, true, Context);
+			Frame = FFrameRate::TransformTime(FFrameTime(Frame, 0), DisplayRate, TickResolution).RoundToFrame();
 		}
+		FTransform Value = Values[Index];
+		Context.LocalTime = TickResolution.AsSeconds(FFrameTime(Frame));
+		ControlRig->SetControlValue<FRigControlValue::FTransform_Float>(ControlName, Value, true, Context);
 	}
 }
 
diff --git a/Engine/Plugins/Animation/ControlRig/Source/ControlRigEditor/Public/ControlRigSequencerEditorLibrary.h b/Engine/Plugins/Animation/ControlRig/Source/ControlRigEditor/Public/ControlRigSequencerEditorLibrary.h
--- a/Engine/Plugins/Animation/ControlRig/Source/ControlRigEditor/Public/ControlRigSequencerEditorLibrary.h
+++ b/Engine/Plugins/Animation/ControlRig/Source/ControlRigEditor/Public/ControlRigSequencerEditorLibrary.h
@@ -371,6 +371,13 @@ public:
 	static TArray<FTransform> GetSkeletalMeshComponentWorldTransforms(ULevelSequence* LevelSequence, USkeletalMeshComponent* SkeletalMeshComponent, const TArray<FFrameNumber>& Frames, 
 		EMovieSceneTimeUnit TimeUnit = EMovieSceneTimeUnit::DisplayRate,FName ReferenceName = NAME_None);
 
+	/**
+	* Sets the interaction state of the following SetLocalControlRig* functions (default is False).
+	* @param bIsInteracting If True the SetLocalControlRig* functions will be treated as if the ControlRig was interacting to react to interaction events.
+	*/
+	UFUNCTION(BlueprintCallable, Category = "Editor Scripting | Sequencer Tools | Control Rig")
+	static void SetInteraction(bool bIsInteracting);
+	
 	/**
 	* Get ControlRig Control's World Transform at a specific time
 	* @param LevelSequence Active Sequence to get transform for
