Anim in Engine: Fix Anim Details view is missing its scrollbar

#rnx

diff --git a/Engine/Plugins/Animation/ControlRig/Source/ControlRigEditor/Private/AnimDetails/Views/SAnimDetailsView.cpp b/Engine/Plugins/Animation/ControlRig/Source/ControlRigEditor/Private/AnimDetails/Views/SAnimDetailsView.cpp
--- a/Engine/Plugins/Animation/ControlRig/Source/ControlRigEditor/Private/AnimDetails/Views/SAnimDetailsView.cpp
+++ b/Engine/Plugins/Animation/ControlRig/Source/ControlRigEditor/Private/AnimDetails/Views/SAnimDetailsView.cpp
@@ -28,6 +28,7 @@
 #include "PropertyEditorDelegates.h"
 #include "PropertyEditorModule.h"
 #include "SAnimOverrideDetailsView.h"
+#include "Widgets/Layout/SScrollBox.h"
 
 namespace UE::ControlRigEditor
 {
@@ -53,7 +54,7 @@ namespace UE::ControlRigEditor
 			DetailsViewArgs.bCustomNameAreaLocation = true;
 			DetailsViewArgs.NameAreaSettings = FDetailsViewArgs::HideNameArea;
 			DetailsViewArgs.bAllowMultipleTopLevelObjects = false;
-			DetailsViewArgs.bShowScrollBar = true;
+			DetailsViewArgs.bShowScrollBar = false;
 		}
 
 		const TSharedRef<IDetailsView> DetailsView = FModuleManager::GetModuleChecked<FPropertyEditorModule>("PropertyEditor").CreateDetailView(DetailsViewArgs);
@@ -107,16 +108,14 @@ namespace UE::ControlRigEditor
 			+ SVerticalBox::Slot()
 			.FillHeight(1.f)
 			[
-				SNew(SVerticalBox)
+				SNew(SScrollBox)
 
-				+ SVerticalBox::Slot()
-				.AutoHeight()
+				+ SScrollBox::Slot()
 				[
 					DetailsView
 				]
 
-				+ SVerticalBox::Slot()
-				.AutoHeight()
+				+ SScrollBox::Slot()
 				[
 					SNew(SAnimOverrideDetailsView)
 					.OnRequestRefreshDetails(this, &SAnimDetailsView::RefreshDetailsView)
diff --git a/Engine/Plugins/Animation/ControlRig/Source/ControlRigEditor/Private/AnimDetails/Views/SAnimOverrideDetailsView.cpp b/Engine/Plugins/Animation/ControlRig/Source/ControlRigEditor/Private/AnimDetails/Views/SAnimOverrideDetailsView.cpp
--- a/Engine/Plugins/Animation/ControlRig/Source/ControlRigEditor/Private/AnimDetails/Views/SAnimOverrideDetailsView.cpp
+++ b/Engine/Plugins/Animation/ControlRig/Source/ControlRigEditor/Private/AnimDetails/Views/SAnimOverrideDetailsView.cpp
@@ -85,7 +85,7 @@ namespace UE::ControlRigEditor
 			DetailsViewArgs.bCustomNameAreaLocation = true;
 			DetailsViewArgs.NameAreaSettings = FDetailsViewArgs::HideNameArea;
 			DetailsViewArgs.bAllowMultipleTopLevelObjects = true;
-			DetailsViewArgs.bShowScrollBar = true;
+			DetailsViewArgs.bShowScrollBar = false;
 		}
 		
 		const TSharedPtr<IDetailsView> DetailsView = FModuleManager::GetModuleChecked<FPropertyEditorModule>("PropertyEditor").CreateDetailView(DetailsViewArgs);
