Anim in Engine: Sequencer selection is now adopted by [REDACTED_PERSON] as in previous versions

#rnx

diff --git a/Engine/Plugins/Animation/ControlRig/Source/ControlRigEditor/Private/AnimDetails/AnimDetailsProxyManager.cpp b/Engine/Plugins/Animation/ControlRig/Source/ControlRigEditor/Private/AnimDetails/AnimDetailsProxyManager.cpp
--- a/Engine/Plugins/Animation/ControlRig/Source/ControlRigEditor/Private/AnimDetails/AnimDetailsProxyManager.cpp
+++ b/Engine/Plugins/Animation/ControlRig/Source/ControlRigEditor/Private/AnimDetails/AnimDetailsProxyManager.cpp
@@ -163,13 +163,7 @@ void UAnimDetailsProxyManager::OnControlRigSelectionChanged(UControlRig* Control
 
 void UAnimDetailsProxyManager::OnSequencerSelectionChanged(TArray<FGuid> ObjectGuids)
 {
-	// Only update when the sequencer selection actually changed
-	if (CachedSequencerSelection != ObjectGuids)
-	{
-		CachedSequencerSelection = ObjectGuids;
-
-		RequestUpdateProxies();
-	}
+	RequestUpdateProxies();
 }
 
 void UAnimDetailsProxyManager::OnObjectsReplaced(const TMap<UObject*, UObject*>& OldToNewInstanceMap)
diff --git a/Engine/Plugins/Animation/ControlRig/Source/ControlRigEditor/Private/AnimDetails/AnimDetailsProxyManager.h b/Engine/Plugins/Animation/ControlRig/Source/ControlRigEditor/Private/AnimDetails/AnimDetailsProxyManager.h
--- a/Engine/Plugins/Animation/ControlRig/Source/ControlRigEditor/Private/AnimDetails/AnimDetailsProxyManager.h
+++ b/Engine/Plugins/Animation/ControlRig/Source/ControlRigEditor/Private/AnimDetails/AnimDetailsProxyManager.h
@@ -137,9 +137,6 @@ private:
 	UPROPERTY()
 	TObjectPtr<UAnimDetailsSelection> AnimDetailsSelection;
 
-	/** Cached sequencer selection, useful to detect changes */
-	TArray<FGuid> CachedSequencerSelection;
-
 	/** The filter used for anim details */
 	UE::ControlRigEditor::FAnimDetailsFilter AnimDetailsFilter;
 
