Anim in Engine: Anim details' keying section now shows if a control has a key on another frame

diff --git a/Engine/Plugins/Animation/ControlRig/Source/ControlRigEditor/Private/AnimDetails/Proxies/AnimDetailsProxyBase.cpp b/Engine/Plugins/Animation/ControlRig/Source/ControlRigEditor/Private/AnimDetails/Proxies/AnimDetailsProxyBase.cpp
--- a/Engine/Plugins/Animation/ControlRig/Source/ControlRigEditor/Private/AnimDetails/Proxies/AnimDetailsProxyBase.cpp
+++ b/Engine/Plugins/Animation/ControlRig/Source/ControlRigEditor/Private/AnimDetails/Proxies/AnimDetailsProxyBase.cpp
@@ -496,11 +496,13 @@ namespace UE::ControlRigEditor
 					{
 						break;
 					}
-					if (ChannelIndices.Contains(ChannelIndex++) == false)
+					else if (ChannelIndices.Contains(ChannelIndex++) == false)
 					{
 						continue;
 					}
-					if (Channel->GetNumKeys() == 0)
+				
+					const int32 NumKeys = Channel->GetNumKeys();
+					if (NumKeys == 0)
 					{
 						++EmptyChannelCount;
 						continue;
@@ -508,24 +510,29 @@ namespace UE::ControlRigEditor
 
 					SectionKeyedStatus = FMath::Max(SectionKeyedStatus, EPropertyKeyedStatus::KeyedInOtherFrame);
 
-					TArray<FFrameNumber> KeyTimes;
-					Channel->GetKeys(Range, &KeyTimes, nullptr);
-					if (KeyTimes.IsEmpty())
+					TArray<FFrameNumber> KeyTimesInRange;
+					Channel->GetKeys(Range, &KeyTimesInRange, nullptr);
+					if (KeyTimesInRange.IsEmpty() && NumKeys == 0)
 					{
 						++EmptyChannelCount;
 					}
+					else if (KeyTimesInRange.IsEmpty() && NumKeys > 0)
+					{
+						SectionKeyedStatus = FMath::Max(SectionKeyedStatus, EPropertyKeyedStatus::KeyedInOtherFrame);
+					}
 					else
 					{
-						SectionKeyedStatus = FMath::Max(SectionKeyedStatus, EPropertyKeyedStatus::PartiallyKeyed);
+						SectionKeyedStatus = FMath::Max(SectionKeyedStatus, EPropertyKeyedStatus::KeyedInFrame);
 					}
 				}
 
 				break; // just do it for one type
 			}
 
-			if (EmptyChannelCount == 0 && SectionKeyedStatus == EPropertyKeyedStatus::PartiallyKeyed)
+			const bool bIsPartiallySelectedGroup = EmptyChannelCount > 0;
+			if (bIsPartiallySelectedGroup)
 			{
-				SectionKeyedStatus = EPropertyKeyedStatus::KeyedInFrame;
+				SectionKeyedStatus = EPropertyKeyedStatus::PartiallyKeyed;
 			}
 
 			return SectionKeyedStatus;
