Control Rig: Fix crash with modular rig open and closing editor

#rnx

diff --git a/Engine/Plugins/Animation/ControlRig/Source/ControlRig/Private/ControlRig.cpp b/Engine/Plugins/Animation/ControlRig/Source/ControlRig/Private/ControlRig.cpp
--- a/Engine/Plugins/Animation/ControlRig/Source/ControlRig/Private/ControlRig.cpp
+++ b/Engine/Plugins/Animation/ControlRig/Source/ControlRig/Private/ControlRig.cpp
@@ -41,6 +41,7 @@
 #include "ModularRig.h"
 #include "Units/Modules/RigUnit_ConnectorExecution.h"
 #include "ControlRigReplay.h"
+#include "CoreGlobals.h"
 
 #include UE_INLINE_GENERATED_CPP_BY_NAME(ControlRig)
 
@@ -147,7 +148,7 @@ void UControlRig::BeginDestroy()
 	PreAdditiveValuesApplicationEvent.Clear();
 
 #if WITH_EDITOR
-	if (!HasAnyFlags(RF_ClassDefaultObject))
+	if (!HasAnyFlags(RF_ClassDefaultObject) && !IsEngineExitRequested())
 	{
 		if(const UClass* Class = GetClass())
 		{
