Control Rig: Improve UI for array connectors

#rnx

diff --git a/Engine/Plugins/Animation/ControlRig/Source/ControlRigEditor/Private/Editor/SModularRigModel.cpp b/Engine/Plugins/Animation/ControlRig/Source/ControlRigEditor/Private/Editor/SModularRigModel.cpp
--- a/Engine/Plugins/Animation/ControlRig/Source/ControlRigEditor/Private/Editor/SModularRigModel.cpp
+++ b/Engine/Plugins/Animation/ControlRig/Source/ControlRigEditor/Private/Editor/SModularRigModel.cpp
@@ -106,6 +106,7 @@ void SModularRigModel::Construct(const FArguments& InArgs, TSharedRef<IControlRi
 	bShowOptionalConnectors = false;
 	bShowUnresolvedConnectors = true;
 	bIsPerformingSelection = false;
+	bKeepCurrentEditedConnectors = false;
 
 	// setup all delegates for the modular rig model widget
 	FModularRigTreeDelegates Delegates;
@@ -122,6 +123,7 @@ void SModularRigModel::Construct(const FArguments& InArgs, TSharedRef<IControlRi
 	Delegates.OnResolveConnector = FOnModularRigTreeResolveConnector::CreateSP(this, &SModularRigModel::HandleConnectorResolved);
 	Delegates.OnDisconnectConnector = FOnModularRigTreeDisconnectConnector::CreateSP(this, &SModularRigModel::HandleConnectorDisconnect);
 	Delegates.OnSelectionChanged = FOnModularRigTreeSelectionChanged::CreateSP(this, &SModularRigModel::HandleSelectionChanged);
+	Delegates.OnAlwaysShowConnector = FOnModularRigTreeAlwaysShowConnector::CreateSP(this, &SModularRigModel::ShouldAlwaysShowConnector);
 
 	HeaderRowWidget = SNew(SHeaderRow)
 		.Visibility(EVisibility::Visible);
@@ -403,6 +405,11 @@ void SModularRigModel::HandlePreCompileModularRigs(URigVMBlueprint* InBlueprint)
 
 void SModularRigModel::HandlePostCompileModularRigs(URigVMBlueprint* InBlueprint)
 {
+	if(!bKeepCurrentEditedConnectors)
+	{
+		CurrentlyEditedConnectors.Reset();
+	}
+	
 	RefreshTreeView();
 	if (ControlRigEditor.IsValid())
 	{
@@ -1198,6 +1205,12 @@ void SModularRigModel::HandleConnectorResolved(const FRigElementKey& InConnector
 
 		if (const UModularRig* ModularRig = GetModularRig())
 		{
+			if(!bKeepCurrentEditedConnectors)
+			{
+				CurrentlyEditedConnectors.Reset();
+			}
+			const TGuardValue<bool> KeepCurrentEditedConnectorsGuard(bKeepCurrentEditedConnectors, true); 
+			CurrentlyEditedConnectors.Add(InConnector.Name);
 			Controller->ConnectConnectorToElements(InConnector, InTargets, true, ModularRig->GetModularRigSettings().bAutoResolve);
 		}
 	}
@@ -1212,6 +1225,12 @@ void SModularRigModel::HandleConnectorDisconnect(const FRigElementKey& InConnect
 		UModularRigController* Controller = ControlRigBlueprint->GetModularRigController();
 		check(Controller);
 
+		if(!bKeepCurrentEditedConnectors)
+		{
+			CurrentlyEditedConnectors.Reset();
+		}
+		const TGuardValue<bool> KeepCurrentEditedConnectorsGuard(bKeepCurrentEditedConnectors, true); 
+		CurrentlyEditedConnectors.Add(InConnector.Name);
 		Controller->DisconnectConnector(InConnector, false, true);
 	}
 }
@@ -1236,6 +1255,11 @@ void SModularRigModel::HandleSelectionChanged(TSharedPtr<FModularRigTreeElement>
 	}
 }
 
+bool SModularRigModel::ShouldAlwaysShowConnector(const FName& InConnectorName) const
+{
+	return CurrentlyEditedConnectors.Contains(InConnectorName);
+}
+
 void SModularRigModel::OnModularRigModified(EModularRigNotification InNotif, const FRigModuleReference* InModule)
 {
 	if (!ControlRigBlueprint.IsValid())
diff --git a/Engine/Plugins/Animation/ControlRig/Source/ControlRigEditor/Private/Editor/SModularRigModel.h b/Engine/Plugins/Animation/ControlRig/Source/ControlRigEditor/Private/Editor/SModularRigModel.h
--- a/Engine/Plugins/Animation/ControlRig/Source/ControlRigEditor/Private/Editor/SModularRigModel.h
+++ b/Engine/Plugins/Animation/ControlRig/Source/ControlRigEditor/Private/Editor/SModularRigModel.h
@@ -100,6 +100,9 @@ private:
 	/** Set Selection Changed */
 	void HandleSelectionChanged(TSharedPtr<FModularRigTreeElement> Selection, ESelectInfo::Type SelectInfo);
 
+	/** Returns true if a given connector should always be shown */
+	bool ShouldAlwaysShowConnector(const FName& InConnectorName) const;
+
 	TSharedPtr< SWidget > CreateContextMenuWidget();
 	void OnItemClicked(TSharedPtr<FModularRigTreeElement> InItem);
 	void OnItemDoubleClicked(TSharedPtr<FModularRigTreeElement> InItem);
@@ -160,6 +163,8 @@ private:
 
 	TSharedPtr<SSearchBox> FilterBox;
 	bool bIsPerformingSelection;
+	bool bKeepCurrentEditedConnectors;
+	TSet<FName> CurrentlyEditedConnectors;
 
 public:
 
diff --git a/Engine/Plugins/Animation/ControlRig/Source/ControlRigEditor/Private/Editor/SModularRigTreeView.cpp b/Engine/Plugins/Animation/ControlRig/Source/ControlRigEditor/Private/Editor/SModularRigTreeView.cpp
--- a/Engine/Plugins/Animation/ControlRig/Source/ControlRigEditor/Private/Editor/SModularRigTreeView.cpp
+++ b/Engine/Plugins/Animation/ControlRig/Source/ControlRigEditor/Private/Editor/SModularRigTreeView.cpp
@@ -700,6 +700,14 @@ bool SModularRigTreeView::AddElement(FString InKey, FString InParentKey, bool bA
 							bShouldFilterByConnectorType = !bMatchesFilter;
 						}
 
+						if(bShouldFilterByConnectorType)
+						{
+							if(Delegates.ShouldAlwaysShowConnector(Key.GetPathFName()))
+							{
+								bShouldFilterByConnectorType = false;
+							}
+						}
+
 						if(bShouldFilterByConnectorType)
 						{
 							const bool bIsConnected = Model.Connections.HasConnection(FRigElementKey(Key.GetPathFName(), ERigElementType::Connector));
diff --git a/Engine/Plugins/Animation/ControlRig/Source/ControlRigEditor/Private/Editor/SModularRigTreeView.h b/Engine/Plugins/Animation/ControlRig/Source/ControlRigEditor/Private/Editor/SModularRigTreeView.h
--- a/Engine/Plugins/Animation/ControlRig/Source/ControlRigEditor/Private/Editor/SModularRigTreeView.h
+++ b/Engine/Plugins/Animation/ControlRig/Source/ControlRigEditor/Private/Editor/SModularRigTreeView.h
@@ -19,6 +19,7 @@ DECLARE_DELEGATE_RetVal_TwoParams(FName, FOnModularRigTreeRenameElement, const F
 DECLARE_DELEGATE_TwoParams(FOnModularRigTreeResolveConnector, const FRigElementKey& /*Connector*/, const TArray<FRigElementKey>& /*Targets*/);
 DECLARE_DELEGATE_OneParam(FOnModularRigTreeDisconnectConnector, const FRigElementKey& /*Connector*/);
 DECLARE_DELEGATE_RetVal_ThreeParams(bool, FOnModularRigTreeVerifyElementNameChanged, const FName& /*OldModuleName*/, const FName& /*NewName*/, FText& /*OutErrorMessage*/);
+DECLARE_DELEGATE_RetVal_OneParam(bool, FOnModularRigTreeAlwaysShowConnector, const FName& /*ConnectorName*/);
 
 typedef STreeView<TSharedPtr<FModularRigTreeElement>>::FOnMouseButtonClick FOnModularRigTreeMouseButtonClick;
 typedef STreeView<TSharedPtr<FModularRigTreeElement>>::FOnMouseButtonDoubleClick FOnModularRigTreeMouseButtonDoubleClick;
@@ -43,6 +44,7 @@ public:
 	FOnModularRigTreeResolveConnector OnResolveConnector;
 	FOnModularRigTreeDisconnectConnector OnDisconnectConnector;
 	FOnModularRigTreeSelectionChanged OnSelectionChanged;
+	FOnModularRigTreeAlwaysShowConnector OnAlwaysShowConnector;
 	
 	FModularRigTreeDelegates()
 	{
@@ -105,6 +107,15 @@ public:
 		(void)OnSelectionChanged.ExecuteIfBound(Selection, SelectInfo);
 	}
 
+	bool ShouldAlwaysShowConnector(const FName& InConnectorName) const
+	{
+		if(OnAlwaysShowConnector.IsBound())
+		{
+			return OnAlwaysShowConnector.Execute(InConnectorName);
+		}
+		return false;
+	}
+
 private:
 
 	bool bSuspendSelectionDelegate = false;
diff --git a/Engine/Plugins/Animation/ControlRig/Source/ControlRigEditor/Private/Editor/SRigConnectorTargetWidget.cpp b/Engine/Plugins/Animation/ControlRig/Source/ControlRigEditor/Private/Editor/SRigConnectorTargetWidget.cpp
--- a/Engine/Plugins/Animation/ControlRig/Source/ControlRigEditor/Private/Editor/SRigConnectorTargetWidget.cpp
+++ b/Engine/Plugins/Animation/ControlRig/Source/ControlRigEditor/Private/Editor/SRigConnectorTargetWidget.cpp
@@ -48,6 +48,7 @@ void SRigConnectorTargetComboButton::Construct(const FArguments& InArgs)
 		// Wrap in configurable box to restrain height/width of menu
 		SNew(SBox)
 		.MinDesiredWidth(InArgs._ButtonMinDesiredWith)
+		.Padding(FMargin(0, 2, 0, 2))
 		[
 			SAssignNew(VerticalButtonBox, SVerticalBox)
 		]
@@ -107,64 +108,46 @@ void SRigConnectorTargetComboButton::PopulateButtonBox()
 	.AutoHeight()
 	.VAlign(VAlign_Top)
 	.HAlign(HAlign_Fill)
-	.Padding(4.0, 0.0, 4.0, 0.0)
+	.Padding(0.0, 0.0, 4.0, 0.0)
 	[
-		SNew( SButton )
-		.ButtonStyle(FAppStyle::Get(), "SimpleButton")
-		.ContentPadding(FMargin(0.0))
+		SNew(SBorder)
+		.Padding(FMargin(2.0, 2.0, 5.0, 2.0))
+		.BorderImage(RoundedBoxBrush)
+		.BorderBackgroundColor(FSlateColor(FLinearColor::Transparent))
+		.Content()
 		[
-			SAssignNew(RowBox, SHorizontalBox)
+			SAssignNew(ButtonBox, SHorizontalBox)
 			+ SHorizontalBox::Slot()
 			.AutoWidth()
 			.VAlign(VAlign_Center)
-			.HAlign(HAlign_Fill)
+			.HAlign(HAlign_Left)
+			.Padding(FMargin(0.f, 0.f, 3.f, 0.f))
+			[
+				SNew(SImage)
+				.Image_Lambda([GetIconAndColor]()
+				{
+					return GetIconAndColor().Key;
+				})
+				.ColorAndOpacity_Lambda([GetIconAndColor]()
+				{
+					return GetIconAndColor().Value;
+				})
+			]
+
+			+ SHorizontalBox::Slot()
+			.AutoWidth()
+			.VAlign(VAlign_Center)
+			.HAlign(HAlign_Left)
 			.Padding(0)
 			[
-				SNew(SBorder)
-				.Padding(FMargin(5.0, 2.0, 5.0, 2.0))
-				.BorderImage(RoundedBoxBrush)
-				.BorderBackgroundColor(FSlateColor(FLinearColor::Transparent))
-				.Content()
-				[
-					SAssignNew(ButtonBox, SHorizontalBox)
-					+ SHorizontalBox::Slot()
-					.AutoWidth()
-					.VAlign(VAlign_Center)
-					.HAlign(HAlign_Left)
-					.Padding(FMargin(0.f, 0.f, 3.f, 0.f))
-					[
-						SNew(SImage)
-						.Image_Lambda([GetIconAndColor]()
-						{
-							return GetIconAndColor().Key;
-						})
-						.ColorAndOpacity_Lambda([GetIconAndColor]()
-						{
-							return GetIconAndColor().Value;
-						})
-					]
-
-					+ SHorizontalBox::Slot()
-					.AutoWidth()
-					.VAlign(VAlign_Center)
-					.HAlign(HAlign_Left)
-					.Padding(0)
-					[
-						SNew( STextBlock )
-						.Text_Lambda([this, Hierarchy]()
-						{
-							const EElementNameDisplayMode NameDisplayMode = RigTreeDelegates.GetDisplaySettings().NameDisplayMode;
-							return Hierarchy->GetDisplayNameForUI(TargetKey.Get(), NameDisplayMode);
-						})
-						.Font( IDetailLayoutBuilder::GetDetailFont() )
-					]
-
-					+ SHorizontalBox::Slot()
-					.FillWidth(1.f)
-					[
-						SNew(SSpacer)
-					]
-				]
+				SNew( STextBlock )
+				.OverflowPolicy(ETextOverflowPolicy::Clip)
+				.Text_Lambda([this, Hierarchy]()
+				{
+					const EElementNameDisplayMode NameDisplayMode = RigTreeDelegates.GetDisplaySettings().NameDisplayMode;
+					return Hierarchy->GetDisplayNameForUI(TargetKey.Get(), NameDisplayMode);
+				})
+				.Font( IDetailLayoutBuilder::GetDetailFont() )
 			]
 		]
 	];
@@ -255,6 +238,7 @@ void SRigConnectorTargetWidget::Construct(const FArguments& InArgs)
 		DetailsViewArgs.bAllowFavoriteSystem = false;
 		DetailsViewArgs.bLockable = false;
 		DetailsViewArgs.bUpdatesFromSelection = false;
+		DetailsViewArgs.ColumnWidth = 1.f;
 
 		const TSharedRef<IDetailsView> DetailsView = PropertyEditorModule.CreateDetailView(DetailsViewArgs);
 
@@ -302,6 +286,7 @@ void SRigConnectorTargetWidget::Construct(const FArguments& InArgs)
 					SNew(SRigConnectorTargetComboButton)
 					.ConnectorKey(Connector)
 					.TargetKey(this, &SRigConnectorTargetWidget::GetSingleTargetKey)
+					.ContentPadding(FMargin(0, 3, 3, 0))
 					.OnSetTarget_Lambda([this](FRigElementKey InTarget) -> bool
 					{
 						SingleTarget = InTarget;
@@ -381,7 +366,7 @@ void FRigConnectorTargetWidgetCustomization::CustomizeHeader(TSharedRef<IPropert
 	HeaderRow.NameContent()
 	.MaxDesiredWidth(30)
 	[
-		StructPropertyHandle->CreatePropertyNameWidget()
+		SNullWidget::NullWidget
 	];
 
 	if(TargetsDetailWrapper.IsValid())
@@ -391,10 +376,13 @@ void FRigConnectorTargetWidgetCustomization::CustomizeHeader(TSharedRef<IPropert
 		if(TargetsDetailWrapper->TargetArray.IsValidIndex(ArrayIndex))
 		{
 			HeaderRow.ValueContent()
+			.HAlign(HAlign_Fill)
 			[
 				SAssignNew(ComboButton, SRigConnectorTargetComboButton)
+				.Clipping(EWidgetClipping::ClipToBounds)
 				.ConnectorKey(TargetsDetailWrapper->Connector)
 				.TargetKey(this, &FRigConnectorTargetWidgetCustomization::GetElementKey)
+				.ContentPadding(FMargin(0, 3, 3, 0))
 				.OnSetTarget_Lambda([this, ArrayIndex](FRigElementKey InTarget) -> bool
 				{
 					TargetsDetailWrapper->TargetArray[ArrayIndex] = InTarget;
