Control RIg: Fix adding duplicate controls when sorting
#rnx

diff --git a/Engine/Plugins/Animation/ControlRig/Source/ControlRig/Private/ModularRig.cpp b/Engine/Plugins/Animation/ControlRig/Source/ControlRig/Private/ModularRig.cpp
--- a/Engine/Plugins/Animation/ControlRig/Source/ControlRig/Private/ModularRig.cpp
+++ b/Engine/Plugins/Animation/ControlRig/Source/ControlRig/Private/ModularRig.cpp
@@ -474,7 +474,7 @@ void UModularRig::GetControlsInOrder(TArray<FRigControlElement*>& SortedControls
 		if (FRigControlElement* Control = Cast<FRigControlElement>(Element))
 		{
 			const FName ModuleName = DynamicHierarchy->GetModuleFName(Control->GetKey());
-			ControlsByModule.FindOrAdd(ModuleName).Add(Control);
+			ControlsByModule.FindOrAdd(ModuleName).AddUnique(Control);
 		}
 		bContinue = true;
 	});
