Control Rig: Arrange controls in modular rigs by module in sequencer (2nd attempt)
#rnx

diff --git a/Engine/Plugins/Animation/ControlRig/Source/ControlRig/Private/ModularRig.cpp b/Engine/Plugins/Animation/ControlRig/Source/ControlRig/Private/ModularRig.cpp
--- a/Engine/Plugins/Animation/ControlRig/Source/ControlRig/Private/ModularRig.cpp
+++ b/Engine/Plugins/Animation/ControlRig/Source/ControlRig/Private/ModularRig.cpp
@@ -458,8 +458,6 @@ const TArray<FName>& UModularRig::GetSupportedEvents() const
 
 void UModularRig::GetControlsInOrder(TArray<FRigControlElement*>& SortedControls) const
 {
-	UControlRig::GetControlsInOrder(SortedControls);
-	/* not working for MR_BIPED 
 	SortedControls.Reset();
 
 	if(DynamicHierarchy == nullptr)
@@ -492,7 +490,6 @@ void UModularRig::GetControlsInOrder(TArray<FRigControlElement*>& SortedControls
 	{
 		SortedControls.Append(*Controls);
 	}
-	*/
 }
 
 const FModularRigSettings& UModularRig::GetModularRigSettings() const
