Control Rig: Last minute fixes for overrides feature

#rnx

diff --git a/Engine/Plugins/Animation/ControlRig/Source/ControlRigEditor/Private/AnimDetails/Views/SAnimOverrideDetailsView.cpp b/Engine/Plugins/Animation/ControlRig/Source/ControlRigEditor/Private/AnimDetails/Views/SAnimOverrideDetailsView.cpp
--- a/Engine/Plugins/Animation/ControlRig/Source/ControlRigEditor/Private/AnimDetails/Views/SAnimOverrideDetailsView.cpp
+++ b/Engine/Plugins/Animation/ControlRig/Source/ControlRigEditor/Private/AnimDetails/Views/SAnimOverrideDetailsView.cpp
@@ -123,6 +123,11 @@ namespace UE::ControlRigEditor
 				const TArray<TObjectPtr<UAnimDetailsProxyBase>>& SelectedProxies = ProxyManager->GetExternalSelection();
 				for(const TObjectPtr<UAnimDetailsProxyBase>& SelectedProxy : SelectedProxies)
 				{
+					const FRigControlElement* ControlElement = SelectedProxy->GetControlElement();
+					if(!ControlElement || ControlElement->IsAnimationChannel())
+					{
+						continue;
+					}
 					SelectedProxy->UpdateOverrideableProperties();
 					OverrideableProxies.Add(SelectedProxy);
 				}
@@ -275,7 +280,7 @@ namespace UE::ControlRigEditor
 	{
 		if(InObjectA && InObjectB)
 		{
-			return InObjectA->IsA<UAnimDetailsProxyBase>() && (InObjectA->GetClass() == InObjectB->GetClass()); 
+			return InObjectA->IsA<UAnimDetailsProxyBase>() && InObjectB->IsA<UAnimDetailsProxyBase>(); 
 		}
 		return false;
 	}
