Anim In Engine: Pose library now uses Mirror Project settings table. Needed for new MR_biped rigs coming in.

diff --git a/Engine/Plugins/Animation/ControlRig/Source/ControlRig/Private/Tools/ControlRigPose.cpp b/Engine/Plugins/Animation/ControlRig/Source/ControlRig/Private/Tools/ControlRigPose.cpp
--- a/Engine/Plugins/Animation/ControlRig/Source/ControlRig/Private/Tools/ControlRigPose.cpp
+++ b/Engine/Plugins/Animation/ControlRig/Source/ControlRig/Private/Tools/ControlRigPose.cpp
@@ -455,11 +455,14 @@ void UControlRigPoseAsset::ReplaceControlName(const FName& CurrentName, const FN
 	Pose.ReplaceControlName(CurrentName, NewName);
 }
 
+void UControlRigPoseAsset::SetUpMirrorMatchTable(UControlRig* InControlRig)
+{
+	MirrorMatchTable.SetUpMirrorTable(InControlRig);
+}
+
 bool UControlRigPoseAsset::DoesMirrorMatch(UControlRig* ControlRig, const FName& ControlName) const
 {
-	FControlRigPoseMirrorTable MirrorTable;
-	MirrorTable.SetUpMirrorTable(ControlRig);
-	return (MirrorTable.IsMatched(ControlName));
+	return (MirrorMatchTable.IsMatched(ControlName));
 }
 
 
diff --git a/Engine/Plugins/Animation/ControlRig/Source/ControlRig/Private/Tools/ControlRigPoseMirrorSettings.cpp b/Engine/Plugins/Animation/ControlRig/Source/ControlRig/Private/Tools/ControlRigPoseMirrorSettings.cpp
--- a/Engine/Plugins/Animation/ControlRig/Source/ControlRig/Private/Tools/ControlRigPoseMirrorSettings.cpp
+++ b/Engine/Plugins/Animation/ControlRig/Source/ControlRig/Private/Tools/ControlRigPoseMirrorSettings.cpp
@@ -5,8 +5,6 @@
 
 UControlRigPoseMirrorSettings::UControlRigPoseMirrorSettings()
 {
-	RightSide = TEXT("_r_");
-	LeftSide = TEXT("_l_");
 	MirrorAxis = EAxis::X;
 	AxisToFlip = EAxis::X;
 }
diff --git a/Engine/Plugins/Animation/ControlRig/Source/ControlRig/Private/Tools/ControlRigPoseMirrorTable.cpp b/Engine/Plugins/Animation/ControlRig/Source/ControlRig/Private/Tools/ControlRigPoseMirrorTable.cpp
--- a/Engine/Plugins/Animation/ControlRig/Source/ControlRig/Private/Tools/ControlRigPoseMirrorTable.cpp
+++ b/Engine/Plugins/Animation/ControlRig/Source/ControlRig/Private/Tools/ControlRigPoseMirrorTable.cpp
@@ -5,30 +5,80 @@
 #include "ControlRig.h"
 #include "Tools/ControlRigPose.h"
 #include "RigVMFunctions/Math/RigVMMathLibrary.h"
+#include "Animation/AnimationSettings.h"
+#include "Animation/MirrorDataTable.h"
+#include "Internationalization/Regex.h"
 
+namespace UE::ControlRigPoseMirrorTable
+{
+FString GetMirrorString(const FString& InNameString, const TArray<FMirrorFindReplaceExpression>& MirrorFindReplaceExpressions)
+{
+	
+	for (const FMirrorFindReplaceExpression& RegExStr : MirrorFindReplaceExpressions)
+	{
+		FString FindString = RegExStr.FindExpression.ToString();
+		FString ReplaceString = RegExStr.ReplaceExpression.ToString();
+		if (RegExStr.FindReplaceMethod == EMirrorFindReplaceMethod::Prefix)
+		{
+			// convert prefix expression to regex that matches start of string, prefix, and any number of characters
+			FindString = FindString + TEXT("([^}]*)");
+			FindString = TEXT("^") + FindString;
+			ReplaceString = ReplaceString + TEXT("$1");
+
+		}
+		else if (RegExStr.FindReplaceMethod == EMirrorFindReplaceMethod::Suffix)
+		{
+			// convert suffix expression to regex that matches any number of characters start, the suffix, and end of string
+			FindString = TEXT("([^}]*)") + FindString + TEXT("$");
+			ReplaceString = TEXT("$1") + ReplaceString;
+		}
+
+		FRegexPattern MatherPatter(FindString);
+		FRegexMatcher Matcher(MatherPatter, InNameString);
+		bool bFound = false;
+		while (Matcher.FindNext())
+		{
+			for (int32 CaptureIndex = 1; CaptureIndex < 10; CaptureIndex++)
+			{
+				FString CaptureResult = Matcher.GetCaptureGroup(CaptureIndex);
+				int32 CaptureBegin = Matcher.GetCaptureGroupBeginning(CaptureIndex);
+				int32 CaptureEnd = Matcher.GetCaptureGroupEnding(CaptureIndex);
+				FString CaptureRegion = CaptureResult.Mid(CaptureBegin, CaptureEnd - CaptureBegin);
+				if (CaptureResult.IsEmpty())
+				{
+					break;
+				}
+				FString MatchString = FString::Printf(TEXT("$%i"), CaptureIndex);
+				ReplaceString = ReplaceString.Replace(*MatchString, *CaptureResult);
+			}
+			bFound = true;
+		}
+		if (bFound)
+		{
+			return ReplaceString;
+		}
+	}
+	return FString();
+}
+}
 void FControlRigPoseMirrorTable::SetUpMirrorTable(const UControlRig* ControlRig)
 {
-	const UControlRigPoseMirrorSettings* Settings = GetDefault<UControlRigPoseMirrorSettings>();
 	MatchedControls.Reset();
-	if (Settings && ControlRig)
+	if (ControlRig)
 	{
 		TArray<FRigControlElement*> CurrentControls = ControlRig->AvailableControls();
-		for (FRigControlElement* ControlElement : CurrentControls)
+		if (UAnimationSettings* AnimSettings = UAnimationSettings::Get())
 		{
-			FString CurrentString = ControlElement->GetName();
-			if (CurrentString.Contains(Settings->RightSide))
-			{
-				FString NewString = CurrentString.Replace(*Settings->RightSide, *Settings->LeftSide);
-				FName CurrentName(*CurrentString);
-				FName NewName(*NewString);
-				MatchedControls.Add(NewName, CurrentName);
-			}
-			else if (CurrentString.Contains(Settings->LeftSide))
+			for (FRigControlElement* ControlElement : CurrentControls)
 			{
-				FString NewString = CurrentString.Replace(*Settings->LeftSide, *Settings->RightSide);
-				FName CurrentName(*CurrentString);
-				FName NewName(*NewString);
-				MatchedControls.Add(NewName, CurrentName);
+				FString CurrentString = ControlElement->GetName();
+				FString NewString = UE::ControlRigPoseMirrorTable::GetMirrorString(CurrentString, AnimSettings->MirrorFindReplaceExpressions);
+				if (NewString.IsEmpty() == false && CurrentString != NewString)
+				{
+					FName CurrentName(*CurrentString);
+					FName NewName(*NewString);
+					MatchedControls.Add(NewName, CurrentName);
+				}
 			}
 		}
 	}
diff --git a/Engine/Plugins/Animation/ControlRig/Source/ControlRig/Public/Tools/ControlRigPose.h b/Engine/Plugins/Animation/ControlRig/Source/ControlRig/Public/Tools/ControlRigPose.h
--- a/Engine/Plugins/Animation/ControlRig/Source/ControlRig/Public/Tools/ControlRigPose.h
+++ b/Engine/Plugins/Animation/ControlRig/Source/ControlRig/Public/Tools/ControlRigPose.h
@@ -147,6 +147,10 @@ public:
 	UFUNCTION(BlueprintCallable, Category = "Pose")
 	UE_API void ReplaceControlName(const FName& CurrentName, const FName& NewName);
 
+	/* Set up an internal mirror table to be used by subsequent calls to DoesMirrorMatch,  need to call for each Control Rig*/
+	UFUNCTION(BlueprintCallable, Category = "Pose")
+	UE_API void SetUpMirrorMatchTable(UControlRig* InControlRig);
+
 	UFUNCTION(BlueprintPure, Category = "Pose")
 	UE_API bool DoesMirrorMatch(UControlRig* ControlRig, const FName& ControlName) const;
 
@@ -155,6 +159,8 @@ public:
 	UPROPERTY(BlueprintReadWrite, EditAnywhere, Category = "Pose")
 	FControlRigControlPose Pose;
 
+private:
+	FControlRigPoseMirrorTable MirrorMatchTable;
 };
 
 #undef UE_API
diff --git a/Engine/Plugins/Animation/ControlRig/Source/ControlRig/Public/Tools/ControlRigPoseMirrorSettings.h b/Engine/Plugins/Animation/ControlRig/Source/ControlRig/Public/Tools/ControlRigPoseMirrorSettings.h
--- a/Engine/Plugins/Animation/ControlRig/Source/ControlRig/Public/Tools/ControlRigPoseMirrorSettings.h
+++ b/Engine/Plugins/Animation/ControlRig/Source/ControlRig/Public/Tools/ControlRigPoseMirrorSettings.h
@@ -12,7 +12,8 @@
 
 
 UCLASS(MinimalAPI, config = EditorPerProjectUserSettings)
-class UControlRigPoseMirrorSettings : public UObject
+class 
+UControlRigPoseMirrorSettings : public UObject
 {
 
 public:
@@ -20,12 +21,12 @@ public:
 
 	UE_API UControlRigPoseMirrorSettings();
 
-	/** The String to denote Controls on the Right Side */
-	UPROPERTY(config, EditAnywhere, BlueprintReadWrite, Category = "Mirror Settings", meta = (ToolTip = "String To Denote Right Side"))
+	UE_DEPRECATED(5.7, "This property is no longer used, instead we use the MirrorFindExpressions defined in UAnimationSettings")
+	UPROPERTY()
 	FString RightSide;
 
-	/** The String to denote Controls on the Left Side */
-	UPROPERTY(config, EditAnywhere, BlueprintReadWrite, Category = "Mirror Settings", meta = (ToolTip = "String To Denote Left Side"))
+	UE_DEPRECATED(5.7, "This property is no longer used, instead we use the MirrorFindExpressions defined in UAnimationSettings")
+	UPROPERTY()
 	FString LeftSide;
 
 	// the axis to mirror translations against
diff --git a/Engine/Plugins/Animation/ControlRig/Source/ControlRigEditor/Private/EditMode/SControlRigControlViews.cpp b/Engine/Plugins/Animation/ControlRig/Source/ControlRigEditor/Private/EditMode/SControlRigControlViews.cpp
--- a/Engine/Plugins/Animation/ControlRig/Source/ControlRigEditor/Private/EditMode/SControlRigControlViews.cpp
+++ b/Engine/Plugins/Animation/ControlRig/Source/ControlRigEditor/Private/EditMode/SControlRigControlViews.cpp
@@ -535,7 +535,7 @@ SControlRigPoseView::~SControlRigPoseView()
 
 void SControlRigPoseView::NotifyPostChange(const FPropertyChangedEvent& PropertyChangedEvent, FProperty* PropertyThatChanged)
 {
-	UpdateStatusBlocks();
+	bUpdateStatusBlocks = true;
 }
 
 ECheckBoxState SControlRigPoseView::IsKeyPoseChecked() const
@@ -578,7 +578,7 @@ void SControlRigPoseView::OnMirrorPoseChecked(ECheckBoxState NewState)
 	{
 		bIsMirror = false;
 	}
-	UpdateStatusBlocks();
+	bUpdateStatusBlocks = true;
 }
 
 bool SControlRigPoseView::IsMirrorEnabled() const
@@ -605,7 +605,7 @@ void SControlRigPoseView::OnAdditivePoseChecked(ECheckBoxState NewState)
 	{
 		bIsAdditive = false;
 	}
-	UpdateStatusBlocks();
+	bUpdateStatusBlocks = true;
 }
 
 FReply SControlRigPoseView::OnPastePose()
@@ -805,12 +805,23 @@ void SControlRigPoseView::HandleControlAdded(UControlRig* ControlRig, bool bIsAd
 			(ControlRig)->ControlSelected().RemoveAll(this);
 		}
 	}
-	UpdateStatusBlocks();
+	bUpdateStatusBlocks = true;
 }
 
 void SControlRigPoseView::HandleControlSelected(UControlRig* Subject, FRigControlElement* InControl, bool bSelected)
 {
-	UpdateStatusBlocks();
+	bUpdateStatusBlocks = true;
+}
+
+void SControlRigPoseView::Tick(const FGeometry& AllottedGeometry, const double InCurrentTime, const float InDeltaTime)
+{
+	SCompoundWidget::Tick(AllottedGeometry, InCurrentTime, InDeltaTime);
+
+	if (bUpdateStatusBlocks)
+	{
+		UpdateStatusBlocks();
+		bUpdateStatusBlocks = false;
+	}
 }
 
 void SControlRigPoseView::UpdateStatusBlocks()
@@ -829,7 +840,7 @@ void SControlRigPoseView::UpdateStatusBlocks()
 		uint32 MirrorMatching = 0;
 		for (UControlRig* ControlRig : ControlRigs)
 		{
-
+			PoseAsset->SetUpMirrorMatchTable(ControlRig);
 			TArray<FName> SelectedNames = ControlRig->CurrentControlSelection();
 			TotalSelected += SelectedNames.Num();
 			for (const FName& ControlName : ControlNames)
diff --git a/Engine/Plugins/Animation/ControlRig/Source/ControlRigEditor/Private/EditMode/SControlRigControlViews.h b/Engine/Plugins/Animation/ControlRig/Source/ControlRigEditor/Private/EditMode/SControlRigControlViews.h
--- a/Engine/Plugins/Animation/ControlRig/Source/ControlRigEditor/Private/EditMode/SControlRigControlViews.h
+++ b/Engine/Plugins/Animation/ControlRig/Source/ControlRigEditor/Private/EditMode/SControlRigControlViews.h
@@ -42,8 +42,9 @@ class SControlRigPoseView : public SCompoundWidget, public FNotifyHook
 
 	//FNotifyHook
 	virtual void NotifyPostChange(const FPropertyChangedEvent& PropertyChangedEvent, FProperty* PropertyThatChanged) override;
-
+	virtual void Tick(const FGeometry& AllottedGeometry, const double InCurrentTime, const float InDeltaTime) override;
 private:
+	bool bUpdateStatusBlocks = false;
 	void UpdateStatusBlocks();
 
 	/*
