Control Rig: Guard against edit mode updates after destroying the Control Rig

#rnx

diff --git a/Engine/Plugins/Animation/ControlRig/Source/ControlRigEditor/Private/EditMode/ControlRigEditMode.cpp b/Engine/Plugins/Animation/ControlRig/Source/ControlRigEditor/Private/EditMode/ControlRigEditMode.cpp
--- a/Engine/Plugins/Animation/ControlRig/Source/ControlRigEditor/Private/EditMode/ControlRigEditMode.cpp
+++ b/Engine/Plugins/Animation/ControlRig/Source/ControlRigEditor/Private/EditMode/ControlRigEditMode.cpp
@@ -6253,12 +6253,15 @@ void FControlRigEditMode::PostPoseUpdate(const FEditorViewportClient* InViewport
 	const bool bAreEditingControlRigDirectly = AreEditingControlRigDirectly();
 	auto UpdateShapes = [this, bAreEditingControlRigDirectly](const UControlRig* ControlRig, const TArray<TObjectPtr<AControlRigShapeActor>>& ShapeActors)
 	{
-		const FTransform ComponentTransform = bAreEditingControlRigDirectly ? FTransform::Identity : GetHostingSceneComponentTransform(ControlRig); 
-
-		FShapeUpdateParams Params(ControlRig, ComponentTransform, IsControlRigSkelMeshVisible(ControlRig), IsInLevelEditor());
-		for (AControlRigShapeActor* ShapeActor : ShapeActors)
+		if(IsValid(ControlRig))
 		{
-			UpdateControlShape(ShapeActor, ControlRig->FindControl(ShapeActor->ControlName), Params);
+			const FTransform ComponentTransform = bAreEditingControlRigDirectly ? FTransform::Identity : GetHostingSceneComponentTransform(ControlRig); 
+
+			FShapeUpdateParams Params(ControlRig, ComponentTransform, IsControlRigSkelMeshVisible(ControlRig), IsInLevelEditor());
+			for (AControlRigShapeActor* ShapeActor : ShapeActors)
+			{
+				UpdateControlShape(ShapeActor, ControlRig->FindControl(ShapeActor->ControlName), Params);
+			}
 		}
 	};
 	
