Anim In Engine: Set or add keys when recording contorl rig keys, previously we would just add keys causing duplicates.

diff --git a/Engine/Plugins/Animation/ControlRig/Source/ControlRig/Private/Sequencer/MovieSceneControlRigParameterSection.cpp b/Engine/Plugins/Animation/ControlRig/Source/ControlRig/Private/Sequencer/MovieSceneControlRigParameterSection.cpp
--- a/Engine/Plugins/Animation/ControlRig/Source/ControlRig/Private/Sequencer/MovieSceneControlRigParameterSection.cpp
+++ b/Engine/Plugins/Animation/ControlRig/Source/ControlRig/Private/Sequencer/MovieSceneControlRigParameterSection.cpp
@@ -3492,6 +3492,72 @@ void UMovieSceneControlRigParameterSection::AutoSetTangents(const FName& Control
 	}
 }
 #if WITH_EDITOR
+namespace UE::Private
+{
+	void SetOrAddKey(TMovieSceneChannelData<FMovieSceneFloatValue>& ChannelData, FFrameNumber Time, float Value, const EMovieSceneKeyInterpolation Interpolation)
+	{
+		int32 ExistingIndex = ChannelData.FindKey(Time);
+		if (ExistingIndex != INDEX_NONE)
+		{
+			FMovieSceneFloatValue& FloatValue = ChannelData.GetValues()[ExistingIndex]; //-V758
+			FloatValue.Value = Value;
+		}
+		else
+		{
+			FMovieSceneFloatValue NewKey(Value);
+			ERichCurveTangentWeightMode WeightedMode = RCTWM_WeightedNone;
+			NewKey.InterpMode = ERichCurveInterpMode::RCIM_Cubic;
+			NewKey.TangentMode = ERichCurveTangentMode::RCTM_Auto;
+			NewKey.Tangent.ArriveTangent = 0.0f;
+			NewKey.Tangent.LeaveTangent = 0.0f;
+			NewKey.Tangent.TangentWeightMode = WeightedMode;
+			NewKey.Tangent.ArriveTangentWeight = 0.0f;
+			NewKey.Tangent.LeaveTangentWeight = 0.0f;
+
+			switch (Interpolation)
+			{
+			case EMovieSceneKeyInterpolation::SmartAuto:
+			{
+				NewKey.InterpMode = ERichCurveInterpMode::RCIM_Cubic;
+				NewKey.TangentMode = ERichCurveTangentMode::RCTM_SmartAuto;
+			}
+			break;
+			case EMovieSceneKeyInterpolation::Auto:
+			{
+				NewKey.InterpMode = ERichCurveInterpMode::RCIM_Cubic;
+				NewKey.TangentMode = ERichCurveTangentMode::RCTM_Auto;
+			}
+			break;
+			case EMovieSceneKeyInterpolation::User:
+			{
+				NewKey.InterpMode = ERichCurveInterpMode::RCIM_Cubic;
+				NewKey.TangentMode = ERichCurveTangentMode::RCTM_User;
+			}
+			break;
+			case EMovieSceneKeyInterpolation::Break:
+			{
+				NewKey.InterpMode = ERichCurveInterpMode::RCIM_Cubic;
+				NewKey.TangentMode = ERichCurveTangentMode::RCTM_Auto;
+			}
+			break;
+
+			case EMovieSceneKeyInterpolation::Linear:
+			{
+				NewKey.InterpMode = ERichCurveInterpMode::RCIM_Linear;
+			}
+			break;
+
+			case EMovieSceneKeyInterpolation::Constant:
+			{
+				NewKey.InterpMode = ERichCurveInterpMode::RCIM_Constant;
+			}
+			break;
+
+			}
+			ChannelData.AddKey(Time, NewKey);
+		}
+	}
+}
 
 void UMovieSceneControlRigParameterSection::RecordControlRigKey(FFrameNumber FrameNumber, bool bSetDefault, EMovieSceneKeyInterpolation InInterpMode, bool bbOntoSelectedControls)
 {
@@ -3505,32 +3571,12 @@ void UMovieSceneControlRigParameterSection::RecordControlRigKey(FFrameNumber Fra
 		// Helper lambda to add a FVector key to the float channels
 		auto AddVectorKeyToFloatChannels = [&FloatChannels, InInterpMode](int32& ChannelIndex, FFrameNumber FrameNumber, const FVector& Value)
 		{
-			switch (InInterpMode)
-			{
-			case EMovieSceneKeyInterpolation::Linear:
-				FloatChannels[ChannelIndex++]->AddLinearKey(FrameNumber, Value.X);
-				FloatChannels[ChannelIndex++]->AddLinearKey(FrameNumber, Value.Y);
-				FloatChannels[ChannelIndex++]->AddLinearKey(FrameNumber, Value.Z);
-				break;
-
-			case EMovieSceneKeyInterpolation::Constant:
-				FloatChannels[ChannelIndex++]->AddConstantKey(FrameNumber, Value.X);
-				FloatChannels[ChannelIndex++]->AddConstantKey(FrameNumber, Value.Y);
-				FloatChannels[ChannelIndex++]->AddConstantKey(FrameNumber, Value.Z);
-				break;
-
-			case  EMovieSceneKeyInterpolation::Auto:
-				FloatChannels[ChannelIndex++]->AddCubicKey(FrameNumber, Value.X, ERichCurveTangentMode::RCTM_Auto);
-				FloatChannels[ChannelIndex++]->AddCubicKey(FrameNumber, Value.Y, ERichCurveTangentMode::RCTM_Auto);
-				FloatChannels[ChannelIndex++]->AddCubicKey(FrameNumber, Value.Z, ERichCurveTangentMode::RCTM_Auto);
-				break;
-			case  EMovieSceneKeyInterpolation::SmartAuto:
-			default:
-				FloatChannels[ChannelIndex++]->AddCubicKey(FrameNumber, Value.X, ERichCurveTangentMode::RCTM_SmartAuto);
-				FloatChannels[ChannelIndex++]->AddCubicKey(FrameNumber, Value.Y, ERichCurveTangentMode::RCTM_SmartAuto);
-				FloatChannels[ChannelIndex++]->AddCubicKey(FrameNumber, Value.Z, ERichCurveTangentMode::RCTM_SmartAuto);
-				break;
-			}
+			TMovieSceneChannelData<FMovieSceneFloatValue> ChannelDataX = FloatChannels[ChannelIndex++]->GetData();
+			UE::Private::SetOrAddKey(ChannelDataX, FrameNumber, Value.X, InInterpMode);
+			TMovieSceneChannelData<FMovieSceneFloatValue> ChannelDataY = FloatChannels[ChannelIndex++]->GetData();
+			UE::Private::SetOrAddKey(ChannelDataY, FrameNumber, Value.Y, InInterpMode);
+			TMovieSceneChannelData<FMovieSceneFloatValue> ChannelDataZ = FloatChannels[ChannelIndex++]->GetData();
+			UE::Private::SetOrAddKey(ChannelDataZ, FrameNumber, Value.Z, InInterpMode);
 		};
 
 		TArray<FRigControlElement*> Controls;
@@ -3579,7 +3625,16 @@ void UMovieSceneControlRigParameterSection::RecordControlRigKey(FFrameNumber Fra
 					{
 						BoolChannels[ChannelIndex]->SetDefault(Val);
 					}
-					BoolChannels[ChannelIndex]->GetData().AddKey(FrameNumber, Val);
+					int32 ExistingIndex = BoolChannels[ChannelIndex]->GetData().FindKey(FrameNumber);
+					if (ExistingIndex != INDEX_NONE)
+					{
+						bool& BValue = BoolChannels[ChannelIndex]->GetData().GetValues()[ExistingIndex]; //-V758
+						BValue = Val;
+					}
+					else
+					{
+						BoolChannels[ChannelIndex]->GetData().AddKey(FrameNumber, Val);
+					}
 					break;
 				}
 				case ERigControlType::Integer:
@@ -3591,7 +3646,18 @@ void UMovieSceneControlRigParameterSection::RecordControlRigKey(FFrameNumber Fra
 						{
 							EnumChannels[ChannelIndex]->SetDefault(Val);
 						}
-						EnumChannels[ChannelIndex]->GetData().AddKey(FrameNumber, Val);
+						int32 ExistingIndex = EnumChannels[ChannelIndex]->GetData().FindKey(FrameNumber);
+						if (ExistingIndex != INDEX_NONE)
+						{
+							uint8& BValue = EnumChannels[ChannelIndex]->GetData().GetValues()[ExistingIndex]; //-V758
+							BValue = Val;
+						}
+						else
+						{
+							EnumChannels[ChannelIndex]->GetData().AddKey(FrameNumber, Val);
+						}
+						break;
+
 					}
 					else
 					{
@@ -3600,7 +3666,17 @@ void UMovieSceneControlRigParameterSection::RecordControlRigKey(FFrameNumber Fra
 						{
 							IntChannels[ChannelIndex]->SetDefault(Val);
 						}
-						IntChannels[ChannelIndex]->GetData().AddKey(FrameNumber, Val);
+
+						int32 ExistingIndex = IntChannels[ChannelIndex]->GetData().FindKey(FrameNumber);
+						if (ExistingIndex != INDEX_NONE)
+						{
+							int32& BValue = IntChannels[ChannelIndex]->GetData().GetValues()[ExistingIndex]; //-V758
+							BValue = Val;
+						}
+						else
+						{
+							IntChannels[ChannelIndex]->GetData().AddKey(FrameNumber, Val);
+						}
 					}
 					break;
 				}
@@ -3612,27 +3688,8 @@ void UMovieSceneControlRigParameterSection::RecordControlRigKey(FFrameNumber Fra
 					{
 						FloatChannels[ChannelIndex]->SetDefault(Val);
 					}
-
-					switch (InInterpMode)
-					{
-					case EMovieSceneKeyInterpolation::Linear:
-						FloatChannels[ChannelIndex++]->AddLinearKey(FrameNumber, Val);
-						break;
-
-					case EMovieSceneKeyInterpolation::Constant:
-						FloatChannels[ChannelIndex++]->AddConstantKey(FrameNumber, Val);
-						break;
-
-					case EMovieSceneKeyInterpolation::Auto:
-						FloatChannels[ChannelIndex++]->AddCubicKey(FrameNumber, Val, ERichCurveTangentMode::RCTM_Auto);
-						break;
-
-					case EMovieSceneKeyInterpolation::SmartAuto:
-					default:
-						FloatChannels[ChannelIndex++]->AddCubicKey(FrameNumber, Val, ERichCurveTangentMode::RCTM_SmartAuto);
-						break;
-					}
-
+					TMovieSceneChannelData<FMovieSceneFloatValue> ChannelDataX = FloatChannels[ChannelIndex++]->GetData();
+					UE::Private::SetOrAddKey(ChannelDataX, FrameNumber, Val, InInterpMode);
 					break;
 				}
 				case ERigControlType::Vector2D:
@@ -3643,29 +3700,11 @@ void UMovieSceneControlRigParameterSection::RecordControlRigKey(FFrameNumber Fra
 						FloatChannels[ChannelIndex]->SetDefault(Val.X);
 						FloatChannels[ChannelIndex + 1]->SetDefault(Val.Y);
 					}
+					TMovieSceneChannelData<FMovieSceneFloatValue> ChannelDataX = FloatChannels[ChannelIndex++]->GetData();
+					UE::Private::SetOrAddKey(ChannelDataX, FrameNumber, Val.X, InInterpMode);
 
-					switch (InInterpMode)
-					{
-					case EMovieSceneKeyInterpolation::Linear:
-						FloatChannels[ChannelIndex++]->AddLinearKey(FrameNumber, Val.X);
-						FloatChannels[ChannelIndex++]->AddLinearKey(FrameNumber, Val.Y);
-						break;
-
-					case  EMovieSceneKeyInterpolation::Constant:
-						FloatChannels[ChannelIndex++]->AddConstantKey(FrameNumber, Val.X);
-						FloatChannels[ChannelIndex++]->AddConstantKey(FrameNumber, Val.Y);
-						break;
-					case  EMovieSceneKeyInterpolation::Auto:
-						FloatChannels[ChannelIndex++]->AddCubicKey(FrameNumber, Val.X, ERichCurveTangentMode::RCTM_Auto);
-						FloatChannels[ChannelIndex++]->AddCubicKey(FrameNumber, Val.Y, ERichCurveTangentMode::RCTM_Auto);
-						break;
-					case EMovieSceneKeyInterpolation::SmartAuto:
-					default:
-						FloatChannels[ChannelIndex++]->AddCubicKey(FrameNumber, Val.X, ERichCurveTangentMode::RCTM_SmartAuto);
-						FloatChannels[ChannelIndex++]->AddCubicKey(FrameNumber, Val.Y, ERichCurveTangentMode::RCTM_SmartAuto);
-						break;
-					}
-
+					TMovieSceneChannelData<FMovieSceneFloatValue> ChannelDataY = FloatChannels[ChannelIndex++]->GetData();
+					UE::Private::SetOrAddKey(ChannelDataY, FrameNumber, Val.Y, InInterpMode);
 					break;
 				}
 				case ERigControlType::Position:
@@ -4765,3 +4804,4 @@ bool UMovieSceneControlRigParameterSection::PopulateEvaluationFieldImpl(const TR
 #undef LOCTEXT_NAMESPACE 
 
 
+
