Anim Layers: Stop duplicate weight values
#trival

diff --git a/Engine/Plugins/Animation/ControlRig/Source/ControlRigEditor/Private/Sequencer/AnimLayers/AnimLayers.cpp b/Engine/Plugins/Animation/ControlRig/Source/ControlRigEditor/Private/Sequencer/AnimLayers/AnimLayers.cpp
--- a/Engine/Plugins/Animation/ControlRig/Source/ControlRigEditor/Private/Sequencer/AnimLayers/AnimLayers.cpp
+++ b/Engine/Plugins/Animation/ControlRig/Source/ControlRigEditor/Private/Sequencer/AnimLayers/AnimLayers.cpp
@@ -1587,16 +1587,17 @@ static void CopySectionIntoAnother(UMovieSceneSection* ToSection, UMovieSceneSec
 		if (ChannelTypeName == FMovieSceneFloatChannel::StaticStruct()->GetFName())
 		{
 			TArrayView<FMovieSceneFloatChannel*> BaseFloatChannels = ToSection->GetChannelProxy().GetChannels<FMovieSceneFloatChannel>();
-			const int StartIndex = 0;
-			const int EndIndex = BaseFloatChannels.Num() - 1;
+			const int32 StartIndex = 0;
+			const int32 EndOffset = ToSection->IsA<UMovieSceneControlRigParameterSection>() ? 2: 1; //if CR section skip weight
+			const int32 EndIndex = BaseFloatChannels.Num() - EndOffset;
 			MovieSceneToolHelpers::MergeSections<FMovieSceneFloatChannel>(ToSection, AbsoluteSections, AdditiveSections, 
 				StartIndex, EndIndex, Range);
 		}
 		else if (ChannelTypeName == FMovieSceneDoubleChannel::StaticStruct()->GetFName())
 		{
 			TArrayView<FMovieSceneDoubleChannel*> BaseDoubleChannels = ToSection->GetChannelProxy().GetChannels<FMovieSceneDoubleChannel>();
-			const int StartIndex = 0;
-			const int EndIndex = BaseDoubleChannels.Num() - 1;
+			const int32 StartIndex = 0;
+			const int32 EndIndex = BaseDoubleChannels.Num() - 1;
 			MovieSceneToolHelpers::MergeSections<FMovieSceneDoubleChannel>(ToSection, AbsoluteSections, AdditiveSections, 
 				StartIndex, EndIndex, Range);
 		}
