ControlRig: only change rotation order if it's actually used

#rnx

diff --git a/Engine/Plugins/Animation/ControlRig/Source/ControlRig/Private/Rigs/RigHierarchy.cpp b/Engine/Plugins/Animation/ControlRig/Source/ControlRig/Private/Rigs/RigHierarchy.cpp
--- a/Engine/Plugins/Animation/ControlRig/Source/ControlRig/Private/Rigs/RigHierarchy.cpp
+++ b/Engine/Plugins/Animation/ControlRig/Source/ControlRig/Private/Rigs/RigHierarchy.cpp
@@ -5876,8 +5876,15 @@ void URigHierarchy::SetPreferredEulerAnglesFromValue(FRigControlElement* InContr
 		{
 			const FVector3f EulerXYZf = InValue.Get<FVector3f>();
 			const FVector EulerXYZ(EulerXYZf.X, EulerXYZf.Y, EulerXYZf.Z);
-			const FVector Angle = AnimationCore::ChangeEulerRotationOrder(EulerXYZ, EEulerRotationOrder::XYZ, RotationOrder);
-			PreferredEulerAngles.SetAngles(Angle, bInitial, RotationOrder, bFixEulerFlips);
+			if (InControlElement->Settings.bUsePreferredRotationOrder)
+			{
+				const FVector Angle = AnimationCore::ChangeEulerRotationOrder(EulerXYZ, EEulerRotationOrder::XYZ, RotationOrder);
+				PreferredEulerAngles.SetAngles(Angle, bInitial, RotationOrder, bFixEulerFlips);
+			}
+			else
+			{
+				PreferredEulerAngles.SetRotator(FRotator::MakeFromEuler(EulerXYZ), bInitial, bFixEulerFlips);
+			}
 			break;
 		}
 	case ERigControlType::EulerTransform:
