ControlRig & ECS: fixed rotation not correctly being applied on rotators with rotation order set
- ECS version of 39626696

#rnx

diff --git a/Engine/Plugins/Animation/ControlRig/Source/ControlRig/Private/Sequencer/MovieSceneControlRigParameterBuffer.cpp b/Engine/Plugins/Animation/ControlRig/Source/ControlRig/Private/Sequencer/MovieSceneControlRigParameterBuffer.cpp
--- a/Engine/Plugins/Animation/ControlRig/Source/ControlRig/Private/Sequencer/MovieSceneControlRigParameterBuffer.cpp
+++ b/Engine/Plugins/Animation/ControlRig/Source/ControlRig/Private/Sequencer/MovieSceneControlRigParameterBuffer.cpp
@@ -1005,7 +1005,20 @@ void FControlRigParameterValueHeader::Apply(UControlRig* Rig, FName Name, int32
 		case EControlRigControlType::Parameter_Vector:
 		{
 			const FVector3f* Ptr = static_cast<const FVector3f*>(ParameterBuffer);
-			Rig->SetControlValue<FVector3f>(Name, Ptr[Index], bNotify, ModifiedContext, bSetupUndo);
+
+			FRigControlElement* ControlElement = Rig->FindControl(Name);
+			URigHierarchy* RigHierarchy = Rig->GetHierarchy();
+			if (RigHierarchy && ControlElement && ControlElement->Settings.ControlType == ERigControlType::Rotator)
+            {
+				const FVector EulerValue(Ptr[Index].X, Ptr[Index].Y, Ptr[Index].Z);
+				const FRotator Rotator(RigHierarchy->GetControlQuaternion(ControlElement, EulerValue));
+				RigHierarchy->SetControlSpecifiedEulerAngle(ControlElement, EulerValue);
+				Rig->SetControlValue<FRotator>(Name, Rotator, bNotify, ModifiedContext, bSetupUndo);
+			}
+			else
+			{
+				Rig->SetControlValue<FVector3f>(Name, Ptr[Index], bNotify, ModifiedContext, bSetupUndo);	
+			}
 			break;
 		}
 		case EControlRigControlType::Parameter_Transform:
