Control Rig: Add option to show only selected rigs/modes in the outliner
#rnx

diff --git a/Engine/Plugins/Animation/ControlRig/Source/ControlRig/Private/Settings/ControlRigSettings.cpp b/Engine/Plugins/Animation/ControlRig/Source/ControlRig/Private/Settings/ControlRigSettings.cpp
--- a/Engine/Plugins/Animation/ControlRig/Source/ControlRig/Private/Settings/ControlRigSettings.cpp
+++ b/Engine/Plugins/Animation/ControlRig/Source/ControlRig/Private/Settings/ControlRigSettings.cpp
@@ -36,6 +36,7 @@ UControlRigEditorSettings::UControlRigEditorSettings(const FObjectInitializer& O
 	bFlattenModules = false;
 	bFocusOnSelection = true;
 	ElementNameDisplayMode = EElementNameDisplayMode::AssetDefault;
+	OutlinerMultiRigDisplayMode = EMultiRigTreeDisplayMode::All;
 #endif
 }
 
diff --git a/Engine/Plugins/Animation/ControlRig/Source/ControlRig/Public/Settings/ControlRigSettings.h b/Engine/Plugins/Animation/ControlRig/Source/ControlRig/Public/Settings/ControlRigSettings.h
--- a/Engine/Plugins/Animation/ControlRig/Source/ControlRig/Public/Settings/ControlRigSettings.h
+++ b/Engine/Plugins/Animation/ControlRig/Source/ControlRig/Public/Settings/ControlRigSettings.h
@@ -29,6 +29,14 @@ struct FControlRigSettingsPerPinBool
 	TMap<FString, bool> Values;
 };
 
+UENUM()
+enum class EMultiRigTreeDisplayMode
+{
+	All, // show all rigs
+	SelectedRigs, // show only rigs with selected controls
+	SelectedModules // show only modules with selected controls (for modular rigs)
+};
+
 /**
  * Default ControlRig settings.
  */
@@ -145,6 +153,12 @@ public:
 	UPROPERTY(EditAnywhere, config, Category = Outliner)
 	bool bFlattenModules;
 
+	/**
+	 *  Outliner reflection mode for multi rigs 
+	 */
+	UPROPERTY(EditAnywhere, config, Category = Outliner)
+	EMultiRigTreeDisplayMode OutlinerMultiRigDisplayMode;
+
 	/**
 	 * Defines how the element names will be displayed in treeviews, anim outliner etc
 	 */
diff --git a/Engine/Plugins/Animation/ControlRig/Source/ControlRigEditor/Private/ControlRigElementDetails.cpp b/Engine/Plugins/Animation/ControlRig/Source/ControlRigEditor/Private/ControlRigElementDetails.cpp
--- a/Engine/Plugins/Animation/ControlRig/Source/ControlRigEditor/Private/ControlRigElementDetails.cpp
+++ b/Engine/Plugins/Animation/ControlRig/Source/ControlRigEditor/Private/ControlRigElementDetails.cpp
@@ -4290,6 +4290,7 @@ void FRigControlElementDetails::CustomizeAvailableSpaces(IDetailLayoutBuilder& D
 	DisplaySettings.bArrangeByModules = false;
 	DisplaySettings.bFlattenModules = false;
 	DisplaySettings.NameDisplayMode = EElementNameDisplayMode::AssetDefault;
+	DisplaySettings.OutlinerDisplayMode = EMultiRigTreeDisplayMode::All;
 
 	const TSharedRef<SHorizontalBox> HeaderContentWidget = SNew(SHorizontalBox);
 	HeaderContentWidget->AddSlot()
diff --git a/Engine/Plugins/Animation/ControlRig/Source/ControlRigEditor/Private/ControlRigHierarchyCommands.cpp b/Engine/Plugins/Animation/ControlRig/Source/ControlRigEditor/Private/ControlRigHierarchyCommands.cpp
--- a/Engine/Plugins/Animation/ControlRig/Source/ControlRigEditor/Private/ControlRigHierarchyCommands.cpp
+++ b/Engine/Plugins/Animation/ControlRig/Source/ControlRigEditor/Private/ControlRigHierarchyCommands.cpp
@@ -35,6 +35,9 @@ void FControlRigHierarchyCommands::RegisterCommands()
 	UI_COMMAND(HideParentsWhenFiltering, "Hide Parents When Filtering", "Whether to show parent items grayed out, or hide them entirely when filtering", EUserInterfaceActionType::ToggleButton, FInputChord());
 	UI_COMMAND(ArrangeByModules, "Arrange by Modules", "Arrange the hierarchy by the modules that spawn each control", EUserInterfaceActionType::ToggleButton, FInputChord());
 	UI_COMMAND(FlattenModules, "Flatten Modules", "Flatten the modular rig into a list of modules", EUserInterfaceActionType::ToggleButton, FInputChord());
+	UI_COMMAND(SetMultiRigMode_All, "All", "Change outliner multi rig display mode", EUserInterfaceActionType::RadioButton, FInputChord());
+	UI_COMMAND(SetMultiRigMode_SelectedRigs, "Selected Rigs", "Change outliner multi rig display mode", EUserInterfaceActionType::RadioButton, FInputChord());
+	UI_COMMAND(SetMultiRigMode_SelectedModules, "Selected Modules", "Change outliner multi rig display mode", EUserInterfaceActionType::RadioButton, FInputChord());
 	UI_COMMAND(FocusOnSelection, "Focus On Selection", "Focus on the selection when it changes", EUserInterfaceActionType::ToggleButton, FInputChord());
 	UI_COMMAND(ShowImportedBones, "Show Imported Bones", "Whether to show or hide imported bones", EUserInterfaceActionType::ToggleButton, FInputChord());
 	UI_COMMAND(ShowBones, "Show Bones", "Whether to show or hide bones", EUserInterfaceActionType::ToggleButton, FInputChord());
diff --git a/Engine/Plugins/Animation/ControlRig/Source/ControlRigEditor/Private/ControlRigHierarchyCommands.h b/Engine/Plugins/Animation/ControlRig/Source/ControlRigEditor/Private/ControlRigHierarchyCommands.h
--- a/Engine/Plugins/Animation/ControlRig/Source/ControlRigEditor/Private/ControlRigHierarchyCommands.h
+++ b/Engine/Plugins/Animation/ControlRig/Source/ControlRigEditor/Private/ControlRigHierarchyCommands.h
@@ -105,6 +105,11 @@ public:
 	/* Flatten the modules in a modular rig */
 	TSharedPtr< FUICommandInfo > FlattenModules;
 
+	/* Options of outliner reflection mode for multi rigs */
+	TSharedPtr< FUICommandInfo > SetMultiRigMode_All; // show all rigs
+	TSharedPtr< FUICommandInfo > SetMultiRigMode_SelectedRigs; // show only rigs with selected controls
+	TSharedPtr< FUICommandInfo > SetMultiRigMode_SelectedModules; // show only modules with selected controls (for modular rigs)
+
 	/* Show imported bones */
 	TSharedPtr< FUICommandInfo > ShowImportedBones;
 
diff --git a/Engine/Plugins/Animation/ControlRig/Source/ControlRigEditor/Private/EditMode/SControlRigOutliner.cpp b/Engine/Plugins/Animation/ControlRig/Source/ControlRigEditor/Private/EditMode/SControlRigOutliner.cpp
--- a/Engine/Plugins/Animation/ControlRig/Source/ControlRigEditor/Private/EditMode/SControlRigOutliner.cpp
+++ b/Engine/Plugins/Animation/ControlRig/Source/ControlRigEditor/Private/EditMode/SControlRigOutliner.cpp
@@ -931,6 +931,22 @@ bool SMultiRigHierarchyTreeView::AddElement(UControlRig* InControlRig, const FRi
 	FMultiRigData ParentData;
 	ParentData.WeakControlRig = InControlRig;
 
+	// If only showing selected modules, 
+	if (InControlRig->IsModularRig() && Settings.bArrangeByModules)
+	{
+		if (Settings.OutlinerDisplayMode == EMultiRigTreeDisplayMode::SelectedModules)
+		{
+			if (const URigHierarchy* Hierarchy = InControlRig->GetHierarchy())
+			{
+				ParentData.SetItemName(FMultiRigData::EMultiRigDataType_Module, *Hierarchy->GetModuleName(InElement->GetKey()));
+				if (!ElementMap.Contains(ParentData))
+				{
+					return false;
+				}
+			}
+		}
+	}
+
 	if (!AddElement(Data, ParentData))
 	{
 		return false;
@@ -946,11 +962,6 @@ bool SMultiRigHierarchyTreeView::AddElement(UControlRig* InControlRig, const FRi
 			{
 				ParentData.SetItemName(FMultiRigData::EMultiRigDataType_Module, *Hierarchy->GetModuleName(InElement->GetKey()));
 
-				FMultiRigData NewData;
-				NewData.Type = FMultiRigData::EMultiRigDataType_Element;
-				NewData.WeakControlRig = InControlRig;
-				NewData.SetElementKey(InElement->GetKey());
-
 				if (TSharedPtr<FMultiRigTreeElement>* ParentElementPtr = ElementMap.Find(ParentData))
 				{
 					if (Settings.FilterText.IsEmpty() || !Settings.bFlattenHierarchyOnFilter)
@@ -979,8 +990,6 @@ bool SMultiRigHierarchyTreeView::AddElement(UControlRig* InControlRig, const FRi
 						}
 					}
 				}
-				
-				AddElement(NewData, ParentData);
 			}
 			else
 			{
@@ -1199,6 +1208,46 @@ void SMultiRigHierarchyTreeView::RefreshTreeView()
 	
 	TSharedPtr<ISequencer> Sequencer = WeakSequencer.Pin();
 
+	const FRigTreeDisplaySettings& Settings = Delegates.GetDisplaySettings();
+
+	TMap<UControlRig*, TArray<FRigElementKey>> SelectedControls;
+	TArray<FRigModuleInstance*> SelectedModules;
+	if (Settings.OutlinerDisplayMode != EMultiRigTreeDisplayMode::All)
+	{
+		EditMode->GetAllSelectedControls(SelectedControls);
+		if (Settings.OutlinerDisplayMode == EMultiRigTreeDisplayMode::SelectedModules)
+		{
+			// If showing current module, always expand
+			ExpansionState.Reset();
+			
+			for (TTuple<UControlRig*, TArray<FRigElementKey>> RigControls : SelectedControls)
+			{
+				if (UModularRig* ModularRig = Cast<UModularRig>(RigControls.Key))
+				{
+					URigHierarchy* Hierarchy = RigControls.Key->GetHierarchy();
+					for (FRigElementKey& Control : RigControls.Value)
+					{
+						// Do not change selected modules because of animation channels or visual cues
+						if (FRigControlElement* ControlElement = Cast<FRigControlElement>(Hierarchy->Find(Control)))
+						{
+							const ERigControlAnimationType& Type = ControlElement->Settings.AnimationType;
+							if (Type == ERigControlAnimationType::AnimationChannel || Type == ERigControlAnimationType::VisualCue)
+							{
+								continue;
+							}
+						}
+						
+						FString ModuleName = Hierarchy->GetModuleName(Control);
+						if (FRigModuleInstance* Module = ModularRig->FindModule(*ModuleName))
+						{
+							SelectedModules.AddUnique(Module);
+						}
+					}
+				}
+			}
+		}
+	}
+
 	auto GetAscestorName = [Sequencer](UObject* Object, FString& InOutName)
 	{
 		if (Sequencer && Object)
@@ -1232,6 +1281,14 @@ void SMultiRigHierarchyTreeView::RefreshTreeView()
 	{
 		if (UControlRig* ControlRig  = ControlRigPtr.Get())
 		{
+			if (Settings.OutlinerDisplayMode != EMultiRigTreeDisplayMode::All)
+			{
+				if (!SelectedControls.Contains(ControlRig))
+				{
+					continue;
+				}
+			}
+			
 			if (TSharedPtr<IControlRigObjectBinding> ObjectBinding = ControlRigPtr->GetObjectBinding())
 			{
 				TArray<FMultiRigData> DataHierarchy;
@@ -1278,17 +1335,23 @@ void SMultiRigHierarchyTreeView::RefreshTreeView()
 				// add root element
 				AddElement(CRData, *ParentData);
 
-
-				const FRigTreeDisplaySettings& Settings = Delegates.GetDisplaySettings();
 				if (ControlRig->IsModularRig() && Settings.bArrangeByModules)
 				{
 					if (UModularRig* ModularRig = Cast<UModularRig>(ControlRig))
 					{
 						ModularRig->ForEachModule([&](FRigModuleInstance* Module)
 						{
+							if (Settings.OutlinerDisplayMode == EMultiRigTreeDisplayMode::SelectedModules)
+							{
+								if (!SelectedModules.Contains(Module))
+								{
+									return true;
+								}
+							}
 							FMultiRigData ParentData;
 							ParentData.WeakControlRig = ControlRig;
-							if (Settings.bFlattenModules || Module->ParentModuleName.IsNone())
+							if (Settings.bFlattenModules || Module->ParentModuleName.IsNone() ||
+								Settings.OutlinerDisplayMode == EMultiRigTreeDisplayMode::SelectedModules)
 							{
 								ParentData = CRData;
 							}
@@ -1455,7 +1518,6 @@ void SMultiRigHierarchyTreeView::RefreshTreeView()
 		else 
 		{
 			// otherwise, iterate thru the elements if there's a filter as the root might have been skipped because of it  
-			const FRigTreeDisplaySettings& Settings = Delegates.GetDisplaySettings();
 			const FString FilteredString = Settings.FilterText.ToString();
 			if (!FilteredString.IsEmpty())
 			{
@@ -1656,6 +1718,7 @@ const FRigTreeDisplaySettings& SSearchableMultiRigHierarchyTreeView::GetDisplayS
 	Settings.bFlattenModules = UControlRigEditorSettings::Get()->bFlattenModules;
 	Settings.bFocusOnSelection = UControlRigEditorSettings::Get()->bFocusOnSelection;
 	Settings.NameDisplayMode = UControlRigEditorSettings::Get()->ElementNameDisplayMode;
+	Settings.OutlinerDisplayMode = UControlRigEditorSettings::Get()->OutlinerMultiRigDisplayMode;
 	return Settings;
 }
 
@@ -1670,6 +1733,27 @@ TSharedRef<SWidget> SSearchableMultiRigHierarchyTreeView::OnGetOptionsMenu()
 		const FControlRigEditModeCommands& EditModeActions = FControlRigEditModeCommands::Get();
 		MenuBuilder.AddMenuEntry(Actions.ArrangeByModules);
 		MenuBuilder.AddMenuEntry(Actions.FlattenModules);
+		
+		MenuBuilder.AddWidget(
+			SNew(SVerticalBox)
+
+			+ SVerticalBox::Slot()
+			.AutoHeight()
+			.Padding(3)
+			[
+				SNew(SEnumComboBox, StaticEnum<EMultiRigTreeDisplayMode>())
+				.CurrentValue_Lambda([this]() -> int32
+				{
+					return static_cast<int32>(GetMultiRigMode());
+				})
+				.OnEnumSelectionChanged_Lambda([this](int32 InEnumValue, ESelectInfo::Type)
+				{
+					SetMultiRigMode(static_cast<EMultiRigTreeDisplayMode>(InEnumValue));
+				})
+			]
+			,
+			LOCTEXT("MultiRigMode", "MultiRig Display Mode")
+		);
 
 		MenuBuilder.AddWidget(
 			SNew(SVerticalBox)
@@ -1868,6 +1952,30 @@ void SSearchableMultiRigHierarchyTreeView::ToggleFocusOnSelection()
 	}
 }
 
+void SSearchableMultiRigHierarchyTreeView::SetMultiRigMode(const EMultiRigTreeDisplayMode InMode)
+{
+	UControlRigEditorSettings* EditorSettings = UControlRigEditorSettings::Get();
+	EditorSettings->OutlinerMultiRigDisplayMode = InMode;
+
+	if (FProperty* Property = UControlRigEditorSettings::StaticClass()->FindPropertyByName(GET_MEMBER_NAME_CHECKED(UControlRigEditorSettings, bFocusOnSelection)))
+	{
+		FPropertyChangedEvent PropertyChangedEvent(Property, EPropertyChangeType::ValueSet);
+		EditorSettings->PostEditChangeProperty(PropertyChangedEvent);
+	}
+
+	GetTreeView()->RequestTreeViewRefresh();
+}
+
+bool SSearchableMultiRigHierarchyTreeView::IsMultiRigMode(const EMultiRigTreeDisplayMode InMode)
+{
+	return UControlRigEditorSettings::Get()->OutlinerMultiRigDisplayMode == InMode;
+}
+
+EMultiRigTreeDisplayMode SSearchableMultiRigHierarchyTreeView::GetMultiRigMode()
+{
+	return UControlRigEditorSettings::Get()->OutlinerMultiRigDisplayMode;
+}
+
 void SSearchableMultiRigHierarchyTreeView::OnFilterTextChanged(const FText& SearchText)
 {
 	FilterText = SearchText;
@@ -1895,6 +2003,7 @@ void SControlRigOutliner::Construct(const FArguments& InArgs, FControlRigEditMod
 	DisplaySettings.bFlattenModules = false;
 	DisplaySettings.bFocusOnSelection = true;
 	DisplaySettings.NameDisplayMode = EElementNameDisplayMode::AssetDefault;
+	DisplaySettings.OutlinerDisplayMode = EMultiRigTreeDisplayMode::All;
 
 	FMultiRigTreeDelegates RigTreeDelegates;
 	RigTreeDelegates.OnGetDisplaySettings = FOnGetRigTreeDisplaySettings::CreateSP(this, &SControlRigOutliner::GetDisplaySettings);
@@ -1987,6 +2096,38 @@ void SControlRigOutliner::HandleControlSelected(UControlRig* Subject, FRigContro
 		// This action was initiated by the outliner
 		return;
 	}
+	const EMultiRigTreeDisplayMode Mode = HierarchyTreeView->GetDisplaySettings().OutlinerDisplayMode;
+	if (Mode != EMultiRigTreeDisplayMode::All)
+	{
+		if (!bSelected)
+		{
+			HierarchyTreeView->GetTreeView()->RequestTreeViewRefresh();
+			return;
+		}
+		
+		FMultiRigData RigData;
+		RigData.WeakControlRig = Subject;
+		RigData.Type = FMultiRigData::EMultiRigDataType_ControlRig;
+		TSharedPtr<FMultiRigTreeElement> Element = HierarchyTreeView->GetTreeView()->FindElement(RigData);
+		if (!Element.IsValid())
+		{
+			HierarchyTreeView->GetTreeView()->RequestTreeViewRefresh();
+		}
+		else if (Mode == EMultiRigTreeDisplayMode::SelectedModules)
+		{
+			if (UModularRig* ModularRig = Cast<UModularRig>(Subject))
+			{
+				FMultiRigData ModuleData;
+				ModuleData.WeakControlRig = Subject;
+				ModuleData.SetItemName(FMultiRigData::EMultiRigDataType_Module, *ModularRig->GetHierarchy()->GetModuleName(ControlElement->GetKey()));
+				TSharedPtr<FMultiRigTreeElement> ModuleElement = HierarchyTreeView->GetTreeView()->FindElement(ModuleData);
+				if (!ModuleElement.IsValid())
+				{
+					HierarchyTreeView->GetTreeView()->RequestTreeViewRefresh();
+				}
+			}
+		}
+	}
 	
 	const FRigElementKey Key = ControlElement->GetKey();
 	FMultiRigData Data(Subject, Key);
diff --git a/Engine/Plugins/Animation/ControlRig/Source/ControlRigEditor/Private/EditMode/SControlRigOutliner.h b/Engine/Plugins/Animation/ControlRig/Source/ControlRigEditor/Private/EditMode/SControlRigOutliner.h
--- a/Engine/Plugins/Animation/ControlRig/Source/ControlRigEditor/Private/EditMode/SControlRigOutliner.h
+++ b/Engine/Plugins/Animation/ControlRig/Source/ControlRigEditor/Private/EditMode/SControlRigOutliner.h
@@ -317,6 +317,9 @@ public:
 	bool CanToggleModuleManipulators() const;
 	bool IsFocusingOnSelection() const;
 	void ToggleFocusOnSelection();
+	void SetMultiRigMode(const EMultiRigTreeDisplayMode InMode);
+	bool IsMultiRigMode(const EMultiRigTreeDisplayMode InMode);
+	EMultiRigTreeDisplayMode GetMultiRigMode();
 
 private:
 
diff --git a/Engine/Plugins/Animation/ControlRig/Source/ControlRigEditor/Private/Editor/SRigHierarchyTreeView.h b/Engine/Plugins/Animation/ControlRig/Source/ControlRigEditor/Private/Editor/SRigHierarchyTreeView.h
--- a/Engine/Plugins/Animation/ControlRig/Source/ControlRigEditor/Private/Editor/SRigHierarchyTreeView.h
+++ b/Engine/Plugins/Animation/ControlRig/Source/ControlRigEditor/Private/Editor/SRigHierarchyTreeView.h
@@ -6,6 +6,7 @@
 #include "Widgets/Views/STreeView.h"
 #include "Rigs/RigHierarchy.h"
 #include "SRigHierarchyTagWidget.h"
+#include "Settings/ControlRigSettings.h"
 
 class SSearchBox;
 class SRigHierarchyTreeView;
@@ -32,6 +33,7 @@ struct CONTROLRIGEDITOR_API FRigTreeDisplaySettings
 		bShowIconColors = true;
 		bArrangeByModules = false;
 		bFlattenModules = false;
+		OutlinerDisplayMode = EMultiRigTreeDisplayMode::All;
 	}
 	
 	FText FilterText;
@@ -80,6 +82,9 @@ struct CONTROLRIGEDITOR_API FRigTreeDisplaySettings
 
 	/** Whether or not to expand and focus the selection when this changes */
 	bool bFocusOnSelection;
+
+	/** Identify which rigs to display based on selection */
+	EMultiRigTreeDisplayMode OutlinerDisplayMode;
 };
 
 DECLARE_DELEGATE_RetVal(const URigHierarchy*, FOnGetRigTreeHierarchy);
