Fix issue where the Control Rig Editor Mode always extended the level viewport toolbar even when it was invoked within a standalone editor.

#rnx

diff --git a/Engine/Plugins/Animation/ControlRig/Source/ControlRigEditor/Private/EditMode/ControlRigEditMode.cpp b/Engine/Plugins/Animation/ControlRig/Source/ControlRigEditor/Private/EditMode/ControlRigEditMode.cpp
--- a/Engine/Plugins/Animation/ControlRig/Source/ControlRigEditor/Private/EditMode/ControlRigEditMode.cpp
+++ b/Engine/Plugins/Animation/ControlRig/Source/ControlRigEditor/Private/EditMode/ControlRigEditMode.cpp
@@ -502,10 +502,13 @@ void FControlRigEditMode::Enter()
 		}
 	}
 
-	UE::ControlRig::PopulateControlRigViewportToolbarTransformSubmenu("LevelEditor.ViewportToolbar.Transform");
-	UE::ControlRig::PopulateControlRigViewportToolbarShowSubmenu(
-		"LevelEditor.ViewportToolbar.Show");
-
+	if (&Owner->GetToolkitHost()->GetEditorModeManager() == &GLevelEditorModeTools())
+	{
+		UE::ControlRig::PopulateControlRigViewportToolbarTransformSubmenu("LevelEditor.ViewportToolbar.Transform");
+		UE::ControlRig::PopulateControlRigViewportToolbarShowSubmenu(
+			"LevelEditor.ViewportToolbar.Show");
+	}
+	
 	RegisterPendingFocusMode();
 
 	// initialize the gizmo context
