Control Rig: Remove unneeded structures and delegate declarations from control rig
#rnx

diff --git a/Engine/Plugins/Animation/ControlRig/Source/ControlRig/Public/ControlRig.h b/Engine/Plugins/Animation/ControlRig/Source/ControlRig/Public/ControlRig.h
--- a/Engine/Plugins/Animation/ControlRig/Source/ControlRig/Public/ControlRig.h
+++ b/Engine/Plugins/Animation/ControlRig/Source/ControlRig/Public/ControlRig.h
@@ -32,7 +32,6 @@
 
 #if WITH_EDITOR
 #include "AnimPreviewInstance.h"
-#include "UnrealWidgetFwd.h"
 #endif 
 
 #include "ControlRig.generated.h"
@@ -44,7 +43,6 @@ class USkeletalMeshComponent;
 class AActor;
 class UTransformableControlHandle;
 class UControlRigReplay;
-class UControlRig;
 class UControlRigOverrideAsset;
 
 struct FReferenceSkeleton;
@@ -55,65 +53,6 @@ struct FControlRigOverrideContainer;
 
 CONTROLRIG_API DECLARE_LOG_CATEGORY_EXTERN(LogControlRig, Log, All);
 
-UENUM()
-enum class EControlRigInteractionTransformSpace
-{
-	World,
-	Local,
-	Parent,
-	Explicit
-};
-
-/**
- * FTransformContext provides a way of passing the various transform parameters to functions that need to know what the transform context is.
- * Extend it if necessary, particularly to avoid overloading some functions signatures.
- */
-USTRUCT(BlueprintType)
-struct FControlRigInteractionTransformContext
-{
-	GENERATED_BODY()
-	
-	FControlRigInteractionTransformContext() = default;
-
-#if WITH_EDITOR
-	FControlRigInteractionTransformContext(const UE::Widget::EWidgetMode& InWidgetMode)
-		: bTranslation(InWidgetMode == UE::Widget::WM_Translate || InWidgetMode == UE::Widget::WM_TranslateRotateZ)
-		, bRotation(InWidgetMode == UE::Widget::WM_Rotate || InWidgetMode == UE::Widget::WM_TranslateRotateZ)
-		, bScale(InWidgetMode == UE::Widget::WM_Scale)
-	{}
-#endif
-	
-	bool CanTransform() const
-	{
-		return bTranslation || bRotation || bScale;
-	}
-
-	bool bTranslation = false;
-	
-	UPROPERTY(BlueprintReadWrite, EditAnywhere, Category = "Interaction")
-	FVector Drag = FVector::ZeroVector;
-
-	bool bRotation = false;
-	
-	UPROPERTY(BlueprintReadWrite, EditAnywhere, Category = "Interaction")
-	FRotator Rot = FRotator::ZeroRotator;
-
-	bool bScale = false;
-
-	UPROPERTY(BlueprintReadWrite, EditAnywhere, Category = "Interaction")
-	FVector Scale = FVector::OneVector;
-
-	UPROPERTY(BlueprintReadWrite, EditAnywhere, Category = "Interaction")
-	EControlRigInteractionTransformSpace Space = EControlRigInteractionTransformSpace::World;
-};
-
-#if WITH_EDITORONLY_DATA
-// A delegate for subscribing / reacting to interaction started.
-DECLARE_DYNAMIC_MULTICAST_DELEGATE_ThreeParams(FControlRigInteractionStartedEvent, UControlRig*, ControlRig, TArray<FRigElementKey>, InteractionKeys, EControlRigInteractionType, InteractionType);
-DECLARE_DYNAMIC_MULTICAST_DELEGATE_TwoParams(FControlRigInteractionEndedEvent, UControlRig*, ControlRig, TArray<FRigElementKey>, InteractionKeys);
-DECLARE_DYNAMIC_MULTICAST_DELEGATE_ThreeParams(FControlRigInteractionUpdatedEvent, UControlRig*, ControlRig, TArray<FRigElementKey>, InteractionKeys, FControlRigInteractionTransformContext, TransformContext);
-#endif
-
 DECLARE_MULTICAST_DELEGATE_OneParam(FControlRigOverrideAssetsChanged, UControlRig*);
 
 /** Runs logic for mapping input data to transforms (the "Rig") */
diff --git a/Engine/Plugins/Animation/ControlRig/Source/ControlRigEditor/Private/ControlRigElementDetails.cpp b/Engine/Plugins/Animation/ControlRig/Source/ControlRigEditor/Private/ControlRigElementDetails.cpp
--- a/Engine/Plugins/Animation/ControlRig/Source/ControlRigEditor/Private/ControlRigElementDetails.cpp
+++ b/Engine/Plugins/Animation/ControlRig/Source/ControlRigEditor/Private/ControlRigElementDetails.cpp
@@ -2988,7 +2988,7 @@ FDetailWidgetRow& FRigTransformElementDetails::CreateEulerTransformValueWidgetRo
 					DebuggedRig->ElementsBeingInteracted.AddUnique(ElementInfo.Element.GetKey());
 				}
 				
-				FControlRigInteractionScope* InteractionScope = new FControlRigInteractionScope(DebuggedRig, DebuggedRig->ElementsBeingInteracted, Type);
+				FControlRigInteractionScope* InteractionScope = new FControlRigInteractionScope(DebuggedRig);
 				StrongThisPtr->InteractionScopes.Add(InteractionScope);
 			}
 		});
diff --git a/Engine/Plugins/Animation/ControlRig/Source/ControlRigEditor/Private/EditMode/ControlRigEditMode.cpp b/Engine/Plugins/Animation/ControlRig/Source/ControlRigEditor/Private/EditMode/ControlRigEditMode.cpp
--- a/Engine/Plugins/Animation/ControlRig/Source/ControlRigEditor/Private/EditMode/ControlRigEditMode.cpp
+++ b/Engine/Plugins/Animation/ControlRig/Source/ControlRigEditor/Private/EditMode/ControlRigEditMode.cpp
@@ -1234,8 +1234,7 @@ bool FControlRigEditMode::HandleBeginTransform(const FEditorViewportClient* InVi
 				if (AreRigElementSelectedAndMovable(ControlRig))
 				{
 					//todo need to add multiple
-					TArray<FRigElementKey> SelectedRigElements = GetSelectedRigElements(ControlRig);
-					FControlRigInteractionScope* InteractionScope = new FControlRigInteractionScope(ControlRig, SelectedRigElements, (EControlRigInteractionType)InteractionType);
+					FControlRigInteractionScope* InteractionScope = new FControlRigInteractionScope(ControlRig);
 					InteractionScopes.Add(ControlRig,InteractionScope);
 					ControlRig->bInteractionJustBegan = true;
 				}
@@ -1526,8 +1525,7 @@ bool FControlRigEditMode::HandleClick(FEditorViewportClient* InViewportClient, H
 
 					if (bInteractionScopePresent)
 					{
-						TArray<FRigElementKey> SelectedRigElements = GetSelectedRigElements(ShapeActor->ControlRig.Get());
-						*InteractionScope = new FControlRigInteractionScope(ShapeActor->ControlRig.Get(), SelectedRigElements, (EControlRigInteractionType) InteractionType);
+						*InteractionScope = new FControlRigInteractionScope(ShapeActor->ControlRig.Get());
 						InteractionScopes.Add(ShapeActor->ControlRig.Get(), *InteractionScope);
 
 					}
diff --git a/Engine/Plugins/Animation/ControlRig/Source/ControlRigEditor/Public/EditMode/ControlRigEditMode.h b/Engine/Plugins/Animation/ControlRig/Source/ControlRigEditor/Public/EditMode/ControlRigEditMode.h
--- a/Engine/Plugins/Animation/ControlRig/Source/ControlRigEditor/Public/EditMode/ControlRigEditMode.h
+++ b/Engine/Plugins/Animation/ControlRig/Source/ControlRigEditor/Public/EditMode/ControlRigEditMode.h
@@ -131,7 +131,7 @@ private:
 
 };
 
-USTRUCT()
+USTRUCT(BlueprintType)
 struct CONTROLRIGEDITOR_API FMultiControlRigElementSelection
 {
 	GENERATED_BODY()
diff --git a/Engine/Plugins/Animation/ControlRig/Source/ControlRigEditor/Public/EditMode/ControlRigEditModeUtil.h b/Engine/Plugins/Animation/ControlRig/Source/ControlRigEditor/Public/EditMode/ControlRigEditModeUtil.h
--- a/Engine/Plugins/Animation/ControlRig/Source/ControlRigEditor/Public/EditMode/ControlRigEditModeUtil.h
+++ b/Engine/Plugins/Animation/ControlRig/Source/ControlRigEditor/Public/EditMode/ControlRigEditModeUtil.h
@@ -11,6 +11,8 @@
 #include "UObject/WeakObjectPtr.h"
 #include "UnrealWidgetFwd.h"
 
+#include "ControlRigEditModeUtil.generated.h"
+
 class FEditorViewportClient;
 class UControlRig;
 class ULevelEditorViewportSettings;
@@ -20,7 +22,53 @@ class UTickableConstraint;
 struct FConvexVolume;
 struct FRigControlModifiedContext;
 struct FRigControlElement;
-struct FControlRigInteractionTransformContext;
+
+
+UENUM()
+enum class EControlRigInteractionTransformSpace
+{
+	World,
+	Local,
+	Parent,
+	Explicit
+};
+
+/**
+ * FControlRigInteractionTransformContext provides a way of passing the various transform parameters to functions that need to know what the transform context is.
+ * Extend it if necessary, particularly to avoid overloading some functions signatures.
+ */
+USTRUCT(BlueprintType)
+struct FControlRigInteractionTransformContext
+{
+	GENERATED_BODY()
+	FControlRigInteractionTransformContext() = default;
+	
+	FControlRigInteractionTransformContext(const UE::Widget::EWidgetMode& InWidgetMode)
+		: bTranslation(InWidgetMode == UE::Widget::WM_Translate || InWidgetMode == UE::Widget::WM_TranslateRotateZ)
+		, bRotation(InWidgetMode == UE::Widget::WM_Rotate || InWidgetMode == UE::Widget::WM_TranslateRotateZ)
+		, bScale(InWidgetMode == UE::Widget::WM_Scale)
+	{}
+	
+	bool CanTransform() const
+	{
+		return bTranslation || bRotation || bScale;
+	}
+
+	bool bTranslation = false;
+	UPROPERTY(BlueprintReadWrite, EditAnywhere, Category = "Interaction")
+	FVector Drag = FVector::ZeroVector;
+
+	bool bRotation = false;
+	UPROPERTY(BlueprintReadWrite, EditAnywhere, Category = "Interaction")
+	FRotator Rot = FRotator::ZeroRotator;
+
+	bool bScale = false;
+	UPROPERTY(BlueprintReadWrite, EditAnywhere, Category = "Interaction")
+	FVector Scale = FVector::OneVector;
+
+	UPROPERTY(BlueprintReadWrite, EditAnywhere, Category = "Interaction")
+	EControlRigInteractionTransformSpace Space = EControlRigInteractionTransformSpace::World;
+};
 
 namespace UE::ControlRigEditMode
 {
