Control Rig: Fix undo / redo not selecting control rig.

This bug is a regression, which SEEMS like caused by CL 41258461. That CL removed the class UControlRigDetailPanelControlProxies. In a 5.5.3 build I tried selecting controls and looking at the recorded objects in the Undo stack... it contained instances of this class.
To workaround this issue, some other place in the control rig architecture now needs to take care of the undo / redo logic of selection. Since the URigHierarchy already saves the selection state, it seems like the right candidate to broadcast when its selection changes due to undo / redo. It seems like the easiest solution.

diff --git a/Engine/Plugins/Animation/ControlRig/Source/ControlRig/Private/Rigs/RigHierarchy.cpp b/Engine/Plugins/Animation/ControlRig/Source/ControlRig/Private/Rigs/RigHierarchy.cpp
--- a/Engine/Plugins/Animation/ControlRig/Source/ControlRig/Private/Rigs/RigHierarchy.cpp
+++ b/Engine/Plugins/Animation/ControlRig/Source/ControlRig/Private/Rigs/RigHierarchy.cpp
@@ -4405,6 +4405,13 @@ bool URigHierarchy::SetTransformStackIndex(int32 InTransformStackIndex)
 
 #if WITH_EDITOR
 
+void URigHierarchy::PreEditUndo()
+{
+	Super::PreEditUndo();
+
+	SelectedKeysBeforeUndo = GetSelectedHierarchyKeys();
+}
+
 void URigHierarchy::PostEditUndo()
 {
 	Super::PostEditUndo();
@@ -4420,6 +4427,8 @@ void URigHierarchy::PostEditUndo()
 	{
 		Controller->SetHierarchy(this);
 	}
+
+	NotifyPostUndoSelectionChanges();
 }
 
 #endif
@@ -8875,6 +8884,35 @@ void URigHierarchy::ComputeAllTransforms()
 	}
 }
 
+#if WITH_EDITOR
+void URigHierarchy::NotifyPostUndoSelectionChanges()
+{
+	for (const FRigHierarchyKey& PreviouslySelectedKey : SelectedKeysBeforeUndo)
+	{
+		FRigBaseElement* Element = Find<FRigBaseElement>(PreviouslySelectedKey.GetElement());
+		if (!Element) // Was element removed by the transaction?
+		{
+			continue;
+		}
+
+		if (!Element->IsSelected())
+		{
+			Notify(ERigHierarchyNotification::ElementDeselected, FRigNotificationSubject(Element));
+		}
+	}
+
+	for (const FRigBaseElement* Element : GetSelectedElements(ERigElementType::All))
+	{
+		if (!SelectedKeysBeforeUndo.Contains(Element->Key))
+		{
+			Notify(ERigHierarchyNotification::ElementSelected, FRigNotificationSubject(Element));
+		}
+	}
+
+	SelectedKeysBeforeUndo.Reset();
+}
+#endif
+
 bool URigHierarchy::IsAnimatable(const FRigElementKey& InKey) const
 {
 	if(const FRigControlElement* ControlElement = Find<FRigControlElement>(InKey))
diff --git a/Engine/Plugins/Animation/ControlRig/Source/ControlRig/Private/Rigs/RigHierarchyController.cpp b/Engine/Plugins/Animation/ControlRig/Source/ControlRig/Private/Rigs/RigHierarchyController.cpp
--- a/Engine/Plugins/Animation/ControlRig/Source/ControlRig/Private/Rigs/RigHierarchyController.cpp
+++ b/Engine/Plugins/Animation/ControlRig/Source/ControlRig/Private/Rigs/RigHierarchyController.cpp
@@ -168,8 +168,15 @@ bool URigHierarchyController::SelectHierarchyKey(FRigHierarchyKey InKey, bool bS
 		{
 			return false;
 		}
+
+		// If recording into transaction buffer, modify the hierarchy so undoing selection works correctly.
+		if (GUndo)
+		{
+			Hierarchy->Modify();
+		}
+		
 		Element->bSelected = bSelect;
-	
+
 		if(bSelect)
 		{
 			Hierarchy->OrderedSelection.Add(Element->GetKey());
@@ -205,6 +212,12 @@ bool URigHierarchyController::SelectHierarchyKey(FRigHierarchyKey InKey, bool bS
 			return false;
 		}
 
+		// If recording into transaction buffer, modify the hierarchy so undoing selection works correctly.
+		if (GUndo)
+		{
+			Hierarchy->Modify();
+		}
+		
 		Component->bSelected = bSelect;
 	
 		if(bSelect)
diff --git a/Engine/Plugins/Animation/ControlRig/Source/ControlRig/Public/Rigs/RigHierarchy.h b/Engine/Plugins/Animation/ControlRig/Source/ControlRig/Public/Rigs/RigHierarchy.h
--- a/Engine/Plugins/Animation/ControlRig/Source/ControlRig/Public/Rigs/RigHierarchy.h
+++ b/Engine/Plugins/Animation/ControlRig/Source/ControlRig/Public/Rigs/RigHierarchy.h
@@ -190,6 +190,7 @@ public:
 	static CONTROLRIG_API void DeclareConstructClasses(TArray<FTopLevelAssetPath>& OutConstructClasses, const UClass* SpecificSubclass);
 #endif
 #if WITH_EDITOR
+	virtual void PreEditUndo() override;
 	CONTROLRIG_API virtual void PostEditUndo() override;
 #endif
 
@@ -4890,6 +4891,11 @@ private:
 	 */
 	TArray<FRigTransformStackEntry> TransformRedoStack;
 
+#if WITH_EDITORONLY_DATA
+	/** The keys that were selected before undo / redo. */
+	TArray<FRigHierarchyKey> SelectedKeysBeforeUndo;
+#endif
+
 	/**
 	 * Sets the transform stack index - which in turns performs a series of undos / redos
 	 * @param InTransformStackIndex The new index for the transform stack
@@ -4929,6 +4935,11 @@ private:
 	 */
 	CONTROLRIG_API void ComputeAllTransforms();
 
+#if WITH_EDITOR
+	/** Compares the selection changes made by between the last PreEditUndo and PostEditUndo calls and calls Notify to inform listening systems. */
+	void NotifyPostUndoSelectionChanges();
+#endif
+
 	/**
 	 * Manages merging transform actions into one during an interaction
 	 */
