Control Rig: Fix ouliner grouping rigs from different actors with same name
#rnx

diff --git a/Engine/Plugins/Animation/ControlRig/Source/ControlRigEditor/Private/EditMode/SControlRigOutliner.cpp b/Engine/Plugins/Animation/ControlRig/Source/ControlRigEditor/Private/EditMode/SControlRigOutliner.cpp
--- a/Engine/Plugins/Animation/ControlRig/Source/ControlRigEditor/Private/EditMode/SControlRigOutliner.cpp
+++ b/Engine/Plugins/Animation/ControlRig/Source/ControlRigEditor/Private/EditMode/SControlRigOutliner.cpp
@@ -159,7 +159,8 @@ bool FMultiRigData::operator== (const FMultiRigData & Other) const
 	if (WeakControlRig == Other.WeakControlRig
 		&& Type == Other.Type
 		&& Key == Other.Key
-		&& Name == Other.Name)
+		&& Name == Other.Name
+		&& UniqueID == Other.UniqueID)
 	{
 		return true;
 	}
@@ -182,7 +183,7 @@ bool FMultiRigData::IsValid() const
 		return true;
 	}
 	
-	return (IsActor() || IsComponent());
+	return (IsActor() || IsComponent()) && UniqueID.IsSet();
 }
 
 URigHierarchy* FMultiRigData::GetHierarchy() const
@@ -1305,6 +1306,7 @@ void SMultiRigHierarchyTreeView::RefreshTreeView()
 							GetAscestorName(Component, ComponentName);
 							
 							Data.SetItemName(FMultiRigData::EMultiRigDataType_Component, *ComponentName);
+							Data.SetUniqueID(Component->GetUniqueID());
 							DataHierarchy.Add(Data);
 						}
 					}
@@ -1316,6 +1318,7 @@ void SMultiRigHierarchyTreeView::RefreshTreeView()
 						GetAscestorName(Actor, ActorName);
 	
 						Data.SetItemName(FMultiRigData::EMultiRigDataType_Actor, *ActorName);
+						Data.SetUniqueID(Actor->GetUniqueID());
 						DataHierarchy.Add(Data);
 					}
 					ParentObject = ParentObject->GetOuter();
diff --git a/Engine/Plugins/Animation/ControlRig/Source/ControlRigEditor/Private/EditMode/SControlRigOutliner.h b/Engine/Plugins/Animation/ControlRig/Source/ControlRigEditor/Private/EditMode/SControlRigOutliner.h
--- a/Engine/Plugins/Animation/ControlRig/Source/ControlRigEditor/Private/EditMode/SControlRigOutliner.h
+++ b/Engine/Plugins/Animation/ControlRig/Source/ControlRigEditor/Private/EditMode/SControlRigOutliner.h
@@ -123,6 +123,7 @@ public:
 	
 	void SetElementKey(const FRigElementKey& InKey) { Type = EMultiRigDataType_Element; Key = InKey; }
 	void SetItemName(EMultiRigDataType InType, const FName& InName) { Type = InType; Name = InName; }
+	void SetUniqueID(const uint32 InID) { UniqueID = InID; }
 	
 	TWeakObjectPtr<UControlRig> WeakControlRig;
 	EMultiRigDataType Type;
@@ -131,6 +132,7 @@ private:
 	
 	TOptional<FName> Name;
 	TOptional<FRigElementKey> Key;
+	TOptional<uint32> UniqueID;
 	
 	mutable TOptional<FText> CachedDisplayName;
 };
