Control Rig: Copy and paste for module settings

#rnx

diff --git a/Engine/Plugins/Animation/ControlRig/Source/ControlRig/Private/ControlRigOverride.cpp b/Engine/Plugins/Animation/ControlRig/Source/ControlRig/Private/ControlRigOverride.cpp
--- a/Engine/Plugins/Animation/ControlRig/Source/ControlRig/Private/ControlRigOverride.cpp
+++ b/Engine/Plugins/Animation/ControlRig/Source/ControlRig/Private/ControlRigOverride.cpp
@@ -57,12 +57,12 @@ FControlRigOverrideValue::FControlRigOverrideValue(const FString& InPath, const
 	}
 }
 
-FControlRigOverrideValue::FControlRigOverrideValue(const FString& InPath, const UStruct* InOwnerStruct, const FString& InValueAsString, const FName& InSubjectKey)
+FControlRigOverrideValue::FControlRigOverrideValue(const FString& InPath, const UStruct* InOwnerStruct, const FString& InValueAsString, const FName& InSubjectKey, const FControlRigOverrideValueErrorPipe::TReportFunction& InReportFunction)
 	: SubjectKey(InSubjectKey)
 {
 	if(SetPropertiesFromPath(InPath, InOwnerStruct))
 	{
-		SetFromString(InValueAsString);
+		SetFromString(InValueAsString, InReportFunction);
 	}
 }
 
@@ -650,7 +650,7 @@ const FString& FControlRigOverrideValue::ToString() const
 	return CachedStringValue.Get(EmptyPath);
 }
 
-bool FControlRigOverrideValue::SetFromString(const FString& InValue)
+bool FControlRigOverrideValue::SetFromString(const FString& InValue, const FControlRigOverrideValueErrorPipe::TReportFunction& InReportFunction)
 {
 	if(InValue != CachedStringValue.Get(FString()))
 	{
@@ -660,9 +660,9 @@ bool FControlRigOverrideValue::SetFromString(const FString& InValue)
 			void* TempDataPtr = AllocateDataIfRequired(TempDataArray); 
 
 			// first try to import the string value using the property
-			FRigVMMemoryStorageImportErrorContext ErrorPipe(false);
+			FControlRigOverrideValueErrorPipe ErrorPipe(ELogVerbosity::Warning, InReportFunction);
 			LeafProperty->ImportText_Direct(*InValue, TempDataPtr, nullptr, PPF_None, &ErrorPipe);
-			if(ErrorPipe.NumErrors == 0)
+			if(ErrorPipe.GetNumErrors() == 0)
 			{
 				FreeDataIfRequired(DataArray);
 				Swap(DataArray, TempDataArray);
diff --git a/Engine/Plugins/Animation/ControlRig/Source/ControlRig/Private/ModularRigController.cpp b/Engine/Plugins/Animation/ControlRig/Source/ControlRig/Private/ModularRigController.cpp
--- a/Engine/Plugins/Animation/ControlRig/Source/ControlRig/Private/ModularRigController.cpp
+++ b/Engine/Plugins/Animation/ControlRig/Source/ControlRig/Private/ModularRigController.cpp
@@ -1446,6 +1446,12 @@ FName UModularRigController::RenameModule(const FName& InModuleName, const FName
 		};
 	}
 
+	// fix overrides
+	for(FControlRigOverrideValue& Override : Module->ConfigOverrides)
+	{
+		Override.SubjectKey = Module->Name;
+	}
+
 	// make sure to update our backwards compat code from path to module name
 	for(TPair<FRigHierarchyModulePath, FName>& Pair : Model->PreviousModulePaths)
 	{
@@ -2041,6 +2047,280 @@ void UModularRigController::RefreshModuleVariables(const FRigModuleReference* In
 #endif
 }
 
+FString UModularRigController::ExportModuleSettingsToString(TArray<FName> InModuleNames) const
+{
+	FModularRigModuleSettingsSetForClipboard Content;
+	for(const FName& ModuleName : InModuleNames)
+	{
+		if(const FRigModuleReference* Module = FindModule(ModuleName))
+		{
+			if(Content.Settings.Contains(ModuleName))
+			{
+				continue;
+			}
+			
+			const UClass* ControlRigClass = Module->Class.Get();
+			if(!ControlRigClass)
+			{
+				UE_LOG(LogControlRig, Error, TEXT("Module '%s' does not have a valid control rig class (%s) associated."), *ModuleName.ToString(), *Module->Class.ToString());
+				return FString();
+			}
+
+			FModularRigModuleSettingsForClipboard& Settings = Content.Settings.Add(ModuleName);
+			
+			Settings.ModuleClass = Module->Class.ToSoftObjectPath();
+
+			// store the overrides as configured by the user
+			for(const FControlRigOverrideValue& Override : Module->ConfigOverrides)
+			{
+				if(Override.IsValid())
+				{
+					Settings.Overrides.Add(Override.GetPath(), Override.ToString());
+				}
+			}
+
+			// also store all of the defaults for the rig (including the changes introduced by the current overrides
+			const UControlRig* CDO = ControlRigClass->GetDefaultObject<UControlRig>();
+			for (TFieldIterator<FProperty> PropertyIt( ControlRigClass ); PropertyIt; ++PropertyIt)
+			{
+				const FProperty* Property = *PropertyIt;
+				if(Property->IsNative())
+				{
+					continue;
+				}
+
+				const FString PropertyPath = Property->GetName();
+
+				// store the default
+				const FControlRigOverrideValue Default(PropertyPath, ControlRigClass, CDO);
+				if(Default.IsValid())
+				{
+					Settings.Defaults.Add(Default.GetPath(), Default.ToString());
+				}
+			}
+
+			// store the bindings as well
+			Settings.Bindings = Module->Bindings;
+		}
+		else
+		{
+			UE_LOG(LogControlRig, Error, TEXT("Module '%s' not found."), *ModuleName.ToString());
+			return FString();
+		}
+	}
+
+	FString ContentAsString;
+	FModularRigModuleSettingsSetForClipboard::StaticStruct()->ExportText(ContentAsString, &Content, &Content, nullptr, PPF_None, nullptr);
+	return ContentAsString; 
+}
+
+bool UModularRigController::ImportModuleSettingsFromString(FString InContent, TArray<FName> InOptionalModuleNames, bool bSetupUndo)
+{
+	FControlRigOverrideValueErrorPipe ErrorPipe(ELogVerbosity::Warning, [this](const TCHAR* V, ELogVerbosity::Type Verbosity)
+	{
+		UE_LOG(LogControlRig, Warning, TEXT("Error during import: %s"), V);
+	});
+	
+	FModularRigModuleSettingsSetForClipboard Content;
+	FModularRigModuleSettingsSetForClipboard::StaticStruct()->ImportText(*InContent, &Content, nullptr, PPF_None, &ErrorPipe, FModularRigModuleSettingsForClipboard::StaticStruct()->GetName(), true);
+	if(ErrorPipe.GetNumErrors() > 0)
+	{
+		return false;
+	}
+
+	TArray<FName> ContentModuleNames;
+	Content.Settings.GenerateKeyArray(ContentModuleNames);
+
+	const TArray<FName> ModuleNames = InOptionalModuleNames.IsEmpty() ? ContentModuleNames : InOptionalModuleNames;
+	if(ModuleNames.Num() != ContentModuleNames.Num())
+	{
+		UE_LOG(LogControlRig, Error, TEXT("The number of modules selected (%d) doesn't match the number of modules on the clipboard (%d)."), ModuleNames.Num(), ContentModuleNames.Num());
+		return false;
+	}
+	
+	TMap<FName, FName> ModuleNameLookup, InvModuleNameLookup;
+	for(int32 Index=0;Index<ModuleNames.Num();Index++)
+	{
+		ModuleNameLookup.Add(ModuleNames[Index], ContentModuleNames[Index]);
+		InvModuleNameLookup.Add(ContentModuleNames[Index], ModuleNames[Index]);
+	}
+
+#if WITH_EDITOR
+	TSharedPtr<FScopedTransaction> TransactionPtr;
+	if (bSetupUndo)
+	{
+		TransactionPtr = MakeShared<FScopedTransaction>(NSLOCTEXT("ModularRigController", "ImportModuleSettingsTransaction", "Import Module Settings"), !GIsTransacting);
+	}
+#endif
+
+	TArray<FName> AffectedModules;
+	auto RecordChange = [this, &AffectedModules](const FName& InModuleName)
+	{
+#if WITH_EDITOR
+		if(UBlueprint* Blueprint = Cast<UBlueprint>(GetOuter()))
+		{
+			Blueprint->Modify();
+		}
+#endif
+		AffectedModules.AddUnique(InModuleName);
+	};
+
+	Notify(EModularRigNotification::InteractionBracketOpened, nullptr);
+
+	for(const FName& InputModuleName : ModuleNames)
+	{
+		if(FRigModuleReference* Module = FindModule(InputModuleName))
+		{
+			const FName& ModuleName = ModuleNameLookup.FindChecked(InputModuleName);
+			
+			if(!Content.Settings.Contains(ModuleName))
+			{
+				UE_LOG(LogControlRig, Error, TEXT("There are no settings provided in the content for module '%s'."), *ModuleName.ToString());
+				continue;
+			}
+
+			const UClass* ControlRigClass = Module->Class.Get();
+			if(!ControlRigClass)
+			{
+				UE_LOG(LogControlRig, Error, TEXT("Module '%s' does not have a valid control rig class (%s) associated."), *ModuleName.ToString(), *Module->Class.ToString());
+				return false;
+			}
+			
+			if(Module->Class.ToSoftObjectPath() != Content.Settings.FindChecked(ModuleName).ModuleClass)
+			{
+				UE_LOG(LogControlRig, Warning, TEXT("Classes for Module '%s' don't match between the current rig and the clipboard content. Still attempting to apply settings."), *ModuleName.ToString());
+			}
+
+			// remove all overrides
+			if(!Module->ConfigOverrides.IsEmpty())
+			{
+				RecordChange(Module->Name);
+				Module->ConfigOverrides.Reset();
+			}
+
+			FModularRigModuleSettingsForClipboard& Settings = Content.Settings.FindChecked(ModuleName);
+
+			// first compare if the defaults are different between the copied content and now
+			const UControlRig* CDO = CastChecked<UControlRig>(ControlRigClass->GetDefaultObject());
+
+			UControlRig* TemporaryRig = NewObject<UControlRig>(GetTransientPackage(), ControlRigClass);
+
+			for(const TPair<FString,FString>& Pair : Settings.Defaults)
+			{
+				// if we have a top level override - let's not worry about this one
+				if(Settings.Overrides.Contains(Pair.Key))
+				{
+					continue;
+				}
+
+				const FControlRigOverrideValueErrorPipe::TReportFunction ReportFunction = [ModuleName, Pair](const TCHAR* V, ELogVerbosity::Type Verbosity)
+				{
+					UE_LOG(LogControlRig, Warning, TEXT("Problem during import of property '%s' for module '%s': %s"), *Pair.Key, *ModuleName.ToString(), V);
+				};
+
+				const FControlRigOverrideValue CopiedDefault(Pair.Key, ControlRigClass, Pair.Value, Module->Name, ReportFunction);
+				if(!CopiedDefault.IsValid())
+				{
+					UE_LOG(LogControlRig, Warning, TEXT("Cannot apply top level override Module '%s' with path '%s'."), *ModuleName.ToString(), *Pair.Key);
+					continue;
+				}
+				
+				const FControlRigOverrideValue DefaultOverride(CopiedDefault.GetPath(), ControlRigClass, CDO, Module->Name);
+				if(DefaultOverride.Identical(CopiedDefault))
+				{
+					continue;
+				}
+				
+				// copy the original default to the temporary rig
+				CopiedDefault.CopyToUObject(TemporaryRig);
+
+				// copy all other overrides that potentially sit under there
+				for(const TPair<FString,FString>& OverridePair : Settings.Overrides)
+				{
+					if(!FControlRigOverrideContainer::IsChildPathOf(OverridePair.Key, DefaultOverride.GetPath()))
+					{
+						continue;
+					}
+					const FControlRigOverrideValue ChildOverride(OverridePair.Key, ControlRigClass, OverridePair.Value, Module->Name, ReportFunction);
+					if(!ChildOverride.IsValid())
+					{
+						continue;
+					}
+
+					ChildOverride.CopyToUObject(TemporaryRig);
+				}
+
+				// construct the new override from the temporary rig, combining the copied default and any additional child override
+				const FControlRigOverrideValue CombinedOverride(CopiedDefault.GetPath(), ControlRigClass, TemporaryRig, Module->Name);
+				if(!CombinedOverride.IsValid())
+				{
+					continue;
+				}
+				
+				RecordChange(Module->Name);
+				Module->ConfigOverrides.Add(CombinedOverride);
+			}
+			
+			// now apply all of the user provided overrides as well
+			for(const TPair<FString,FString>& Pair : Settings.Overrides)
+			{
+				const FControlRigOverrideValueErrorPipe::TReportFunction ReportFunction = [ModuleName, Pair](const TCHAR* V, ELogVerbosity::Type Verbosity)
+				{
+					UE_LOG(LogControlRig, Warning, TEXT("Problem during import of property '%s' for module '%s': %s"), *Pair.Key, *ModuleName.ToString(), V);
+				};
+				
+				const FControlRigOverrideValue UserProvidedOverride(Pair.Key, ControlRigClass, Pair.Value, Module->Name, ReportFunction);
+				if(!UserProvidedOverride.IsValid())
+				{
+					UE_LOG(LogControlRig, Warning, TEXT("Cannot apply top level override Module '%s' with path '%s'."), *ModuleName.ToString(), *Pair.Key);
+					continue;
+				}
+
+				// this fails in case there's already an override on the parent property path.
+				// if there's already an override on .Color, the secondary override on say .Color.R is ignored.
+				RecordChange(Module->Name);
+				Module->ConfigOverrides.Add(UserProvidedOverride);
+			}
+
+			Module->Bindings.Reset();
+
+			// apply the bindings
+			for(const TPair<FName, FString>& Pair : Settings.Bindings)
+			{
+				// potentially remap the source path within the provided set
+				FRigHierarchyModulePath Path = Pair.Value;
+				if(const FName* RemappedName = InvModuleNameLookup.Find(Path.GetModuleFName()))
+				{
+					Path.SetModuleName(*RemappedName);
+				}
+				RecordChange(Module->Name);
+				if(!BindModuleVariable(Module->Name, Pair.Key, Path.GetPath(), bSetupUndo))
+				{
+					UE_LOG(LogControlRig, Warning, TEXT("Cannot recreate binding for module '%s', property '%s' to '%s'."), *Module->Name.ToString(), *Pair.Key.ToString(), *Path.GetPath());
+				}
+			}
+		}
+	}
+
+#if WITH_EDITOR
+	TransactionPtr.Reset();
+#endif
+
+	if(!AffectedModules.IsEmpty())
+	{
+		for(const FName& ModuleName : AffectedModules)
+		{
+			if(const FRigModuleReference* Module = FindModule(ModuleName))
+			{
+				Notify(EModularRigNotification::ModuleConfigValueChanged, Module);
+			}
+		}
+	}
+	Notify(EModularRigNotification::InteractionBracketClosed, nullptr);
+
+	return !AffectedModules.IsEmpty();
+}
+
 void UModularRigController::SanitizeName(FRigName& InOutName, bool bAllowNameSpaces)
 {
 	// Sanitize the name
diff --git a/Engine/Plugins/Animation/ControlRig/Source/ControlRig/Public/ControlRigOverride.h b/Engine/Plugins/Animation/ControlRig/Source/ControlRig/Public/ControlRigOverride.h
--- a/Engine/Plugins/Animation/ControlRig/Source/ControlRig/Public/ControlRigOverride.h
+++ b/Engine/Plugins/Animation/ControlRig/Source/ControlRig/Public/ControlRigOverride.h
@@ -7,6 +7,8 @@
 #include "UObject/WeakFieldPtr.h"
 #include "RigVMStringUtils.h"
 #include "HAL/IConsoleManager.h"
+#include "Misc/OutputDevice.h"
+#include "Logging/LogVerbosity.h"
 #include "ControlRigOverride.generated.h"
 
 #define UE_API CONTROLRIG_API
@@ -17,6 +19,41 @@ class FPropertyPath;
 
 extern CONTROLRIG_API TAutoConsoleVariable<bool> CVarControlRigEnableOverrides;
 
+class FControlRigOverrideValueErrorPipe : public FOutputDevice
+{
+public:
+
+	typedef  TFunction<void(const TCHAR*, ELogVerbosity::Type)> TReportFunction;
+
+	UE_API FControlRigOverrideValueErrorPipe( ELogVerbosity::Type InMaxVerbosity = ELogVerbosity::Warning, TReportFunction InReportFunction = nullptr )
+		: FOutputDevice()
+		, NumErrors(0)
+		, MaxVerbosity(InMaxVerbosity)
+		, ReportFunction(InReportFunction)
+	{
+	}
+
+	UE_API virtual void Serialize(const TCHAR* V, ELogVerbosity::Type Verbosity, const class FName& Category) override
+	{
+		if(Verbosity <= MaxVerbosity)
+		{
+			NumErrors++;
+			if(ReportFunction)
+			{
+				ReportFunction(V, Verbosity);
+			}
+		}
+	}
+
+	UE_API ELogVerbosity::Type GetMaxVerbosity() const { return MaxVerbosity; }
+	UE_API int32 GetNumErrors() const { return NumErrors; }
+
+private:
+	int32 NumErrors;
+	ELogVerbosity::Type MaxVerbosity;
+	TReportFunction ReportFunction;
+};
+
 /**
  * A single value used to represent an override on a subject.
  * The value serialized the data using binary serialization based on the last
@@ -46,7 +83,7 @@ public:
 	UE_API FControlRigOverrideValue(const FString& InPath, const UStruct* InOwnerStruct, const void* InSubjectPtr, const FName& InSubjectKey = NAME_None);
 
 	// constructor given the path, the owning structure as well as the value as string
-	UE_API FControlRigOverrideValue(const FString& InPath, const UStruct* InOwnerStruct, const FString& InValueAsString, const FName& InSubjectKey = NAME_None);
+	UE_API FControlRigOverrideValue(const FString& InPath, const UStruct* InOwnerStruct, const FString& InValueAsString, const FName& InSubjectKey = NAME_None, const FControlRigOverrideValueErrorPipe::TReportFunction& InReportFunction = nullptr);
 
 	UE_API FControlRigOverrideValue(const FControlRigOverrideValue& InOther);
 
@@ -99,7 +136,7 @@ public:
 	UE_API const FString& ToString() const;
 
 	// sets the contents of this value from string
-	UE_API bool SetFromString(const FString& InValue);
+	UE_API bool SetFromString(const FString& InValue, const FControlRigOverrideValueErrorPipe::TReportFunction& InReportFunction = nullptr);
 
 	// copies the override value onto a container
 	[[nodiscard]] UE_API bool CopyToSubject(void* InSubjectPtr, const UStruct* InSubjectStruct = nullptr) const;
@@ -195,6 +232,8 @@ private:
 
 	TArray<FPropertyInfo> Properties;
 	TArray<uint8, TAlignedHeapAllocator<16>> DataArray;
+	
+	friend class UModularRigController;
 };
 
 template<>
diff --git a/Engine/Plugins/Animation/ControlRig/Source/ControlRig/Public/ModularRigController.h b/Engine/Plugins/Animation/ControlRig/Source/ControlRig/Public/ModularRigController.h
--- a/Engine/Plugins/Animation/ControlRig/Source/ControlRig/Public/ModularRigController.h
+++ b/Engine/Plugins/Animation/ControlRig/Source/ControlRig/Public/ModularRigController.h
@@ -111,6 +111,12 @@ class UModularRigController : public UObject
 
 	UE_API void RefreshModuleVariables(bool bSetupUndo = true);
 	UE_API void RefreshModuleVariables(const FRigModuleReference* InModule, bool bSetupUndo = true);
+	
+	UFUNCTION(BlueprintPure, Category = "ControlRig | Modules")
+	UE_API FString ExportModuleSettingsToString(TArray<FName> InModuleNames) const;
+
+	UFUNCTION(BlueprintCallable, Category = "ControlRig | Modules")
+	UE_API bool ImportModuleSettingsFromString(FString InContent, TArray<FName> InOptionalModuleNames, bool bSetupUndo = true);
 
 	static int32 GetMaxNameLength() { return 100; }
 	static UE_API void SanitizeName(FRigName& InOutName, bool bAllowNameSpaces);
diff --git a/Engine/Plugins/Animation/ControlRig/Source/ControlRig/Public/ModularRigModel.h b/Engine/Plugins/Animation/ControlRig/Source/ControlRig/Public/ModularRigModel.h
--- a/Engine/Plugins/Animation/ControlRig/Source/ControlRig/Public/ModularRigModel.h
+++ b/Engine/Plugins/Animation/ControlRig/Source/ControlRig/Public/ModularRigModel.h
@@ -417,4 +417,33 @@ private:
 	friend struct FRigModuleReference;
 };
 
+// A transient struct used for copy & paste of module settings
+USTRUCT(BlueprintType)
+struct FModularRigModuleSettingsForClipboard
+{
+	GENERATED_BODY()
+
+	UPROPERTY(EditAnywhere, BlueprintReadWrite, Category="Module")
+	FSoftObjectPath ModuleClass;
+
+	UPROPERTY(EditAnywhere, BlueprintReadWrite, Category="Module")
+   	TMap<FString, FString> Defaults;
+
+	UPROPERTY(EditAnywhere, BlueprintReadWrite, Category="Module")
+	TMap<FString, FString> Overrides;
+
+	UPROPERTY(EditAnywhere, BlueprintReadWrite, Category="Module")
+	TMap<FName, FString> Bindings;
+};
+
+// A transient struct used for copy & paste of module settings
+USTRUCT(BlueprintType)
+struct FModularRigModuleSettingsSetForClipboard
+{
+	GENERATED_BODY()
+
+	UPROPERTY(EditAnywhere, BlueprintReadWrite, Category="Module")
+	TMap<FName, FModularRigModuleSettingsForClipboard> Settings;
+};
+
 #undef UE_API
diff --git a/Engine/Plugins/Animation/ControlRig/Source/ControlRigDeveloper/Private/ControlRigBlueprint.cpp b/Engine/Plugins/Animation/ControlRig/Source/ControlRigDeveloper/Private/ControlRigBlueprint.cpp
--- a/Engine/Plugins/Animation/ControlRig/Source/ControlRigDeveloper/Private/ControlRigBlueprint.cpp
+++ b/Engine/Plugins/Animation/ControlRig/Source/ControlRigDeveloper/Private/ControlRigBlueprint.cpp
@@ -2000,6 +2000,7 @@ void UControlRigBlueprint::RecompileModularRig()
 		if (UModularRig* DefaultObject = Cast<UModularRig>(MyControlRigClass->GetDefaultObject(false)))
 		{
 			PropagateModuleHierarchyFromBPToInstances();
+			RequestConstructionOnAllModules();
 		}
 	}
 	UpdateModularDependencyDelegates();
diff --git a/Engine/Plugins/Animation/ControlRig/Source/ControlRigEditor/Private/ControlRigModularRigCommands.cpp b/Engine/Plugins/Animation/ControlRig/Source/ControlRigEditor/Private/ControlRigModularRigCommands.cpp
--- a/Engine/Plugins/Animation/ControlRig/Source/ControlRigEditor/Private/ControlRigModularRigCommands.cpp
+++ b/Engine/Plugins/Animation/ControlRig/Source/ControlRigEditor/Private/ControlRigModularRigCommands.cpp
@@ -12,6 +12,8 @@ void FControlRigModularRigCommands::RegisterCommands()
 	UI_COMMAND(MirrorModuleItem, "Mirror", "Mirror module from the rig.", EUserInterfaceActionType::Button, FInputChord());
 	UI_COMMAND(ReresolveModuleItem, "Auto Resolve", "Auto Resolve secondary connectors.", EUserInterfaceActionType::Button, FInputChord());
 	UI_COMMAND(SwapModuleClassItem, "Swap Module", "Swap module of common class.", EUserInterfaceActionType::Button, FInputChord());
+	UI_COMMAND(CopyModuleSettings, "Copy Settings", "Copies the module(s) settings to the clipboard.", EUserInterfaceActionType::Button, FInputChord(EModifierKey::Control, EKeys::C));
+	UI_COMMAND(PasteModuleSettings, "Paste Settings", "Pastes the module(s) settings from the clipboard.", EUserInterfaceActionType::Button, FInputChord(EModifierKey::Control, EKeys::V));
 }
 
 #undef LOCTEXT_NAMESPACE
diff --git a/Engine/Plugins/Animation/ControlRig/Source/ControlRigEditor/Private/ControlRigModularRigCommands.h b/Engine/Plugins/Animation/ControlRig/Source/ControlRigEditor/Private/ControlRigModularRigCommands.h
--- a/Engine/Plugins/Animation/ControlRig/Source/ControlRigEditor/Private/ControlRigModularRigCommands.h
+++ b/Engine/Plugins/Animation/ControlRig/Source/ControlRigEditor/Private/ControlRigModularRigCommands.h
@@ -36,6 +36,12 @@ public:
 	/** Swap Module Class */
 	TSharedPtr< FUICommandInfo > SwapModuleClassItem;
 
+	/** Copy module settings */
+	TSharedPtr< FUICommandInfo > CopyModuleSettings;
+
+	/** Paste module settings */
+	TSharedPtr< FUICommandInfo > PasteModuleSettings;
+
 	/**
 	 * Initialize commands
 	 */
diff --git a/Engine/Plugins/Animation/ControlRig/Source/ControlRigEditor/Private/Editor/SModularRigModel.cpp b/Engine/Plugins/Animation/ControlRig/Source/ControlRigEditor/Private/Editor/SModularRigModel.cpp
--- a/Engine/Plugins/Animation/ControlRig/Source/ControlRigEditor/Private/Editor/SModularRigModel.cpp
+++ b/Engine/Plugins/Animation/ControlRig/Source/ControlRigEditor/Private/Editor/SModularRigModel.cpp
@@ -306,6 +306,14 @@ void SModularRigModel::BindCommands()
 	CommandList->MapAction(Commands.SwapModuleClassItem,
 		FExecuteAction::CreateSP(this, &SModularRigModel::HandleSwapClassForModules),
 		FCanExecuteAction::CreateSP(this, &SModularRigModel::CanSwapModules));
+
+	CommandList->MapAction(Commands.CopyModuleSettings,
+		FExecuteAction::CreateSP(this, &SModularRigModel::HandleCopyModuleSettings),
+		FCanExecuteAction::CreateSP(this, &SModularRigModel::CanCopyModuleSettings));
+
+	CommandList->MapAction(Commands.PasteModuleSettings,
+		FExecuteAction::CreateSP(this, &SModularRigModel::HandlePasteModuleSettings),
+		FCanExecuteAction::CreateSP(this, &SModularRigModel::CanPasteModuleSettings));
 }
 
 FReply SModularRigModel::OnKeyDown(const FGeometry& MyGeometry, const FKeyEvent& InKeyEvent)
@@ -596,6 +604,8 @@ void SModularRigModel::CreateContextMenu()
 					ModulesSection.AddMenuEntry(Commands.MirrorModuleItem);
 					ModulesSection.AddMenuEntry(Commands.ReresolveModuleItem);
 					ModulesSection.AddMenuEntry(Commands.SwapModuleClassItem);
+					ModulesSection.AddMenuEntry(Commands.CopyModuleSettings);
+					ModulesSection.AddMenuEntry(Commands.PasteModuleSettings);
 				}
 			})
 		);
@@ -1194,6 +1204,96 @@ void SModularRigModel::HandleSwapClassForModules(const TArray<FName>& InModuleNa
 	SwapModulesDialog->ShowNormal();
 }
 
+bool SModularRigModel::CanCopyModuleSettings() const
+{
+	if(!ControlRigEditor.IsValid())
+	{
+		return false;
+	}
+	return !GetSelectedItems().IsEmpty();
+}
+
+void SModularRigModel::HandleCopyModuleSettings()
+{
+	if (ControlRigBlueprint.IsValid())
+	{
+		TArray<FName> SelectedModuleNames;
+		const TArray<TSharedPtr<FModularRigTreeElement>> SelectedItems = GetSelectedItems();
+		for (const TSharedPtr<FModularRigTreeElement>& SelectedItem : SelectedItems)
+		{
+			if (SelectedItem.IsValid())
+			{
+				SelectedModuleNames.AddUnique(SelectedItem->ModuleName);
+			}
+		}
+		
+		const UModularRigController* Controller = ControlRigBlueprint->GetModularRigController();
+		check(Controller);
+
+		const FString ContentAsString = Controller->ExportModuleSettingsToString(SelectedModuleNames);
+		if(!ContentAsString.IsEmpty())
+		{
+			FPlatformApplicationMisc::ClipboardCopy(*ContentAsString);
+		}
+	}
+}
+
+bool SModularRigModel::CanPasteModuleSettings() const
+{
+	if(!ControlRigEditor.IsValid())
+	{
+		return false;
+	}
+
+	FString ContentAsString;
+	FPlatformApplicationMisc::ClipboardPaste(ContentAsString);
+	if(ContentAsString.IsEmpty())
+	{
+		return false;
+	}
+
+	FControlRigOverrideValueErrorPipe ErrorPipe;
+	FModularRigModuleSettingsSetForClipboard Content;
+	FModularRigModuleSettingsSetForClipboard::StaticStruct()->ImportText(*ContentAsString, &Content, nullptr, PPF_None, &ErrorPipe, FModularRigModuleSettingsForClipboard::StaticStruct()->GetName(), true);
+	if(ErrorPipe.GetNumErrors() > 0)
+	{
+		return false;
+	}
+	
+	return GetSelectedItems().Num() == Content.Settings.Num();
+}
+
+void SModularRigModel::HandlePasteModuleSettings()
+{
+	if (ControlRigBlueprint.IsValid())
+	{
+		FString ContentAsString;
+		FPlatformApplicationMisc::ClipboardPaste(ContentAsString);
+		if(ContentAsString.IsEmpty())
+		{
+			return;
+		}
+
+		TArray<FName> SelectedModuleNames;
+		const TArray<TSharedPtr<FModularRigTreeElement>> SelectedItems = GetSelectedItems();
+		for (const TSharedPtr<FModularRigTreeElement>& SelectedItem : SelectedItems)
+		{
+			if (SelectedItem.IsValid())
+			{
+				SelectedModuleNames.AddUnique(SelectedItem->ModuleName);
+			}
+		}
+
+		FScopedTransaction Transaction(LOCTEXT("ModularRigModelResolveConnector", "Resolve Connector"));
+
+		UModularRigController* Controller = ControlRigBlueprint->GetModularRigController();
+		check(Controller);
+
+		(void)Controller->ImportModuleSettingsFromString(ContentAsString, SelectedModuleNames);
+	}
+
+}
+
 void SModularRigModel::HandleConnectorResolved(const FRigElementKey& InConnector, const TArray<FRigElementKey>& InTargets)
 {
 	if (ControlRigBlueprint.IsValid())
diff --git a/Engine/Plugins/Animation/ControlRig/Source/ControlRigEditor/Private/Editor/SModularRigModel.h b/Engine/Plugins/Animation/ControlRig/Source/ControlRigEditor/Private/Editor/SModularRigModel.h
--- a/Engine/Plugins/Animation/ControlRig/Source/ControlRigEditor/Private/Editor/SModularRigModel.h
+++ b/Engine/Plugins/Animation/ControlRig/Source/ControlRigEditor/Private/Editor/SModularRigModel.h
@@ -91,6 +91,12 @@ private:
 	void HandleSwapClassForModules();
 	void HandleSwapClassForModules(const TArray<FName>& InModuleNames);
 
+	/** Copy & paste module settings */
+	bool CanCopyModuleSettings() const;
+	void HandleCopyModuleSettings();
+	bool CanPasteModuleSettings() const;
+	void HandlePasteModuleSettings();
+
 	/** Resolve connector */
 	void HandleConnectorResolved(const FRigElementKey& InConnector, const TArray<FRigElementKey>& InTargets);
 
