Anim In Engine: Need to keep track of Guid for Anim Layer scene objects to handle spawnables. We let the track editor handle the delegates, didn't want to pollute the layer code with figuring out how to get/release players as things re-activate.

diff --git a/Engine/Plugins/Animation/ControlRig/Source/ControlRigEditor/Private/Sequencer/AnimLayers/AnimLayers.cpp b/Engine/Plugins/Animation/ControlRig/Source/ControlRigEditor/Private/Sequencer/AnimLayers/AnimLayers.cpp
--- a/Engine/Plugins/Animation/ControlRig/Source/ControlRigEditor/Private/Sequencer/AnimLayers/AnimLayers.cpp
+++ b/Engine/Plugins/Animation/ControlRig/Source/ControlRigEditor/Private/Sequencer/AnimLayers/AnimLayers.cpp
@@ -43,6 +43,7 @@
 #include "MovieSceneToolHelpers.h"
 #define LOCTEXT_NAMESPACE "AnimLayers"
 
+
 static TArray<FGuid> GetSelectedOutlinerGuids(ISequencer* SequencerPtr)
 {
 	using namespace UE::Sequencer;
@@ -67,6 +68,7 @@ struct FObjectAndTrack
 {
 	UMovieScenePropertyTrack* Track;
 	UObject* BoundObject;
+	FGuid SequencerGuid;
 };
 
 static void GetSelectedControlRigsAndBoundObjects(ISequencer* SequencerPtr, TArray<FControlRigAndControlsAndTrack>& OutSelectedCRs, TArray<FObjectAndTrack>& OutBoundObjects)
@@ -137,6 +139,7 @@ static void GetSelectedControlRigsAndBoundObjects(ISequencer* SequencerPtr, TArr
 								FObjectAndTrack ObjectAndTrack;
 								ObjectAndTrack.BoundObject = BoundObject;
 								ObjectAndTrack.Track = Track;
+								ObjectAndTrack.SequencerGuid = Binding.GetObjectGuid();
 								OutBoundObjects.Add(ObjectAndTrack);
 							}
 						}
@@ -179,6 +182,22 @@ FText FAnimLayerState::AnimLayerTypeToText() const
 	return FText(LOCTEXT("Base", "Base"));
 }
 
+void FAnimLayerItem::MakeCopy(const FGuid& NewGuid, const TWeakObjectPtr<UObject>& NewObject,  FAnimLayerItem& OutCopy) const
+{
+	OutCopy.SequencerGuid = NewGuid;
+	for (const FAnimLayerSectionItem& SectionItem : SectionItems)
+	{
+		if (SectionItem.Section.IsValid())
+		{
+			FAnimLayerSectionItem CopySectionItem;
+			CopySectionItem.Section = SectionItem.Section;
+			CopySectionItem.AnimLayerSet = SectionItem.AnimLayerSet;
+			CopySectionItem.AnimLayerSet.BoundObject = NewObject;
+			OutCopy.SectionItems.Add(CopySectionItem);
+		}
+	}
+}
+
 
 UAnimLayer::UAnimLayer(class FObjectInitializer const& ObjectInitializer)
 	: Super(ObjectInitializer)
@@ -186,6 +205,75 @@ UAnimLayer::UAnimLayer(class FObjectInitializer const& ObjectInitializer)
 	WeightProxy = NewObject<UAnimLayerWeightProxy>(this, TEXT("Weight"), RF_Transactional);
 }
 
+
+
+void UAnimLayer::UpdateSceneObjectorGuidsForItems(ISequencer* InSequencer)
+{
+	TArray<TPair<TWeakObjectPtr<UObject>, FGuid>> DeadObjects;
+	for (TPair<TWeakObjectPtr<UObject>, FAnimLayerItem>& Pair : AnimLayerItems)
+	{
+		if (Pair.Key.IsValid() && Pair.Key->IsA<UControlRig>() == false)
+		{
+			UObject* Object = Pair.Key.Pin().Get();
+			//if guid is not set we need to set it.
+			if (Pair.Value.SequencerGuid.IsValid() == false)
+			{
+				USceneComponent* SceneComponent = nullptr;
+				AActor *Actor = nullptr;
+				SceneComponent = Cast<USceneComponent>(Object);
+				if (SceneComponent)
+				{
+					Actor = SceneComponent->GetOwner();
+				}
+				else
+				{
+					Actor = Cast<AActor>(Object);
+				}
+				FGuid Binding;
+				if (SceneComponent)
+				{
+					Binding = InSequencer->GetHandleToObject(SceneComponent, false /*bCreateHandle*/);
+					if (Binding.IsValid())
+					{
+						Pair.Value.SequencerGuid = Binding;
+					}
+					
+				}
+				if(Binding.IsValid() == false && Actor)
+				{
+					Binding = InSequencer->GetHandleToObject(Actor, false /*bCreateHandle*/);
+					if (Binding.IsValid())
+					{
+						Pair.Value.SequencerGuid = Binding;
+					}
+				}
+			}
+		}
+		else if(Pair.Value.SequencerGuid.IsValid())//not valid so update it
+		{
+			TPair<TWeakObjectPtr<UObject>, FGuid> DeadObject;
+			DeadObject.Key = Pair.Key;
+			DeadObject.Value = Pair.Value.SequencerGuid;
+			DeadObjects.Add(DeadObject);
+		}
+	}
+	for (TPair<TWeakObjectPtr<UObject>, FGuid>& DeadObject: DeadObjects)
+	{
+		for (TWeakObjectPtr<> BoundObject : InSequencer->FindBoundObjects(DeadObject.Value, InSequencer->GetFocusedTemplateID()))
+		{
+			if (!BoundObject.IsValid())
+			{
+				continue;
+			}
+			FAnimLayerItem* Item = AnimLayerItems.Find(DeadObject.Key);
+			FAnimLayerItem NewCopy;
+			Item->MakeCopy(DeadObject.Value, BoundObject, NewCopy);
+			AnimLayerItems.Remove(DeadObject.Key);
+			AnimLayerItems.Add(BoundObject,NewCopy);
+		}
+	}
+}
+
 void UAnimLayer::SetKey(TSharedPtr<ISequencer>& Sequencer, const IPropertyHandle& KeyedPropertyHandle)
 {
 	FScopedTransaction PropertyChangedTransaction(LOCTEXT("KeyAnimLayerWeight", "Key Anim Layer Weight"), !GIsTransacting);
@@ -547,6 +635,7 @@ bool UAnimLayer::AddSelectedInSequencer()
 	for (FObjectAndTrack& ObjectAndTrack : SelectedBoundObjects)
 	{
 		FAnimLayerItem& AnimLayerItem = AnimLayerItems.FindOrAdd(ObjectAndTrack.BoundObject);
+		AnimLayerItem.SequencerGuid = ObjectAndTrack.SequencerGuid;
 		FAnimLayerSectionItem SectionItem;
 		SectionItem.AnimLayerSet.BoundObject = ObjectAndTrack.BoundObject;
 		/* if we ever support channels
@@ -2867,6 +2956,8 @@ int32 UAnimLayers::AddAnimLayerFromSelection(ISequencer* SequencerPtr)
 			FAnimLayerItem& AnimLayerItem = AnimLayer->AnimLayerItems.FindOrAdd(ObjectAndTrack.BoundObject);
 			FAnimLayerSectionItem SectionItem;
 			SectionItem.AnimLayerSet.BoundObject = ObjectAndTrack.BoundObject;
+			AnimLayerItem.SequencerGuid = ObjectAndTrack.SequencerGuid;
+
 			/*
 			for (FName& ControlName : SelectedControls)
 			{
diff --git a/Engine/Plugins/Animation/ControlRig/Source/ControlRigEditor/Private/Sequencer/AnimLayers/AnimLayers.h b/Engine/Plugins/Animation/ControlRig/Source/ControlRigEditor/Private/Sequencer/AnimLayers/AnimLayers.h
--- a/Engine/Plugins/Animation/ControlRig/Source/ControlRigEditor/Private/Sequencer/AnimLayers/AnimLayers.h
+++ b/Engine/Plugins/Animation/ControlRig/Source/ControlRigEditor/Private/Sequencer/AnimLayers/AnimLayers.h
@@ -11,6 +11,7 @@
 #include "Styling/SlateTypes.h"
 #include "ISequencerPropertyKeyedStatus.h"
 #include "BakingAnimationKeySettings.h"
+#include "Misc/Guid.h"
 #include "AnimLayers.generated.h"
 
 class ULevelSequence;
@@ -75,6 +76,16 @@ struct FAnimLayerSelectionSet
 	GENERATED_BODY()
 	FAnimLayerSelectionSet() :BoundObject(nullptr) {};
 
+	FAnimLayerSelectionSet& operator = (const FAnimLayerSelectionSet& Other)
+	{
+		BoundObject = Other.BoundObject;
+		for (const TPair<FName, FAnimLayerPropertyAndChannels>& Pair : Other.Names)
+		{
+			Names.Add(Pair.Key, Pair.Value);
+		}
+		return *this;
+	}
+
 	UPROPERTY(BlueprintReadWrite, EditAnywhere, Category = "Data")
 	TWeakObjectPtr<UObject> BoundObject; //bound object is either a CR or a bound sequencer object
 
@@ -108,6 +119,12 @@ struct FAnimLayerItem
 	UPROPERTY(BlueprintReadWrite, EditAnywhere, Category = "Data")
 	TArray< FAnimLayerSectionItem>  SectionItems;
 
+	UPROPERTY()
+	FGuid SequencerGuid;
+
+	//make copy with new guid
+	void MakeCopy(const FGuid& NewGuid, const TWeakObjectPtr<UObject>& NewObject,  FAnimLayerItem& OutCopy) const;
+
 	//find section that matches based upon incoming section (same type and property)
 	FAnimLayerSectionItem* FindMatchingSectionItem(UMovieSceneSection* InMovieSceneSection);
 };
@@ -176,6 +193,8 @@ struct FAnimLayerSceneObject
 	//TArray<FName>  PropertyNames;
 };
 
+
+
 USTRUCT(BlueprintType)
 struct FAnimLayerObjects
 {
@@ -266,6 +285,10 @@ public:
 	//for key property status for weights
 	void SetKey(TSharedPtr<ISequencer>& Sequencer, const IPropertyHandle& KeyedPropertyHandle);
 	EPropertyKeyedStatus GetPropertyKeyedStatus(TSharedPtr<ISequencer>& Sequencer, const IPropertyHandle& PropertyHandle);
+	
+	//function to make sure we have Guids stored on the anim layers now so we can correctly keep track of spawned actors.
+	void UpdateSceneObjectorGuidsForItems(ISequencer* InSequencer);
+
 private:
 
 	UPROPERTY(VisibleAnywhere, Category = "Data")
diff --git a/Engine/Plugins/Animation/ControlRig/Source/ControlRigEditor/Private/Sequencer/ControlRigParameterTrackEditor.cpp b/Engine/Plugins/Animation/ControlRig/Source/ControlRigEditor/Private/Sequencer/ControlRigParameterTrackEditor.cpp
--- a/Engine/Plugins/Animation/ControlRig/Source/ControlRigEditor/Private/Sequencer/ControlRigParameterTrackEditor.cpp
+++ b/Engine/Plugins/Animation/ControlRig/Source/ControlRigEditor/Private/Sequencer/ControlRigParameterTrackEditor.cpp
@@ -114,6 +114,7 @@
 #include "FrontendFilterBase.h"
 #include "InstancedPropertyBagStructureDataProvider.h"
 #include "Editor/RigVMEditorTools.h"
+#include "Sequencer/AnimLayers/AnimLayers.h"
 
 #define LOCTEXT_NAMESPACE "FControlRigParameterTrackEditor"
 
@@ -300,6 +301,9 @@ void FControlRigParameterTrackEditor::OnInitialize()
 	SequencerChangedHandle = Sequencer->OnMovieSceneDataChanged().AddRaw(this, &FControlRigParameterTrackEditor::OnSequencerDataChanged);
 	OnActivateSequenceChangedHandle = Sequencer->OnActivateSequence().AddRaw(this, &FControlRigParameterTrackEditor::OnActivateSequenceChanged);
 	OnChannelChangedHandle = Sequencer->OnChannelChanged().AddRaw(this, &FControlRigParameterTrackEditor::OnChannelChanged);
+	OnMovieSceneBindingsChangeHandle = Sequencer->OnMovieSceneBindingsChanged().AddRaw(this, &FControlRigParameterTrackEditor::OnMovieSceneBindingsChanged);
+
+
 	if (MovieScene)
 	{
 		OnMovieSceneChannelChangedHandle = MovieScene->OnChannelChanged().AddRaw(this, &FControlRigParameterTrackEditor::OnChannelChanged);
@@ -714,6 +718,12 @@ void FControlRigParameterTrackEditor::OnRelease()
 			OnChannelChangedHandle.Reset();
 		}
 
+		if (OnMovieSceneBindingsChangeHandle.IsValid())
+		{
+			Sequencer->OnChannelChanged().Remove(OnMovieSceneBindingsChangeHandle);
+			OnMovieSceneBindingsChangeHandle.Reset();
+		}
+
 		if (OnMovieSceneChannelChangedHandle.IsValid())
 		{
 			UMovieSceneSequence* OwnerSequence = Sequencer->GetFocusedMovieSceneSequence();
@@ -2611,6 +2621,24 @@ struct FChannelChangedStruct
 	TMap<UMovieSceneControlRigParameterSection*, TSet<FName>> SectionControlNames;
 };
 
+void FControlRigParameterTrackEditor::OnMovieSceneBindingsChanged()
+{
+	const TSharedPtr<ISequencer> Sequencer = GetSequencer();
+	if (!Sequencer)
+	{
+		return;
+	}
+	if (UAnimLayers* AnimLayers = UAnimLayers::GetAnimLayers(Sequencer.Get()))
+	{
+		for (UAnimLayer* AnimLayer : AnimLayers->AnimLayers)
+		{
+			if (AnimLayer)
+			{
+				AnimLayer->UpdateSceneObjectorGuidsForItems(Sequencer.Get());
+			}
+		}
+	}
+}
 
 //this function used to set up a set of control names per control and evaluate them on next tick, but this doesn't work with 
 //certain rigs that send events immediately on an evaluation which would happen before the tick. In that case we would get the set key events
@@ -2845,6 +2873,21 @@ void FControlRigParameterTrackEditor::OnActivateSequenceChanged(FMovieSceneSeque
 			}
 		});
 	}
+
+	//update bindings here
+	if (Sequencer)
+	{
+		if (UAnimLayers* AnimLayers = UAnimLayers::GetAnimLayers(Sequencer.Get()))
+		{
+			for (UAnimLayer* AnimLayer : AnimLayers->AnimLayers)
+			{
+				if (AnimLayer)
+				{
+					AnimLayer->UpdateSceneObjectorGuidsForItems(Sequencer.Get());
+				}
+			}
+		}
+	}
 }
 
 
diff --git a/Engine/Plugins/Animation/ControlRig/Source/ControlRigEditor/Private/Sequencer/ControlRigParameterTrackEditor.h b/Engine/Plugins/Animation/ControlRig/Source/ControlRigEditor/Private/Sequencer/ControlRigParameterTrackEditor.h
--- a/Engine/Plugins/Animation/ControlRig/Source/ControlRigEditor/Private/Sequencer/ControlRigParameterTrackEditor.h
+++ b/Engine/Plugins/Animation/ControlRig/Source/ControlRigEditor/Private/Sequencer/ControlRigParameterTrackEditor.h
@@ -175,6 +175,9 @@ private:
 	/** Delegate for when track immediately changes, so we need to do an interaction edit for things like fk/ik*/
 	void OnChannelChanged(const FMovieSceneChannelMetaData* MetaData, UMovieSceneSection* InSection);
 
+	/** Delegate for when bindings changed, anim layers need to get updated*/
+	void OnMovieSceneBindingsChanged();
+
 	/** Delegate for difference focused movie scene sequence*/
 	void OnActivateSequenceChanged(FMovieSceneSequenceIDRef ID);
 
@@ -357,6 +360,7 @@ private:
 	FDelegateHandle OnActivateSequenceChangedHandle;
 	FDelegateHandle CurveChangedHandle;
 	FDelegateHandle OnChannelChangedHandle;
+	FDelegateHandle OnMovieSceneBindingsChangeHandle;
 	FDelegateHandle OnMovieSceneChannelChangedHandle;
 	FDelegateHandle OnActorAddedToSequencerHandle;
 
