Control Rig: Added Post update gizmo interaction delegate in the control rig edit mode
#rnx

diff --git a/Engine/Plugins/Animation/ControlRig/Source/ControlRigEditor/Private/EditMode/ControlRigEditMode.cpp b/Engine/Plugins/Animation/ControlRig/Source/ControlRigEditor/Private/EditMode/ControlRigEditMode.cpp
--- a/Engine/Plugins/Animation/ControlRig/Source/ControlRigEditor/Private/EditMode/ControlRigEditMode.cpp
+++ b/Engine/Plugins/Animation/ControlRig/Source/ControlRigEditor/Private/EditMode/ControlRigEditMode.cpp
@@ -2061,6 +2061,27 @@ bool FControlRigEditMode::InputDelta(FEditorViewportClient* InViewportClient, FV
 		return false;
 	}();
 
+	auto GatherSelectedKeys = [&]() -> FMultiControlRigElementSelection
+	{
+		FMultiControlRigElementSelection MultiRigSelection;
+		for (TPair<UControlRig*, TArray<TObjectPtr<AControlRigShapeActor>>>& Pairs : ControlRigShapeActors)
+		{
+			MultiRigSelection.Rigs.Add(Pairs.Key);
+			FRigElementKeyCollection& Keys = MultiRigSelection.KeysPerRig.Add_GetRef({});
+			if (AreRigElementsSelected(ValidControlTypeMask(), Pairs.Key))
+			{
+				for (TObjectPtr<AControlRigShapeActor> ShapeActor : Pairs.Value)
+				{
+					if (ShapeActor->IsEnabled() && ShapeActor->IsSelected())
+					{
+						Keys.Add(ShapeActor->GetElementKey());
+					}
+				}
+			}
+		}
+		return MultiRigSelection;
+	};
+
 	if (InteractionScopes.Num() > 0 && bMouseButtonDown && CurrentAxis != EAxisList::None
 		&& TransformContext.CanTransform())
 	{
@@ -2071,25 +2092,10 @@ bool FControlRigEditMode::InputDelta(FEditorViewportClient* InViewportClient, FV
 		}
 
 		// The interaction update event must be broadcasted before setting any control value or executing the rig
-		if (OnGizmoInteractionUpdatedDelegate.IsBound())
+		if (OnGizmoInteractionPreUpdatedDelegate.IsBound())
 		{
-			FMultiControlRigElementSelection MultiRigSelection;
-			for (TPair<UControlRig*, TArray<TObjectPtr<AControlRigShapeActor>>>& Pairs : ControlRigShapeActors)
-			{
-				MultiRigSelection.Rigs.Add(Pairs.Key);
-				FRigElementKeyCollection& Keys = MultiRigSelection.KeysPerRig.Add_GetRef({});
-				if (AreRigElementsSelected(ValidControlTypeMask(), Pairs.Key))
-				{
-					for (TObjectPtr<AControlRigShapeActor> ShapeActor : Pairs.Value)
-					{
-						if (ShapeActor->IsEnabled() && ShapeActor->IsSelected())
-						{
-							Keys.Add(ShapeActor->GetElementKey());
-						}
-					}
-				}
-			}
-			OnGizmoInteractionUpdatedDelegate.Broadcast(MultiRigSelection, TransformContext);
+			FMultiControlRigElementSelection MultiRigSelection = GatherSelectedKeys();
+			OnGizmoInteractionPreUpdatedDelegate.Broadcast(MultiRigSelection, TransformContext);
 		}
 		
 		for (auto& Pairs : ControlRigShapeActors)
@@ -2272,6 +2278,12 @@ bool FControlRigEditMode::InputDelta(FEditorViewportClient* InViewportClient, FV
 
 	UpdatePivotTransforms();
 
+	if (OnGizmoInteractionPostUpdatedDelegate.IsBound())
+		{
+			FMultiControlRigElementSelection MultiRigSelection = GatherSelectedKeys();
+			OnGizmoInteractionPostUpdatedDelegate.Broadcast(MultiRigSelection, TransformContext);
+		}
+
 	if (bManipulatorMadeChange)
 	{
 		TArray<TWeakObjectPtr<UControlRig>> RigsToTick;
diff --git a/Engine/Plugins/Animation/ControlRig/Source/ControlRigEditor/Public/EditMode/ControlRigEditMode.h b/Engine/Plugins/Animation/ControlRig/Source/ControlRigEditor/Public/EditMode/ControlRigEditMode.h
--- a/Engine/Plugins/Animation/ControlRig/Source/ControlRigEditor/Public/EditMode/ControlRigEditMode.h
+++ b/Engine/Plugins/Animation/ControlRig/Source/ControlRigEditor/Public/EditMode/ControlRigEditMode.h
@@ -245,7 +245,9 @@ public:
 	FControlRigEditModeInteractionEndedEvent& OnGizmoInteractionEnded() { return OnGizmoInteractionEndedDelegate; }
 
 	/** Broadcasts a notification when a gizmo manipulation has been updated (before the update actually takes place) */
-	FControlRigEditModeInteractionUpdatedEvent& OnGizmoInteractionUpdated() { return OnGizmoInteractionUpdatedDelegate; }
+	FControlRigEditModeInteractionUpdatedEvent& OnGizmoInteractionPreUpdated() { return OnGizmoInteractionPreUpdatedDelegate; }
+	/** Broadcasts a notification when a gizmo manipulation has been updated (after the transforms have been updated) */
+	FControlRigEditModeInteractionUpdatedEvent& OnGizmoInteractionPostUpdated() { return OnGizmoInteractionPostUpdatedDelegate; }
 
 	// callback that gets called when rig element is selected in other view
 	void OnHierarchyModified(ERigHierarchyNotification InNotif, URigHierarchy* InHierarchy, const FRigNotificationSubject& InSubject);
@@ -476,8 +478,11 @@ protected:
 	FControlRigEditModeInteractionEndedEvent OnGizmoInteractionEndedDelegate;
 
 	/** Broadcasts a notification when a gizmo manipulation has been updated (before the update actually takes place) */
-	FControlRigEditModeInteractionUpdatedEvent OnGizmoInteractionUpdatedDelegate;
+	FControlRigEditModeInteractionUpdatedEvent OnGizmoInteractionPreUpdatedDelegate;
 
+	/** Broadcasts a notification when a gizmo manipulation has been updated (after the transforms have been updated) */
+	FControlRigEditModeInteractionUpdatedEvent OnGizmoInteractionPostUpdatedDelegate;
+	
 	/** Broadcasts a notification when a control rig shape actors were recreated */
 	FSimpleMulticastDelegate OnControlRigShapeActorsRecreatedDelegate;
 
