Anim in Engine: Fix editing multiple controls in Anim Details may create more than one undo step

#rnx

diff --git a/Engine/Plugins/Animation/ControlRig/Source/ControlRigEditor/Private/AnimDetails/AnimDetailsMultiEditUtil.cpp b/Engine/Plugins/Animation/ControlRig/Source/ControlRigEditor/Private/AnimDetails/AnimDetailsMultiEditUtil.cpp
--- a/Engine/Plugins/Animation/ControlRig/Source/ControlRigEditor/Private/AnimDetails/AnimDetailsMultiEditUtil.cpp
+++ b/Engine/Plugins/Animation/ControlRig/Source/ControlRigEditor/Private/AnimDetails/AnimDetailsMultiEditUtil.cpp
@@ -11,6 +11,7 @@
 #include "AnimDetails/Proxies/AnimDetailsProxyBase.h"
 #include "AnimDetailsMathOperation.h"
 #include "Editor.h"
+#include "MovieSceneTrackEditor.h"
 #include "PropertyHandle.h"
 #include "ScopedTransaction.h"
 
@@ -49,7 +50,7 @@ namespace UE::ControlRigEditor
 					EPropertyChangeType::Interactive : 
 					EPropertyChangeType::ValueSet;
 
-				 PropertyHandle->NotifyPostChange(EPropertyChangeType::ValueSet);
+				 PropertyHandle->NotifyPostChange(ChangeType);
 			}
 		}
 
@@ -199,6 +200,8 @@ namespace UE::ControlRigEditor
 			return;
 		}
 
+		const FMovieSceneTrackEditor::FScopeSuspendTransactionsOnPropertyChanges SuspendSequencerTransactions = FMovieSceneTrackEditor::SuspendTransactionsOnPropertyChanges();
+
 		const FScopedTransaction ScopedTransaction(LOCTEXT("MultiEditSetTransaction", "Set Property Value"));
 
 		for (const TSharedRef<IPropertyHandle>& PropertyHandle : GetPropertiesBeingEdited<ValueType>(ProxyManager, InstigatorProperty))
@@ -224,6 +227,8 @@ namespace UE::ControlRigEditor
 			return;
 		}
 
+		const FMovieSceneTrackEditor::FScopeSuspendTransactionsOnPropertyChanges SuspendSequencerTransactions = FMovieSceneTrackEditor::SuspendTransactionsOnPropertyChanges();
+
 		const FScopedTransaction ScopedTransaction(LOCTEXT("MultiEditMathTransaction", "Set Property Value"));
 
 		for (const TSharedRef<IPropertyHandle>& PropertyHandle : GetPropertiesBeingEdited<ValueType>(ProxyManager, InstigatorProperty))
@@ -264,6 +269,8 @@ namespace UE::ControlRigEditor
 		const TSharedRef<IPropertyHandle>& InstigatorProperty, 
 		const bool bInteractive)
 	{
+		const FMovieSceneTrackEditor::FScopeSuspendTransactionsOnPropertyChanges SuspendSequencerTransactions = FMovieSceneTrackEditor::SuspendTransactionsOnPropertyChanges();
+
 		if (bInteractive && !IsInteractive())
 		{
 			GEditor->BeginTransaction(LOCTEXT("MultiEditSetPropertyValues", "Set Property Value"));
diff --git a/Engine/Source/Editor/Sequencer/Private/MovieSceneTrackEditor.cpp b/Engine/Source/Editor/Sequencer/Private/MovieSceneTrackEditor.cpp
--- a/Engine/Source/Editor/Sequencer/Private/MovieSceneTrackEditor.cpp
+++ b/Engine/Source/Editor/Sequencer/Private/MovieSceneTrackEditor.cpp
@@ -44,6 +44,8 @@ namespace UE::Sequencer
 	}
 }
 
+int32 FMovieSceneTrackEditor::SuspendTransactionCounter = 0;
+
 FMovieSceneTrackEditor::FMovieSceneTrackEditor(TSharedRef<ISequencer> InSequencer)
 	: WeakSequencerPtr(InSequencer)
 { 
@@ -112,9 +114,13 @@ void FMovieSceneTrackEditor::AnimatablePropertyChanged( FOnKeyProperty OnKeyProp
 		// @todo Sequencer - The sequencer probably should have taken care of this
 		MovieSceneSequence->SetFlags(RF_Transactional);
 	
-		// Create a transaction record because we are about to add keys
-		const bool bShouldActuallyTransact = !GIsTransacting;		// Don't transact if we're recording in a PIE world.  That type of keyframe capture cannot be undone.
-		FScopedTransaction AutoKeyTransaction( NSLOCTEXT("AnimatablePropertyTool", "PropertyChanged", "Animatable Property Changed"), bShouldActuallyTransact );
+		// Create a transaction record because we are about to add keys	
+		TOptional<FScopedTransaction> AutoKeyTransaction;
+		if (SuspendTransactionCounter == 0)
+		{
+			const bool bShouldActuallyTransact = !GIsTransacting;		// Don't transact if we're recording in a PIE world.  That type of keyframe capture cannot be undone.
+			AutoKeyTransaction = FScopedTransaction(NSLOCTEXT("AnimatablePropertyTool", "PropertyChanged", "Animatable Property Changed"), bShouldActuallyTransact);
+		}
 
 		FKeyPropertyResult KeyPropertyResult = OnKeyProperty.Execute( KeyTime );
 
diff --git a/Engine/Source/Editor/Sequencer/Public/MovieSceneTrackEditor.h b/Engine/Source/Editor/Sequencer/Public/MovieSceneTrackEditor.h
--- a/Engine/Source/Editor/Sequencer/Public/MovieSceneTrackEditor.h
+++ b/Engine/Source/Editor/Sequencer/Public/MovieSceneTrackEditor.h
@@ -164,6 +164,21 @@ public:
 		return Result;
 	}
 
+	/** A struct to suspend track editor from generating transactions on property changes */
+	struct FScopeSuspendTransactionsOnPropertyChanges
+	{
+	public:
+		FScopeSuspendTransactionsOnPropertyChanges()
+		{
+			FMovieSceneTrackEditor::SuspendTransactionCounter++;
+		}
+
+		~FScopeSuspendTransactionsOnPropertyChanges()
+		{
+			FMovieSceneTrackEditor::SuspendTransactionCounter--;
+		}
+	};
+
 	/** @return The sequencer bound to this handler */
 	TSharedPtr<ISequencer> GetSequencer() const;
 
@@ -198,6 +213,12 @@ public:
 	virtual bool SupportsSequence(UMovieSceneSequence* InSequence) const override { return true; }
 	virtual void Tick(float DeltaTime) override;
 
+	/** Creates a scope in which track editor doesn't generate transactions on property changes */
+	[[nodiscard]] static FScopeSuspendTransactionsOnPropertyChanges SuspendTransactionsOnPropertyChanges()
+	{
+		return FScopeSuspendTransactionsOnPropertyChanges();
+	}
+
 protected:
 
 	/**
@@ -208,7 +229,6 @@ protected:
 	UMovieScene* GetFocusedMovieScene() const;
 
 private:
-
 	/** The sequencer bound to this handler.  Used to access movie scene and time info during auto-key */
 	TWeakPtr<ISequencer> WeakSequencerPtr;
 
@@ -217,4 +237,7 @@ private:
 
 	/** Indicates whether we're currently in a keying operation where multiple keys may be created (ie. drag and drop) */
 	static bool bKeying;
+
+	/** Counter to suspend transactions on generic property changes */
+	static int32 SuspendTransactionCounter;
 };
