Anim in Engine: Anim Details no longer show parent controls of the current selection, avoiding showing too many channels when a selected control has hosted channels

diff --git a/Engine/Plugins/Animation/ControlRig/Source/ControlRigEditor/Private/AnimDetails/AnimDetailsProxyManager.cpp b/Engine/Plugins/Animation/ControlRig/Source/ControlRigEditor/Private/AnimDetails/AnimDetailsProxyManager.cpp
--- a/Engine/Plugins/Animation/ControlRig/Source/ControlRigEditor/Private/AnimDetails/AnimDetailsProxyManager.cpp
+++ b/Engine/Plugins/Animation/ControlRig/Source/ControlRigEditor/Private/AnimDetails/AnimDetailsProxyManager.cpp
@@ -235,7 +235,7 @@ void UAnimDetailsProxyManager::ForceUpdateProxies()
 					UAnimDetailsProxyBase* Proxy = GetOrCreateControlRigProxy(ControlRigToSelectedKeysPair.Key, ControlElement);
 					AddProxyToExternalSelectionIfValid(Proxy);
 
-					// Also add parents of the selected elements
+					// Also add children of the selected elements
 					if (ElementKey.Type != ERigElementType::Control ||
 						!ControlElement->CanDriveControls())
 					{
@@ -251,34 +251,6 @@ void UAnimDetailsProxyManager::ForceUpdateProxies()
 							AddProxyToExternalSelectionIfValid(DrivenProxy);
 						}
 					}
-
-					Hierarchy->ForEach<FRigControlElement>(
-						[this, ControlRig, ControlElement, DrivenKeys, Hierarchy](FRigControlElement* AnimationChannelControl) -> bool
-						{
-							if (AnimationChannelControl->IsAnimationChannel())
-							{
-								if (FRigControlElement* ParentControlElement = Cast<FRigControlElement>(Hierarchy->GetFirstParent(AnimationChannelControl)))
-								{
-									if (DrivenKeys.Contains(ParentControlElement->GetKey()) ||
-										ParentControlElement->GetKey() == ControlElement->GetKey())
-									{
-										UAnimDetailsProxyBase* ParentProxy = GetOrCreateControlRigProxy(ControlRig, ParentControlElement);
-										AddProxyToExternalSelectionIfValid(ParentProxy);
-
-										return true;
-									}
-
-									if (AnimationChannelControl->Settings.Customization.AvailableSpaces.FindByKey(ControlElement->GetKey()))
-									{
-										UAnimDetailsProxyBase* ParentProxy = GetOrCreateControlRigProxy(ControlRig, ParentControlElement);
-										AddProxyToExternalSelectionIfValid(ParentProxy);
-
-										return true;
-									}
-								}
-							}
-							return true;
-						});
 				}
 			}
 		}
