Sequencer: Fix CIS

#rnx

diff --git a/Engine/Plugins/Animation/ControlRig/Source/ControlRig/Private/Sequencer/MovieSceneControlRigParameterSection.cpp b/Engine/Plugins/Animation/ControlRig/Source/ControlRig/Private/Sequencer/MovieSceneControlRigParameterSection.cpp
--- a/Engine/Plugins/Animation/ControlRig/Source/ControlRig/Private/Sequencer/MovieSceneControlRigParameterSection.cpp
+++ b/Engine/Plugins/Animation/ControlRig/Source/ControlRig/Private/Sequencer/MovieSceneControlRigParameterSection.cpp
@@ -440,7 +440,7 @@ struct FParameterTransformChannelEditorData
 		const EAxisList::Type YAxis = EAxisList::Left;
 		const EAxisList::Type ZAxis = EAxisList::Up;
 
-		const FIntVector4& Swizzle = AxisDisplayInfo::GetTransformAxisSwizzle();
+		const FIntVector4 Swizzle = AxisDisplayInfo::GetTransformAxisSwizzle();
 		const FIntVector4 ReverseSwizzle = ReverseSwizzleFunc(Swizzle);
 		const int32 TranslationOrderOffset = 0;
 		const int32 RotationOrderOffset = TranslationOrderOffset + 3;
diff --git a/Engine/Plugins/Animation/ControlRig/Source/ControlRigEditor/Private/Sequencer/ControlRigParameterTrackEditor.cpp b/Engine/Plugins/Animation/ControlRig/Source/ControlRigEditor/Private/Sequencer/ControlRigParameterTrackEditor.cpp
--- a/Engine/Plugins/Animation/ControlRig/Source/ControlRigEditor/Private/Sequencer/ControlRigParameterTrackEditor.cpp
+++ b/Engine/Plugins/Animation/ControlRig/Source/ControlRigEditor/Private/Sequencer/ControlRigParameterTrackEditor.cpp
@@ -6373,7 +6373,7 @@ void FControlRigParameterSection::BuildSectionContextMenu(FMenuBuilder& MenuBuil
 					}
 				};
 
-				const FIntVector4& Swizzle = AxisDisplayInfo::GetTransformAxisSwizzle();
+				const FIntVector4 Swizzle = AxisDisplayInfo::GetTransformAxisSwizzle();
 				for (int32 MenuItemIndex = 0; MenuItemIndex < NumMenuItems; MenuItemIndex++)
 				{
 					const int32 SwizzledComponentIndex = Swizzle[MenuItemIndex];
@@ -6430,7 +6430,7 @@ void FControlRigParameterSection::BuildSectionContextMenu(FMenuBuilder& MenuBuil
 					}
 				};
 
-				const FIntVector4& Swizzle = AxisDisplayInfo::GetTransformAxisSwizzle();
+				const FIntVector4 Swizzle = AxisDisplayInfo::GetTransformAxisSwizzle();
 				for (int32 MenuItemIndex = 0; MenuItemIndex < NumMenuItems; MenuItemIndex++)
 				{
 					const int32 SwizzledComponentIndex = Swizzle[MenuItemIndex];
diff --git a/Engine/Source/Editor/MovieSceneTools/Private/Sections/CinematicShotSection.cpp b/Engine/Source/Editor/MovieSceneTools/Private/Sections/CinematicShotSection.cpp
--- a/Engine/Source/Editor/MovieSceneTools/Private/Sections/CinematicShotSection.cpp
+++ b/Engine/Source/Editor/MovieSceneTools/Private/Sections/CinematicShotSection.cpp
@@ -431,7 +431,7 @@ void FCinematicShotSection::AddShotMenuSection(FMenuBuilder& MenuBuilder, const
 				}
 			};
 
-			const FIntVector4& Swizzle = AxisDisplayInfo::GetTransformAxisSwizzle();
+			const FIntVector4 Swizzle = AxisDisplayInfo::GetTransformAxisSwizzle();
 			for (int32 MenuItemIndex = 0; MenuItemIndex < NumMenuItems; MenuItemIndex++)
 			{
 				const int32 SwizzledComponentIndex = Swizzle[MenuItemIndex];
diff --git a/Engine/Source/Editor/MovieSceneTools/Private/Sections/TransformPropertySection.cpp b/Engine/Source/Editor/MovieSceneTools/Private/Sections/TransformPropertySection.cpp
--- a/Engine/Source/Editor/MovieSceneTools/Private/Sections/TransformPropertySection.cpp
+++ b/Engine/Source/Editor/MovieSceneTools/Private/Sections/TransformPropertySection.cpp
@@ -120,7 +120,7 @@ void FTransformSection::BuildSectionContextMenu(FMenuBuilder& MenuBuilder, const
 					}
 				};
 
-				const FIntVector4& Swizzle = AxisDisplayInfo::GetTransformAxisSwizzle();
+				const FIntVector4 Swizzle = AxisDisplayInfo::GetTransformAxisSwizzle();
 				for (int32 MenuItemIndex = 0; MenuItemIndex < NumMenuItems; MenuItemIndex++)
 				{
 					const int32 SwizzledComponentIndex = Swizzle[MenuItemIndex];
@@ -177,7 +177,7 @@ void FTransformSection::BuildSectionContextMenu(FMenuBuilder& MenuBuilder, const
 					}
 				};
 
-				const FIntVector4& Swizzle = AxisDisplayInfo::GetTransformAxisSwizzle();
+				const FIntVector4 Swizzle = AxisDisplayInfo::GetTransformAxisSwizzle();
 				for (int32 MenuItemIndex = 0; MenuItemIndex < NumMenuItems; MenuItemIndex++)
 				{
 					const int32 SwizzledComponentIndex = Swizzle[MenuItemIndex];
diff --git a/Engine/Source/Editor/MovieSceneTools/Private/TrackEditors/SubTrackEditor.cpp b/Engine/Source/Editor/MovieSceneTools/Private/TrackEditors/SubTrackEditor.cpp
--- a/Engine/Source/Editor/MovieSceneTools/Private/TrackEditors/SubTrackEditor.cpp
+++ b/Engine/Source/Editor/MovieSceneTools/Private/TrackEditors/SubTrackEditor.cpp
@@ -207,7 +207,7 @@ public:
 					}
 				};
 
-				const FIntVector4& Swizzle = AxisDisplayInfo::GetTransformAxisSwizzle();
+				const FIntVector4 Swizzle = AxisDisplayInfo::GetTransformAxisSwizzle();
 				for (int32 MenuItemIndex = 0; MenuItemIndex < NumMenuItems; MenuItemIndex++)
 				{
 					const int32 SwizzledComponentIndex = Swizzle[MenuItemIndex];
diff --git a/Engine/Source/Runtime/MovieScene/Private/Sections/MovieSceneSubSection.cpp b/Engine/Source/Runtime/MovieScene/Private/Sections/MovieSceneSubSection.cpp
--- a/Engine/Source/Runtime/MovieScene/Private/Sections/MovieSceneSubSection.cpp
+++ b/Engine/Source/Runtime/MovieScene/Private/Sections/MovieSceneSubSection.cpp
@@ -66,7 +66,7 @@ struct FSubSectionEditorData
 		const EAxisList::Type YAxis = EAxisList::Left;
 		const EAxisList::Type ZAxis = EAxisList::Up;
 
-		const FIntVector4& Swizzle = AxisDisplayInfo::GetTransformAxisSwizzle();
+		const FIntVector4 Swizzle = AxisDisplayInfo::GetTransformAxisSwizzle();
 		const FIntVector4 ReverseSwizzle = ReverseSwizzleFunc(Swizzle);
 		const int32 TranslationOrderOffset = 0;
 		const int32 RotationOrderOffset = TranslationOrderOffset + 3;
diff --git a/Engine/Source/Runtime/MovieSceneTracks/Private/Sections/MovieScene3DTransformSection.cpp b/Engine/Source/Runtime/MovieSceneTracks/Private/Sections/MovieScene3DTransformSection.cpp
--- a/Engine/Source/Runtime/MovieSceneTracks/Private/Sections/MovieScene3DTransformSection.cpp
+++ b/Engine/Source/Runtime/MovieSceneTracks/Private/Sections/MovieScene3DTransformSection.cpp
@@ -70,7 +70,7 @@ struct F3DTransformChannelEditorData
 		const EAxisList::Type YAxis = EAxisList::Left;
 		const EAxisList::Type ZAxis = EAxisList::Up;
 
-		const FIntVector4& Swizzle = AxisDisplayInfo::GetTransformAxisSwizzle();
+		const FIntVector4 Swizzle = AxisDisplayInfo::GetTransformAxisSwizzle();
 		const FIntVector4 ReverseSwizzle = ReverseSwizzleFunc(Swizzle);
 		const int32 TranslationOrderOffset = 0;
 		const int32 RotationOrderOffset = TranslationOrderOffset + 3;
