Sequencer LUF: Use GetAxisDisplayCoordinateSystem to determine invert

#rnx

diff --git a/Engine/Plugins/Animation/ControlRig/Source/ControlRig/Private/Sequencer/MovieSceneControlRigParameterSection.cpp b/Engine/Plugins/Animation/ControlRig/Source/ControlRig/Private/Sequencer/MovieSceneControlRigParameterSection.cpp
--- a/Engine/Plugins/Animation/ControlRig/Source/ControlRig/Private/Sequencer/MovieSceneControlRigParameterSection.cpp
+++ b/Engine/Plugins/Animation/ControlRig/Source/ControlRig/Private/Sequencer/MovieSceneControlRigParameterSection.cpp
@@ -474,7 +474,7 @@ struct FParameterTransformChannelEditorData
 			MetaData[1].Color = AxisDisplayInfo::GetAxisColor(YAxis);
 			MetaData[1].SortOrder = SortStartIndex + TranslationOrderOffset + ReverseSwizzle[1];
 			MetaData[1].bCanCollapseToTrack = true;
-			MetaData[1].bInvertValue = AxisDisplayInfo::UseForwardRightUpDisplayNames();
+			MetaData[1].bInvertValue = AxisDisplayInfo::GetAxisDisplayCoordinateSystem() == EAxisList::LeftUpForward;
 			for (const FName& PropertyMetaDataKey : PropertyMetaDataKeys)
 			{
 				MetaData[1].PropertyMetaData.Add(PropertyMetaDataKey, RelativeLocationProperty->GetMetaData(PropertyMetaDataKey));
diff --git a/Engine/Source/Runtime/MovieScene/Private/Sections/MovieSceneSubSection.cpp b/Engine/Source/Runtime/MovieScene/Private/Sections/MovieSceneSubSection.cpp
--- a/Engine/Source/Runtime/MovieScene/Private/Sections/MovieSceneSubSection.cpp
+++ b/Engine/Source/Runtime/MovieScene/Private/Sections/MovieSceneSubSection.cpp
@@ -84,7 +84,8 @@ struct FSubSectionEditorData
 		MetaData[1].bEnabled = EnumHasAllFlags(Mask, EMovieSceneTransformChannel::TranslationY);
 		MetaData[1].Color = AxisDisplayInfo::GetAxisColor(YAxis);
 		MetaData[1].bCanCollapseToTrack = false;
-		
+		MetaData[1].bInvertValue = AxisDisplayInfo::GetAxisDisplayCoordinateSystem() == EAxisList::LeftUpForward;
+
 		MetaData[2].SetIdentifiers("Override.Location.Z", AxisDisplayInfo::GetAxisDisplayName(ZAxis), LocationGroup);
 		MetaData[2].SubPropertyPath = TEXT("Location.Z");
 		MetaData[2].SortOrder = TranslationOrderOffset + ReverseSwizzle[2];
diff --git a/Engine/Source/Runtime/MovieSceneTracks/Private/Sections/MovieScene3DTransformSection.cpp b/Engine/Source/Runtime/MovieSceneTracks/Private/Sections/MovieScene3DTransformSection.cpp
--- a/Engine/Source/Runtime/MovieSceneTracks/Private/Sections/MovieScene3DTransformSection.cpp
+++ b/Engine/Source/Runtime/MovieSceneTracks/Private/Sections/MovieScene3DTransformSection.cpp
@@ -97,7 +97,7 @@ struct F3DTransformChannelEditorData
 			MetaData[1].Color = AxisDisplayInfo::GetAxisColor(YAxis);
 			MetaData[1].SortOrder = SortOrderStart + TranslationOrderOffset + ReverseSwizzle[1];
 			MetaData[1].bCanCollapseToTrack = false;
-			MetaData[1].bInvertValue = AxisDisplayInfo::UseForwardRightUpDisplayNames();
+			MetaData[1].bInvertValue = AxisDisplayInfo::GetAxisDisplayCoordinateSystem() == EAxisList::LeftUpForward;
 			for (const FName& PropertyMetaDataKey : PropertyMetaDataKeys)
 			{
 				MetaData[1].PropertyMetaData.Add(PropertyMetaDataKey, RelativeLocationProperty->GetMetaData(PropertyMetaDataKey));
