Fix Crashes when clearning anim instance from skeletal mesh componoent with a control rig track in sequencer, or when recompiling anim blueprint with it open.
-Added callbacks to template and ECS evaluation of control rig tracks to unregister the custom anim instance from a skeletal mesh component when the object is changed, or if the anim blueprint is recompiled.

diff --git a/Engine/Plugins/Animation/ControlRig/Source/ControlRig/Private/Sequencer/MovieSceneControlRigParameterTemplate.cpp b/Engine/Plugins/Animation/ControlRig/Source/ControlRig/Private/Sequencer/MovieSceneControlRigParameterTemplate.cpp
--- a/Engine/Plugins/Animation/ControlRig/Source/ControlRig/Private/Sequencer/MovieSceneControlRigParameterTemplate.cpp
+++ b/Engine/Plugins/Animation/ControlRig/Source/ControlRig/Private/Sequencer/MovieSceneControlRigParameterTemplate.cpp
@@ -28,6 +28,10 @@
 #include "Constraints/ControlRigTransformableHandle.h"
 #include "Misc/ScopeLock.h"
 
+#if WITH_EDITOR
+#include "Editor.h"
+#endif
+
 #include UE_INLINE_GENERATED_CPP_BY_NAME(MovieSceneControlRigParameterTemplate)
 
 //#include "Particles/ParticleSystemComponent.h"
@@ -588,11 +592,11 @@ static void SelectControls(UControlRig* ControlRig, TArray<FName>& SelectedNames
 	}
 }
 
-void FControlRigBindingHelper::BindToSequencerInstance(UControlRig* ControlRig)
+bool FControlRigBindingHelper::BindToSequencerInstance(UControlRig* ControlRig)
 {
 	if (!ControlRig)
 	{
-		return;
+		return false;
 	}
 	if (USkeletalMeshComponent* SkeletalMeshComponent = Cast<USkeletalMeshComponent>(ControlRig->GetObjectBinding()->GetBoundObject()))
 	{
@@ -618,6 +622,7 @@ void FControlRigBindingHelper::BindToSequencerInstance(UControlRig* ControlRig)
 					}
 				}
 			}
+			return bWasCreated;
 		}
 	}
 	else if (UControlRigComponent* ControlRigComponent = Cast<UControlRigComponent>(ControlRig->GetObjectBinding()->GetBoundObject()))
@@ -641,6 +646,7 @@ void FControlRigBindingHelper::BindToSequencerInstance(UControlRig* ControlRig)
 			}
 		}
 	}
+	return false;
 }
 
 void FControlRigBindingHelper::UnBindFromSequencerInstance(UControlRig* ControlRig)
@@ -1221,8 +1227,9 @@ static UTickableConstraint* CreateConstraintIfNeeded(UWorld* InWorld, const FCon
 struct FControlRigParameterExecutionToken : IMovieSceneExecutionToken
 {
 	FControlRigParameterExecutionToken(const UMovieSceneControlRigParameterSection* InSection,
-		const FEvaluatedControlRigParameterSectionValues& Values)
-	:	Section(InSection)
+		const FEvaluatedControlRigParameterSectionValues& Values, FMovieSceneControlRigParameterTemplate* InTemplate)
+	:	Section(InSection),
+		Template(InTemplate)
 	{
 		BoolValues = Values.BoolValues;
 		IntegerValues = Values.IntegerValues;
@@ -1300,7 +1307,19 @@ struct FControlRigParameterExecutionToken : IMovieSceneExecutionToken
 				Player.SavePreAnimatedState(*SkeletalMeshComponent, FControlRigSkeletalMeshComponentBindingTokenProducer::GetAnimTypeID(), FControlRigSkeletalMeshComponentBindingTokenProducer(Operand.SequenceID, ControlRig));
 			}
 
-			FControlRigBindingHelper::BindToSequencerInstance(ControlRig);
+#if WITH_EDITOR
+			TWeakObjectPtr<UAnimInstance> PreviousAnimInstanceWeakPtr = nullptr;
+			
+			if (ControlRig->GetObjectBinding())
+			{
+				if (const USkeletalMeshComponent* SkeletalMeshComponent = Cast<USkeletalMeshComponent>(ControlRig->GetObjectBinding()->GetBoundObject()))
+				{
+					PreviousAnimInstanceWeakPtr = SkeletalMeshComponent->GetAnimInstance();
+				}
+			}
+#endif
+
+			bool bWasCreated = FControlRigBindingHelper::BindToSequencerInstance(ControlRig);
 
 			if (ControlRig->GetObjectBinding())
 			{
@@ -1308,6 +1327,37 @@ struct FControlRigParameterExecutionToken : IMovieSceneExecutionToken
 				{
 					if (UControlRigLayerInstance* AnimInstance = Cast<UControlRigLayerInstance>(SkeletalMeshComponent->GetAnimInstance()))
 					{
+#if WITH_EDITOR
+						if (GEditor && bWasCreated)
+						{
+							TWeakObjectPtr<USkeletalMeshComponent> WeakSkeletalMeshComponentPtr = SkeletalMeshComponent;
+							FDelegateHandle PreCompileHandle = GEditor->OnBlueprintPreCompile().AddLambda([PreviousAnimInstanceWeakPtr, WeakSkeletalMeshComponentPtr](const UBlueprint* InBlueprint)
+							{
+								const TStrongObjectPtr<UAnimInstance> PinnedAnimInstance = PreviousAnimInstanceWeakPtr.Pin();
+								const TStrongObjectPtr<USkeletalMeshComponent> PinnedSkeletalMeshComponent = WeakSkeletalMeshComponentPtr.Pin();
+								if (PinnedAnimInstance && PinnedSkeletalMeshComponent && InBlueprint && PinnedAnimInstance->GetClass() == InBlueprint->GeneratedClass)
+								{
+									FAnimCustomInstanceHelper::UnbindFromSkeletalMeshComponent<UControlRigLayerInstance>(PinnedSkeletalMeshComponent.Get());
+								}
+							});
+							
+							auto UnregisteredLambda = USkeletalMeshComponent::FOnSkeletalMeshUnregisteredDelegate::CreateLambda([WeakSkeletalMeshComponentPtr](const USkeletalMeshComponent* InMeshComponent)
+							{
+								const TStrongObjectPtr<USkeletalMeshComponent> PinnedSkeletalMeshComponent = WeakSkeletalMeshComponentPtr.Pin();
+								if (PinnedSkeletalMeshComponent && InMeshComponent == PinnedSkeletalMeshComponent.Get())
+								{
+									FAnimCustomInstanceHelper::UnbindFromSkeletalMeshComponent<UControlRigLayerInstance>(PinnedSkeletalMeshComponent.Get());
+								}
+							});
+
+							FDelegateHandle SkeletalMeshUnregisteredHandle = SkeletalMeshComponent->RegisterOnSkeletalMeshUnregistered(UnregisteredLambda);
+
+							Template->PreCompileHandles.Add(PreCompileHandle);
+							Template->MeshUnregisteredHandles.Add({SkeletalMeshUnregisteredHandle, SkeletalMeshComponent});
+						}
+						
+#endif
+
 						float Weight = Section->EvaluateEasing(Context.GetTime());
 						if (EnumHasAllFlags(Section->TransformMask.GetChannels(), EMovieSceneTransformChannel::Weight))
 						{
@@ -1537,6 +1587,8 @@ struct FControlRigParameterExecutionToken : IMovieSceneExecutionToken
 	TArray<FControlSpaceAndValue, TInlineAllocator<2>> SpaceValues;
 	/** Array of evaluated bool values */
 	TArray<FConstraintAndActiveValue, TInlineAllocator<2>> ConstraintsValues;
+
+	FMovieSceneControlRigParameterTemplate* Template;
 };
 
 
@@ -1961,7 +2013,7 @@ void FMovieSceneControlRigParameterTemplate::Evaluate(const FMovieSceneEvaluatio
 		}
 
 		//Do basic token
-		FControlRigParameterExecutionToken ExecutionToken(Section,Values);
+		FControlRigParameterExecutionToken ExecutionToken(Section,Values, const_cast<FMovieSceneControlRigParameterTemplate*>(this));
 		ExecutionTokens.Add(MoveTemp(ExecutionToken));
 
 		EMovieSceneBlendType BlendType = Section->GetBlendType().IsValid() ? Section->GetBlendType().Get() : EMovieSceneBlendType::Absolute;
@@ -2448,6 +2500,26 @@ FMovieSceneAnimTypeID FMovieSceneControlRigParameterTemplate::GetAnimTypeID()
 	return TMovieSceneAnimTypeID<FMovieSceneControlRigParameterTemplate>();
 }
 
+#if WITH_EDITOR
+FMovieSceneControlRigParameterTemplate::~FMovieSceneControlRigParameterTemplate()
+{
+	if (GEditor)
+	{
+		for (FDelegateHandle& Handle : PreCompileHandles)
+		{
+			GEditor->OnBlueprintPreCompile().Remove(Handle);
+		}
+
+		for (TTuple<FDelegateHandle, TWeakObjectPtr<USkeletalMeshComponent>>& HandlePair : MeshUnregisteredHandles)
+		{
+			if (TStrongObjectPtr<USkeletalMeshComponent> SkeletalMeshComponent = HandlePair.Value.Pin())
+			{
+				SkeletalMeshComponent->UnregisterOnSkeletalMeshUnregistered(HandlePair.Key);
+			}
+		}
+	}
+}
+#endif
 
 void FMovieSceneControlRigParameterTemplate::Interrogate(const FMovieSceneContext& Context, FMovieSceneInterrogationData& Container, UObject* BindingOverride) const
 {
diff --git a/Engine/Plugins/Animation/ControlRig/Source/ControlRig/Private/Sequencer/MovieSceneControlRigSystem.cpp b/Engine/Plugins/Animation/ControlRig/Source/ControlRig/Private/Sequencer/MovieSceneControlRigSystem.cpp
--- a/Engine/Plugins/Animation/ControlRig/Source/ControlRig/Private/Sequencer/MovieSceneControlRigSystem.cpp
+++ b/Engine/Plugins/Animation/ControlRig/Source/ControlRig/Private/Sequencer/MovieSceneControlRigSystem.cpp
@@ -39,6 +39,11 @@
 #include "Transform/TransformableHandle.h"
 #include "Constraints/ControlRigTransformableHandle.h"
 
+#if WITH_EDITOR
+#include "Editor.h"
+#include "AnimCustomInstanceHelper.h"
+#endif
+
 #include UE_INLINE_GENERATED_CPP_BY_NAME(MovieSceneControlRigSystem)
 
 namespace UE::MovieScene
@@ -1974,7 +1979,77 @@ void UMovieSceneControlRigParameterEvaluatorSystem::InitializeBaseRigComponent(
 		OutBaseComponentTracker.Add(ControlRig, &OutBaseData);
 	}
 
-	FControlRigBindingHelper::BindToSequencerInstance(ControlRig);
+#if WITH_EDITOR
+	TWeakObjectPtr<UAnimInstance> PreviousAnimInstanceWeakPtr = nullptr;
+			
+	if (ControlRig->GetObjectBinding())
+	{
+		if (const USkeletalMeshComponent* SkeletalMeshComponent = Cast<USkeletalMeshComponent>(ControlRig->GetObjectBinding()->GetBoundObject()))
+		{
+			PreviousAnimInstanceWeakPtr = SkeletalMeshComponent->GetAnimInstance();
+		}
+	}
+#endif
+
+	bool bWasCreated = FControlRigBindingHelper::BindToSequencerInstance(ControlRig);
 	this->AccumulatedValues.InitializeRig(Track, ControlRig);
+
+#if WITH_EDITOR
+	if (GEditor && bWasCreated)
+	{
+		if (ControlRig->GetObjectBinding())
+		{
+			if (USkeletalMeshComponent* SkeletalMeshComponent = Cast<USkeletalMeshComponent>(ControlRig->GetObjectBinding()->GetBoundObject()))
+			{
+				TWeakObjectPtr<USkeletalMeshComponent> WeakSkeletalMeshComponentPtr = SkeletalMeshComponent;
+				FDelegateHandle PreCompileHandle = GEditor->OnBlueprintPreCompile().AddLambda([PreviousAnimInstanceWeakPtr, WeakSkeletalMeshComponentPtr](const UBlueprint* InBlueprint)
+				{
+					const TStrongObjectPtr<UAnimInstance> PinnedAnimInstance = PreviousAnimInstanceWeakPtr.Pin();
+					const TStrongObjectPtr<USkeletalMeshComponent> PinnedSkeletalMeshComponent = WeakSkeletalMeshComponentPtr.Pin();
+					if (PinnedAnimInstance && PinnedSkeletalMeshComponent && InBlueprint && PinnedAnimInstance->GetClass() == InBlueprint->GeneratedClass)
+					{
+						FAnimCustomInstanceHelper::UnbindFromSkeletalMeshComponent<UControlRigLayerInstance>(PinnedSkeletalMeshComponent.Get());
+					}
+				});
+								
+				auto UnregisteredLambda = USkeletalMeshComponent::FOnSkeletalMeshUnregisteredDelegate::CreateLambda([WeakSkeletalMeshComponentPtr](const USkeletalMeshComponent* InMeshComponent)
+				{
+					const TStrongObjectPtr<USkeletalMeshComponent> PinnedSkeletalMeshComponent = WeakSkeletalMeshComponentPtr.Pin();
+					if (PinnedSkeletalMeshComponent && InMeshComponent == PinnedSkeletalMeshComponent.Get())
+					{
+						FAnimCustomInstanceHelper::UnbindFromSkeletalMeshComponent<UControlRigLayerInstance>(PinnedSkeletalMeshComponent.Get());
+					}
+				});
+
+				FDelegateHandle SkeletalMeshUnregisteredHandle = SkeletalMeshComponent->RegisterOnSkeletalMeshUnregistered(UnregisteredLambda);
+
+				PreCompileHandles.Add(PreCompileHandle);
+				MeshUnregisteredHandles.Add({SkeletalMeshUnregisteredHandle, SkeletalMeshComponent});
+			}
+		}
+	}
+#endif
+	
 }
 
+#if WITH_EDITOR
+UMovieSceneControlRigParameterEvaluatorSystem::~UMovieSceneControlRigParameterEvaluatorSystem()
+{
+	if (GEditor)
+	{
+		for (FDelegateHandle& Handle : PreCompileHandles)
+		{
+			GEditor->OnBlueprintPreCompile().Remove(Handle);
+		}
+
+		for (TTuple<FDelegateHandle, TWeakObjectPtr<USkeletalMeshComponent>>& HandlePair : MeshUnregisteredHandles)
+		{
+			if (TStrongObjectPtr<USkeletalMeshComponent> SkeletalMeshComponent = HandlePair.Value.Pin())
+			{
+				SkeletalMeshComponent->UnregisterOnSkeletalMeshUnregistered(HandlePair.Key);
+			}
+		}
+	}
+}
+#endif
+
diff --git a/Engine/Plugins/Animation/ControlRig/Source/ControlRig/Public/Sequencer/MovieSceneControlRigParameterSection.h b/Engine/Plugins/Animation/ControlRig/Source/ControlRig/Public/Sequencer/MovieSceneControlRigParameterSection.h
--- a/Engine/Plugins/Animation/ControlRig/Source/ControlRig/Public/Sequencer/MovieSceneControlRigParameterSection.h
+++ b/Engine/Plugins/Animation/ControlRig/Source/ControlRig/Public/Sequencer/MovieSceneControlRigParameterSection.h
@@ -73,7 +73,7 @@ namespace UE::MovieScene
 
 struct FControlRigBindingHelper
 {
-	static UE_API void BindToSequencerInstance(UControlRig* ControlRig);
+	static UE_API bool BindToSequencerInstance(UControlRig* ControlRig);
 	static UE_API void UnBindFromSequencerInstance(UControlRig* ControlRig);
 };
 
diff --git a/Engine/Plugins/Animation/ControlRig/Source/ControlRig/Public/Sequencer/MovieSceneControlRigParameterTemplate.h b/Engine/Plugins/Animation/ControlRig/Source/ControlRig/Public/Sequencer/MovieSceneControlRigParameterTemplate.h
--- a/Engine/Plugins/Animation/ControlRig/Source/ControlRig/Public/Sequencer/MovieSceneControlRigParameterTemplate.h
+++ b/Engine/Plugins/Animation/ControlRig/Source/ControlRig/Public/Sequencer/MovieSceneControlRigParameterTemplate.h
@@ -20,6 +20,13 @@ struct FMovieSceneControlRigParameterTemplate : public FMovieSceneParameterSecti
 	FMovieSceneControlRigParameterTemplate(const UMovieSceneControlRigParameterSection& Section, const UMovieSceneControlRigParameterTrack& Track);
 	static FMovieSceneAnimTypeID GetAnimTypeID();
 
+#if WITH_EDITOR
+	virtual ~FMovieSceneControlRigParameterTemplate() override;
+
+	TArray<FDelegateHandle> PreCompileHandles;
+	TArray<TTuple<FDelegateHandle, TWeakObjectPtr<USkeletalMeshComponent>>> MeshUnregisteredHandles;
+#endif
+	
 private:
 	virtual UScriptStruct& GetScriptStructImpl() const override { return *StaticStruct(); }
 	virtual void Evaluate(const FMovieSceneEvaluationOperand& Operand, const FMovieSceneContext& Context, const FPersistentEvaluationData& PersistentData, FMovieSceneExecutionTokens& ExecutionTokens) const override;
diff --git a/Engine/Plugins/Animation/ControlRig/Source/ControlRig/Public/Sequencer/MovieSceneControlRigSystem.h b/Engine/Plugins/Animation/ControlRig/Source/ControlRig/Public/Sequencer/MovieSceneControlRigSystem.h
--- a/Engine/Plugins/Animation/ControlRig/Source/ControlRig/Public/Sequencer/MovieSceneControlRigSystem.h
+++ b/Engine/Plugins/Animation/ControlRig/Source/ControlRig/Public/Sequencer/MovieSceneControlRigSystem.h
@@ -239,5 +239,13 @@ public:
 	/** Cached blender system used for blending parameters */
 	UPROPERTY()
 	TObjectPtr<UMovieScenePiecewiseDoubleBlenderSystem> DoubleBlenderSystem;
+
+#if WITH_EDITOR
+	virtual  ~UMovieSceneControlRigParameterEvaluatorSystem() override;
+	
+	TArray<FDelegateHandle> PreCompileHandles;
+	TArray<TTuple<FDelegateHandle, TWeakObjectPtr<USkeletalMeshComponent>>> MeshUnregisteredHandles;
+#endif
+	
 };
 
