Constraints: Make PVS happy by doing a ref capture directly.

#rnx

diff --git a/Engine/Plugins/Animation/ControlRig/Source/ControlRig/Private/Constraints/ControlRigTransformableHandle.cpp b/Engine/Plugins/Animation/ControlRig/Source/ControlRig/Private/Constraints/ControlRigTransformableHandle.cpp
--- a/Engine/Plugins/Animation/ControlRig/Source/ControlRig/Private/Constraints/ControlRigTransformableHandle.cpp
+++ b/Engine/Plugins/Animation/ControlRig/Source/ControlRig/Private/Constraints/ControlRigTransformableHandle.cpp
@@ -506,8 +506,7 @@ const UE::Anim::FAnimationEvaluationTask& UTransformableControlHandle::GetEvalua
 
 		URigHierarchy* Hierarchy = ControlRig->GetHierarchy();
 		const int32 CtrlIndex = Hierarchy->GetIndex(Control->GetKey());
-		FTransform &CacheRef = Cache;
-		EvaluationTask.PostEvaluationFunction = [&CacheRef, CtrlIndex, WeakHierarchy = TWeakObjectPtr<URigHierarchy>(Hierarchy)]()
+		EvaluationTask.PostEvaluationFunction = [&CacheRef = Cache, CtrlIndex, WeakHierarchy = TWeakObjectPtr<URigHierarchy>(Hierarchy)]()
 		{
 			if (WeakHierarchy.IsValid())
 			{
