Anim Layers: Fix Modify on undo
#trivial

diff --git a/Engine/Plugins/Animation/ControlRig/Source/ControlRigEditor/Private/Sequencer/AnimLayers/AnimLayers.cpp b/Engine/Plugins/Animation/ControlRig/Source/ControlRigEditor/Private/Sequencer/AnimLayers/AnimLayers.cpp
--- a/Engine/Plugins/Animation/ControlRig/Source/ControlRigEditor/Private/Sequencer/AnimLayers/AnimLayers.cpp
+++ b/Engine/Plugins/Animation/ControlRig/Source/ControlRigEditor/Private/Sequencer/AnimLayers/AnimLayers.cpp
@@ -2539,7 +2539,7 @@ static void MergeTransformSections(UMovieScene3DTransformSection* BaseSection, U
 	{
 		EMovieSceneTransformChannel BaseChannelMask = BaseSection->GetMask().GetChannels();
 		EMovieSceneTransformChannel ChannelMask = Section->GetMask().GetChannels();
-
+		BaseSection->Modify();
 		TArray<UMovieSceneSection*> TrackSections;
 		GetValidSections(BaseSection, TrackSections);
 		const bool bDoAllTransform = EnumHasAllFlags(ChannelMask, EMovieSceneTransformChannel::AllTransform);
