Anim In Engine: Updated Icons for motion  trails anim layer keying.

diff --git a/Engine/Plugins/Animation/ControlRig/Source/ControlRigEditor/Private/ControlRigEditorStyle.h b/Engine/Plugins/Animation/ControlRig/Source/ControlRigEditor/Private/ControlRigEditorStyle.h
--- a/Engine/Plugins/Animation/ControlRig/Source/ControlRigEditor/Private/ControlRigEditorStyle.h
+++ b/Engine/Plugins/Animation/ControlRig/Source/ControlRigEditor/Private/ControlRigEditorStyle.h
@@ -94,6 +94,8 @@ public:
 			Set("ControlRig.AnimLayerSelected", new IMAGE_BRUSH_SVG("Slate/AnimLayerSelected", Icon16x16));
 			Set("ControlRig.FilterAnimLayerSelected", new IMAGE_BRUSH_SVG("Slate/FilterAnimLayerSelected", Icon16x16));
 			Set("ControlRig.AnimLayers", new IMAGE_BRUSH_SVG("Slate/AnimLayers", Icon16x16));
+			Set("ControlRig.KeyAdd", new IMAGE_BRUSH_SVG("Slate/KeyAdd_14", Icon14x14));
+			Set("ControlRig.KeySpecial", new IMAGE_BRUSH_SVG("Slate/KeySpecial_14", Icon14x14));
 
 		}
 		// Control Rig Editor styles
diff --git a/Engine/Plugins/Animation/ControlRig/Source/ControlRigEditor/Private/Sequencer/AnimLayers/SAnimLayers.cpp b/Engine/Plugins/Animation/ControlRig/Source/ControlRigEditor/Private/Sequencer/AnimLayers/SAnimLayers.cpp
--- a/Engine/Plugins/Animation/ControlRig/Source/ControlRigEditor/Private/Sequencer/AnimLayers/SAnimLayers.cpp
+++ b/Engine/Plugins/Animation/ControlRig/Source/ControlRigEditor/Private/Sequencer/AnimLayers/SAnimLayers.cpp
@@ -687,7 +687,7 @@ public:
 						.ToolTipText(LOCTEXT("AnimLayerKeyTooltip", "Key selected controls or all controls if none selected"))
 						[
 							SNew(SImage)
-								.Image(FAppStyle::GetBrush("Sequencer.AddKey.Details"))
+								.Image(FSlateIcon(FControlRigEditorStyle::Get().GetStyleSetName(), "ControlRig.KeyAdd").GetIcon())
 								.ColorAndOpacity(FSlateColor::UseForeground())
 						]
 				]
@@ -708,7 +708,7 @@ public:
 						.ToolTipText(LOCTEXT("AnimLayerPassthroughTooltip", "Set Default Pose(Base), Zero Key(Additive) or Passthrough(Override) key"))
 						[
 							SNew(SImage)
-								.Image(FAppStyle::GetBrush("Sequencer.AddKey.Details"))
+								.Image(FSlateIcon(FControlRigEditorStyle::Get().GetStyleSetName(), "ControlRig.KeySpecial").GetIcon())
 								.ColorAndOpacity(FSlateColor::UseForeground())
 						]
 				];
