Fix for Control Rig Hierarchy updating more transforms than needed in propagate dirty flags

diff --git a/Engine/Plugins/Animation/ControlRig/Source/ControlRig/Private/Rigs/RigHierarchy.cpp b/Engine/Plugins/Animation/ControlRig/Source/ControlRig/Private/Rigs/RigHierarchy.cpp
--- a/Engine/Plugins/Animation/ControlRig/Source/ControlRig/Private/Rigs/RigHierarchy.cpp
+++ b/Engine/Plugins/Animation/ControlRig/Source/ControlRig/Private/Rigs/RigHierarchy.cpp
@@ -7127,127 +7127,45 @@ void URigHierarchy::ShrinkComponentStorage()
 
 void URigHierarchy::PropagateDirtyFlags(FRigTransformElement* InTransformElement, bool bInitial, bool bAffectChildren, bool bComputeOpposed, bool bMarkDirty) const
 {
-	if(!bEnableDirtyPropagation)
+	if (!bEnableDirtyPropagation)
 	{
 		return;
 	}
-	
+
 	check(InTransformElement);
 
 	const ERigTransformType::Type LocalType = bInitial ? ERigTransformType::InitialLocal : ERigTransformType::CurrentLocal;
 	const ERigTransformType::Type GlobalType = bInitial ? ERigTransformType::InitialGlobal : ERigTransformType::CurrentGlobal;
 
-	if(bComputeOpposed)
+	const ERigTransformType::Type TypeToCompute = bAffectChildren ? LocalType : GlobalType;
+	const ERigTransformType::Type TypeToDirty = SwapLocalAndGlobal(TypeToCompute);
+
+	for (const FRigTransformElement::FElementToDirty& ElementToDirty : InTransformElement->ElementsToDirty)
 	{
-		for(const FRigTransformElement::FElementToDirty& ElementToDirty : InTransformElement->ElementsToDirty)
+		if (ElementToDirty.Element->IsElementDirty(TypeToDirty))
 		{
-			ERigTransformType::Type TypeToCompute = bAffectChildren ? LocalType : GlobalType;
-			ERigTransformType::Type TypeToDirty = SwapLocalAndGlobal(TypeToCompute);
-			
-			if(FRigControlElement* ControlElement = Cast<FRigControlElement>(ElementToDirty.Element))
-			{
-				// animation channels never dirty their local value
-				if(ControlElement->IsAnimationChannel())
-				{
-					if(ERigTransformType::IsLocal(TypeToDirty))
-					{
-						Swap(TypeToDirty, TypeToCompute);
-					}
-				}
-				
-				if(ERigTransformType::IsGlobal(TypeToDirty))
-				{
-					if(ControlElement->GetOffsetDirtyState().IsDirty(TypeToDirty) &&
-						ControlElement->GetDirtyState().IsDirty(TypeToDirty) &&
-						ControlElement->GetShapeDirtyState().IsDirty(TypeToDirty))
-					{
-						continue;
-					}
-				}
-			}
-			else if(FRigMultiParentElement* MultiParentElement = Cast<FRigMultiParentElement>(ElementToDirty.Element))
-			{
-				if(ERigTransformType::IsGlobal(TypeToDirty))
-				{
-					if(MultiParentElement->GetDirtyState().IsDirty(TypeToDirty))
-					{
-						continue;
-					}
-				}
-			}
-			else
-			{
-				if(ElementToDirty.Element->GetDirtyState().IsDirty(TypeToDirty))
-				{
-					continue;
-				}
-			}
+			continue;
+		}
 
-			if(FRigControlElement* ControlElement = Cast<FRigControlElement>(ElementToDirty.Element))
+		if (bComputeOpposed)
+		{
+			if (FRigControlElement* ControlElement = Cast<FRigControlElement>(ElementToDirty.Element))
 			{
 				GetControlOffsetTransform(ControlElement, LocalType);
 			}
 			GetTransform(ElementToDirty.Element, TypeToCompute); // make sure the local / global transform is up 2 date
 
-			PropagateDirtyFlags(ElementToDirty.Element, bInitial, bAffectChildren, true, false);
+			if (bAffectChildren)
+			{
+				PropagateDirtyFlags(ElementToDirty.Element, bInitial, bAffectChildren, true, false);
+			}
 		}
-	}
 
-	if(bMarkDirty)
-	{
-		for(const FRigTransformElement::FElementToDirty& ElementToDirty : InTransformElement->ElementsToDirty)
+		if (bMarkDirty)
 		{
-			ERigTransformType::Type TypeToCompute = bAffectChildren ? LocalType : GlobalType;
-			ERigTransformType::Type TypeToDirty = SwapLocalAndGlobal(TypeToCompute);
-
-			if(FRigControlElement* ControlElement = Cast<FRigControlElement>(ElementToDirty.Element))
-			{
-				// animation channels never dirty their local value
-				if(ControlElement->IsAnimationChannel())
-				{
-					if(ERigTransformType::IsLocal(TypeToDirty))
-					{
-						Swap(TypeToDirty, TypeToCompute);
-					}
-				}
+			ElementToDirty.Element->MarkElementDirty(TypeToDirty);
 
-				if(ERigTransformType::IsGlobal(TypeToDirty))
-				{
-					if(ControlElement->GetOffsetDirtyState().IsDirty(TypeToDirty) &&
-						ControlElement->GetDirtyState().IsDirty(TypeToDirty) &&
-					    ControlElement->GetShapeDirtyState().IsDirty(TypeToDirty))
-					{
-						continue;
-					}
-				}
-			}
-			else if(FRigMultiParentElement* MultiParentElement = Cast<FRigMultiParentElement>(ElementToDirty.Element))
-			{
-				if(ERigTransformType::IsGlobal(TypeToDirty))
-				{
-					if(MultiParentElement->GetDirtyState().IsDirty(TypeToDirty))
-					{
-						continue;
-					}
-				}
-			}
-			else
-			{
-				if(ElementToDirty.Element->GetDirtyState().IsDirty(TypeToDirty))
-				{
-					continue;
-				}
-			}
-						
-			ElementToDirty.Element->GetDirtyState().MarkDirty(TypeToDirty);
-		
-			if(FRigControlElement* ControlElement = Cast<FRigControlElement>(ElementToDirty.Element))
-			{
-				ControlElement->GetOffsetDirtyState().MarkDirty(GlobalType);
-				ControlElement->GetShapeDirtyState().MarkDirty(GlobalType);
-			}
-
-			if(bAffectChildren)
+			if (bAffectChildren)
 			{
 				PropagateDirtyFlags(ElementToDirty.Element, bInitial, bAffectChildren, false, true);
 			}
diff --git a/Engine/Plugins/Animation/ControlRig/Source/ControlRig/Private/Rigs/RigHierarchyElements.cpp b/Engine/Plugins/Animation/ControlRig/Source/ControlRig/Private/Rigs/RigHierarchyElements.cpp
--- a/Engine/Plugins/Animation/ControlRig/Source/ControlRig/Private/Rigs/RigHierarchyElements.cpp
+++ b/Engine/Plugins/Animation/ControlRig/Source/ControlRig/Private/Rigs/RigHierarchyElements.cpp
@@ -1067,6 +1067,16 @@ FRigCurrentAndInitialDirtyState& FRigTransformElement::GetDirtyState()
 	return PoseDirtyState;
 }
 
+bool FRigTransformElement::IsElementDirty(ERigTransformType::Type InTransformType) const
+{
+	return PoseDirtyState.IsDirty(InTransformType);
+}
+
+void FRigTransformElement::MarkElementDirty(ERigTransformType::Type InTransformType)
+{
+	PoseDirtyState.MarkDirty(InTransformType);
+}
+
 void FRigTransformElement::CopyFrom(const FRigBaseElement* InOther)
 {
 	Super::CopyFrom(InOther);
@@ -1317,6 +1327,17 @@ void FRigMultiParentElement::CopyPose(FRigBaseElement* InOther, bool bCurrent, b
 	}
 }
 
+bool FRigMultiParentElement::IsElementDirty(ERigTransformType::Type InTransformType) const
+{
+	return PoseDirtyState.IsDirty(ERigTransformType::MakeGlobal(InTransformType));
+}
+
+void FRigMultiParentElement::MarkElementDirty(ERigTransformType::Type InTransformType)
+{
+	PoseDirtyState.MarkDirty(ERigTransformType::MakeGlobal(InTransformType));
+}
+
+
 ////////////////////////////////////////////////////////////////////////////////
 // FRigBoneElement
 ////////////////////////////////////////////////////////////////////////////////
@@ -1957,6 +1978,39 @@ void FRigControlElement::CopyPose(FRigBaseElement* InOther, bool bCurrent, bool
 	}
 }
 
+bool FRigControlElement::IsElementDirty(ERigTransformType::Type InTransformType) const
+{
+	if (IsAnimationChannel())
+	{
+		// animation channels never dirty their local value
+		InTransformType = ERigTransformType::MakeGlobal(InTransformType);
+	}
+
+	if (ERigTransformType::IsGlobal(InTransformType))
+	{
+		return (GetOffsetDirtyState().IsDirty(InTransformType)
+			&& GetDirtyState().IsDirty(InTransformType)
+			&& GetShapeDirtyState().IsDirty(InTransformType));
+	}
+
+	return PoseDirtyState.IsDirty(InTransformType);
+}
+
+void FRigControlElement::MarkElementDirty(ERigTransformType::Type InTransformType)
+{
+	const ERigTransformType::Type GlobalType = ERigTransformType::MakeGlobal(InTransformType);
+
+	if (IsAnimationChannel())
+	{
+		// animation channels never dirty their local value
+		InTransformType = GlobalType;
+	}
+
+	PoseDirtyState.MarkDirty(InTransformType);
+	GetOffsetDirtyState().MarkDirty(GlobalType);
+	GetShapeDirtyState().MarkDirty(GlobalType);
+}
+
 void FRigControlElement::LinkStorage(const TArrayView<FTransform>& InTransforms, const TArrayView<bool>& InDirtyStates,
 	const TArrayView<float>& InCurves)
 {
diff --git a/Engine/Plugins/Animation/ControlRig/Source/ControlRig/Public/Rigs/RigHierarchyElements.h b/Engine/Plugins/Animation/ControlRig/Source/ControlRig/Public/Rigs/RigHierarchyElements.h
--- a/Engine/Plugins/Animation/ControlRig/Source/ControlRig/Public/Rigs/RigHierarchyElements.h
+++ b/Engine/Plugins/Animation/ControlRig/Source/ControlRig/Public/Rigs/RigHierarchyElements.h
@@ -1069,6 +1069,9 @@ struct FRigTransformElement : public FRigBaseElement
 	UE_API const FRigCurrentAndInitialDirtyState& GetDirtyState() const;
 	UE_API FRigCurrentAndInitialDirtyState& GetDirtyState();
 
+	UE_API virtual bool IsElementDirty(ERigTransformType::Type InTransformType) const;
+	UE_API virtual void MarkElementDirty(ERigTransformType::Type InTransformType);
+
 protected:
 	
 	FRigTransformElement(URigHierarchy* InOwner, const ERigElementType InType) :
@@ -1282,6 +1285,9 @@ struct FRigMultiParentElement : public FRigTransformElement
 
 	UE_API virtual void CopyPose(FRigBaseElement* InOther, bool bCurrent, bool bInitial, bool bWeights) override;
 	
+	UE_API virtual bool IsElementDirty(ERigTransformType::Type InTransformType) const override;
+	UE_API virtual void MarkElementDirty(ERigTransformType::Type InTransformType)  override;
+
 	FRigElementParentConstraintArray ParentConstraints;
 	TMap<FRigElementKey, int32> IndexLookup;
 
@@ -1795,6 +1801,9 @@ struct FRigControlElement final : public FRigMultiParentElement
 
 	UE_API virtual void CopyPose(FRigBaseElement* InOther, bool bCurrent, bool bInitial, bool bWeights) override;
 
+	UE_API virtual bool IsElementDirty(ERigTransformType::Type InTransformType) const override;
+	UE_API virtual void MarkElementDirty(ERigTransformType::Type InTransformType) override;
+
 protected:
 
 	UE_API virtual void LinkStorage(const TArrayView<FTransform>& InTransforms, const TArrayView<bool>& InDirtyStates, const TArrayView<float>& InCurves) override;
