Control Rig: Use stack allocator to improve allocation of local array when sorting element storage
#rnx

diff --git a/Engine/Plugins/Animation/ControlRig/Source/ControlRig/Private/Rigs/RigHierarchy.cpp b/Engine/Plugins/Animation/ControlRig/Source/ControlRig/Private/Rigs/RigHierarchy.cpp
--- a/Engine/Plugins/Animation/ControlRig/Source/ControlRig/Private/Rigs/RigHierarchy.cpp
+++ b/Engine/Plugins/Animation/ControlRig/Source/ControlRig/Private/Rigs/RigHierarchy.cpp
@@ -7530,7 +7530,9 @@ bool URigHierarchy::SortElementStorage()
 	}
 
 	const int32 TotalNumberOfTransforms = TotalNumberOfTransformsPerTransformType * NumTransformTypes;
-	TArray<int32> TransformIndexToAllocatedIndex;
+
+	FMemMark Mark(FMemStack::Get()); // Guard to clean up the stack allocator
+	TArray<int32, FAnimStackAllocator> TransformIndexToAllocatedIndex;
 	TransformIndexToAllocatedIndex.AddUninitialized(TotalNumberOfTransforms);
 
 	// we sort elements by laying them out first by transform type, then element type, then storage type.
