Remove top level component concept from Control Rig Hierarchy

diff --git a/Engine/Plugins/Animation/ControlRig/Source/ControlRig/Private/Rigs/RigHierarchy.cpp b/Engine/Plugins/Animation/ControlRig/Source/ControlRig/Private/Rigs/RigHierarchy.cpp
--- a/Engine/Plugins/Animation/ControlRig/Source/ControlRig/Private/Rigs/RigHierarchy.cpp
+++ b/Engine/Plugins/Animation/ControlRig/Source/ControlRig/Private/Rigs/RigHierarchy.cpp
@@ -679,7 +679,6 @@ void URigHierarchy::Load(FArchive& Ar)
 	{
 		ElementComponents.Reset();
 		ComponentIndexLookup.Reset();
-		TopLevelComponentIndices.Reset();
 
 		int32 NumComponents = 0;
 		*ArchiveForElements << NumComponents;
@@ -734,19 +733,12 @@ void URigHierarchy::Load(FArchive& Ar)
 
 				Component->IndexInHierarchy = ComponentIndex;
 
-				if(Component->GetElementKey() == GetTopLevelComponentElementKey())
-				{
-					Component->Element = nullptr;
-					Component->IndexInElement = TopLevelComponentIndices.Add(Component->IndexInHierarchy);
-				}
-				else
-				{
-					FRigBaseElement* Element = Find(Component->GetElementKey());
-					check(Element);
+				FRigBaseElement* Element = Find(Component->GetElementKey());
+				check(Element);
 					
-					Component->Element = Element;
-					Component->IndexInElement = Element->ComponentIndices.Add(Component->IndexInHierarchy);
-				}
+				Component->Element = Element;
+				Component->IndexInElement = Element->ComponentIndices.Add(Component->IndexInHierarchy);
+
 				ComponentIndexLookup.Add(Component->GetKey(), ComponentIndex);
 			}
 		}
@@ -814,7 +806,6 @@ void URigHierarchy::Reset_Impl(bool bResetElements)
 		// clearing the components first will avoid cost of looping over
 		// each component per elements and destroying them one by onw
 		ElementComponents.Reset();
-		TopLevelComponentIndices.Reset();
 
 		// walk in reverse since certain elements might not have been allocated themselves
 		for(int32 ElementIndex = Elements.Num() - 1; ElementIndex >= 0; ElementIndex--)
@@ -1140,8 +1131,6 @@ void URigHierarchy::CopyHierarchy(URigHierarchy* InHierarchy)
 		ElementComponents.Reserve(InHierarchy->ElementComponents.Num());
 		ComponentIndexLookup.Reset();
 		ComponentIndexLookup.Reserve(InHierarchy->ElementComponents.Num());
-		TopLevelComponentIndices.Reset();
-		TopLevelComponentIndices.Reserve(InHierarchy->TopLevelComponentIndices.Num());
 
 		for(int32 SourceComponentIndex = 0; SourceComponentIndex < InHierarchy->ElementComponents.Num(); SourceComponentIndex++)
 		{
@@ -1153,16 +1142,9 @@ void URigHierarchy::CopyHierarchy(URigHierarchy* InHierarchy)
 			FRigBaseComponent* TargetComponent = ElementComponents[TargetComponentIndex].GetMutablePtr<FRigBaseComponent>();
 			TargetComponent->IndexInHierarchy = TargetComponentIndex;
 
-			if(TargetComponent->IsTopLevel())
-			{
-				TargetComponent->Element = nullptr;
-				TargetComponent->IndexInElement = TopLevelComponentIndices.Add(TargetComponentIndex);
-			}
-			else
-			{
-				TargetComponent->Element = Find(TargetComponent->GetElementKey());
-				TargetComponent->IndexInElement = TargetComponent->Element->ComponentIndices.Add(TargetComponent->IndexInHierarchy);
-			}
+			TargetComponent->Element = Find(TargetComponent->GetElementKey());
+			TargetComponent->IndexInElement = TargetComponent->Element->ComponentIndices.Add(TargetComponent->IndexInHierarchy);
+	
 			ComponentIndexLookup.Add(TargetComponent->GetKey(), TargetComponentIndex);
 		}
 
@@ -1853,34 +1835,6 @@ TArray<const FRigBaseComponent*> URigHierarchy::GetComponents(const UScriptStruc
 	return Result;
 }
 
-const FRigBaseComponent* URigHierarchy::FindTopLevelComponent(const FName& InName) const
-{
-	return FindComponent(FRigComponentKey(GetTopLevelComponentElementKey(), InName));
-}
-
-FRigBaseComponent* URigHierarchy::FindTopLevelComponent(const FName& InName)
-{
-	return FindComponent(FRigComponentKey(GetTopLevelComponentElementKey(), InName));
-}
-
-const FRigBaseComponent* URigHierarchy::GetTopLevelComponent(int32 InTopLevelComponentIndex) const
-{
-	if(TopLevelComponentIndices.IsValidIndex(InTopLevelComponentIndex))
-	{
-		return GetComponent(TopLevelComponentIndices[InTopLevelComponentIndex]);
-	}
-	return nullptr;
-}
-
-FRigBaseComponent* URigHierarchy::GetTopLevelComponent(int32 InTopLevelComponentIndex)
-{
-	if(TopLevelComponentIndices.IsValidIndex(InTopLevelComponentIndex))
-	{
-		return GetComponent(TopLevelComponentIndices[InTopLevelComponentIndex]);
-	}
-	return nullptr;
-}
-
 int32 URigHierarchy::GetComponentIndex(const FRigComponentKey& InComponentKey, bool bFollowRedirector) const
 {
 	if(bFollowRedirector)
@@ -1899,7 +1853,7 @@ int32 URigHierarchy::GetComponentIndex(const FRigComponentKey& InComponentKey, b
 		}
 	}
 	
-	if(bFollowRedirector && !InComponentKey.IsTopLevel())
+	if(bFollowRedirector)
 	{
 		const int32 ElementIndex = GetIndex(InComponentKey.ElementKey);
 		if(ElementIndex != INDEX_NONE)
@@ -1967,10 +1921,6 @@ int32 URigHierarchy::NumComponents(const UScriptStruct* InComponentStruct) const
 
 int32 URigHierarchy::NumComponents(FRigElementKey InElement) const
 {
-	if(InElement == GetTopLevelComponentElementKey())
-	{
-		return NumTopLevelComponents();
-	}
 	if(const FRigBaseElement* Element = Find(InElement))
 	{
 		return Element->NumComponents();
@@ -1987,10 +1937,6 @@ TArray<FRigComponentKey> URigHierarchy::GetAllComponentKeys() const
 
 TArray<FRigComponentKey> URigHierarchy::GetComponentKeys(FRigElementKey InElement) const
 {
-	if(InElement == GetTopLevelComponentElementKey())
-	{
-		return GetTopLevelComponentKeys();
-	}
 	if(const FRigBaseElement* Element = Find(InElement))
 	{
 		return Element->GetComponentKeys();
@@ -2000,13 +1946,6 @@ TArray<FRigComponentKey> URigHierarchy::GetComponentKeys(FRigElementKey InElemen
 
 FRigComponentKey URigHierarchy::GetComponentKey(FRigElementKey InElement, int32 InComponentIndex) const
 {
-	if(InElement == GetTopLevelComponentElementKey())
-	{
-		if(const FRigBaseComponent* Component = GetTopLevelComponent(InComponentIndex))
-		{
-			return Component->GetKey();
-		}
-	}
 	if(const FRigBaseElement* Element = Find(InElement))
 	{
 		if(const FRigBaseComponent* Component = Element->GetComponent(InComponentIndex))
@@ -2019,13 +1958,6 @@ FRigComponentKey URigHierarchy::GetComponentKey(FRigElementKey InElement, int32
 
 FName URigHierarchy::GetComponentName(FRigElementKey InElement, int32 InComponentIndex) const
 {
-	if(InElement == GetTopLevelComponentElementKey())
-	{
-		if(const FRigBaseComponent* Component = GetTopLevelComponent(InComponentIndex))
-		{
-			return Component->GetFName();
-		}
-	}
 	if(const FRigBaseElement* Element = Find(InElement))
 	{
 		if(const FRigBaseComponent* Component = Element->GetComponent(InComponentIndex))
@@ -2038,13 +1970,6 @@ FName URigHierarchy::GetComponentName(FRigElementKey InElement, int32 InComponen
 
 UScriptStruct* URigHierarchy::GetComponentType(FRigElementKey InElement, int32 InComponentIndex) const
 {
-	if(InElement == GetTopLevelComponentElementKey())
-	{
-		if(const FRigBaseComponent* Component = GetTopLevelComponent(InComponentIndex))
-		{
-			return Component->GetScriptStruct();
-		}
-	}
 	if(const FRigBaseElement* Element = Find(InElement))
 	{
 		if(const FRigBaseComponent* Component = Element->GetComponent(InComponentIndex))
@@ -2057,13 +1982,6 @@ UScriptStruct* URigHierarchy::GetComponentType(FRigElementKey InElement, int32 I
 
 FString URigHierarchy::GetComponentContent(FRigElementKey InElement, int32 InComponentIndex) const
 {
-	if(InElement == GetTopLevelComponentElementKey())
-	{
-		if(const FRigBaseComponent* Component = GetTopLevelComponent(InComponentIndex))
-		{
-			return Component->GetContentAsText();
-		}
-	}
 	if(const FRigBaseElement* Element = Find(InElement))
 	{
 		if(const FRigBaseComponent* Component = Element->GetComponent(InComponentIndex))
@@ -2074,47 +1992,11 @@ FString URigHierarchy::GetComponentContent(FRigElementKey InElement, int32 InCom
 	return FString();
 }
 
-int32 URigHierarchy::NumTopLevelComponents() const
-{
-	return TopLevelComponentIndices.Num();
-}
-
-TArray<FRigComponentKey> URigHierarchy::GetTopLevelComponentKeys() const
-{
-	TArray<FRigComponentKey> Keys;
-	Keys.Reserve(NumTopLevelComponents());
-	for(int32 Index = 0; Index < NumTopLevelComponents(); Index++)
-	{
-		Keys.Add(GetTopLevelComponentKey(Index));
-	}
-	return Keys;
-}
-
-FRigComponentKey URigHierarchy::GetTopLevelComponentKey(int32 InTopLevelComponentIndex) const
-{
-	return GetComponentKey(GetTopLevelComponentElementKey(), InTopLevelComponentIndex);
-}
-
-FName URigHierarchy::GetTopLevelComponentName(int32 InTopLevelComponentIndex) const
-{
-	return GetComponentName(GetTopLevelComponentElementKey(), InTopLevelComponentIndex);
-}
-
-UScriptStruct* URigHierarchy::GetTopLevelComponentType(int32 InTopLevelComponentIndex) const
-{
-	return GetComponentType(GetTopLevelComponentElementKey(), InTopLevelComponentIndex);
-}
-
-FString URigHierarchy::GetTopLevelComponentContent(int32 InTopLevelComponentIndex) const
-{
-	return GetComponentContent(GetTopLevelComponentElementKey(), InTopLevelComponentIndex);
-}
-
 bool URigHierarchy::CanAddComponent(FRigElementKey InElementKey, const UScriptStruct* InComponentStruct, FString* OutFailureReason) const
 {
 	check(InComponentStruct);
 
-	if(InElementKey != GetTopLevelComponentElementKey() && !InElementKey.IsValid())
+	if(!InElementKey.IsValid())
 	{
 		if(OutFailureReason)
 		{
@@ -2133,7 +2015,7 @@ bool URigHierarchy::CanAddComponent(FRigElementKey InElementKey, const UScriptSt
 	}
 
 	const FRigBaseElement* Element = Find(InElementKey);
-	if(Element == nullptr && InElementKey != GetTopLevelComponentElementKey())
+	if(Element == nullptr)
 	{
 		if(OutFailureReason)
 		{
@@ -2164,23 +2046,6 @@ bool URigHierarchy::CanAddComponent(FRigElementKey InElementKey, const UScriptSt
 					}
 				}
 			}
-			else
-			{
-				for(int32 Index = 0; Index < NumTopLevelComponents(); Index++)
-				{
-					if(const FRigBaseComponent* ExistingComponent = GetTopLevelComponent(Index))
-					{
-						if(ExistingComponent->GetScriptStruct() == InComponentStruct)
-						{
-							if(OutFailureReason)
-							{
-								*OutFailureReason = FString::Printf(TEXT("A component of type '%s' already exists at the top level."), *InComponentStruct->GetName());
-							}
-							return false;
-						}
-					}
-				}
-			}
 		}
 		if(!StructMemory->CanBeAddedTo(InElementKey, this, OutFailureReason))
 		{
@@ -2207,12 +2072,6 @@ bool URigHierarchy::CanAddComponent(FRigElementKey InElementKey, const FRigBaseC
 	return true;
 }
 
-const FRigElementKey& URigHierarchy::GetTopLevelComponentElementKey()
-{
-	static const FRigElementKey TopLevelComponentElementKey = FRigElementKey(TEXT("TopLevel"), ERigElementType::None);
-	return TopLevelComponentElementKey;
-}
-
 TArray<FRigSocketState> URigHierarchy::GetSocketStates() const
 {
 	const TArray<FRigElementKey> Keys = GetSocketKeys(true);
@@ -6986,10 +6845,8 @@ FRigBaseComponent* URigHierarchy::MakeComponent(const UScriptStruct* InComponent
 {
 	check(InComponentStruct);
 	check(InComponentStruct->IsChildOf(FRigBaseComponent::StaticStruct()));
-	if(InElement)
-	{
-		check(InElement->GetOwner() == this);
-	}
+	check(InElement != nullptr);
+	check(InElement->GetOwner() == this);
 	
 	int32 ComponentIndex = INDEX_NONE;
 	for(int32 ExistingIndex = 0; ExistingIndex < ElementComponents.Num(); ExistingIndex++)
@@ -7007,22 +6864,14 @@ FRigBaseComponent* URigHierarchy::MakeComponent(const UScriptStruct* InComponent
 		ComponentIndex = ElementComponents.Emplace(InComponentStruct);
 	}
 
-	const FRigName UniqueName = GetSafeNewComponentName(InElement ? InElement->GetKey() : GetTopLevelComponentElementKey(), InName);
+	const FRigName UniqueName = GetSafeNewComponentName(InElement->GetKey(), InName);
 	
 	FRigBaseComponent* Component = ElementComponents[ComponentIndex].GetMutablePtr<FRigBaseComponent>();
 	Component->IndexInHierarchy = ComponentIndex;
-	if(InElement)
-	{
-		Component->Key = FRigComponentKey(InElement->GetKey(), UniqueName.GetFName());
-		Component->Element = InElement;
-		Component->IndexInElement = InElement->ComponentIndices.Add(Component->IndexInHierarchy);
-	}
-	else
-	{
-		Component->Key = FRigComponentKey(GetTopLevelComponentElementKey(), UniqueName.GetFName());
-		Component->Element = nullptr;
-		Component->IndexInElement = TopLevelComponentIndices.Add(Component->IndexInHierarchy);
-	}
+
+	Component->Key = FRigComponentKey(InElement->GetKey(), UniqueName.GetFName());
+	Component->Element = InElement;
+	Component->IndexInElement = InElement->ComponentIndices.Add(Component->IndexInHierarchy);
 
 	ComponentIndexLookup.Add(Component->GetKey(), Component->GetIndexInHierarchy());
 	
@@ -7037,19 +6886,7 @@ void URigHierarchy::DestroyComponent(FRigBaseComponent*& InComponent)
 		const FRigBaseComponent* Component = InstancedStruct.GetMutablePtr<FRigBaseComponent>();
 		if(Component == InComponent)
 		{
-			if(Component->GetElementKey() == GetTopLevelComponentElementKey())
-			{
-				TopLevelComponentIndices.Remove(ComponentIndex);
-				
-				for(int32 TopLevelComponentIndex = 0; TopLevelComponentIndex < NumTopLevelComponents(); TopLevelComponentIndex++)
-				{
-					if(FRigBaseComponent* RemainingComponent = GetTopLevelComponent(TopLevelComponentIndex))
-					{
-						RemainingComponent->IndexInElement = TopLevelComponentIndex;
-					}
-				}
-			}
-			else if(FRigBaseElement* Element = Find(Component->GetElementKey()))
+			if(FRigBaseElement* Element = Find(Component->GetElementKey()))
 			{
 				Element->ComponentIndices.Remove(ComponentIndex);
 				
@@ -7098,8 +6935,6 @@ void URigHierarchy::ShrinkComponentStorage()
 		Element->ComponentIndices.Reset();
 	}
 
-	TopLevelComponentIndices.Reset();
-
 	ElementComponents.RemoveAll([](const FInstancedStruct& InstancedStruct) -> bool
 	{
 		return !InstancedStruct.IsValid();
@@ -7117,10 +6952,6 @@ void URigHierarchy::ShrinkComponentStorage()
 			Component->Element = Find(Component->GetElementKey());
 			Component->Element->ComponentIndices.Add(ComponentIndex);
 		}
-		else if(Component->GetElementKey() == GetTopLevelComponentElementKey())
-		{
-			TopLevelComponentIndices.Add(ComponentIndex);
-		}
 		ComponentIndexLookup.Add(Component->GetKey(), ComponentIndex);
 	}
 }
diff --git a/Engine/Plugins/Animation/ControlRig/Source/ControlRig/Private/Rigs/RigHierarchyController.cpp b/Engine/Plugins/Animation/ControlRig/Source/ControlRig/Private/Rigs/RigHierarchyController.cpp
--- a/Engine/Plugins/Animation/ControlRig/Source/ControlRig/Private/Rigs/RigHierarchyController.cpp
+++ b/Engine/Plugins/Animation/ControlRig/Source/ControlRig/Private/Rigs/RigHierarchyController.cpp
@@ -959,15 +959,11 @@ FRigComponentKey URigHierarchyController::AddComponent(UScriptStruct* InComponen
 	URigHierarchy* Hierarchy = GetHierarchy();
 	check(Hierarchy);
 
-	FRigBaseElement* Element = nullptr;
-	if(InElement != URigHierarchy::GetTopLevelComponentElementKey())
+	FRigBaseElement* Element = Hierarchy->Find(InElement);
+	if(Element == nullptr)
 	{
-		Element = Hierarchy->Find(InElement);
-		if(Element == nullptr)
-		{
-			ReportErrorf(TEXT("The element '%s' could not be found."), *InElement.ToString());
-			return FRigComponentKey();
-		}
+		ReportErrorf(TEXT("The element '%s' could not be found."), *InElement.ToString());
+		return FRigComponentKey();
 	}
 
 	FString FailureReason;
@@ -1028,12 +1024,6 @@ FRigComponentKey URigHierarchyController::AddComponent(UScriptStruct* InComponen
 	return NewComponent->GetKey();
 }
 
-FRigComponentKey URigHierarchyController::AddTopLevelComponent(UScriptStruct* InComponentStruct, FName InName, FString InContent, bool bSetupUndo,
-	bool bPrintPythonCommand)
-{
-	return AddComponent(InComponentStruct, InName, URigHierarchy::GetTopLevelComponentElementKey(), InContent, bSetupUndo, bPrintPythonCommand);
-}
-
 bool URigHierarchyController::RemoveComponent(FRigComponentKey InComponent, bool bSetupUndo, bool bPrintPythonCommand)
 {
 	if(!IsValid())
@@ -2800,10 +2790,7 @@ TArray<FString> URigHierarchyController::GetAddComponentPythonCommands(const FRi
 {
 	TArray<FString> Commands;
 
-	const FString ElementKeyStr =
-		Component->GetElement() ?
-		Component->GetElement()->GetKey().ToPythonString() :
-		URigHierarchy::GetTopLevelComponentElementKey().ToPythonString();
+	const FString ElementKeyStr = Component->GetElement()->GetKey().ToPythonString();
 
 	FString ContentStr;
 	Component->GetScriptStruct()->ExportText(ContentStr, Component, Component, nullptr, PPF_None, nullptr);
@@ -3460,7 +3447,7 @@ bool URigHierarchyController::ReparentComponent(FRigBaseComponent* InComponent,
 		return false;
 	}
 
-	const FRigElementKey NewParentKey = InParentElement ? InParentElement->GetKey() : URigHierarchy::GetTopLevelComponentElementKey();
+	const FRigElementKey NewParentKey = InParentElement->GetKey();
 
 	URigHierarchy* Hierarchy = GetHierarchy();
 	if(!Hierarchy->CanAddComponent(NewParentKey, InComponent))
@@ -3478,11 +3465,7 @@ bool URigHierarchyController::ReparentComponent(FRigBaseComponent* InComponent,
 		DeselectComponent(OldKey);
 	}
 
-	if(OldKey.ElementKey == URigHierarchy::GetTopLevelComponentElementKey())
-	{
-		verify(Hierarchy->TopLevelComponentIndices.Remove(InComponent->IndexInHierarchy) > 0);
-	}
-	else if(FRigBaseElement* OldParentElement = Hierarchy->Find(OldKey.ElementKey))
+	if(FRigBaseElement* OldParentElement = Hierarchy->Find(OldKey.ElementKey))
 	{
 		OldParentElement->ComponentIndices.Remove(InComponent->IndexInHierarchy);
 	}
@@ -3491,33 +3474,16 @@ bool URigHierarchyController::ReparentComponent(FRigBaseComponent* InComponent,
 	InComponent->Key.ElementKey = NewParentKey;
 	InComponent->CachedNameString.Reset();
 
-	if(NewParentKey == URigHierarchy::GetTopLevelComponentElementKey())
-	{
-		check(!Hierarchy->TopLevelComponentIndices.Contains(InComponent->IndexInHierarchy));
-		Hierarchy->TopLevelComponentIndices.Add(InComponent->IndexInHierarchy);
-		InComponent->Element = nullptr;
-		
-		for(int32 IndexInElement = 0; IndexInElement < Hierarchy->TopLevelComponentIndices.Num(); IndexInElement++)
-		{
-			if(FRigBaseComponent* RemainingComponent = Hierarchy->GetComponent(Hierarchy->TopLevelComponentIndices[IndexInElement]))
-			{
-				RemainingComponent->IndexInElement = IndexInElement;
-			}
-		}
-	}
-	else
-	{
-		check(InParentElement);
-		check(!InParentElement->ComponentIndices.Contains(InComponent->IndexInHierarchy));
-		InParentElement->ComponentIndices.Add(InComponent->IndexInHierarchy);
-		InComponent->Element = InParentElement;
+	check(InParentElement);
+	check(!InParentElement->ComponentIndices.Contains(InComponent->IndexInHierarchy));
+	InParentElement->ComponentIndices.Add(InComponent->IndexInHierarchy);
+	InComponent->Element = InParentElement;
 
-		for(int32 IndexInElement = 0; IndexInElement < InParentElement->ComponentIndices.Num(); IndexInElement++)
+	for(int32 IndexInElement = 0; IndexInElement < InParentElement->ComponentIndices.Num(); IndexInElement++)
+	{
+		if(FRigBaseComponent* RemainingComponent = Hierarchy->GetComponent(InParentElement->ComponentIndices[IndexInElement]))
 		{
-			if(FRigBaseComponent* RemainingComponent = Hierarchy->GetComponent(InParentElement->ComponentIndices[IndexInElement]))
-			{
-				RemainingComponent->IndexInElement = IndexInElement;
-			}
+			RemainingComponent->IndexInElement = IndexInElement;
 		}
 	}
 	
diff --git a/Engine/Plugins/Animation/ControlRig/Source/ControlRig/Private/Rigs/RigHierarchyDefines.cpp b/Engine/Plugins/Animation/ControlRig/Source/ControlRig/Private/Rigs/RigHierarchyDefines.cpp
--- a/Engine/Plugins/Animation/ControlRig/Source/ControlRig/Private/Rigs/RigHierarchyDefines.cpp
+++ b/Engine/Plugins/Animation/ControlRig/Source/ControlRig/Private/Rigs/RigHierarchyDefines.cpp
@@ -498,8 +498,7 @@ void FRigComponentKey::Load(FArchive& Ar)
 
 bool FRigComponentKey::IsValid() const
 {
-	return Name.IsValid() && Name != NAME_None &&
-		(ElementKey.IsValid() || ElementKey == URigHierarchy::GetTopLevelComponentElementKey());
+	return Name.IsValid() && Name != NAME_None && ElementKey.IsValid();
 }
 
 FString FRigComponentKey::ToPythonString() const
@@ -513,11 +512,6 @@ FString FRigComponentKey::ToPythonString() const
 #endif
 }
 
-bool FRigComponentKey::IsTopLevel() const
-{
-	return ElementKey == URigHierarchy::GetTopLevelComponentElementKey();
-}
-
 FRigComponentKey FRigComponentKey::ConvertToModuleNameFormat(const TMap<FRigHierarchyModulePath, FName>* InModulePathToModuleName) const
 {
 	FRigComponentKey Result = *this;
diff --git a/Engine/Plugins/Animation/ControlRig/Source/ControlRig/Private/Units/Hierarchy/RigUnit_Component.cpp b/Engine/Plugins/Animation/ControlRig/Source/ControlRig/Private/Units/Hierarchy/RigUnit_Component.cpp
--- a/Engine/Plugins/Animation/ControlRig/Source/ControlRig/Private/Units/Hierarchy/RigUnit_Component.cpp
+++ b/Engine/Plugins/Animation/ControlRig/Source/ControlRig/Private/Units/Hierarchy/RigUnit_Component.cpp
@@ -190,71 +190,6 @@ void FRigDispatch_SpawnComponent::Execute(FRigVMExtendedExecuteContext& InContex
 	}
 }
 
-//////////////////////////////////////////////////////////////////////////////////////////////////////
-// FRigDispatch_SpawnTopLevelComponent
-//////////////////////////////////////////////////////////////////////////////////////////////////////
-
-const TArray<FRigVMTemplateArgumentInfo>& FRigDispatch_SpawnTopLevelComponent::GetArgumentInfos() const
-{
-	if (Infos.IsEmpty())
-	{
-		NameArgIndex = Infos.Emplace(NameArgName, ERigVMPinDirection::Input, RigVMTypeUtils::TypeIndex::FName);
-		ComponentArgIndex = Infos.Add(FRigVMTemplateArgumentInfo(ComponentArgName, ERigVMPinDirection::Input, {FRigVMTemplateArgument::ETypeCategory_SingleScriptStructValue}, &FRigDispatch_SpawnTopLevelComponent::IsTypeSupported));
-		KeyArgIndex = Infos.Emplace(KeyArgName, ERigVMPinDirection::Output, GetComponentKeyType());
-		SuccessArgIndex = Infos.Emplace(SuccessArgName, ERigVMPinDirection::Output, RigVMTypeUtils::TypeIndex::Bool);
-	}
-	return Infos;
-}
-
-#if WITH_EDITOR
-
-FString FRigDispatch_SpawnTopLevelComponent::GetKeywords() const
-{
-	return FRigDispatch_ComponentBase::GetKeywords() + TEXT(",AddComponent,CreateComponent,TopLevel,RootComponent");
-}
-
-#endif
-
-void FRigDispatch_SpawnTopLevelComponent::Execute(FRigVMExtendedExecuteContext& InContext, FRigVMMemoryHandleArray Handles, FRigVMPredicateBranchArray /* not needed */)
-{
-	const FRigDispatch_SpawnTopLevelComponent* Factory = static_cast<const FRigDispatch_SpawnTopLevelComponent*>(InContext.Factory);
-		
-	// unpack the memory
-	const FName& Name = *(const FName*)Handles[Factory->NameArgIndex].GetData();
-	const uint8* SourceComponent = Handles[Factory->ComponentArgIndex].GetData();
-	FRigComponentKey& ComponentKey = *(FRigComponentKey*)Handles[Factory->KeyArgIndex].GetData();
-	bool& bSuccess = *(bool*)Handles[Factory->SuccessArgIndex].GetData();
-
-	const FStructProperty* StructProperty = CastField<FStructProperty>(Handles[Factory->ComponentArgIndex].GetResolvedProperty());
-	check(StructProperty);
-	UScriptStruct* ComponentStruct = StructProperty->Struct;
-	check(ComponentStruct && ComponentStruct->IsChildOf(FRigBaseComponent::StaticStruct()));
-
-	ComponentKey = FRigComponentKey();
-	bSuccess = false;
-
-	// extract the component from the hierarchy
-	const FControlRigExecuteContext& Context = InContext.GetPublicDataSafe<FControlRigExecuteContext>();
-	if(URigHierarchyController* Controller = Context.Hierarchy->GetController())
-	{
-		FRigHierarchyControllerInstructionBracket InstructionBracket(Controller, Context.GetInstructionIndex());
-		ComponentKey = Controller->AddTopLevelComponent(ComponentStruct, Name);
-		if(ComponentKey.IsValid())
-		{
-			FRigBaseComponent* TargetComponent = Context.Hierarchy->FindComponent(ComponentKey);
-			check(TargetComponent);
-			
-			// copy the public content over (members which are not UProperty will be skipped) 
-			for (TFieldIterator<FProperty> It(ComponentStruct); It; ++It)
-			{
-				It->CopyCompleteValue_InContainer(TargetComponent, SourceComponent);
-			}
-
-			bSuccess = true;
-		}
-	}
-}
-
 //////////////////////////////////////////////////////////////////////////////////////////////////////
 // FRigDispatch_GetComponentContent
 //////////////////////////////////////////////////////////////////////////////////////////////////////
@@ -370,141 +305,3 @@ void FRigDispatch_SetComponentContent::Execute(FRigVMExtendedExecuteContext& InC
 	}
 	bSuccess = true;
 }
-
-//////////////////////////////////////////////////////////////////////////////////////////////////////
-// FRigDispatch_GetTopLevelComponentContent
-//////////////////////////////////////////////////////////////////////////////////////////////////////
-
-const TArray<FRigVMTemplateArgumentInfo>& FRigDispatch_GetTopLevelComponentContent::GetArgumentInfos() const
-{
-	if (Infos.IsEmpty())
-	{
-		NameArgIndex = Infos.Emplace(NameArgName, ERigVMPinDirection::Input, RigVMTypeUtils::TypeIndex::FName);
-		ComponentArgIndex = Infos.Add(FRigVMTemplateArgumentInfo(ComponentArgName, ERigVMPinDirection::Output, {FRigVMTemplateArgument::ETypeCategory_SingleScriptStructValue}, &FRigDispatch_GetTopLevelComponentContent::IsTypeSupported));
-		SuccessArgIndex = Infos.Emplace(SuccessArgName, ERigVMPinDirection::Output, RigVMTypeUtils::TypeIndex::Bool);
-	}
-	
-	return Infos;
-}
-
-#if WITH_EDITOR
-
-FString FRigDispatch_GetTopLevelComponentContent::GetKeywords() const
-{
-	return FRigDispatch_ComponentBase::GetKeywords() + TEXT(",GetComponent,ReadComponent");
-}
-
-#endif
-
-void FRigDispatch_GetTopLevelComponentContent::Execute(FRigVMExtendedExecuteContext& InContext, FRigVMMemoryHandleArray Handles, FRigVMPredicateBranchArray)
-{
-	const FRigDispatch_GetTopLevelComponentContent* Factory = static_cast<const FRigDispatch_GetTopLevelComponentContent*>(InContext.Factory);
-			
-	// unpack the memory
-	const FName& Name = *(const FName*)Handles[Factory->NameArgIndex].GetData();
-	const FRigComponentKey ComponentKey(URigHierarchy::GetTopLevelComponentElementKey(), Name);
-	uint8* TargetComponent = Handles[Factory->ComponentArgIndex].GetData();
-	bool& bSuccess = *(bool*)Handles[Factory->SuccessArgIndex].GetData();
-
-	const FStructProperty* StructProperty = CastField<FStructProperty>(Handles[Factory->ComponentArgIndex].GetResolvedProperty());
-	check(StructProperty);
-	const UScriptStruct* ComponentStruct = StructProperty->Struct;
-	check(ComponentStruct && ComponentStruct->IsChildOf(FRigBaseComponent::StaticStruct()));
-
-	// extract the component from the hierarchy
-	const FControlRigExecuteContext& Context = InContext.GetPublicDataSafe<FControlRigExecuteContext>();
-	const FRigBaseComponent* SourceComponent = Context.Hierarchy->FindComponent(ComponentKey);
-	if(SourceComponent == nullptr)
-	{
-		ComponentStruct->InitializeDefaultValue(TargetComponent);
-#if WITH_EDITOR
-		Context.Report(FRigVMLogSettings(EMessageSeverity::Warning), Context.GetFunctionName(), Context.GetInstructionIndex(), FString::Printf(TEXT("%s not found."), *ComponentKey.ToString()));
-#endif
-		bSuccess = false;
-		return;
-	}
-
-	if(SourceComponent->GetScriptStruct() != ComponentStruct)
-	{
-		ComponentStruct->InitializeDefaultValue(TargetComponent);
-#if WITH_EDITOR
-		const FString Message = FString::Printf(TEXT("Component pin type (%s) doesn't match component in hierarchy (%s)."), *ComponentStruct->GetName(), *SourceComponent->GetScriptStruct()->GetName());
-		Context.Report(FRigVMLogSettings(EMessageSeverity::Warning), Context.GetFunctionName(), Context.GetInstructionIndex(), Message);
-#endif
-		bSuccess = false;
-		return;
-	}
-
-	ComponentStruct->CopyScriptStruct(TargetComponent, SourceComponent, 1);
-	bSuccess = true;
-}
-
-//////////////////////////////////////////////////////////////////////////////////////////////////////
-// FRigDispatch_SetTopLevelComponentContent
-//////////////////////////////////////////////////////////////////////////////////////////////////////
-
-const TArray<FRigVMTemplateArgumentInfo>& FRigDispatch_SetTopLevelComponentContent::GetArgumentInfos() const
-{
-	if (Infos.IsEmpty())
-	{
-		NameArgIndex = Infos.Emplace(NameArgName, ERigVMPinDirection::Input, RigVMTypeUtils::TypeIndex::FName);
-		ComponentArgIndex = Infos.Add(FRigVMTemplateArgumentInfo(ComponentArgName, ERigVMPinDirection::Input, {FRigVMTemplateArgument::ETypeCategory_SingleScriptStructValue}, &FRigDispatch_SetTopLevelComponentContent::IsTypeSupported));
-		SuccessArgIndex = Infos.Emplace(SuccessArgName, ERigVMPinDirection::Output, RigVMTypeUtils::TypeIndex::Bool);
-	}
-	
-	return Infos;
-}
-
-#if WITH_EDITOR
-
-FString FRigDispatch_SetTopLevelComponentContent::GetKeywords() const
-{
-	return FRigDispatch_ComponentBase::GetKeywords() + TEXT(",SetComponent,WriteComponent");
-}
-
-#endif
-
-void FRigDispatch_SetTopLevelComponentContent::Execute(FRigVMExtendedExecuteContext& InContext, FRigVMMemoryHandleArray Handles, FRigVMPredicateBranchArray)
-{
-	const FRigDispatch_SetTopLevelComponentContent* Factory = static_cast<const FRigDispatch_SetTopLevelComponentContent*>(InContext.Factory);
-			
-	// unpack the memory
-	const FName& Name = *(const FName*)Handles[Factory->NameArgIndex].GetData();
-	const FRigComponentKey ComponentKey(URigHierarchy::GetTopLevelComponentElementKey(), Name);
-	const  uint8* SourceComponent = Handles[Factory->ComponentArgIndex].GetData();
-	bool& bSuccess = *(bool*)Handles[Factory->SuccessArgIndex].GetData();
-
-	const FStructProperty* StructProperty = CastField<FStructProperty>(Handles[Factory->ComponentArgIndex].GetResolvedProperty());
-	check(StructProperty);
-	const UScriptStruct* ComponentStruct = StructProperty->Struct;
-	check(ComponentStruct && ComponentStruct->IsChildOf(FRigBaseComponent::StaticStruct()));
-
-	// extract the component from the hierarchy
-	const FControlRigExecuteContext& Context = InContext.GetPublicDataSafe<FControlRigExecuteContext>();
-	FRigBaseComponent* TargetComponent = Context.Hierarchy->FindComponent(ComponentKey);
-	if(TargetComponent == nullptr)
-	{
-#if WITH_EDITOR
-		Context.Report(FRigVMLogSettings(EMessageSeverity::Warning), Context.GetFunctionName(), Context.GetInstructionIndex(), FString::Printf(TEXT("%s not found."), *ComponentKey.ToString()));
-#endif
-		bSuccess = false;
-		return;
-	}
-
-	if(TargetComponent->GetScriptStruct() != ComponentStruct)
-	{
-#if WITH_EDITOR
-		const FString Message = FString::Printf(TEXT("Component pin type (%s) doesn't match component in hierarchy (%s)."), *ComponentStruct->GetName(), *TargetComponent->GetScriptStruct()->GetName());
-		Context.Report(FRigVMLogSettings(EMessageSeverity::Warning), Context.GetFunctionName(), Context.GetInstructionIndex(), Message);
-#endif
-		bSuccess = false;
-		return;
-	}
-
-	// copy the public content over (members which are not UProperty will be skipped) 
-	for (TFieldIterator<FProperty> It(ComponentStruct); It; ++It)
-	{
-		It->CopyCompleteValue_InContainer(TargetComponent, SourceComponent);
-	}
-	bSuccess = true;
-}
\ No newline at end of file
diff --git a/Engine/Plugins/Animation/ControlRig/Source/ControlRig/Public/Rigs/RigHierarchy.h b/Engine/Plugins/Animation/ControlRig/Source/ControlRig/Public/Rigs/RigHierarchy.h
--- a/Engine/Plugins/Animation/ControlRig/Source/ControlRig/Public/Rigs/RigHierarchy.h
+++ b/Engine/Plugins/Animation/ControlRig/Source/ControlRig/Public/Rigs/RigHierarchy.h
@@ -831,86 +831,6 @@ public:
 		return Result;
 	}
 
-	/**
-	 * Returns a component for a given top level name
-	 * @param InName The name of the component to retrieve.
-	 * @return The retrieved component or nullptr.
-	 */
-	CONTROLRIG_API const FRigBaseComponent* FindTopLevelComponent(const FName& InName) const;
-
-	/**
-	 * Returns a component for given top level name
-	 * @param InName The name of the component to retrieve.
-	 * @return The retrieved component or nullptr.
-	 */
-	CONTROLRIG_API FRigBaseComponent* FindTopLevelComponent(const FName& InName);
-
-	/**
-	 * Returns a component at a given index or nullptr.
-	 * @param InTopLevelComponentIndex The index of the top level component to retrieve.
-	 * @return The retrieved component or nullptr.
-	 */
-	CONTROLRIG_API const FRigBaseComponent* GetTopLevelComponent(int32 InTopLevelComponentIndex) const;
-
-	/**
-	 * Returns a component at a given index or nullptr.
-	 * @param InTopLevelComponentIndex The index of the top level component to retrieve.
-	 * @return The retrieved component or nullptr.
-	 */
-	CONTROLRIG_API FRigBaseComponent* GetTopLevelComponent(int32 InTopLevelComponentIndex);
-
-	/**
-	 * Returns a component at a given index or nullptr.
-	 * This templated method also casts to the chosen
-	 * component type but does not guarantee a valid result.
-	 * @param InTopLevelComponentIndex The index of the top level component to retrieve.
-	 * @return The retrieved component or nullptr.
-	 */
-	template<typename T>
-	const T* GetTopLevelComponent(int32 InTopLevelComponentIndex) const
-	{
-		return Cast<T>(GetTopLevelComponent(InTopLevelComponentIndex));
-	}
-
-	/**
-	 * Returns a component at a given index or nullptr.
-	 * This templated method also casts to the chosen
-	 * component type but does not guarantee a valid result.
-	 * @param InTopLevelComponentIndex The index of the top level component to retrieve.
-	 * @return The retrieved component or nullptr.
-	 */
-	template<typename T>
-	T* GetTopLevelComponent(int32 InTopLevelComponentIndex)
-	{
-		return Cast<T>(GetTopLevelComponent(InTopLevelComponentIndex));
-	}
-
-	/**
-	 * Returns a component at a given index.
-	 * This templated method also casts to the chosen
-	 * component type and checks for a the valid result.
-	 * @param InTopLevelComponentIndex The index of the top level component to retrieve.
-	 * @return The retrieved component or nullptr.
-	 */
-	template<typename T>
-	const T* GetTopLevelComponentChecked(int32 InTopLevelComponentIndex) const
-	{
-		return CastChecked<T>(GetTopLevelComponent(InTopLevelComponentIndex));
-	}
-
-	/**
-	 * Returns a component at a given index.
-	 * This templated method also casts to the chosen
-	 * component type and checks for a the valid result.
-	 * @param InTopLevelComponentIndex The index of the top level component to retrieve.
-	 * @return The retrieved component or nullptr.
-	 */
-	template<typename T>
-    T* GetTopLevelComponentChecked(int32 InTopLevelComponentIndex)
-	{
-		return CastChecked<T>(GetTopLevelComponent(InTopLevelComponentIndex));
-	}
-
 	/**
 	 * Returns the count of all components
 	 * @return The count of all components
@@ -955,29 +875,9 @@ public:
 	UFUNCTION(BlueprintPure, Category = URigHierarchy)
 	CONTROLRIG_API FString GetComponentContent(FRigElementKey InElement, int32 InComponentIndex) const;
 
-	UFUNCTION(BlueprintPure, Category = URigHierarchy)
-	CONTROLRIG_API int32 NumTopLevelComponents() const;
-
-	UFUNCTION(BlueprintPure, Category = URigHierarchy)
-	CONTROLRIG_API TArray<FRigComponentKey> GetTopLevelComponentKeys() const;
-
-	UFUNCTION(BlueprintPure, Category = URigHierarchy)
-	CONTROLRIG_API FRigComponentKey GetTopLevelComponentKey(int32 InTopLevelComponentIndex) const;
-
-	UFUNCTION(BlueprintPure, Category = URigHierarchy)
-	CONTROLRIG_API FName GetTopLevelComponentName(int32 InTopLevelComponentIndex) const;
-	
-	UFUNCTION(BlueprintPure, Category = URigHierarchy)
-	CONTROLRIG_API UScriptStruct* GetTopLevelComponentType(int32 InTopLevelComponentIndex) const;
-
-	UFUNCTION(BlueprintPure, Category = URigHierarchy)
-	CONTROLRIG_API FString GetTopLevelComponentContent(int32 InTopLevelComponentIndex) const;
-
 	CONTROLRIG_API bool CanAddComponent(FRigElementKey InElementKey, const UScriptStruct* InComponentStruct, FString* OutFailureReason = nullptr) const;
 	CONTROLRIG_API bool CanAddComponent(FRigElementKey InElementKey, const FRigBaseComponent* InComponent, FString* OutFailureReason = nullptr) const;
 
-	static CONTROLRIG_API const FRigElementKey& GetTopLevelComponentElementKey();
-
 private:
 	/**
 	* Returns bone element for a given key, for scripting purpose only, for cpp usage, use Find<FRigBoneElement>()
@@ -4851,9 +4751,6 @@ private:
 	// the component storage.
 	TArray<FInstancedStruct> ElementComponents;
 
-	// List of top level component indices
-	TArray<int32> TopLevelComponentIndices;
-
 	// A quick-lookup cache for elements' children. Each element that has a ChildCacheIndex
 	// not equal to INDEX_NONE is an index into the offset and count cache below, which in
 	// turn contains an offset into the ChildElementCache, which stores consecutive runs of
diff --git a/Engine/Plugins/Animation/ControlRig/Source/ControlRig/Public/Rigs/RigHierarchyController.h b/Engine/Plugins/Animation/ControlRig/Source/ControlRig/Public/Rigs/RigHierarchyController.h
--- a/Engine/Plugins/Animation/ControlRig/Source/ControlRig/Public/Rigs/RigHierarchyController.h
+++ b/Engine/Plugins/Animation/ControlRig/Source/ControlRig/Public/Rigs/RigHierarchyController.h
@@ -337,18 +337,6 @@ public:
 	UFUNCTION(BlueprintCallable, Category = URigHierarchyController)
 	UE_API FRigComponentKey AddComponent(UScriptStruct* InComponentStruct, FName InName, FRigElementKey InElement, FString InContent = TEXT(""), bool bSetupUndo = false, bool bPrintPythonCommand = false);
 
-	/**
-	 * Adds a component to the hierarchy top level of the hierarchy
-	 * @param InComponentStruct The script struct of the component to add
-	 * @param InName The suggested name of the new component.
-	 * @param InContent The (optional) serialized text default for the component
-	 * @param bSetupUndo If set to true the stack will record the change for undo / redo
-	 * @param bPrintPythonCommand If set to true a python command equivalent to this call will be printed out
-	 * @return The name of the newly created component
-	 */
-	UFUNCTION(BlueprintCallable, Category = URigHierarchyController)
-	UE_API FRigComponentKey AddTopLevelComponent(UScriptStruct* InComponentStruct, FName InName, FString InContent = TEXT(""), bool bSetupUndo = false, bool bPrintPythonCommand = false);
-
 	/**
      * Removes a component from the hierarchy
      * @param InComponent The component to remove
diff --git a/Engine/Plugins/Animation/ControlRig/Source/ControlRig/Public/Units/Hierarchy/RigUnit_Component.h b/Engine/Plugins/Animation/ControlRig/Source/ControlRig/Public/Units/Hierarchy/RigUnit_Component.h
--- a/Engine/Plugins/Animation/ControlRig/Source/ControlRig/Public/Units/Hierarchy/RigUnit_Component.h
+++ b/Engine/Plugins/Animation/ControlRig/Source/ControlRig/Public/Units/Hierarchy/RigUnit_Component.h
@@ -77,36 +77,6 @@ protected:
 	static UE_API void Execute(FRigVMExtendedExecuteContext& InContext, FRigVMMemoryHandleArray Handles, FRigVMPredicateBranchArray /* not needed */);
 };
 
-/**
- * Spawns a top level component
- */
-USTRUCT(meta=(DisplayName="Spawn Top Level Component"))
-struct FRigDispatch_SpawnTopLevelComponent : public FRigDispatch_ComponentBase
-{
-	GENERATED_BODY()
-
-	FRigDispatch_SpawnTopLevelComponent()
-	{
-		// this is only needed for engine test
-		FactoryScriptStruct = FRigDispatch_SpawnTopLevelComponent::StaticStruct();
-	}
-
-	UE_API virtual const TArray<FRigVMTemplateArgumentInfo>& GetArgumentInfos() const override;
-#if WITH_EDITOR
-	UE_API virtual FString GetKeywords() const override;
-#endif
-
-protected:
-
-	virtual bool IsMutable() const override { return true; }
-
-	virtual FRigVMFunctionPtr GetDispatchFunctionImpl(const FRigVMTemplateTypeMap& InTypes) const override
-	{
-		return &Execute;
-	}
-	static UE_API void Execute(FRigVMExtendedExecuteContext& InContext, FRigVMMemoryHandleArray Handles, FRigVMPredicateBranchArray /* not needed */);
-};
-
 /**
  * Gets the component
  */
@@ -159,62 +129,4 @@ protected:
 	static UE_API void Execute(FRigVMExtendedExecuteContext& InContext, FRigVMMemoryHandleArray Handles, FRigVMPredicateBranchArray /* not needed */);
 };
 
-/**
- * Gets the top level component
- */
-USTRUCT(meta=(DisplayName="Get Top Level Component"))
-struct FRigDispatch_GetTopLevelComponentContent : public FRigDispatch_ComponentBase
-{
-	GENERATED_BODY()
-
-	FRigDispatch_GetTopLevelComponentContent()
-	{
-		// this is only needed for engine test
-		FactoryScriptStruct = FRigDispatch_GetTopLevelComponentContent::StaticStruct();
-	}
-
-	UE_API virtual const TArray<FRigVMTemplateArgumentInfo>& GetArgumentInfos() const override;
-#if WITH_EDITOR
-	UE_API virtual FString GetKeywords() const override;
-#endif
-
-protected:
-
-	virtual FRigVMFunctionPtr GetDispatchFunctionImpl(const FRigVMTemplateTypeMap& InTypes) const override
-	{
-		return &Execute;
-	}
-	static UE_API void Execute(FRigVMExtendedExecuteContext& InContext, FRigVMMemoryHandleArray Handles, FRigVMPredicateBranchArray /* not needed */);
-};
-
-/**
- * Set the content of a component
- */
-USTRUCT(meta=(DisplayName="Set Top Level Component"))
-struct FRigDispatch_SetTopLevelComponentContent : public FRigDispatch_ComponentBase
-{
-	GENERATED_BODY()
-
-	FRigDispatch_SetTopLevelComponentContent()
-	{
-		// this is only needed for engine test
-		FactoryScriptStruct = FRigDispatch_SetTopLevelComponentContent::StaticStruct();
-	}
-
-	UE_API virtual const TArray<FRigVMTemplateArgumentInfo>& GetArgumentInfos() const override;
-#if WITH_EDITOR
-	UE_API virtual FString GetKeywords() const override;
-#endif
-
-protected:
-
-	virtual bool IsMutable() const override { return true; }
-
-	virtual FRigVMFunctionPtr GetDispatchFunctionImpl(const FRigVMTemplateTypeMap& InTypes) const override
-	{
-		return &Execute;
-	}
-	static UE_API void Execute(FRigVMExtendedExecuteContext& InContext, FRigVMMemoryHandleArray Handles, FRigVMPredicateBranchArray /* not needed */);
-};
-
 #undef UE_API
diff --git a/Engine/Plugins/Animation/ControlRig/Source/ControlRigEditor/Private/Editor/ControlRigEditor.cpp b/Engine/Plugins/Animation/ControlRig/Source/ControlRigEditor/Private/Editor/ControlRigEditor.cpp
--- a/Engine/Plugins/Animation/ControlRig/Source/ControlRigEditor/Private/Editor/ControlRigEditor.cpp
+++ b/Engine/Plugins/Animation/ControlRig/Source/ControlRigEditor/Private/Editor/ControlRigEditor.cpp
@@ -5958,8 +5958,6 @@ void FControlRigBaseEditor::HandleMakeComponentGetterSetter(bool bIsGetter, TArr
 
 	static const FRigVMDispatchFactory* GetComponentFactory = FRigVMRegistry_RWLock::Get().FindDispatchFactory(FRigDispatch_GetComponentContent().GetFactoryName());
 	static const FRigVMDispatchFactory* SetComponentFactory = FRigVMRegistry_RWLock::Get().FindDispatchFactory(FRigDispatch_SetComponentContent().GetFactoryName());
-	static const FRigVMDispatchFactory* GetTopLevelComponentFactory = FRigVMRegistry_RWLock::Get().FindDispatchFactory(FRigDispatch_GetTopLevelComponentContent().GetFactoryName());
-	static const FRigVMDispatchFactory* SetTopLevelComponentFactory = FRigVMRegistry_RWLock::Get().FindDispatchFactory(FRigDispatch_SetTopLevelComponentContent().GetFactoryName());
 
 	TArray<FName> NewNodes;
 
@@ -5979,14 +5977,7 @@ void FControlRigBaseEditor::HandleMakeComponentGetterSetter(bool bIsGetter, TArr
 		}
 
 		const FRigVMDispatchFactory* Factory = nullptr;
-		if(Key.IsTopLevel())
-		{
-			Factory = bIsGetter ? GetTopLevelComponentFactory : SetTopLevelComponentFactory;
-		}
-		else
-		{
-			Factory = bIsGetter ? GetComponentFactory : SetComponentFactory;
-		}
+		Factory = bIsGetter ? GetComponentFactory : SetComponentFactory;
 
 		check(Factory);
 
diff --git a/Engine/Plugins/Animation/ControlRig/Source/ControlRigEditor/Private/Editor/SRigHierarchy.cpp b/Engine/Plugins/Animation/ControlRig/Source/ControlRigEditor/Private/Editor/SRigHierarchy.cpp
--- a/Engine/Plugins/Animation/ControlRig/Source/ControlRigEditor/Private/Editor/SRigHierarchy.cpp
+++ b/Engine/Plugins/Animation/ControlRig/Source/ControlRigEditor/Private/Editor/SRigHierarchy.cpp
@@ -1540,10 +1540,6 @@ void SRigHierarchy::CreateContextMenu()
 										SelectedElements.Add(SelectedItem->Key.GetElement());
 									}
 								}
-								if(SelectedElements.IsEmpty())
-								{
-									SelectedElements.Add(URigHierarchy::GetTopLevelComponentElementKey());
-								}
 
 								for(UScriptStruct* RigComponentStruct : RigComponentStructs)
 								{
diff --git a/Engine/Plugins/Animation/ControlRig/Source/ControlRigEditor/Private/Editor/SRigHierarchyTreeView.cpp b/Engine/Plugins/Animation/ControlRig/Source/ControlRigEditor/Private/Editor/SRigHierarchyTreeView.cpp
--- a/Engine/Plugins/Animation/ControlRig/Source/ControlRigEditor/Private/Editor/SRigHierarchyTreeView.cpp
+++ b/Engine/Plugins/Animation/ControlRig/Source/ControlRigEditor/Private/Editor/SRigHierarchyTreeView.cpp
@@ -967,12 +967,6 @@ void SRigHierarchyTreeView::RefreshTreeView(bool bRebuildContent)
 				bContinue = true;
 			});
 
-			// first add the root components
-			for(int32 Index = 0; Index < Hierarchy->NumTopLevelComponents(); Index++)
-			{
-				AddComponent(Hierarchy->GetTopLevelComponent(Index));
-			}
-
 			// sort the sockets by depth
 			Algo::SortBy(Sockets, [ElementDepth](const FRigSocketElement* Socket) -> int32
 			{
diff --git a/Engine/Plugins/Experimental/PhysicsControl/Source/ControlRigPhysics/Private/RigPhysicsSimulation.cpp b/Engine/Plugins/Experimental/PhysicsControl/Source/ControlRigPhysics/Private/RigPhysicsSimulation.cpp
--- a/Engine/Plugins/Experimental/PhysicsControl/Source/ControlRigPhysics/Private/RigPhysicsSimulation.cpp
+++ b/Engine/Plugins/Experimental/PhysicsControl/Source/ControlRigPhysics/Private/RigPhysicsSimulation.cpp
@@ -108,15 +108,8 @@ bool FRigPhysicsSimulation::ShouldComponentBeInSimulation(
 			// Note that getting the parent of an element at the root doesn't return the top-level element
 			ElementKey = Hierarchy->GetFirstParent(ElementKey);
 		}
-
-		for (FRigComponentKey CK : Hierarchy->GetTopLevelComponentKeys())
-		{
-			if (CK == SolverComponentKey)
-			{
-				return true;
-			}
-		}
 	}
+
 	return false;
 }
 
