Anim In Engine:  Make sure property handle is valid for anim details.

diff --git a/Engine/Plugins/Animation/ControlRig/Source/ControlRigEditor/Private/AnimDetails/Widgets/SAnimDetailsValueNumeric.cpp b/Engine/Plugins/Animation/ControlRig/Source/ControlRigEditor/Private/AnimDetails/Widgets/SAnimDetailsValueNumeric.cpp
--- a/Engine/Plugins/Animation/ControlRig/Source/ControlRigEditor/Private/AnimDetails/Widgets/SAnimDetailsValueNumeric.cpp
+++ b/Engine/Plugins/Animation/ControlRig/Source/ControlRigEditor/Private/AnimDetails/Widgets/SAnimDetailsValueNumeric.cpp
@@ -142,12 +142,11 @@ namespace UE::ControlRigEditor
 				{
 					FAnimDetailsMultiEditUtil::Get().MultiEditSet(*WeakProxyManager.Get(), Value, PropertyHandle.ToSharedRef());
 				}
-			}
-
-			NumericType NewValue;
-			if (PropertyHandle->GetValue(NewValue) == FPropertyAccess::Success)
-			{
-				return NewValue;
+				NumericType NewValue;
+				if (PropertyHandle->GetValue(NewValue) == FPropertyAccess::Success)
+				{
+					return NewValue;
+				}
 			}
 
 			return TOptional<NumericType>();
