Control rig: Improve fix for not being able to undo / redo controller selection (in viewport). The previous change broke Animation.AnimationData.Actions.Curves test in EngineTest due to an added Modify call in URigHierarchyController. This change makes it so only the selection code calls Modify.

diff --git a/Engine/Plugins/Animation/ControlRig/Source/ControlRig/Private/ControlRig.cpp b/Engine/Plugins/Animation/ControlRig/Source/ControlRig/Private/ControlRig.cpp
--- a/Engine/Plugins/Animation/ControlRig/Source/ControlRig/Private/ControlRig.cpp
+++ b/Engine/Plugins/Animation/ControlRig/Source/ControlRig/Private/ControlRig.cpp
@@ -2571,24 +2571,24 @@ const TArray<TSoftObjectPtr<UControlRigShapeLibrary>>& UControlRig::GetShapeLibr
 	return Libraries;
 }
 
-void UControlRig::SelectControl(const FName& InControlName, bool bSelect)
+void UControlRig::SelectControl(const FName& InControlName, bool bSelect, bool bSetupUndo)
 {
 	if(DynamicHierarchy)
 	{
 		if(URigHierarchyController* Controller = DynamicHierarchy->GetController(true))
 		{
-			Controller->SelectElement(FRigElementKey(InControlName, ERigElementType::Control), bSelect);
+			Controller->SelectElement(FRigElementKey(InControlName, ERigElementType::Control), bSelect, false, bSetupUndo);
 		}
 	}
 }
 
-bool UControlRig::ClearControlSelection()
+bool UControlRig::ClearControlSelection(bool bSetupUndo)
 {
 	if(DynamicHierarchy)
 	{
 		if(URigHierarchyController* Controller = DynamicHierarchy->GetController(true))
 		{
-			return Controller->ClearSelection();
+			return Controller->ClearSelection(bSetupUndo);
 		}
 	}
 	return false;
diff --git a/Engine/Plugins/Animation/ControlRig/Source/ControlRig/Private/Rigs/RigHierarchyController.cpp b/Engine/Plugins/Animation/ControlRig/Source/ControlRig/Private/Rigs/RigHierarchyController.cpp
--- a/Engine/Plugins/Animation/ControlRig/Source/ControlRig/Private/Rigs/RigHierarchyController.cpp
+++ b/Engine/Plugins/Animation/ControlRig/Source/ControlRig/Private/Rigs/RigHierarchyController.cpp
@@ -85,17 +85,17 @@ void URigHierarchyController::SetHierarchy(URigHierarchy* InHierarchy)
 	}
 }
 
-bool URigHierarchyController::SelectElement(FRigElementKey InKey, bool bSelect, bool bClearSelection)
+bool URigHierarchyController::SelectElement(FRigElementKey InKey, bool bSelect, bool bClearSelection, bool bSetupUndo)
 {
-	return SelectHierarchyKey(InKey, bSelect, bClearSelection);
+	return SelectHierarchyKey(InKey, bSelect, bClearSelection, bSetupUndo);
 }
 
-bool URigHierarchyController::SelectComponent(FRigComponentKey InKey, bool bSelect, bool bClearSelection)
+bool URigHierarchyController::SelectComponent(FRigComponentKey InKey, bool bSelect, bool bClearSelection, bool bSetupUndo)
 {
-	return SelectHierarchyKey(InKey, bSelect, bClearSelection);
+	return SelectHierarchyKey(InKey, bSelect, bClearSelection, bSetupUndo);
 }
 
-bool URigHierarchyController::SelectHierarchyKey(FRigHierarchyKey InKey, bool bSelect, bool bClearSelection)
+bool URigHierarchyController::SelectHierarchyKey(FRigHierarchyKey InKey, bool bSelect, bool bClearSelection, bool bSetupUndo)
 {
 	if(!IsValid())
 	{
@@ -115,7 +115,7 @@ bool URigHierarchyController::SelectHierarchyKey(FRigHierarchyKey InKey, bool bS
 	{
 		if(URigHierarchyController* ControllerForSelection = HierarchyForSelection->GetController())
 		{
-			return ControllerForSelection->SelectHierarchyKey(InKey, bSelect, bClearSelection);
+			return ControllerForSelection->SelectHierarchyKey(InKey, bSelect, bClearSelection, bSetupUndo);
 		}
 	}
 
@@ -170,7 +170,7 @@ bool URigHierarchyController::SelectHierarchyKey(FRigHierarchyKey InKey, bool bS
 		}
 
 		// If recording into transaction buffer, modify the hierarchy so undoing selection works correctly.
-		if (GUndo)
+		if (bSetupUndo && GUndo)
 		{
 			Hierarchy->Modify();
 		}
@@ -213,7 +213,7 @@ bool URigHierarchyController::SelectHierarchyKey(FRigHierarchyKey InKey, bool bS
 		}
 
 		// If recording into transaction buffer, modify the hierarchy so undoing selection works correctly.
-		if (GUndo)
+		if (bSetupUndo && GUndo)
 		{
 			Hierarchy->Modify();
 		}
@@ -244,7 +244,7 @@ bool URigHierarchyController::SelectHierarchyKey(FRigHierarchyKey InKey, bool bS
 	return true;
 }
 
-bool URigHierarchyController::SetSelection(const TArray<FRigElementKey>& InKeys, bool bPrintPythonCommand)
+bool URigHierarchyController::SetSelection(const TArray<FRigElementKey>& InKeys, bool bPrintPythonCommand, bool bSetupUndo)
 {
 	TArray<FRigHierarchyKey> Keys;
 	Keys.Reserve(InKeys.Num());
@@ -252,7 +252,7 @@ bool URigHierarchyController::SetSelection(const TArray<FRigElementKey>& InKeys,
 	{
 		Keys.Emplace(Key);
 	}
-	return SetHierarchySelection(Keys);
+	return SetHierarchySelection(Keys, bPrintPythonCommand, bSetupUndo);
 }
 
 bool URigHierarchyController::SetComponentSelection(const TArray<FRigComponentKey>& InKeys, bool bPrintPythonCommand)
@@ -266,7 +266,7 @@ bool URigHierarchyController::SetComponentSelection(const TArray<FRigComponentKe
 	return SetHierarchySelection(Keys);
 }
 
-bool URigHierarchyController::SetHierarchySelection(const TArray<FRigHierarchyKey>& InKeys, bool bPrintPythonCommand)
+bool URigHierarchyController::SetHierarchySelection(const TArray<FRigHierarchyKey>& InKeys, bool bPrintPythonCommand, bool bSetupUndo)
 {
 	if(!IsValid())
 	{
@@ -279,7 +279,7 @@ bool URigHierarchyController::SetHierarchySelection(const TArray<FRigHierarchyKe
 	{
 		if(URigHierarchyController* ControllerForSelection = HierarchyForSelection->GetController())
 		{
-			return ControllerForSelection->SetHierarchySelection(InKeys);
+			return ControllerForSelection->SetHierarchySelection(InKeys, false, bSetupUndo);
 		}
 	}
 
@@ -294,7 +294,7 @@ bool URigHierarchyController::SetHierarchySelection(const TArray<FRigHierarchyKe
 		{
 			if(!InKeys.Contains(KeyToDeselect))
 			{
-				if(!SelectHierarchyKey(KeyToDeselect, false))
+				if(!SelectHierarchyKey(KeyToDeselect, false, false, bSetupUndo))
 				{
 					bResult = false;
 				}
@@ -305,7 +305,7 @@ bool URigHierarchyController::SetHierarchySelection(const TArray<FRigHierarchyKe
 		{
 			if(!PreviousSelection.Contains(KeyToSelect))
 			{
-				if(!SelectHierarchyKey(KeyToSelect, true))
+				if(!SelectHierarchyKey(KeyToSelect, true, false, bSetupUndo))
 				{
 					bResult = false;
 				}
diff --git a/Engine/Plugins/Animation/ControlRig/Source/ControlRig/Public/ControlRig.h b/Engine/Plugins/Animation/ControlRig/Source/ControlRig/Public/ControlRig.h
--- a/Engine/Plugins/Animation/ControlRig/Source/ControlRig/Public/ControlRig.h
+++ b/Engine/Plugins/Animation/ControlRig/Source/ControlRig/Public/ControlRig.h
@@ -365,10 +365,17 @@ public:
 	CONTROLRIG_API virtual void CreateRigControlsForCurveContainer();
 	CONTROLRIG_API virtual void GetControlsInOrder(TArray<FRigControlElement*>& SortedControls) const;
 
+	/**
+	 * Selects or deselects an element in the hierarchy
+	 * @param InControlName The key of the element to select
+	 * @param bSelect If set to false the element will be deselected
+	 * @param bSetupUndo If set to true the stack will record the change for undo / redo
+	 */
 	UFUNCTION(BlueprintCallable, Category = "Control Rig")
-	CONTROLRIG_API virtual void SelectControl(const FName& InControlName, bool bSelect = true);
+	CONTROLRIG_API virtual void SelectControl(const FName& InControlName, bool bSelect = true, bool bSetupUndo = false);
 	UFUNCTION(BlueprintCallable, Category = "Control Rig")
-	CONTROLRIG_API virtual bool ClearControlSelection();
+	/** @param bSetupUndo If set to true the stack will record the change for undo / redo */
+	CONTROLRIG_API virtual bool ClearControlSelection(bool bSetupUndo = false);
 	UFUNCTION(BlueprintPure, Category = "Control Rig")
 	CONTROLRIG_API virtual TArray<FName> CurrentControlSelection() const;
 	UFUNCTION(BlueprintPure, Category = "Control Rig")
diff --git a/Engine/Plugins/Animation/ControlRig/Source/ControlRig/Public/Rigs/RigHierarchyController.h b/Engine/Plugins/Animation/ControlRig/Source/ControlRig/Public/Rigs/RigHierarchyController.h
--- a/Engine/Plugins/Animation/ControlRig/Source/ControlRig/Public/Rigs/RigHierarchyController.h
+++ b/Engine/Plugins/Animation/ControlRig/Source/ControlRig/Public/Rigs/RigHierarchyController.h
@@ -45,10 +45,11 @@ public:
 	 * @param InKey The key of the element to select
 	 * @param bSelect If set to false the element will be deselected
 	 * @param bClearSelection If this is true the selection will be cleared
+	 * @param bSetupUndo If set to true the stack will record the change for undo / redo
 	 * @return Returns true if the selection was applied
 	 */
 	UFUNCTION(BlueprintCallable, Category = URigHierarchyController)
-	UE_API bool SelectElement(FRigElementKey InKey, bool bSelect = true, bool bClearSelection = false);
+	UE_API bool SelectElement(FRigElementKey InKey, bool bSelect = true, bool bClearSelection = false, bool bSetupUndo = false);
 
 	/**
 	 * Deselects or deselects an element in the hierarchy
@@ -66,10 +67,11 @@ public:
 	 * @param InKey The key of the component to select
 	 * @param bSelect If set to false the component will be deselected
 	 * @param bClearSelection If this is true the selection will be cleared
+	 * @param bSetupUndo If set to true the stack will record the change for undo / redo
 	 * @return Returns true if the selection was applied
 	 */
 	UFUNCTION(BlueprintCallable, Category = URigHierarchyController)
-	UE_API bool SelectComponent(FRigComponentKey InKey, bool bSelect = true, bool bClearSelection = false);
+	UE_API bool SelectComponent(FRigComponentKey InKey, bool bSelect = true, bool bClearSelection = false, bool bSetupUndo = false);
 
 	/**
 	 * Deselects or deselects a component in the hierarchy
@@ -87,32 +89,37 @@ public:
 	 * @param InKey The key of the component or an element to select
 	 * @param bSelect If set to false the component or an element will be deselected
 	 * @param bClearSelection If this is true the selection will be cleared
+	 * @param bSetupUndo If set to true the stack will record the change for undo / redo
 	 * @return Returns true if the selection was applied
 	 */
 	UFUNCTION(BlueprintCallable, Category = URigHierarchyController)
-	UE_API bool SelectHierarchyKey(FRigHierarchyKey InKey, bool bSelect = true, bool bClearSelection = false);
+	UE_API bool SelectHierarchyKey(FRigHierarchyKey InKey, bool bSelect = true, bool bClearSelection = false, bool bSetupUndo = false);
 
 	/**
 	 * Deselects or deselects a component or an element in the hierarchy
 	 * @param InKey The key of the component or element to deselect
+	 * @param bSetupUndo If set to true the stack will record the change for undo / redo
 	 * @return Returns true if the selection was applied
 	 */
 	UFUNCTION(BlueprintCallable, Category = URigHierarchyController)
-	bool DeselectHierarchyKey(FRigHierarchyKey InKey)
+	bool DeselectHierarchyKey(FRigHierarchyKey InKey, bool bSetupUndo = false)
 	{
-		return SelectHierarchyKey(InKey, false);
+		return SelectHierarchyKey(InKey, false, false, bSetupUndo);
 	}
 
 	/**
 	 * Sets the selection based on a list of element keys
 	 * @param InKeys The array of keys of the elements to select
+	 * @param bPrintPythonCommand If set to true a python command equivalent to this call will be printed out
+	 * @param bSetupUndo If set to true the stack will record the change for undo / redo
 	 * @return Returns true if the selection was applied
 	 */
 	UFUNCTION(BlueprintCallable, Category = URigHierarchyController)
-	UE_API bool SetSelection(const TArray<FRigElementKey>& InKeys, bool bPrintPythonCommand = false);
+	UE_API bool SetSelection(const TArray<FRigElementKey>& InKeys, bool bPrintPythonCommand = false, bool bSetupUndo = false);
 
 	/**
 	 * Sets the selection based on a list of component keys
+	 * @param bPrintPythonCommand If set to true a python command equivalent to this call will be printed out
 	 * @param InKeys The array of keys of the component to select
 	 * @return Returns true if the selection was applied
 	 */
@@ -122,19 +129,22 @@ public:
 	/**
 	 * Sets the selection based on a list of component keys
 	 * @param InKeys The array of keys of the component to select
+	 * @param bPrintPythonCommand If set to true a python command equivalent to this call will be printed out
+	 * @param bSetupUndo If set to true the stack will record the change for undo / redo
 	 * @return Returns true if the selection was applied
 	 */
 	UFUNCTION(BlueprintCallable, Category = URigHierarchyController)
-	UE_API bool SetHierarchySelection(const TArray<FRigHierarchyKey>& InKeys, bool bPrintPythonCommand = false);
+	UE_API bool SetHierarchySelection(const TArray<FRigHierarchyKey>& InKeys, bool bPrintPythonCommand = false, bool bSetupUndo = false);
 
 	/**
 	 * Clears the selection
+	 * @param bSetupUndo If set to true the stack will record the change for undo / redo
 	 * @return Returns true if the selection was applied
 	 */
 	UFUNCTION(BlueprintCallable, Category = URigHierarchyController)
-    bool ClearSelection()
+    bool ClearSelection(bool bSetupUndo = false)
 	{
-		return SetSelection(TArray<FRigElementKey>());
+		return SetSelection(TArray<FRigElementKey>(), false, bSetupUndo);
 	}
 	
 	/**
diff --git a/Engine/Plugins/Animation/ControlRig/Source/ControlRigEditor/Private/EditMode/ControlRigEditMode.cpp b/Engine/Plugins/Animation/ControlRig/Source/ControlRigEditor/Private/EditMode/ControlRigEditMode.cpp
--- a/Engine/Plugins/Animation/ControlRig/Source/ControlRigEditor/Private/EditMode/ControlRigEditMode.cpp
+++ b/Engine/Plugins/Animation/ControlRig/Source/ControlRigEditor/Private/EditMode/ControlRigEditMode.cpp
@@ -2377,7 +2377,8 @@ void FControlRigEditMode::ClearRigElementSelection(uint32 InTypes)
 	{
 		if (URigHierarchyController* Controller = GetController(RuntimeRigPtr.Get()))
 		{
-			Controller->ClearSelection();
+			constexpr bool bSetupUndo = true;
+			Controller->ClearSelection(bSetupUndo);
 		}
 	}
 }
@@ -2387,7 +2388,8 @@ void FControlRigEditMode::SetRigElementSelectionInternal(UControlRig* ControlRig
 {
 	if(URigHierarchyController* Controller = ControlRig->GetHierarchy()->GetController())
 	{
-		Controller->SelectElement(FRigElementKey(InRigElementName, Type), bSelected);
+		constexpr bool bSetupUndo = true;
+		Controller->SelectElement(FRigElementKey(InRigElementName, Type), bSelected, false, bSetupUndo);
 	}
 }
 
@@ -2448,7 +2450,8 @@ void FControlRigEditMode::SetRigElementsSelectionInternal(const TMap<UControlRig
 		{
 			for (const FRigElementKey& Element : Elements)
 			{
-				Controller->SelectElement(Element, bSelected);
+				constexpr bool bSetupUndo = true;
+				Controller->SelectElement(Element, bSelected, false, bSetupUndo);
 			}					
 		}
 	}
diff --git a/Engine/Plugins/Animation/ControlRig/Source/ControlRigEditor/Private/Sequencer/ControlRigParameterTrackEditor.cpp b/Engine/Plugins/Animation/ControlRig/Source/ControlRigEditor/Private/Sequencer/ControlRigParameterTrackEditor.cpp
--- a/Engine/Plugins/Animation/ControlRig/Source/ControlRigEditor/Private/Sequencer/ControlRigParameterTrackEditor.cpp
+++ b/Engine/Plugins/Animation/ControlRig/Source/ControlRigEditor/Private/Sequencer/ControlRigParameterTrackEditor.cpp
@@ -3007,7 +3007,8 @@ void FControlRigParameterTrackEditor::OnSelectionChanged(TArray<UMovieSceneTrack
 				if (ControlRig && ControlRig->CurrentControlSelection().Num() > 0)
 				{
 					FScopedTransaction ScopedTransaction(LOCTEXT("SelectControlTransaction", "Select Control"), !GIsTransacting);
-					ControlRig->ClearControlSelection();
+					constexpr bool bSetupUndo = true;
+					ControlRig->ClearControlSelection(bSetupUndo);
 				}
 			}
 		}
