Control Rig Edit Mode: only remove the toolbar extension when working in the level editor

#rnx

diff --git a/Engine/Plugins/Animation/ControlRig/Source/ControlRigEditor/Private/EditMode/ControlRigEditMode.cpp b/Engine/Plugins/Animation/ControlRig/Source/ControlRigEditor/Private/EditMode/ControlRigEditMode.cpp
--- a/Engine/Plugins/Animation/ControlRig/Source/ControlRigEditor/Private/EditMode/ControlRigEditMode.cpp
+++ b/Engine/Plugins/Animation/ControlRig/Source/ControlRigEditor/Private/EditMode/ControlRigEditMode.cpp
@@ -502,11 +502,10 @@ void FControlRigEditMode::Enter()
 		}
 	}
 
-	if (&Owner->GetToolkitHost()->GetEditorModeManager() == &GLevelEditorModeTools())
+	if (IsInLevelEditor())
 	{
 		UE::ControlRig::PopulateControlRigViewportToolbarTransformSubmenu("LevelEditor.ViewportToolbar.Transform");
-		UE::ControlRig::PopulateControlRigViewportToolbarShowSubmenu(
-			"LevelEditor.ViewportToolbar.Show");
+		UE::ControlRig::PopulateControlRigViewportToolbarShowSubmenu("LevelEditor.ViewportToolbar.Show");
 	}
 	
 	RegisterPendingFocusMode();
@@ -549,8 +548,11 @@ void FControlRigEditMode::Exit()
 		WeakGizmoContext->RotationContext = FRotationContext();
 		WeakGizmoContext.Reset();
 	}
-	
-	UE::ControlRig::RemoveControlRigViewportToolbarExtensions();
+
+	if (IsInLevelEditor())
+	{
+		UE::ControlRig::RemoveControlRigViewportToolbarExtensions();
+	}
 
 	UnregisterPendingFocusMode();
 	
