Anim In Engine: Make sure Rotator/Location/Scale controls have a curve color

diff --git a/Engine/Plugins/Animation/ControlRig/Source/ControlRig/Private/Sequencer/MovieSceneControlRigParameterSection.cpp b/Engine/Plugins/Animation/ControlRig/Source/ControlRig/Private/Sequencer/MovieSceneControlRigParameterSection.cpp
--- a/Engine/Plugins/Animation/ControlRig/Source/ControlRig/Private/Sequencer/MovieSceneControlRigParameterSection.cpp
+++ b/Engine/Plugins/Animation/ControlRig/Source/ControlRig/Private/Sequencer/MovieSceneControlRigParameterSection.cpp
@@ -199,12 +199,16 @@ struct FParameterVectorChannelEditorData
 		ParameterName = InName;
 		FString NameAsString = InName.ToString();
 		FString TotalName = NameAsString;
+		const EAxisList::Type XAxis = EAxisList::X;
+		const EAxisList::Type YAxis = EAxisList::Y;
+		const EAxisList::Type ZAxis = EAxisList::Z;
 
 		{
 			TotalName += ".X";
 			MetaData[0].SetIdentifiers(FName(*TotalName), FCommonChannelData::ChannelX);
 			TotalName = NameAsString;
 			MetaData[0].IntentName = FCommonChannelData::ChannelX;
+			MetaData[0].Color = AxisDisplayInfo::GetAxisColor(XAxis);
 			MetaData[0].Group = GroupName;
 			MetaData[0].bEnabled = bEnabledOverride;
 			MetaData[0].SortOrder = SortStartIndex++;
@@ -215,6 +219,7 @@ struct FParameterVectorChannelEditorData
 			MetaData[1].SetIdentifiers(FName(*TotalName), FCommonChannelData::ChannelY);
 			TotalName = NameAsString;
 			MetaData[1].IntentName = FCommonChannelData::ChannelY;
+			MetaData[1].Color = AxisDisplayInfo::GetAxisColor(YAxis);
 			MetaData[1].Group = GroupName;
 			MetaData[1].bEnabled = bEnabledOverride;
 			MetaData[1].SortOrder = SortStartIndex++;
@@ -225,6 +230,7 @@ struct FParameterVectorChannelEditorData
 			MetaData[2].SetIdentifiers(FName(*TotalName), FCommonChannelData::ChannelZ);
 			TotalName = NameAsString;
 			MetaData[2].IntentName = FCommonChannelData::ChannelZ;
+			MetaData[2].Color = AxisDisplayInfo::GetAxisColor(ZAxis);
 			MetaData[2].Group = GroupName;
 			MetaData[2].bEnabled = bEnabledOverride;
 			MetaData[2].SortOrder = SortStartIndex++;
