[ControlRig] Update custom tooltip to match CB style

diff --git a/Engine/Plugins/Animation/ControlRig/Source/ControlRigEditor/Private/Editor/SRigModuleAssetBrowser.cpp b/Engine/Plugins/Animation/ControlRig/Source/ControlRigEditor/Private/Editor/SRigModuleAssetBrowser.cpp
--- a/Engine/Plugins/Animation/ControlRig/Source/ControlRigEditor/Private/Editor/SRigModuleAssetBrowser.cpp
+++ b/Engine/Plugins/Animation/ControlRig/Source/ControlRigEditor/Private/Editor/SRigModuleAssetBrowser.cpp
@@ -2,6 +2,7 @@
 
 #include "SRigModuleAssetBrowser.h"
 
+#include "AssetDefinitionRegistry.h"
 #include "ContentBrowserModule.h"
 #include "ControlRigBlueprint.h"
 #include "IContentBrowserSingleton.h"
@@ -11,11 +12,31 @@
 #include "ControlRigEditor.h"
 #include "FrontendFilterBase.h"
 #include "Editor/RigVMEditorTools.h"
+#include "Styling/SlateIconFinder.h"
+#include "Widgets/Layout/SSeparator.h"
 #include "Widgets/Layout/SUniformGridPanel.h"
 #include "Widgets/SRigVMVariantWidget.h"
 
 #define LOCTEXT_NAMESPACE "RigModuleAssetBrowser"
 
+namespace UE::Editor::ContentBrowser
+{
+	static bool IsNewStyleEnabled()
+	{
+		static bool bIsNewStyleEnabled = [&]()
+		{
+			if (const IConsoleVariable* CVar = IConsoleManager::Get().FindConsoleVariable(TEXT("ContentBrowser.EnableNewStyle")))
+			{
+				ensureAlwaysMsgf(!EnumHasAnyFlags(CVar->GetFlags(), ECVF_Default), TEXT("The CVar should have already been set from commandline, @see: UnrealEdGlobals.cpp, UE::Editor::ContentBrowser::EnableContentBrowserNewStyleCVarRegistration."));
+				return CVar->GetBool();
+			}
+			return false;
+		}();
+
+		return bIsNewStyleEnabled;
+	}
+}
+
 typedef UE::RigVM::Editor::Tools::FFilterByAssetTag FRigVMFilterTag; 
 
 void SRigModuleAssetBrowser::Construct(
@@ -172,6 +193,11 @@ void SRigModuleAssetBrowser::OnAssetDoubleClicked(const FAssetData& AssetData)
 
 TSharedRef<SToolTip> SRigModuleAssetBrowser::CreateCustomAssetToolTip(FAssetData& AssetData)
 {
+	if (UE::Editor::ContentBrowser::IsNewStyleEnabled())
+	{
+		return CreateCustomAssetToolTipNewStyle(AssetData);
+	}
+
 	// Make a list of tags to show
 	TArray<UObject::FAssetRegistryTag> Tags;
 	UClass* AssetClass = FindObject<UClass>(AssetData.AssetClassPath);
@@ -370,6 +396,253 @@ TSharedRef<SToolTip> SRigModuleAssetBrowser::CreateCustomAssetToolTip(FAssetData
 	return ToolTipWidget;
 }
 
+TSharedRef<SToolTip> SRigModuleAssetBrowser::CreateCustomAssetToolTipNewStyle(FAssetData& AssetData)
+{
+	// Make a list of tags to show
+	TArray<UObject::FAssetRegistryTag> Tags;
+	UClass* AssetClass = FindObject<UClass>(AssetData.AssetClassPath);
+	check(AssetClass);
+	UObject* DefaultObject = AssetClass->GetDefaultObject();
+	FAssetRegistryTagsContextData TagsContext(DefaultObject, EAssetRegistryTagsCaller::Uncategorized);
+	DefaultObject->GetAssetRegistryTags(TagsContext);
+
+	TArray<FName> TagsToShow;
+	static const FName ModulePath(TEXT("Path"));
+	static const FName ModuleSettings(TEXT("RigModuleSettings"));
+	for (const TPair<FName, UObject::FAssetRegistryTag>& TagPair : TagsContext.Tags)
+	{
+		if(TagPair.Key == ModulePath ||
+			TagPair.Key == ModuleSettings)
+		{
+			TagsToShow.Add(TagPair.Key);
+		}
+	}
+
+	TSharedRef<SVerticalBox> OverallTooltipVBox = SNew(SVerticalBox);
+
+	// Asset Name/Type Area
+	{
+		const FSlateBrush* ClassIcon = FAppStyle::GetDefaultBrush();
+		TOptional<FLinearColor> Color;
+		if (const UAssetDefinition* AssetDefinition = UAssetDefinitionRegistry::Get()->GetAssetDefinitionForClass(AssetData.GetClass()))
+		{
+			ClassIcon = AssetDefinition->GetIconBrush(AssetData, AssetData.AssetClassPath.GetAssetName());
+			Color = AssetDefinition->GetAssetColor();
+		}
+
+		if (ClassIcon == nullptr || ClassIcon == FAppStyle::GetDefaultBrush())
+		{
+			ClassIcon = FSlateIconFinder::FindIconForClass(AssetData.GetClass()).GetIcon();
+		}
+
+		FText ClassNameText = LOCTEXT("ClassNameText", "Not Found");
+		if (AssetClass != NULL)
+		{
+			ClassNameText = AssetClass->GetDisplayNameText();
+		}
+		else if (!AssetData.AssetClassPath.IsNull())
+		{
+			ClassNameText = FText::FromString(AssetData.AssetClassPath.ToString());
+		}
+
+		const FText NameText = FText::FromString(AssetData.AssetName.ToString());
+
+		// Name/Type Slot
+		OverallTooltipVBox->AddSlot()
+		.AutoHeight()
+		[
+			SNew(SVerticalBox)
+
+			+ SVerticalBox::Slot()
+			.Padding(0.f, 0.f, 0.f, 6.f)
+			.AutoHeight()
+			[
+				SNew(STextBlock)
+				.Text(NameText)
+				.ColorAndOpacity(FStyleColors::White)
+			]
+
+			+ SVerticalBox::Slot()
+			.Padding(0.f, 0.f, 0.f, 6.f)
+			.AutoHeight()
+			[
+				SNew(SHorizontalBox)
 
+				+ SHorizontalBox::Slot()
+				.AutoWidth()
+				.Padding(0.f, 0.f, 4.f, 0.f)
+				[
+					SNew(SBox)
+					.WidthOverride(16.f)
+					.HeightOverride(16.f)
+					[
+						SNew(SImage)
+						.Image(ClassIcon)
+						.ColorAndOpacity_Lambda([Color] () { return Color.IsSet() ? Color.GetValue() : FStyleColors::White;})
+					]
+				]
+
+				+ SHorizontalBox::Slot()
+				.AutoWidth()
+				.VAlign(VAlign_Center)
+				[
+					SNew(STextBlock)
+					.Text(ClassNameText)
+				]
+			]
+		];
+	}
+
+	// Separator
+	OverallTooltipVBox->AddSlot()
+	.Padding(0.f,0.f, 0.f, 6.f)
+	.AutoHeight()
+	[
+		SNew(SSeparator)
+		.Orientation(Orient_Horizontal)
+		.Thickness(1.f)
+		.ColorAndOpacity(COLOR("#484848FF"))
+		.SeparatorImage(FAppStyle::Get().GetBrush("WhiteBrush"))
+	];
+
+	// Add asset registry tags to a text list; except skeleton as that is implied in Persona
+	TMap<FName, FText> TagsAndValuesToShow;
+
+	static const FName AssetVariantPropertyName = GET_MEMBER_NAME_CHECKED(URigVMBlueprint, AssetVariant);
+	const FProperty* AssetVariantProperty = CastField<FProperty>(URigVMBlueprint::StaticClass()->FindPropertyByName(AssetVariantPropertyName));
+	const FString VariantStr = AssetData.GetTagValueRef<FString>(AssetVariantPropertyName);
+	if(!VariantStr.IsEmpty())
+	{
+		FRigVMVariant AssetVariant;
+		AssetVariantProperty->ImportText_Direct(*VariantStr, &AssetVariant, nullptr, EPropertyPortFlags::PPF_None);
+
+		if(!AssetVariant.Tags.IsEmpty())
+		{
+			OverallTooltipVBox->AddSlot()
+			.AutoHeight()
+			.Padding(0,0,5,0)
+			[
+				SNew(SHorizontalBox)
+				+ SHorizontalBox::Slot()
+				.HAlign(HAlign_Left)
+				.VAlign(VAlign_Center)
+				.AutoWidth()
+				[
+					SNew(STextBlock)
+					.Text(LOCTEXT("AssetBrowser_RigVMTagsLabel_NewStyle", "Tags :"))
+					.ColorAndOpacity(FSlateColor::UseSubduedForeground())
+				]
+
+				+ SHorizontalBox::Slot()
+				.AutoWidth()
+				.HAlign(HAlign_Left)
+				.VAlign(VAlign_Center)
+				.Padding(4, 0, 0, 0)
+				[
+					SNew(SRigVMVariantTagWidget)
+					.CapsuleTagBorder(FRigVMEditorStyle::Get().GetBrush("RigVM.TagCapsuleDark"))
+					.Visibility(EVisibility::Visible)
+					.CanAddTags(false)
+					.EnableContextMenu(false)
+					.EnableTick(false)
+					.Orientation(EOrientation::Orient_Horizontal)
+					.OnGetTags_Lambda([AssetVariant]() { return AssetVariant.Tags; })
+				]
+			];
+		}
+	}
+	
+	for(TPair<FName, FAssetTagValueRef> TagPair : AssetData.TagsAndValues)
+	{
+		if(TagsToShow.Contains(TagPair.Key))
+		{
+			// Check for DisplayName metadata
+			FName DisplayName;
+			if (FProperty* Field = FindFProperty<FProperty>(AssetClass, TagPair.Key))
+			{
+				DisplayName = *Field->GetDisplayNameText().ToString();
+			}
+			else
+			{
+				DisplayName = TagPair.Key;
+			}
+
+			if (TagPair.Key == ModuleSettings)
+			{
+				FRigModuleSettings Settings;
+				FRigVMPinDefaultValueImportErrorContext ErrorPipe;
+				FRigModuleSettings::StaticStruct()->ImportText(*TagPair.Value.GetValue(), &Settings, nullptr, PPF_None, &ErrorPipe, FString());
+				if (ErrorPipe.NumErrors == 0)
+				{
+					TagsAndValuesToShow.Add(TEXT("Default Name"), FText::FromString(Settings.Identifier.Name));
+					TagsAndValuesToShow.Add(TEXT("Category"), FText::FromString(Settings.Category));
+					TagsAndValuesToShow.Add(TEXT("Keywords"), FText::FromString(Settings.Keywords));
+					TagsAndValuesToShow.Add(TEXT("Description"), FText::FromString(Settings.Description));
+				}
+			}
+			else
+			{
+				TagsAndValuesToShow.Add(DisplayName, TagPair.Value.AsText());
+			}
+		}
+	}
+
+	for (const TPair<FName, FText>& TagPair : TagsAndValuesToShow)
+	{
+		OverallTooltipVBox->AddSlot()
+		.AutoHeight()
+		.Padding(0,0,5,0)
+		[
+			SNew(SHorizontalBox)
+			+ SHorizontalBox::Slot()
+			.AutoWidth()
+			[
+				SNew(STextBlock)
+				.Text(FText::Format(LOCTEXT("AssetTagKey_NewStyle", "{0}: "), FText::FromName(TagPair.Key)))
+				.ColorAndOpacity(FSlateColor::UseSubduedForeground())
+			]
+
+			+ SHorizontalBox::Slot()
+			.AutoWidth()
+			[
+				SNew(STextBlock)
+				.Text(TagPair.Value)
+				.ColorAndOpacity(FSlateColor::UseForeground())
+			]
+		];
+	}
+
+	OverallTooltipVBox->AddSlot()
+		.AutoHeight()
+		.Padding(0,0,5,0)
+		[
+			SNew(SHorizontalBox)
+			+ SHorizontalBox::Slot()
+			.AutoWidth()
+			[
+				SNew(STextBlock)
+				.Text(LOCTEXT("AssetBrowser_FolderPathLabel_NewStyle", "Folder :"))
+				.ColorAndOpacity(FSlateColor::UseSubduedForeground())
+			]
+
+			+ SHorizontalBox::Slot()
+			.AutoWidth()
+			[
+				SNew(STextBlock)
+				.Text(FText::FromName(AssetData.PackagePath))
+				.ColorAndOpacity(FSlateColor::UseForeground())
+				.WrapTextAt(300.f)
+			]
+		];
+
+	TSharedRef<SToolTip> ToolTipWidget = SNew(SToolTip)
+	.TextMargin(FMargin(12.f, 8.f, 12.f, 8.f))
+	.BorderImage(FAppStyle::GetBrush("AssetThumbnail.Tooltip.Border"))
+	[
+		OverallTooltipVBox
+	];
+
+	return ToolTipWidget;
+}
 
 #undef LOCTEXT_NAMESPACE
diff --git a/Engine/Plugins/Animation/ControlRig/Source/ControlRigEditor/Private/Editor/SRigModuleAssetBrowser.h b/Engine/Plugins/Animation/ControlRig/Source/ControlRigEditor/Private/Editor/SRigModuleAssetBrowser.h
--- a/Engine/Plugins/Animation/ControlRig/Source/ControlRigEditor/Private/Editor/SRigModuleAssetBrowser.h
+++ b/Engine/Plugins/Animation/ControlRig/Source/ControlRigEditor/Private/Editor/SRigModuleAssetBrowser.h
@@ -28,6 +28,7 @@ private:
 	TSharedPtr<SWidget> OnGetAssetContextMenu(const TArray<FAssetData>& SelectedAssets) const;
 	void OnAssetDoubleClicked(const FAssetData& AssetData);
 	TSharedRef<SToolTip> CreateCustomAssetToolTip(FAssetData& AssetData);
+	TSharedRef<SToolTip> CreateCustomAssetToolTipNewStyle(FAssetData& AssetData);
 
 	/** Used to get the currently selected assets */
 	FGetCurrentSelectionDelegate GetCurrentSelectionDelegate;
diff --git a/Engine/Plugins/Runtime/RigVM/Source/RigVMEditor/Private/Widgets/SRigVMVariantTagWidget.cpp b/Engine/Plugins/Runtime/RigVM/Source/RigVMEditor/Private/Widgets/SRigVMVariantTagWidget.cpp
--- a/Engine/Plugins/Runtime/RigVM/Source/RigVMEditor/Private/Widgets/SRigVMVariantTagWidget.cpp
+++ b/Engine/Plugins/Runtime/RigVM/Source/RigVMEditor/Private/Widgets/SRigVMVariantTagWidget.cpp
@@ -6,7 +6,6 @@
 #include "Framework/MultiBox/MultiBoxBuilder.h"
 #include "Framework/Application/SlateApplication.h"
 #include "RigVMSettings.h"
-#include "Editor/RigVMEditorStyle.h"
 #include "Styling/AppStyle.h"
 
 #define LOCTEXT_NAMESPACE "SRigVMVariantTagWidget"
@@ -44,7 +43,7 @@ void SRigVMVariantCapsule::Construct(
 		SNew(SBorder)
 		.Padding(0.0f)
 		.HAlign(HAlign_Fill)
-		.BorderImage(FAppStyle::Get().GetBrush("FilterBar.FilterBackground"))
+		.BorderImage(InArgs._CapsuleTagBorder)
 		[
 			SNew(SHorizontalBox)
 			+SHorizontalBox::Slot()
@@ -140,6 +139,7 @@ void SRigVMVariantTagWidget::Construct(
 	EnableContextMenu = InArgs._EnableContextMenu;
 	LastTagHash = 0;
 	MinDesiredLabelWidth = InArgs._MinDesiredLabelWidth;
+	CapsuleTagBorder = InArgs._CapsuleTagBorder;
 
 	const TAttribute<EVisibility> BoxVisibilityAttribute = TAttribute<EVisibility>::CreateLambda([this]() -> EVisibility
 	{
@@ -288,7 +288,8 @@ void SRigVMVariantTagWidget::UpdateCapsules()
 		.ToolTipText(AssignedTag.ToolTip)
 		.EnableContextMenu(EnableContextMenu)
 		.MinDesiredLabelWidth(MinDesiredLabelWidth)
-		.OnRemoveTag(OnRemoveTag);
+		.OnRemoveTag(OnRemoveTag)
+		.CapsuleTagBorder(CapsuleTagBorder);
 
 		if(HorizontalCapsuleBox)
 		{
diff --git a/Engine/Plugins/Runtime/RigVM/Source/RigVMEditor/Public/Editor/RigVMEditorStyle.h b/Engine/Plugins/Runtime/RigVM/Source/RigVMEditor/Public/Editor/RigVMEditorStyle.h
--- a/Engine/Plugins/Runtime/RigVM/Source/RigVMEditor/Public/Editor/RigVMEditorStyle.h
+++ b/Engine/Plugins/Runtime/RigVM/Source/RigVMEditor/Public/Editor/RigVMEditorStyle.h
@@ -102,6 +102,12 @@ class RIGVMEDITOR_API FRigVMEditorStyle final
 			
 			Set( "TagButton", TagButton );
 		}
+
+		// Tag Capsule styles
+		{
+			Set("RigVM.TagCapsule", new FSlateRoundedBoxBrush(FStyleColors::Secondary, 3.0f));
+			Set("RigVM.TagCapsuleDark", new FSlateRoundedBoxBrush(FStyleColors::Header, 3.0f));
+		}
 	}
 public:
 
diff --git a/Engine/Plugins/Runtime/RigVM/Source/RigVMEditor/Public/Widgets/SRigVMVariantTagWidget.h b/Engine/Plugins/Runtime/RigVM/Source/RigVMEditor/Public/Widgets/SRigVMVariantTagWidget.h
--- a/Engine/Plugins/Runtime/RigVM/Source/RigVMEditor/Public/Widgets/SRigVMVariantTagWidget.h
+++ b/Engine/Plugins/Runtime/RigVM/Source/RigVMEditor/Public/Widgets/SRigVMVariantTagWidget.h
@@ -2,6 +2,7 @@
 
 #pragma once
 
+#include "Editor/RigVMEditorStyle.h"
 #include "Widgets/Layout/SBox.h"
 #include "Widgets/Input/SButton.h"
 #include "Widgets/SBoxPanel.h"
@@ -19,6 +20,7 @@ public:
 		: _Padding(0.f)
 		, _EnableContextMenu(false)
 		, _MinDesiredLabelWidth(0.f)
+		, _CapsuleTagBorder(FRigVMEditorStyle::Get().GetBrush("RigVM.TagCapsule"))
 	{
 	}
 	SLATE_ATTRIBUTE(FName, Name)
@@ -30,6 +32,7 @@ public:
 	SLATE_EVENT(FRigVMVariant_OnRemoveTag, OnRemoveTag)
 	SLATE_ATTRIBUTE(bool, EnableContextMenu)
 	SLATE_ARGUMENT(float, MinDesiredLabelWidth)
+	SLATE_ATTRIBUTE(const FSlateBrush*, CapsuleTagBorder)
 	SLATE_END_ARGS()
 
 	SRigVMVariantCapsule();
@@ -62,6 +65,7 @@ public:
 		, _EnableContextMenu(false)
 		, _MinDesiredLabelWidth(0.f)
 		, _EnableTick(true)
+		, _CapsuleTagBorder(FRigVMEditorStyle::Get().GetBrush("RigVM.TagCapsule"))
 	{
 	}
 	SLATE_ARGUMENT(EOrientation, Orientation)
@@ -72,6 +76,7 @@ public:
 	SLATE_EVENT(FRigVMVariant_OnRemoveTag, OnRemoveTag)
 	SLATE_ARGUMENT(float, MinDesiredLabelWidth)
 	SLATE_ARGUMENT(bool, EnableTick)
+	SLATE_ATTRIBUTE(const FSlateBrush*, CapsuleTagBorder)
 	SLATE_END_ARGS()
 
 	SRigVMVariantTagWidget();
@@ -96,6 +101,7 @@ private:
 	FRigVMVariant_OnRemoveTag OnRemoveTag;
 	TAttribute<bool> CanAddTags;
 	TAttribute<bool> EnableContextMenu;
+	TAttribute<const FSlateBrush*> CapsuleTagBorder;
 	uint32 LastTagHash;
 	float MinDesiredLabelWidth;
 };
