Modular Rig: Make sure the dependency classes are loaded before we add module instances
#rnx

diff --git a/Engine/Plugins/Animation/ControlRig/Source/ControlRig/Private/ModularRig.cpp b/Engine/Plugins/Animation/ControlRig/Source/ControlRig/Private/ModularRig.cpp
--- a/Engine/Plugins/Animation/ControlRig/Source/ControlRig/Private/ModularRig.cpp
+++ b/Engine/Plugins/Animation/ControlRig/Source/ControlRig/Private/ModularRig.cpp
@@ -58,9 +58,7 @@ const FRigModuleInstance* FModuleInstanceHandle::Get() const
 UModularRig::UModularRig(const FObjectInitializer& ObjectInitializer)
 	: Super(ObjectInitializer)
 {
-#if WITH_EDITOR
-	FCoreUObjectDelegates::OnObjectsReinstanced.AddUObject(this, &UModularRig::OnObjectsReplaced);
-#endif
+
 }
 
 void UModularRig::BeginDestroy()
@@ -260,6 +258,29 @@ void UModularRig::PostInitProperties()
 	UpdateSupportedEvents();
 }
 
+void UModularRig::GetPreloadDependencies(TArray<UObject*>& OutDeps)
+{
+	Super::GetPreloadDependencies(OutDeps);
+
+	ModularRigModel.ForEachModule([&OutDeps](const FRigModuleReference* Module) -> bool
+	{
+		OutDeps.AddUnique(Module->Class.Get());
+		return true;
+	});
+}
+
+void UModularRig::PostInitInstance(URigVMHost* InCDO)
+{
+	Super::PostInitInstance(InCDO);
+
+#if WITH_EDITOR
+	if (IsInGameThread())
+	{
+		FCoreUObjectDelegates::OnObjectsReinstanced.AddUObject(this, &UModularRig::OnObjectsReplaced);
+	}
+#endif
+}
+
 void UModularRig::Serialize(FArchive& Ar)
 {
 	Super::Serialize(Ar);
@@ -900,7 +921,8 @@ FRigModuleInstance* UModularRig::AddModuleInstance(const FName& InModuleName, TS
 	NewModule.SetRig(NewModuleRig);
 
 	UpdateCachedChildren();
-	for (const FName& EventName : NewModule.GetRig()->GetSupportedEvents())
+	const TArray<FName> NewModuleSupportedEvents = NewModule.GetRig()->GetSupportedEvents();
+	for (const FName& EventName : NewModuleSupportedEvents)
 	{
 		SupportedEvents.AddUnique(EventName);
 	}
diff --git a/Engine/Plugins/Animation/ControlRig/Source/ControlRig/Public/ModularRig.h b/Engine/Plugins/Animation/ControlRig/Source/ControlRig/Public/ModularRig.h
--- a/Engine/Plugins/Animation/ControlRig/Source/ControlRig/Public/ModularRig.h
+++ b/Engine/Plugins/Animation/ControlRig/Source/ControlRig/Public/ModularRig.h
@@ -142,8 +142,10 @@ class UModularRig : public UControlRig
 public:
 
 	UE_API virtual void PostInitProperties() override;
+	virtual void GetPreloadDependencies(TArray<UObject*>& OutDeps) override;
 
 	// BEGIN ControlRig
+	UE_API virtual void PostInitInstance(URigVMHost* InCDO) override;
 	UE_API virtual void Serialize(FArchive& Ar) override;
 	UE_API virtual void PostLoad() override;
 	UE_API virtual void InitializeVMs(bool bRequestInit = true) override;
