Control Rig: Fix crash when duplicated modular rig character in sequencer
#rnx

diff --git a/Engine/Plugins/Animation/ControlRig/Source/ControlRig/Private/ModularRig.cpp b/Engine/Plugins/Animation/ControlRig/Source/ControlRig/Private/ModularRig.cpp
--- a/Engine/Plugins/Animation/ControlRig/Source/ControlRig/Private/ModularRig.cpp
+++ b/Engine/Plugins/Animation/ControlRig/Source/ControlRig/Private/ModularRig.cpp
@@ -796,7 +796,8 @@ void UModularRig::UpdateCachedChildren()
 void UModularRig::UpdateSupportedEvents() const
 {
 	SupportedEvents.Reset();
-	ModularRigModel.ForEachModule([this](const FRigModuleReference* Module) -> bool
+	const FModularRigModel& Model = GetModularRigModel();
+	Model.ForEachModule([this](const FRigModuleReference* Module) -> bool
 	{
 		if (Module->Class.IsValid())
 		{
