Control Rig: Fix multi target connector losing all the data when attempting to add an invalid target
#rnx

diff --git a/Engine/Plugins/Animation/ControlRig/Source/ControlRig/Private/ModularRigController.cpp b/Engine/Plugins/Animation/ControlRig/Source/ControlRig/Private/ModularRigController.cpp
--- a/Engine/Plugins/Animation/ControlRig/Source/ControlRig/Private/ModularRigController.cpp
+++ b/Engine/Plugins/Animation/ControlRig/Source/ControlRig/Private/ModularRigController.cpp
@@ -275,17 +275,23 @@ bool UModularRigController::ConnectConnectorToElements(const FRigElementKey& InC
 	}
 	
 	FText ErrorMessage;
-	if (bCheckValidConnection && !CanConnectConnectorToElements(InConnectorKey, InTargetKeys, ErrorMessage))
+	auto GenerateTargetKeyString = [&TargetKeys]() -> FString
 	{
 		FString TargetKeyString = TargetKeys[0].ToString();
 		if(TargetKeys.Num() > 1)
 		{
 			for(int32 Index = 1; Index < TargetKeys.Num(); Index++)
 			{
-				TargetKeyString = TargetKeyString + TEXT(", ") + TargetKeys[Index].ToString();				
+				TargetKeyString = TargetKeyString + TEXT(", ") + TargetKeys[Index].ToString();
 			}
 			TargetKeyString = TEXT("{") + TargetKeyString = TEXT("}");
 		}
+		return TargetKeyString;
+	};
+	
+	if (bCheckValidConnection && !CanConnectConnectorToElements(InConnectorKey, InTargetKeys, ErrorMessage))
+	{
+		FString TargetKeyString = GenerateTargetKeyString();
 		UE_LOG(LogControlRig, Error, TEXT("Could not connect %s to %s: %s"), *InConnectorKey.ToString(), *TargetKeyString, *ErrorMessage.ToString());
 		return false;
 	}
@@ -402,9 +408,12 @@ bool UModularRigController::ConnectConnectorToElements(const FRigElementKey& InC
 	// If the connection that was attempted generated cycles, try to reestablish the previous connections
 	if (DisconnectedConnectors.Contains(InConnectorKey))
 	{
+		FString TargetKeyString = GenerateTargetKeyString();
+		UE_LOG(LogControlRig, Error, TEXT("Could not connect %s to %s: cycles detected"), *InConnectorKey.ToString(), *TargetKeyString);
+		
 		if (TArray<FRigElementKey>* Previous = PreviousConnections.Find(InConnectorKey))
 		{
-			ConnectConnectorToElements(InConnectorKey, *Previous, bSetupUndo, bAutomaticReparenting, bCheckValidConnection);
+			ConnectConnectorToElements(InConnectorKey, *Previous, bSetupUndo, bAutomaticReparenting, false);
 		}
 	}
 
