Control Rig: Fix outliner not scrolling to element when only selected rigs/modules are displayed
#rnx

diff --git a/Engine/Plugins/Animation/ControlRig/Source/ControlRigEditor/Private/EditMode/SControlRigOutliner.cpp b/Engine/Plugins/Animation/ControlRig/Source/ControlRigEditor/Private/EditMode/SControlRigOutliner.cpp
--- a/Engine/Plugins/Animation/ControlRig/Source/ControlRigEditor/Private/EditMode/SControlRigOutliner.cpp
+++ b/Engine/Plugins/Animation/ControlRig/Source/ControlRigEditor/Private/EditMode/SControlRigOutliner.cpp
@@ -170,6 +170,11 @@ bool FMultiRigData::operator== (const FMultiRigData & Other) const
 
 bool FMultiRigData::IsValid() const
 {
+	if (Type >= EMultiRigDataType_Invalid)
+	{
+		return false;
+	}
+	
 	if (WeakControlRig.IsValid())
 	{
 		if (Key.IsSet())
@@ -1147,30 +1152,47 @@ TSharedPtr<FMultiRigTreeElement> SMultiRigHierarchyTreeView::GetParentElement(TS
 	return nullptr;
 }
 
-void SMultiRigHierarchyTreeView::RequestTreeViewRefresh(TFunction<void()>&& InRefreshFunction)
+void SMultiRigHierarchyTreeView::ScrollToElement(const FMultiRigData& InElement)
 {
-	UnregisterPendingRefresh();
-
-	if (InRefreshFunction)
+	if (bRefreshPending)
 	{
-		PendingTreeViewRefreshHandle = RegisterActiveTimer( 0.f,
-			FWidgetActiveTimerDelegate::CreateLambda([this, RefreshFunction = MoveTemp(InRefreshFunction)](double, float)
-			{
-				RefreshFunction();
-				return EActiveTimerReturnType::Stop;
-			})
-		);		
+		ScrollToElementAfterRefresh = InElement;
 	}
 	else
 	{
-		PendingTreeViewRefreshHandle = RegisterActiveTimer( 0.f,
-			FWidgetActiveTimerDelegate::CreateLambda([this](double, float)
+		if (TSharedPtr<FMultiRigTreeElement>* Element = ElementMap.Find(InElement))
+		{
+			// Expand parents
+			TSharedPtr<FMultiRigTreeElement> Parent = GetParentElement(*Element);
+			while (Parent)
 			{
-				RefreshTreeView();
-				return EActiveTimerReturnType::Stop;
-			})
-		);
+				SetItemExpansion(Parent, true);
+				Parent = GetParentElement(Parent);
+			}
+			
+			RequestScrollIntoView(*Element);
+		}
+	}
+}
+
+void SMultiRigHierarchyTreeView::RequestTreeViewRefresh()
+{
+	if (bRefreshPending)
+	{
+		return;
 	}
+	
+	UnregisterPendingRefresh();
+
+	bRefreshPending = true;
+	
+	PendingTreeViewRefreshHandle = RegisterActiveTimer( 0.f,
+		FWidgetActiveTimerDelegate::CreateLambda([this](double, float)
+		{
+			RefreshTreeView();
+			return EActiveTimerReturnType::Stop;
+		})
+	);
 }
 
 void SMultiRigHierarchyTreeView::UnregisterPendingRefresh()
@@ -1187,6 +1209,8 @@ void SMultiRigHierarchyTreeView::UnregisterPendingRefresh()
 
 void SMultiRigHierarchyTreeView::RefreshTreeView()
 {
+	bRefreshPending = false;
+	
 	// store expansion state
 	TMap<FMultiRigData, bool> ExpansionState;
 	for (TPair<FMultiRigData, TSharedPtr<FMultiRigTreeElement>> Pair : ElementMap)
@@ -1536,6 +1560,12 @@ void SMultiRigHierarchyTreeView::RefreshTreeView()
 			}
 		}
 	}
+
+	if (ScrollToElementAfterRefresh.IsValid())
+	{
+		ScrollToElement(ScrollToElementAfterRefresh);
+		ScrollToElementAfterRefresh.Reset();
+	}
 }
 
 void SMultiRigHierarchyTreeView::SetExpansionRecursive(TSharedPtr<FMultiRigTreeElement> InElement, bool bShouldBeExpanded)
@@ -2094,6 +2124,12 @@ void SControlRigOutliner::HandleControlSelected(UControlRig* Subject, FRigContro
 {
 	FControlRigBaseDockableView::HandleControlSelected(Subject, ControlElement, bSelected);
 
+	if (ControlElement->IsAnimationChannel())
+	{
+		// This control is not present in the outliner, so no need to handle
+		return;
+	}
+
 	if (bIsChangingRigHierarchy)
 	{
 		// This action was initiated by the outliner
@@ -2136,35 +2172,23 @@ void SControlRigOutliner::HandleControlSelected(UControlRig* Subject, FRigContro
 	FMultiRigData Data(Subject, Key);
 
 	bool bScrollToElement = UControlRigEditorSettings::Get()->bFocusOnSelection;
-	for (int32 RootIndex = 0; RootIndex < HierarchyTreeView->GetTreeView()->GetRootElements().Num(); ++RootIndex)
+	
+	TSharedPtr<FMultiRigTreeElement> Found = HierarchyTreeView->GetTreeView()->FindElement(Data);
+	if (Found.IsValid())
 	{
-		TSharedPtr<FMultiRigTreeElement> Found = HierarchyTreeView->GetTreeView()->FindElement(Data, HierarchyTreeView->GetTreeView()->GetRootElements()[RootIndex]);
-		if (Found.IsValid())
-		{
-			TGuardValue<bool> GuardRigHierarchyChanges(bIsChangingRigHierarchy, true);
-
-			if (!FSlateApplication::Get().GetModifierKeys().IsShiftDown() && !FSlateApplication::Get().GetModifierKeys().IsControlDown())
-			{
-				HierarchyTreeView->GetTreeView()->ClearSelection();
-			}
-
-			HierarchyTreeView->GetTreeView()->SetItemSelection(Found, bSelected, ESelectInfo::Direct);
+		TGuardValue<bool> GuardRigHierarchyChanges(bIsChangingRigHierarchy, true);
 
-			if (bSelected && bScrollToElement)
-			{
-				TArray<TSharedPtr<FMultiRigTreeElement>> SelectedItems = HierarchyTreeView->GetTreeView()->GetSelectedItems();
-				if (SelectedItems.Num() > 0)
-				{
-					TSharedPtr<FMultiRigTreeElement> Parent = HierarchyTreeView->GetTreeView()->GetParentElement(Found);
-					while (Parent)
-					{
-						HierarchyTreeView->GetTreeView()->SetItemExpansion(Parent, true);
-						Parent = HierarchyTreeView->GetTreeView()->GetParentElement(Parent);
-					}
-					HierarchyTreeView->GetTreeView()->RequestScrollIntoView(SelectedItems.Last());
-				}
-			}
+		if (!FSlateApplication::Get().GetModifierKeys().IsShiftDown() && !FSlateApplication::Get().GetModifierKeys().IsControlDown())
+		{
+			HierarchyTreeView->GetTreeView()->ClearSelection();
 		}
+
+		HierarchyTreeView->GetTreeView()->SetItemSelection(Found, bSelected, ESelectInfo::Direct);
+	}
+	
+	if (bSelected && bScrollToElement)
+	{
+		HierarchyTreeView->GetTreeView()->ScrollToElement(Data);
 	}
 }
 
diff --git a/Engine/Plugins/Animation/ControlRig/Source/ControlRigEditor/Private/EditMode/SControlRigOutliner.h b/Engine/Plugins/Animation/ControlRig/Source/ControlRigEditor/Private/EditMode/SControlRigOutliner.h
--- a/Engine/Plugins/Animation/ControlRig/Source/ControlRigEditor/Private/EditMode/SControlRigOutliner.h
+++ b/Engine/Plugins/Animation/ControlRig/Source/ControlRigEditor/Private/EditMode/SControlRigOutliner.h
@@ -88,7 +88,7 @@ class FMultiRigData
 {
 
 public:
-	enum EMultiRigDataType
+	enum EMultiRigDataType : uint8
 	{
 		EMultiRigDataType_ControlRig,
 		EMultiRigDataType_Element,
@@ -96,7 +96,8 @@ public:
 		EMultiRigDataType_Component,
 		EMultiRigDataType_Actor,
 
-		EMultiRigDataType_MAX
+		EMultiRigDataType_MAX,
+		EMultiRigDataType_Invalid = EMultiRigDataType_MAX,
 	};
 	
 	FMultiRigData() { Type = EMultiRigDataType_ControlRig; };
@@ -108,6 +109,7 @@ public:
 	FText GetToolTipText(const FRigTreeDisplaySettings& Settings) const;
 	bool operator == (const FMultiRigData & Other) const;
 	bool IsValid() const;
+	void Reset() { Type = EMultiRigDataType_Invalid; }
 	URigHierarchy* GetHierarchy() const;
 	bool IsModularRig() const;
 	FRigModuleInstance* GetModuleInstance() const;
@@ -258,12 +260,12 @@ public:
 	TArray<URigHierarchy*> GetHierarchy() const;
 	void SetControlRigs(const TArrayView<TWeakObjectPtr<UControlRig>>& InControlRigs);
 
-	/** Requests a tree view refresh. The actual refresh will be done lazily.
-	 * Note that the refresh function will be set to RefreshTreeView if no function is passed. 
-	 * InRefreshFunction can capture what's necessary if anything has to be done once the tree rebuilt. */
-	void RequestTreeViewRefresh(TFunction<void()>&& InRefreshFunction = nullptr);
+	/** Requests a tree view refresh. The actual refresh will be done lazily. */
+	void RequestTreeViewRefresh();
 
 	TSharedPtr<FMultiRigTreeElement> GetParentElement(TSharedPtr<FMultiRigTreeElement> InElement) const;
+
+	void ScrollToElement(const FMultiRigData& InElement);
 	
 private:
 
@@ -287,6 +289,8 @@ private:
 
 	/** Pending function to refresh the tree view. */
 	TWeakPtr<FActiveTimerHandle> PendingTreeViewRefreshHandle;
+	bool bRefreshPending = false;
+	FMultiRigData ScrollToElementAfterRefresh;
 	void UnregisterPendingRefresh();
 	void RefreshTreeView();
 };
