Control Rig: Fix selection not being undo / redoable when selection through Anim Details.

diff --git a/Engine/Plugins/Animation/ControlRig/Source/ControlRigEditor/Private/EditMode/SControlRigOutliner.cpp b/Engine/Plugins/Animation/ControlRig/Source/ControlRigEditor/Private/EditMode/SControlRigOutliner.cpp
--- a/Engine/Plugins/Animation/ControlRig/Source/ControlRigEditor/Private/EditMode/SControlRigOutliner.cpp
+++ b/Engine/Plugins/Animation/ControlRig/Source/ControlRigEditor/Private/EditMode/SControlRigOutliner.cpp
@@ -2362,6 +2362,7 @@ void SControlRigOutliner::HandleSelectionChanged(TSharedPtr<FMultiRigTreeElement
 		HierarchyTreeView->GetTreeView()->SetItemSelection(AddToSelection, true);
 	}
 	
+	const bool bSetupUndo = bEndTransaction;
 	//due to how Sequencer Tree View will redo selection on next tick if we aren't keeping or toggling selection we need to clear it out
 	if (FSlateApplication::Get().GetModifierKeys().IsShiftDown() == false || FSlateApplication::Get().GetModifierKeys().IsControlDown() == false)
 	{
@@ -2373,7 +2374,7 @@ void SControlRigOutliner::HandleSelectionChanged(TSharedPtr<FMultiRigTreeElement
 			{
 				if (CurrentSelection.Key)
 				{
-					CurrentSelection.Key->ClearControlSelection();
+					CurrentSelection.Key->ClearControlSelection(bSetupUndo);
 				}
 			}
 			if (GEditor)
@@ -2401,7 +2402,7 @@ void SControlRigOutliner::HandleSelectionChanged(TSharedPtr<FMultiRigTreeElement
 		{
 			URigHierarchyController* Controller = ((URigHierarchy*)Hierarchy)->GetController(true);
 			check(Controller);
-			Controller->SetSelection(RigAndKeys.Value);
+			Controller->SetSelection(RigAndKeys.Value, false, bSetupUndo);
 		}
 	}
 	if (bEndTransaction)
