Control Rig: Fix interaction not adding all the selected elements
#rnx

diff --git a/Engine/Plugins/Animation/ControlRig/Source/ControlRigEditor/Private/EditMode/ControlRigEditMode.cpp b/Engine/Plugins/Animation/ControlRig/Source/ControlRigEditor/Private/EditMode/ControlRigEditMode.cpp
--- a/Engine/Plugins/Animation/ControlRig/Source/ControlRigEditor/Private/EditMode/ControlRigEditMode.cpp
+++ b/Engine/Plugins/Animation/ControlRig/Source/ControlRigEditor/Private/EditMode/ControlRigEditMode.cpp
@@ -1236,7 +1236,8 @@ bool FControlRigEditMode::HandleBeginTransform(const FEditorViewportClient* InVi
 				if (AreRigElementSelectedAndMovable(ControlRig))
 				{
 					//todo need to add multiple
-					FControlRigInteractionScope* InteractionScope = new FControlRigInteractionScope(ControlRig);
+					TArray<FRigElementKey> SelectedRigElements = GetSelectedRigElements(ControlRig);
+					FControlRigInteractionScope* InteractionScope = new FControlRigInteractionScope(ControlRig, SelectedRigElements, (EControlRigInteractionType)InteractionType);
 					InteractionScopes.Add(ControlRig,InteractionScope);
 					ControlRig->bInteractionJustBegan = true;
 				}
@@ -1527,7 +1528,8 @@ bool FControlRigEditMode::HandleClick(FEditorViewportClient* InViewportClient, H
 
 					if (bInteractionScopePresent)
 					{
-						*InteractionScope = new FControlRigInteractionScope(ShapeActor->ControlRig.Get());
+						TArray<FRigElementKey> SelectedRigElements = GetSelectedRigElements(ShapeActor->ControlRig.Get());
+						*InteractionScope = new FControlRigInteractionScope(ShapeActor->ControlRig.Get(), SelectedRigElements, (EControlRigInteractionType)InteractionType);
 						InteractionScopes.Add(ShapeActor->ControlRig.Get(), *InteractionScope);
 
 					}
