Control Rig: Convert manual memory management to automaticly managed TUniquePtr.

diff --git a/Engine/Plugins/Animation/ControlRig/Source/ControlRig/Public/ControlRig.h b/Engine/Plugins/Animation/ControlRig/Source/ControlRig/Public/ControlRig.h
--- a/Engine/Plugins/Animation/ControlRig/Source/ControlRig/Public/ControlRig.h
+++ b/Engine/Plugins/Animation/ControlRig/Source/ControlRig/Public/ControlRig.h
@@ -1055,14 +1055,14 @@ public:
 
 	FControlRigInteractionScope(
 		UControlRig* InControlRig,
-		const TArray<FRigElementKey>& InKeys,
+		TArray<FRigElementKey> InKeys,
 		EControlRigInteractionType InInteractionType = EControlRigInteractionType::All
 	)
 		: ControlRig(InControlRig)
 		, InteractionBracketScope(InControlRig->InteractionBracket)
-	, SyncBracketScope(InControlRig->InterRigSyncBracket)
+		, SyncBracketScope(InControlRig->InterRigSyncBracket)
 	{
-		InControlRig->ElementsBeingInteracted = InKeys;
+		InControlRig->ElementsBeingInteracted = MoveTemp(InKeys);
 		InControlRig->InteractionType = (uint8)InInteractionType;
 		InControlRig->bInteractionJustBegan = true;
 		InControlRig->GetHierarchy()->StartInteraction();
diff --git a/Engine/Plugins/Animation/ControlRig/Source/ControlRigEditor/Private/ControlRigElementDetails.cpp b/Engine/Plugins/Animation/ControlRig/Source/ControlRigEditor/Private/ControlRigElementDetails.cpp
--- a/Engine/Plugins/Animation/ControlRig/Source/ControlRigEditor/Private/ControlRigElementDetails.cpp
+++ b/Engine/Plugins/Animation/ControlRig/Source/ControlRigEditor/Private/ControlRigElementDetails.cpp
@@ -2988,8 +2988,7 @@ FDetailWidgetRow& FRigTransformElementDetails::CreateEulerTransformValueWidgetRo
 					DebuggedRig->ElementsBeingInteracted.AddUnique(ElementInfo.Element.GetKey());
 				}
 				
-				FControlRigInteractionScope* InteractionScope = new FControlRigInteractionScope(DebuggedRig);
-				StrongThisPtr->InteractionScopes.Add(InteractionScope);
+				StrongThisPtr->InteractionScopes.Emplace(MakeUnique<FControlRigInteractionScope>(DebuggedRig));
 			}
 		});
 	TransformWidgetArgs.OnEndSliderMovement_Lambda(
@@ -3012,13 +3011,6 @@ FDetailWidgetRow& FRigTransformElementDetails::CreateEulerTransformValueWidgetRo
 				DebuggedRig->InteractionType = (uint8)EControlRigInteractionType::None;
 				DebuggedRig->ElementsBeingInteracted.Reset();
 			}
-			for (FControlRigInteractionScope* InteractionScope : StrongThisPtr->InteractionScopes)
-			{
-				if (InteractionScope)
-				{
-					delete InteractionScope; 
-				}
-			}
 			StrongThisPtr->InteractionScopes.Reset();
 		});
 
diff --git a/Engine/Plugins/Animation/ControlRig/Source/ControlRigEditor/Private/ControlRigElementDetails.h b/Engine/Plugins/Animation/ControlRig/Source/ControlRigEditor/Private/ControlRigElementDetails.h
--- a/Engine/Plugins/Animation/ControlRig/Source/ControlRigEditor/Private/ControlRigElementDetails.h
+++ b/Engine/Plugins/Animation/ControlRig/Source/ControlRigEditor/Private/ControlRigElementDetails.h
@@ -344,7 +344,7 @@ protected:
 
 	TArray<FPerElementInfo> PerElementInfos;
 	FDelegateHandle MetadataHandle;
-	TArray<FControlRigInteractionScope*> InteractionScopes;
+	TArray<TUniquePtr<FControlRigInteractionScope>> InteractionScopes;
 
 	struct FSharedComponent
 	{
diff --git a/Engine/Plugins/Animation/ControlRig/Source/ControlRigEditor/Private/EditMode/ControlRigEditMode.cpp b/Engine/Plugins/Animation/ControlRig/Source/ControlRigEditor/Private/EditMode/ControlRigEditMode.cpp
--- a/Engine/Plugins/Animation/ControlRig/Source/ControlRigEditor/Private/EditMode/ControlRigEditMode.cpp
+++ b/Engine/Plugins/Animation/ControlRig/Source/ControlRigEditor/Private/EditMode/ControlRigEditMode.cpp
@@ -568,21 +568,13 @@ void FControlRigEditMode::Exit()
 		}
 	}
 
-	if (InteractionScopes.Num() >0)
+	if (!InteractionScopes.IsEmpty())
 	{
-
 		if (GEditor)
 		{
 			GEditor->EndTransaction();
 		}
 
-		for (TPair<UControlRig*,FControlRigInteractionScope*>& InteractionScope : InteractionScopes)
-		{
-			if (InteractionScope.Value)
-			{
-				delete InteractionScope.Value;
-			}
-		}
 		InteractionScopes.Reset();
 		bManipulatorMadeChange = false;
 	}
@@ -1237,8 +1229,10 @@ bool FControlRigEditMode::HandleBeginTransform(const FEditorViewportClient* InVi
 				{
 					//todo need to add multiple
 					TArray<FRigElementKey> SelectedRigElements = GetSelectedRigElements(ControlRig);
-					FControlRigInteractionScope* InteractionScope = new FControlRigInteractionScope(ControlRig, SelectedRigElements, (EControlRigInteractionType)InteractionType);
-					InteractionScopes.Add(ControlRig,InteractionScope);
+					TUniquePtr<FControlRigInteractionScope> InteractionScope = MakeUnique<FControlRigInteractionScope>(
+						ControlRig, SelectedRigElements, static_cast<EControlRigInteractionType>(InteractionType)
+						);
+					InteractionScopes.Emplace(ControlRig, MoveTemp(InteractionScope));
 					ControlRig->bInteractionJustBegan = true;
 				}
 				else
@@ -1259,9 +1253,11 @@ bool FControlRigEditMode::HandleBeginTransform(const FEditorViewportClient* InVi
 			
 			OnGizmoInteractionStartedDelegate.Broadcast(Selection, (EControlRigInteractionType)InteractionType);
 		}
-		TArray<FRigElementKey> SelectedRigElements = GetSelectedRigElements(ControlRig);
-		FControlRigInteractionScope* InteractionScope = new FControlRigInteractionScope(ControlRig, SelectedRigElements, (EControlRigInteractionType)InteractionType);
-		InteractionScopes.Add(ControlRig,InteractionScope);
+		
+		TUniquePtr<FControlRigInteractionScope> InteractionScope = MakeUnique<FControlRigInteractionScope>(
+			ControlRig, GetSelectedRigElements(ControlRig), static_cast<EControlRigInteractionType>(InteractionType)
+			);
+		InteractionScopes.Emplace(ControlRig, MoveTemp(InteractionScope));
 	}
 	else
 	{
@@ -1311,17 +1307,16 @@ bool FControlRigEditMode::HandleEndTransform(FEditorViewportClient* InViewportCl
 			OnGizmoInteractionEndedDelegate.Broadcast(Selection);
 		}
 		
-		for (TPair<UControlRig*, FControlRigInteractionScope*>& RigInteractionScope : InteractionScopes)
+		for (const TPair<UControlRig*, TUniquePtr<FControlRigInteractionScope>>& RigInteractionScope : InteractionScopes)
 		{
 			if (UControlRig* ControlRig = RigInteractionScope.Key)
 			{
 				RigsToTick.Add(ControlRig);
 			}
 
-			if (FControlRigInteractionScope* InteractionScope = RigInteractionScope.Value)
+			if (const TUniquePtr<FControlRigInteractionScope>& InteractionScope = RigInteractionScope.Value)
 			{
 				Keyframer.Apply(*InteractionScope, TransformContext);
-				delete InteractionScope; 
 			}
 		}
 
@@ -1482,13 +1477,7 @@ bool FControlRigEditMode::HandleClick(FEditorViewportClient* InViewportClient, H
 					FScopedTransaction ScopedTransaction(LOCTEXT("SelectControlTransaction", "Select Control"), !AreEditingControlRigDirectly() && !GIsTransacting);
 
 					// temporarily disable the interaction scope
-					FControlRigInteractionScope** InteractionScope = InteractionScopes.Find(ShapeActor->ControlRig.Get());
-					const bool bInteractionScopePresent = InteractionScope != nullptr;
-					if (bInteractionScopePresent)
-					{
-						delete* InteractionScope;
-						InteractionScopes.Remove(ShapeActor->ControlRig.Get());
-					}
+					const bool bInteractionScopePresent = InteractionScopes.Remove(ShapeActor->ControlRig.Get()) > 0;
 					
 					const FName& ControlName = ShapeActor->ControlName;
 					if (Click.IsShiftDown()) //guess we just select
@@ -1528,9 +1517,11 @@ bool FControlRigEditMode::HandleClick(FEditorViewportClient* InViewportClient, H
 
 					if (bInteractionScopePresent)
 					{
-						TArray<FRigElementKey> SelectedRigElements = GetSelectedRigElements(ShapeActor->ControlRig.Get());
-						*InteractionScope = new FControlRigInteractionScope(ShapeActor->ControlRig.Get(), SelectedRigElements, (EControlRigInteractionType)InteractionType);
-						InteractionScopes.Add(ShapeActor->ControlRig.Get(), *InteractionScope);
+						UControlRig* Rig = ShapeActor->ControlRig.Get();
+						TUniquePtr<FControlRigInteractionScope> InteractionScope = MakeUnique<FControlRigInteractionScope>(
+							Rig, GetSelectedRigElements(ShapeActor->ControlRig.Get()), static_cast<EControlRigInteractionType>(InteractionType)
+							);
+						InteractionScopes.Add(Rig, MoveTemp(InteractionScope));
 
 					}
 
@@ -2292,7 +2283,7 @@ bool FControlRigEditMode::InputDelta(FEditorViewportClient* InViewportClient, FV
 	{
 		TArray<TWeakObjectPtr<UControlRig>> RigsToTick;
 		RigsToTick.Reserve(InteractionScopes.Num());
-		Algo::Transform(InteractionScopes, RigsToTick, [](const TPair<UControlRig*,FControlRigInteractionScope*>& InteractionData)
+		Algo::Transform(InteractionScopes, RigsToTick, [](const TPair<UControlRig*, TUniquePtr<FControlRigInteractionScope>>& InteractionData)
 		{
 			return InteractionData.Key;
 		});
diff --git a/Engine/Plugins/Animation/ControlRig/Source/ControlRigEditor/Public/EditMode/ControlRigEditMode.h b/Engine/Plugins/Animation/ControlRig/Source/ControlRigEditor/Public/EditMode/ControlRigEditMode.h
--- a/Engine/Plugins/Animation/ControlRig/Source/ControlRigEditor/Public/EditMode/ControlRigEditMode.h
+++ b/Engine/Plugins/Animation/ControlRig/Source/ControlRigEditor/Public/EditMode/ControlRigEditMode.h
@@ -430,7 +430,7 @@ protected:
 
 
 	/** The scope for the interaction, one per manipulated Control rig */
-	TMap<UControlRig*,FControlRigInteractionScope*> InteractionScopes;
+	TMap<UControlRig*, TUniquePtr<FControlRigInteractionScope>> InteractionScopes;
 
 	/** True if there's tracking going on right now */
 	bool bIsTracking;
