Control Rig: Fix crash when moving controls in a rig module editor while debugging a modular rig
#rnx

diff --git a/Engine/Plugins/Animation/ControlRig/Source/ControlRigEditor/Private/EditMode/ControlRigEditMode.cpp b/Engine/Plugins/Animation/ControlRig/Source/ControlRigEditor/Private/EditMode/ControlRigEditMode.cpp
--- a/Engine/Plugins/Animation/ControlRig/Source/ControlRigEditor/Private/EditMode/ControlRigEditMode.cpp
+++ b/Engine/Plugins/Animation/ControlRig/Source/ControlRigEditor/Private/EditMode/ControlRigEditMode.cpp
@@ -1142,6 +1142,18 @@ bool FControlRigEditMode::HandleBeginTransform(const FEditorViewportClient* InVi
 	{
 		return false;
 	}
+
+	// If moving in a control rig editor, but the rig is a module instance, moving controls is forbidden
+	if (AreEditingControlRigDirectly())
+	{
+		for (TWeakObjectPtr<UControlRig> ControlRig : RuntimeControlRigs)
+		{
+			if (ControlRig->IsRigModuleInstance())
+			{
+				return false;
+			}
+		}
+	}
 	
 	InteractionType = GetInteractionType(InViewportClient);
 	bIsTracking = true;
