Control Rig: Curves not being imported and reflected properly when using the import skeleton rig unit, or when editing a rig module

Most of the changes in SRigCurveContainer follow the same logic as SRigHierarchy.

#rnx

diff --git a/Engine/Plugins/Animation/ControlRig/Source/ControlRig/Private/Rigs/RigHierarchyController.cpp b/Engine/Plugins/Animation/ControlRig/Source/ControlRig/Private/Rigs/RigHierarchyController.cpp
--- a/Engine/Plugins/Animation/ControlRig/Source/ControlRig/Private/Rigs/RigHierarchyController.cpp
+++ b/Engine/Plugins/Animation/ControlRig/Source/ControlRig/Private/Rigs/RigHierarchyController.cpp
@@ -1776,6 +1776,17 @@ TArray<FRigElementKey> URigHierarchyController::ImportPreviewSkeletalMesh(USkele
 		}
 	}
 
+	// import the curves for the preview hierarchy
+	// use the ref skeleton so we'll only see the curves that are actually part of the mesh
+	const TArray<FRigElementKey> Curves = ImportCurvesFromSkeletalMesh(InSkeletalMesh, NAME_None, bSelectBones, bSetupUndo);
+	for(const FRigElementKey& Curve : Curves)
+	{
+		if(FRigBaseElement* Element = Hierarchy->Find(Curve))
+		{
+			Element->CreatedAtInstructionIndex = InstructionIndex;
+		}
+	}
+
 	// create a null to store controls under
 	static const FRigElementKey ControlParentKey(TEXT("Controls"), ERigElementType::Null);
 	if(!Hierarchy->Contains(ControlParentKey))
diff --git a/Engine/Plugins/Animation/ControlRig/Source/ControlRig/Private/Units/Execution/RigUnit_DynamicHierarchy.cpp b/Engine/Plugins/Animation/ControlRig/Source/ControlRig/Private/Units/Execution/RigUnit_DynamicHierarchy.cpp
--- a/Engine/Plugins/Animation/ControlRig/Source/ControlRig/Private/Units/Execution/RigUnit_DynamicHierarchy.cpp
+++ b/Engine/Plugins/Animation/ControlRig/Source/ControlRig/Private/Units/Execution/RigUnit_DynamicHierarchy.cpp
@@ -340,14 +340,16 @@ FRigUnit_HierarchyImportFromSkeleton_Execute()
 		
 		USkeleton* Skeleton = nullptr;
 		const FReferenceSkeleton* ReferenceSkeleton = nullptr;
+		USkeletalMesh* SkeletalMesh = nullptr;
 
 		const USkeletalMeshComponent* SkelMeshComponent = ExecuteContext.UnitContext.DataSourceRegistry->RequestSource<USkeletalMeshComponent>(UControlRig::OwnerComponent);
 		if (SkelMeshComponent != nullptr)
 		{
 			if (SkelMeshComponent->GetSkeletalMeshAsset())
 			{
-				Skeleton = SkelMeshComponent->GetSkeletalMeshAsset()->GetSkeleton();
-				ReferenceSkeleton = &SkelMeshComponent->GetSkeletalMeshAsset()->GetRefSkeleton();
+				SkeletalMesh = SkelMeshComponent->GetSkeletalMeshAsset();
+				Skeleton = SkeletalMesh->GetSkeleton();
+				ReferenceSkeleton = &SkeletalMesh->GetRefSkeleton();
 			}
 		}
 		else
@@ -390,17 +392,21 @@ FRigUnit_HierarchyImportFromSkeleton_Execute()
 				Items = Controller->ImportBones(FilteredBoneInfos, FilteredBoneTransforms, NameSpace, false, false, false, false);
 			}
 
-			if(bIncludeCurves && Skeleton != nullptr)
+			if(bIncludeCurves)
 			{
-				Items.Append(Controller->ImportCurves(Skeleton, NameSpace, false, false, false));
-			}
-			if(bIncludeMeshSockets && SkelMeshComponent != nullptr)
-			{
-				if(USkeletalMesh* SkeletalMesh = SkelMeshComponent->GetSkeletalMeshAsset())
+				if (SkeletalMesh)
+				{
+					Items.Append(Controller->ImportCurvesFromSkeletalMesh(SkeletalMesh, NameSpace, false, false, false));
+				}
+				else if (Skeleton)
 				{
-					Items.Append(Controller->ImportSocketsFromSkeletalMesh(SkeletalMesh, NameSpace, false, false, false));
+					Items.Append(Controller->ImportCurves(Skeleton, NameSpace, false, false, false));
 				}
 			}
+			if(bIncludeMeshSockets && SkeletalMesh != nullptr)
+			{
+				Items.Append(Controller->ImportSocketsFromSkeletalMesh(SkeletalMesh, NameSpace, false, false, false));
+			}
 		}
 	}
 }
diff --git a/Engine/Plugins/Animation/ControlRig/Source/ControlRigEditor/Private/Editor/SRigCurveContainer.cpp b/Engine/Plugins/Animation/ControlRig/Source/ControlRigEditor/Private/Editor/SRigCurveContainer.cpp
--- a/Engine/Plugins/Animation/ControlRig/Source/ControlRigEditor/Private/Editor/SRigCurveContainer.cpp
+++ b/Engine/Plugins/Animation/ControlRig/Source/ControlRigEditor/Private/Editor/SRigCurveContainer.cpp
@@ -145,6 +145,7 @@ void SRigCurveContainer::Construct(const FArguments& InArgs, TSharedRef<IControl
 	bIsChangingRigHierarchy = false;
 
 	ControlRigBlueprint->Hierarchy->OnModified().AddRaw(this, &SRigCurveContainer::OnHierarchyModified);
+	ControlRigBlueprint->OnSetObjectBeingDebugged().AddRaw(this, &SRigCurveContainer::HandleSetObjectBeingDebugged);
 	ControlRigBlueprint->OnRefreshEditor().AddRaw(this, &SRigCurveContainer::HandleRefreshEditorFromBlueprint);
 
 	UEditorEngine* Editor = Cast<UEditorEngine>(GEngine);
@@ -207,6 +208,16 @@ void SRigCurveContainer::Construct(const FArguments& InArgs, TSharedRef<IControl
 	{
 		ControlRigEditor.Pin()->OnEditorClosed().AddSP(this, &SRigCurveContainer::OnEditorClose);
 	}
+
+	// after opening the editor the debugged rig won't exist yet. we'll have to wait for a tick so
+	// that we have a valid rig to listen to.
+	RegisterActiveTimer(0.f, FWidgetActiveTimerDelegate::CreateLambda([this](double, float) {
+		if(ControlRigBlueprint.IsValid())
+		{
+			(void)HandleSetObjectBeingDebugged(ControlRigBlueprint->GetDebuggedControlRig());
+		}
+		return EActiveTimerReturnType::Stop;
+	}));
 }
 
 SRigCurveContainer::~SRigCurveContainer()
@@ -316,6 +327,7 @@ void SRigCurveContainer::OnEditorClose(IControlRigBaseEditor* InEditor, UControl
 	{
 		InBlueprint->Hierarchy->OnModified().RemoveAll(this);
 		InBlueprint->OnRefreshEditor().RemoveAll(this);
+		InBlueprint->OnSetObjectBeingDebugged().RemoveAll(this);
 	}
 
 	if (InEditor)
@@ -519,7 +531,7 @@ void SRigCurveContainer::SetCurveValue(const FName& CurveName, float CurveValue)
 
 float SRigCurveContainer::GetCurveValue(const FName& CurveName)
 {
-	URigHierarchy* Hierarchy = GetInstanceHierarchy();
+	URigHierarchy* Hierarchy = GetHierarchy();
 	if (Hierarchy)
 	{
 		return Hierarchy->GetCurveValue(FRigElementKey(CurveName, ERigElementType::Curve));
@@ -645,6 +657,46 @@ void SRigCurveContainer::OnHierarchyModified(ERigHierarchyNotification InNotif,
     }
 }
 
+void SRigCurveContainer::OnPostConstruction(UControlRig* InControlRig, const FName& InName)
+{
+	RefreshCurveList();
+}
+
+void SRigCurveContainer::HandleSetObjectBeingDebugged(UObject* InObject)
+{
+	if(ControlRigBeingDebuggedPtr.Get() == InObject)
+	{
+		return;
+	}
+
+	if(ControlRigBeingDebuggedPtr.IsValid())
+	{
+		if(UControlRig* ControlRigBeingDebugged = ControlRigBeingDebuggedPtr.Get())
+		{
+			if(!URigVMHost::IsGarbageOrDestroyed(ControlRigBeingDebugged))
+			{
+				ControlRigBeingDebugged->GetHierarchy()->OnModified().RemoveAll(this);
+			}
+		}
+	}
+
+	ControlRigBeingDebuggedPtr.Reset();
+	
+	if(UControlRig* ControlRig = Cast<UControlRig>(InObject))
+	{
+		ControlRigBeingDebuggedPtr = ControlRig;
+		if(URigHierarchy* Hierarchy = ControlRig->GetHierarchy())
+		{
+			Hierarchy->OnModified().RemoveAll(this);
+			Hierarchy->OnModified().AddSP(this, &SRigCurveContainer::OnHierarchyModified);
+		}
+		ControlRig->OnPostConstruction_AnyThread().RemoveAll(this);
+		ControlRig->OnPostConstruction_AnyThread().AddSP(this, &SRigCurveContainer::OnPostConstruction);
+	}
+
+	RefreshCurveList();
+}
+
 void SRigCurveContainer::HandleRefreshEditorFromBlueprint(URigVMBlueprint* InBlueprint)
 {
 	RefreshCurveList();
@@ -736,21 +788,27 @@ URigHierarchy* SRigCurveContainer::GetHierarchy() const
 {
 	if (ControlRigBlueprint.IsValid())
 	{
-		return ControlRigBlueprint->Hierarchy;
+		if (UControlRig* DebuggedRig = ControlRigBeingDebuggedPtr.Get())
+		{
+			return DebuggedRig->GetHierarchy();
+		}
 	}
-	return nullptr;
-}
-
-URigHierarchy* SRigCurveContainer::GetInstanceHierarchy() const
-{
 	if (ControlRigEditor.IsValid())
 	{
-		UControlRig* ControlRig = ControlRigEditor.Pin()->GetControlRig();
-		if (ControlRig)
+		if (UControlRig* CurrentRig = ControlRigEditor.Pin()->GetControlRig())
 		{
-			return ControlRig->GetHierarchy();
+			return CurrentRig->GetHierarchy();
 		}
 	}
+	return GetDefaultHierarchy();
+}
+
+URigHierarchy* SRigCurveContainer::GetDefaultHierarchy() const
+{
+	if (ControlRigBlueprint.IsValid())
+	{
+		return ControlRigBlueprint->Hierarchy;
+	}
 	return nullptr;
 }
 
diff --git a/Engine/Plugins/Animation/ControlRig/Source/ControlRigEditor/Private/Editor/SRigCurveContainer.h b/Engine/Plugins/Animation/ControlRig/Source/ControlRigEditor/Private/Editor/SRigCurveContainer.h
--- a/Engine/Plugins/Animation/ControlRig/Source/ControlRigEditor/Private/Editor/SRigCurveContainer.h
+++ b/Engine/Plugins/Animation/ControlRig/Source/ControlRigEditor/Private/Editor/SRigCurveContainer.h
@@ -9,6 +9,7 @@
 #include "Widgets/Views/SListView.h"
 #include "Rigs/RigHierarchy.h"
 #include "ControlRigBlueprint.h"
+#include "ControlRigEditor.h"
 #include "EditorUndoClient.h"
 
 //////////////////////////////////////////////////////////////////////////
@@ -243,6 +244,8 @@ private:
 
 	bool bIsChangingRigHierarchy;
 	void OnHierarchyModified(ERigHierarchyNotification InNotif, URigHierarchy* InHierarchy, const FRigNotificationSubject& InSubject);
+	void OnPostConstruction(UControlRig* InControlRig, const FName& InName);
+	void HandleSetObjectBeingDebugged(UObject* InObject);
 	void HandleRefreshEditorFromBlueprint(URigVMBlueprint* InBlueprint);
 
 	// import curve part
@@ -262,6 +265,7 @@ private:
 	/** Control Rig Blueprint */
 	TWeakObjectPtr<UControlRigBlueprint> ControlRigBlueprint;
 	TWeakPtr<IControlRigBaseEditor> ControlRigEditor;
+	TWeakObjectPtr<UControlRig> ControlRigBeingDebuggedPtr;
 
 	/** Current text typed into NameFilterBox */
 	FText FilterText;
@@ -270,7 +274,7 @@ private:
 	TSharedPtr<FUICommandList> UICommandList;
 
 	URigHierarchy* GetHierarchy() const;
-	URigHierarchy* GetInstanceHierarchy() const;
+	URigHierarchy* GetDefaultHierarchy() const;
 
 	friend class SRigCurveListRow;
 	friend class SRigCurveTypeList;
