Anim Layers: Check nullptr, recent change to fix guid changes for anim layers.

diff --git a/Engine/Plugins/Animation/ControlRig/Source/ControlRigEditor/Private/Sequencer/AnimLayers/AnimLayers.cpp b/Engine/Plugins/Animation/ControlRig/Source/ControlRigEditor/Private/Sequencer/AnimLayers/AnimLayers.cpp
--- a/Engine/Plugins/Animation/ControlRig/Source/ControlRigEditor/Private/Sequencer/AnimLayers/AnimLayers.cpp
+++ b/Engine/Plugins/Animation/ControlRig/Source/ControlRigEditor/Private/Sequencer/AnimLayers/AnimLayers.cpp
@@ -257,7 +257,7 @@ void UAnimLayer::UpdateSceneObjectorGuidsForItems(ISequencer* InSequencer)
 			DeadObjects.Add(DeadObject);
 		}
 	}
-	for (TPair<TWeakObjectPtr<UObject>, FGuid>& DeadObject: DeadObjects)
+	for (TPair<TWeakObjectPtr<UObject>, FGuid>& DeadObject : DeadObjects)
 	{
 		for (TWeakObjectPtr<> BoundObject : InSequencer->FindBoundObjects(DeadObject.Value, InSequencer->GetFocusedTemplateID()))
 		{
@@ -265,11 +265,13 @@ void UAnimLayer::UpdateSceneObjectorGuidsForItems(ISequencer* InSequencer)
 			{
 				continue;
 			}
-			FAnimLayerItem* Item = AnimLayerItems.Find(DeadObject.Key);
-			FAnimLayerItem NewCopy;
-			Item->MakeCopy(DeadObject.Value, BoundObject, NewCopy);
+			if (FAnimLayerItem* Item = AnimLayerItems.Find(DeadObject.Key))
+			{
+				FAnimLayerItem NewCopy;
+				Item->MakeCopy(DeadObject.Value, BoundObject, NewCopy);
+				AnimLayerItems.Add(BoundObject, NewCopy);
+			}
 			AnimLayerItems.Remove(DeadObject.Key);
-			AnimLayerItems.Add(BoundObject,NewCopy);
 		}
 	}
 }
