Control Rig: Add line trace unit nodes

diff --git a/Engine/Plugins/Animation/ControlRig/Source/ControlRig/Private/Units/Collision/RigUnit_WorldCollision.cpp b/Engine/Plugins/Animation/ControlRig/Source/ControlRig/Private/Units/Collision/RigUnit_WorldCollision.cpp
--- a/Engine/Plugins/Animation/ControlRig/Source/ControlRig/Private/Units/Collision/RigUnit_WorldCollision.cpp
+++ b/Engine/Plugins/Animation/ControlRig/Source/ControlRig/Private/Units/Collision/RigUnit_WorldCollision.cpp
@@ -145,6 +145,11 @@ FRigUnit_SphereTraceByTraceChannel_Execute()
 	WorkData.HitNormal = HitNormal;
 }
 
+FRigUnit_LineTraceByTraceChannel_Execute()
+{
+	FRigUnit_SphereTraceByTraceChannel::StaticExecute(ExecuteContext, Start, End, TraceChannel, 0, bHit, HitLocation, HitNormal, WorkData);
+}
+
 FRigUnit_SphereTraceByObjectTypes_Execute()
 {
 	DECLARE_SCOPE_HIERARCHICAL_COUNTER_RIGUNIT()
@@ -228,3 +233,9 @@ FRigUnit_SphereTraceByObjectTypes_Execute()
 	WorkData.HitNormal = HitNormal;
 }
 
+FRigUnit_LineTraceByObjectTypes_Execute()
+{
+	FRigUnit_SphereTraceByObjectTypes::StaticExecute(ExecuteContext, Start, End, ObjectTypes, 0, bHit, HitLocation, HitNormal, WorkData);
+}
+
+
diff --git a/Engine/Plugins/Animation/ControlRig/Source/ControlRig/Public/Units/Collision/RigUnit_WorldCollision.h b/Engine/Plugins/Animation/ControlRig/Source/ControlRig/Public/Units/Collision/RigUnit_WorldCollision.h
--- a/Engine/Plugins/Animation/ControlRig/Source/ControlRig/Public/Units/Collision/RigUnit_WorldCollision.h
+++ b/Engine/Plugins/Animation/ControlRig/Source/ControlRig/Public/Units/Collision/RigUnit_WorldCollision.h
@@ -155,6 +155,58 @@ struct FRigUnit_SphereTraceByTraceChannel : public FRigUnit
 	FRigUnit_SphereTrace_WorkData WorkData;
 };
 
+/**
+ * Performs a line trace against the world and return the first blocking hit using a specific channel. Target objects can have different object types, but they need to have the same trace channel set to "block" in their collision response settings. 
+ * You can create custom trace channels in Project Setting - Collision. 
+ */
+USTRUCT(meta = (DisplayName = "Line Trace By Trace Channel", Category = "Collision", DocumentationPolicy = "Strict", Keywords = "Line,Raytrace,Collision,Collide,Trace", Varying, NodeColor = "0.2 0.4 0.7"))
+struct FRigUnit_LineTraceByTraceChannel : public FRigUnit
+{
+	GENERATED_BODY()
+
+	FRigUnit_LineTraceByTraceChannel()
+		: Start(EForceInit::ForceInitToZero)
+		, End(EForceInit::ForceInitToZero)
+		, TraceChannel(TraceTypeQuery1)
+		, bHit(false)
+		, HitLocation(EForceInit::ForceInitToZero)
+		, HitNormal(0.f, 0.f, 1.f)
+	{
+
+	}
+
+	RIGVM_METHOD()
+	UE_API virtual void Execute() override;
+
+	/** Start of the trace in rig / global space */
+	UPROPERTY(meta = (Input))
+	FVector Start;
+
+	/** End of the trace in rig / global space */
+	UPROPERTY(meta = (Input))
+	FVector End;
+
+	/** The 'channel' that this trace is in, used to determine which components to hit */
+	UPROPERTY(meta = (Input))
+	TEnumAsByte<ETraceTypeQuery> TraceChannel;
+
+	/** Returns true if there was a hit */
+	UPROPERTY(meta = (Output))
+	bool bHit;
+
+	/** Hit location in rig / global Space */
+	UPROPERTY(meta = (Output))
+	FVector HitLocation;
+
+	/** Hit normal in rig / global Space */
+	UPROPERTY(meta = (Output))
+	FVector HitNormal;
+	
+	/** Cache / workstate */
+	UPROPERTY()
+	FRigUnit_SphereTrace_WorkData WorkData;
+};
+
 /**
  * Sweeps a sphere against the world and return the first blocking hit. The trace is filtered by object types only, the collision response settings are ignored.
  * You can create custom object types in Project Setting - Collision
@@ -212,4 +264,56 @@ struct FRigUnit_SphereTraceByObjectTypes : public FRigUnit
 	FRigUnit_SphereTrace_WorkData WorkData;
 };
 
+/**
+ * Performs a line trace against the world and return the first blocking hit. The trace is filtered by object types only, the collision response settings are ignored.
+ * You can create custom object types in Project Setting - Collision
+ */
+USTRUCT(meta = (DisplayName = "Line Trace By Object Types", Category = "Collision", DocumentationPolicy = "Strict", Keywords = "Line,Raytrace,Collision,Collide,Trace", Varying, NodeColor = "0.2 0.4 0.7"))
+struct FRigUnit_LineTraceByObjectTypes : public FRigUnit
+{
+	GENERATED_BODY()
+
+	FRigUnit_LineTraceByObjectTypes()
+	: Start(EForceInit::ForceInitToZero)
+	, End(EForceInit::ForceInitToZero)
+	, ObjectTypes()
+	, bHit(false)
+	, HitLocation(EForceInit::ForceInitToZero)
+	, HitNormal(0.f, 0.f, 1.f)
+	{
+
+	}
+
+	RIGVM_METHOD()
+	UE_API virtual void Execute() override;
+
+	/** Start of the trace in rig / global space */
+	UPROPERTY(meta = (Input))
+	FVector Start;
+
+	/** End of the trace in rig / global space */
+	UPROPERTY(meta = (Input))
+	FVector End;
+
+	/** The types of objects that this trace can hit */
+	UPROPERTY(meta = (Input))
+	TArray<TEnumAsByte<EObjectTypeQuery>> ObjectTypes;
+
+	/** Returns true if there was a hit */
+	UPROPERTY(meta = (Output))
+	bool bHit;
+
+	/** Hit location in rig / global Space */
+	UPROPERTY(meta = (Output))
+	FVector HitLocation;
+
+	/** Hit normal in rig / global Space */
+	UPROPERTY(meta = (Output))
+	FVector HitNormal;
+	
+	/** Cache / workstate */
+	UPROPERTY()
+	FRigUnit_SphereTrace_WorkData WorkData;
+};
+
 #undef UE_API
