Fix for Control Rig Hiearchy iterating over deprecated types and generating data for them

diff --git a/Engine/Plugins/Animation/ControlRig/Source/ControlRig/Private/Rigs/RigHierarchy.cpp b/Engine/Plugins/Animation/ControlRig/Source/ControlRig/Private/Rigs/RigHierarchy.cpp
--- a/Engine/Plugins/Animation/ControlRig/Source/ControlRig/Private/Rigs/RigHierarchy.cpp
+++ b/Engine/Plugins/Animation/ControlRig/Source/ControlRig/Private/Rigs/RigHierarchy.cpp
@@ -819,13 +819,13 @@ void URigHierarchy::Reset_Impl(bool bResetElements)
 		Elements.Reset();
 		ElementsPerType.Reset();
 		const UEnum* ElementTypeEnum = StaticEnum<ERigElementType>();
-		for(int32 TypeIndex=0; ;TypeIndex++)
+		for (int32 EnumTypeBit = (int32)ERigElementType::First; EnumTypeBit <= (int32)ERigElementType::Last; EnumTypeBit <<= 1)
 		{
-			if (static_cast<ERigElementType>(ElementTypeEnum->GetValueByIndex(TypeIndex)) == ERigElementType::All)
+			const int32 ElementTypeIndex = RigElementTypeToFlatIndex(static_cast<ERigElementType>(EnumTypeBit));
+			if (ElementTypeIndex != INDEX_NONE) // Skip deprecated types
 			{
-				break;
+				ElementsPerType.Add(TArray<FRigBaseElement*>());
 			}
-			ElementsPerType.Add(TArray<FRigBaseElement*>());
 		}
 		ElementIndexLookup.Reset();
 		ComponentIndexLookup.Reset();
@@ -910,18 +910,18 @@ void URigHierarchy::CopyHierarchy(URigHierarchy* InHierarchy)
 	if(!bReallocateElements)
 	{
 		const UEnum* ElementTypeEnum = StaticEnum<ERigElementType>();
-		for(int32 ElementTypeIndex = 0; ; ElementTypeIndex++)
+		for (int32 EnumTypeBit = (int32)ERigElementType::First; EnumTypeBit <= (int32)ERigElementType::Last; EnumTypeBit <<= 1)
 		{
-			if (static_cast<ERigElementType>(ElementTypeEnum->GetValueByIndex(ElementTypeIndex)) == ERigElementType::All)
+			const int32 ElementTypeIndex = RigElementTypeToFlatIndex(static_cast<ERigElementType>(EnumTypeBit));
+			if (ElementTypeIndex != INDEX_NONE) // Skip deprecated types
 			{
-				break;
-			}
-			check(ElementsPerType.IsValidIndex(ElementTypeIndex));
-			check(InHierarchy->ElementsPerType.IsValidIndex(ElementTypeIndex));
-			if(ElementsPerType[ElementTypeIndex].Num() != InHierarchy->ElementsPerType[ElementTypeIndex].Num())
-			{
-				bReallocateElements = true;
-				break;
+				check(ElementsPerType.IsValidIndex(ElementTypeIndex));
+				check(InHierarchy->ElementsPerType.IsValidIndex(ElementTypeIndex));
+				if(ElementsPerType[ElementTypeIndex].Num() != InHierarchy->ElementsPerType[ElementTypeIndex].Num())
+				{
+					bReallocateElements = true;
+					break;
+				}
 			}
 		}
 
@@ -1027,13 +1027,13 @@ void URigHierarchy::CopyHierarchy(URigHierarchy* InHierarchy)
 			// shrink the containers accordingly
 			Elements.SetNum(InHierarchy->Elements.Num());
 			const UEnum* ElementTypeEnum = StaticEnum<ERigElementType>();
-			for(int32 ElementTypeIndex = 0; ; ElementTypeIndex++)
+			for (int32 EnumTypeBit = (int32)ERigElementType::First; EnumTypeBit <= (int32)ERigElementType::Last; EnumTypeBit <<= 1)
 			{
-				if ((ERigElementType)ElementTypeEnum->GetValueByIndex(ElementTypeIndex) == ERigElementType::All)
+				const int32 ElementTypeIndex = RigElementTypeToFlatIndex(static_cast<ERigElementType>(EnumTypeBit));
+				if (ElementTypeIndex != INDEX_NONE) // Skip deprecated types
 				{
-					break;
+					ElementsPerType[ElementTypeIndex].SetNum(InHierarchy->ElementsPerType[ElementTypeIndex].Num());
 				}
-				ElementsPerType[ElementTypeIndex].SetNum(InHierarchy->ElementsPerType[ElementTypeIndex].Num());
 			}
 
 			for(int32 Index = 0; Index < InHierarchy->Num(); Index++)
diff --git a/Engine/Plugins/Animation/ControlRig/Source/ControlRig/Public/Rigs/RigHierarchyDefines.h b/Engine/Plugins/Animation/ControlRig/Source/ControlRig/Public/Rigs/RigHierarchyDefines.h
--- a/Engine/Plugins/Animation/ControlRig/Source/ControlRig/Public/Rigs/RigHierarchyDefines.h
+++ b/Engine/Plugins/Animation/ControlRig/Source/ControlRig/Public/Rigs/RigHierarchyDefines.h
@@ -36,10 +36,12 @@ enum class ERigElementType : uint8
 	Reference = 0x020,
 	Connector = 0x040,
 	Socket = 0x080,
-	
+
+	// All is used in iterations, it has to be after the last valid type
+	All = Bone | Null | Control | Curve | Reference | Connector | Socket,
+
 	First = Bone UMETA(Hidden), 
 	Last = Socket UMETA(Hidden), 
-	All = Bone | Null | Control | Curve | Reference | Connector | Socket,
 	ToResetAfterConstructionEvent = Bone | Control | Curve | Socket UMETA(Hidden),
 };
 
