Control Rig: Fix Undo/Redo crash when updating control shapes while the rig is running
#rnx

diff --git a/Engine/Plugins/Animation/ControlRig/Source/ControlRigEditor/Private/AnimDetails/Proxies/AnimDetailsProxyBase.cpp b/Engine/Plugins/Animation/ControlRig/Source/ControlRigEditor/Private/AnimDetails/Proxies/AnimDetailsProxyBase.cpp
--- a/Engine/Plugins/Animation/ControlRig/Source/ControlRigEditor/Private/AnimDetails/Proxies/AnimDetailsProxyBase.cpp
+++ b/Engine/Plugins/Animation/ControlRig/Source/ControlRigEditor/Private/AnimDetails/Proxies/AnimDetailsProxyBase.cpp
@@ -612,7 +612,7 @@ FRigControlElement* UAnimDetailsProxyBase::GetControlElement() const
 
 const FRigElementKey& UAnimDetailsProxyBase::GetControlElementKey() const
 {
-	return CachedRigElement.GetResolvedKey();
+	return CachedRigElement.GetKey();
 }
 
 const FName& UAnimDetailsProxyBase::GetControlName() const
diff --git a/Engine/Plugins/Animation/ControlRig/Source/ControlRigEditor/Private/EditMode/ControlRigEditMode.cpp b/Engine/Plugins/Animation/ControlRig/Source/ControlRigEditor/Private/EditMode/ControlRigEditMode.cpp
--- a/Engine/Plugins/Animation/ControlRig/Source/ControlRigEditor/Private/EditMode/ControlRigEditMode.cpp
+++ b/Engine/Plugins/Animation/ControlRig/Source/ControlRigEditor/Private/EditMode/ControlRigEditMode.cpp
@@ -4429,6 +4429,10 @@ bool FControlRigEditMode::TryUpdatingControlsShapes(UControlRig* InControlRig)
 	// we have matching controls - we should at least sync their settings.
 	// PostPoseUpdate / TickControlShape is going to take care of color, visibility etc.
 	// MeshTransform has to be handled here.
+
+	// Prevent evaluating the rig while we update the shapes. We want to especially prevent running construction during this update.
+	UE::TScopeLock EvaluateLock(InControlRig->GetEvaluateMutex());
+	
 	const TArray<TSoftObjectPtr<UControlRigShapeLibrary>>& ShapeLibraries = InControlRig->GetShapeLibraries();
 	for (int32 ShapeActorIndex = 0; ShapeActorIndex < NumShapes; ShapeActorIndex++)
 	{
