Control Rig now only performs one copy of the hierarchy during initialization

diff --git a/Engine/Plugins/Animation/ControlRig/Source/ControlRig/Private/ControlRig.cpp b/Engine/Plugins/Animation/ControlRig/Source/ControlRig/Private/ControlRig.cpp
--- a/Engine/Plugins/Animation/ControlRig/Source/ControlRig/Private/ControlRig.cpp
+++ b/Engine/Plugins/Animation/ControlRig/Source/ControlRig/Private/ControlRig.cpp
@@ -84,6 +84,7 @@ UControlRig::UControlRig(const FObjectInitializer& ObjectInitializer)
 #endif
 	, bCopyHierarchyBeforeConstruction(true)
 	, bResetInitialTransformsBeforeConstruction(true)
+	, bHierarchyInitializedThisFrame(false)
 	, bResetCurrentTransformsAfterConstruction(false)
 	, bManipulationEnabled(false)
 	, PreConstructionBracket(0)
@@ -749,6 +750,7 @@ void UControlRig::InitializeFromCDO()
 				TGuardValue<bool> Guard(Hierarchy->GetSuspendNotificationsFlag(), true);
 				Hierarchy->CopyHierarchy(CDO->GetHierarchy());
 				Hierarchy->ResetPoseToInitial(ERigElementType::All);
+				bHierarchyInitializedThisFrame = true;
 			}
 
 #if WITH_EDITOR
@@ -1191,37 +1193,45 @@ bool UControlRig::Execute(const FName& InEventName)
 #if WITH_EDITOR
 					FTransientControlScope TransientControlScope(Hierarchy);
 	#endif
+					if (!bHierarchyInitializedThisFrame)
 					{
-						// maintain the initial pose if it ever was set by the client
-						FRigPose InitialPose;
-						if(!bResetInitialTransformsBeforeConstruction)
+						if ((bCopyHierarchyBeforeConstruction || bResetInitialTransformsBeforeConstruction))
 						{
-							InitialPose = Hierarchy->GetPose(true, ERigElementType::ToResetAfterConstructionEvent, FRigElementKeyCollection());
-						}
+							// maintain the initial pose if it ever was set by the client
+							FRigPose InitialPose;
+							if(!bResetInitialTransformsBeforeConstruction)
+							{
+								InitialPose = Hierarchy->GetPose(true, ERigElementType::ToResetAfterConstructionEvent, FRigElementKeyCollection());
+							}
 
-						if(bCopyHierarchyBeforeConstruction)
-						{
-							Hierarchy->ResetToDefault();
+							if(bCopyHierarchyBeforeConstruction)
+							{
+								Hierarchy->ResetToDefault();
+							}
+
+							if(InitialPose.Num() > 0)
+							{
+								const TGuardValue<bool> DisableRecordingCurveChanges(Hierarchy->GetRecordCurveChangesFlag(), false);
+								Hierarchy->SetPose(InitialPose, ERigTransformType::InitialLocal);
+							}
 						}
 
-						if(InitialPose.Num() > 0)
 						{
-							const TGuardValue<bool> DisableRecordingCurveChanges(Hierarchy->GetRecordCurveChangesFlag(), false);
-							Hierarchy->SetPose(InitialPose, ERigTransformType::InitialLocal);
+		#if WITH_EDITOR
+							TUniquePtr<FTransientControlPoseScope> TransientControlPoseScope;
+							if (bConstructionModeEnabled)
+							{
+								// save the transient control value, it should not be constantly reset in construction mode
+								TransientControlPoseScope = MakeUnique<FTransientControlPoseScope>(this);
+							}
+		#endif
+							// reset the pose to initial such that construction event can run from a deterministic initial state
+							Hierarchy->ResetPoseToInitial(ERigElementType::All);
 						}
 					}
-
+					else
 					{
-	#if WITH_EDITOR
-						TUniquePtr<FTransientControlPoseScope> TransientControlPoseScope;
-						if (bConstructionModeEnabled)
-						{
-							// save the transient control value, it should not be constantly reset in construction mode
-							TransientControlPoseScope = MakeUnique<FTransientControlPoseScope>(this);
-						}
-	#endif
-						// reset the pose to initial such that construction event can run from a deterministic initial state
-						Hierarchy->ResetPoseToInitial(ERigElementType::All);
+						bHierarchyInitializedThisFrame = false;
 					}
 
 					RestoreShapeLibrariesFromCDO();
@@ -3741,11 +3751,6 @@ void UControlRig::PostInitInstance(URigVMHost* InCDO)
 			{
 				VM = InCDO->GetVM();
 			}
-
-			if(!IsRigModuleInstance())
-			{
-				DynamicHierarchy->CopyHierarchy(CastChecked<UControlRig>(InCDO)->GetHierarchy());
-			}
 		}
 	}
 	else // we are the CDO
@@ -4003,4 +4008,3 @@ bool UControlRig::ContainsSimulation() const
 }
 
 #undef LOCTEXT_NAMESPACE
-
diff --git a/Engine/Plugins/Animation/ControlRig/Source/ControlRig/Private/ControlRigTestData.cpp b/Engine/Plugins/Animation/ControlRig/Source/ControlRig/Private/ControlRigTestData.cpp
--- a/Engine/Plugins/Animation/ControlRig/Source/ControlRig/Private/ControlRigTestData.cpp
+++ b/Engine/Plugins/Animation/ControlRig/Source/ControlRig/Private/ControlRigTestData.cpp
@@ -581,7 +581,7 @@ bool UControlRigTestData::PerformTest(UControlRig* InSubject, TFunction<void(EMe
 	bool bSuccess = true;
 
 	// initialize
-	InSubject->RequestInit();
+	InSubject->Initialize();
 	Initial.Restore(InSubject, true);
 
 	// run the construction event
diff --git a/Engine/Plugins/Animation/ControlRig/Source/ControlRig/Public/ControlRig.h b/Engine/Plugins/Animation/ControlRig/Source/ControlRig/Public/ControlRig.h
--- a/Engine/Plugins/Animation/ControlRig/Source/ControlRig/Public/ControlRig.h
+++ b/Engine/Plugins/Animation/ControlRig/Source/ControlRig/Public/ControlRig.h
@@ -663,6 +663,7 @@ public:
 protected:
 	bool bCopyHierarchyBeforeConstruction;
 	bool bResetInitialTransformsBeforeConstruction;
+	bool bHierarchyInitializedThisFrame;
 	bool bResetCurrentTransformsAfterConstruction;
 	bool bManipulationEnabled;
 
