Control Rig: Fix crash when refreshing SRigCurveContainer from a worker thread
#rnx

diff --git a/Engine/Plugins/Animation/ControlRig/Source/ControlRigEditor/Private/Editor/SRigCurveContainer.cpp b/Engine/Plugins/Animation/ControlRig/Source/ControlRigEditor/Private/Editor/SRigCurveContainer.cpp
--- a/Engine/Plugins/Animation/ControlRig/Source/ControlRigEditor/Private/Editor/SRigCurveContainer.cpp
+++ b/Engine/Plugins/Animation/ControlRig/Source/ControlRigEditor/Private/Editor/SRigCurveContainer.cpp
@@ -144,7 +144,7 @@ void SRigCurveContainer::Construct(const FArguments& InArgs, TSharedRef<IControl
 	ControlRigBlueprint = InControlRigEditor.Get().GetControlRigBlueprint();
 	bIsChangingRigHierarchy = false;
 
-	ControlRigBlueprint->Hierarchy->OnModified().AddRaw(this, &SRigCurveContainer::OnHierarchyModified);
+	ControlRigBlueprint->Hierarchy->OnModified().AddRaw(this, &SRigCurveContainer::OnHierarchyModified_AnyThread);
 	ControlRigBlueprint->OnSetObjectBeingDebugged().AddRaw(this, &SRigCurveContainer::HandleSetObjectBeingDebugged);
 	ControlRigBlueprint->OnRefreshEditor().AddRaw(this, &SRigCurveContainer::HandleRefreshEditorFromBlueprint);
 
@@ -598,6 +598,51 @@ void SRigCurveContainer::OnSelectionChanged(FDisplayedRigCurveInfoPtr Selection,
 	}
 }
 
+void SRigCurveContainer::OnHierarchyModified_AnyThread(ERigHierarchyNotification InNotif, URigHierarchy* InHierarchy, const FRigNotificationSubject& InSubject)
+{
+	if(bIsChangingRigHierarchy)
+	{
+		return;
+	}
+	
+	if(!ControlRigBeingDebuggedPtr.IsValid())
+	{
+		return;
+	}
+
+	if(InHierarchy != ControlRigBeingDebuggedPtr->GetHierarchy())
+	{
+		return;
+	}
+
+	if(IsInGameThread())
+	{
+		OnHierarchyModified(InNotif, InHierarchy, InSubject);
+	}
+	else
+	{
+		FRigElementKey Key;
+		if(const FRigBaseElement* InElement = InSubject.Element)
+		{
+			Key = InElement->GetKey();
+		}
+		
+		TWeakObjectPtr<URigHierarchy> WeakHierarchy = InHierarchy;
+		FFunctionGraphTask::CreateAndDispatchWhenReady([this, InNotif, WeakHierarchy,Key]()
+		{
+			if(!WeakHierarchy.IsValid())
+			{
+				return;
+			}
+
+			if (const FRigBaseElement* Element = WeakHierarchy.Get()->Find(Key))
+			{
+				OnHierarchyModified(InNotif, WeakHierarchy.Get(), Element);
+			}
+		}, TStatId(), NULL, ENamedThreads::GameThread);
+	}
+}
+
 void SRigCurveContainer::OnHierarchyModified(ERigHierarchyNotification InNotif, URigHierarchy* InHierarchy, const FRigNotificationSubject& InSubject)
 {
 	if (bIsChangingRigHierarchy)
@@ -688,7 +733,7 @@ void SRigCurveContainer::HandleSetObjectBeingDebugged(UObject* InObject)
 		if(URigHierarchy* Hierarchy = ControlRig->GetHierarchy())
 		{
 			Hierarchy->OnModified().RemoveAll(this);
-			Hierarchy->OnModified().AddSP(this, &SRigCurveContainer::OnHierarchyModified);
+			Hierarchy->OnModified().AddSP(this, &SRigCurveContainer::OnHierarchyModified_AnyThread);
 		}
 		ControlRig->OnPostConstruction_AnyThread().RemoveAll(this);
 		ControlRig->OnPostConstruction_AnyThread().AddSP(this, &SRigCurveContainer::OnPostConstruction);
diff --git a/Engine/Plugins/Animation/ControlRig/Source/ControlRigEditor/Private/Editor/SRigCurveContainer.h b/Engine/Plugins/Animation/ControlRig/Source/ControlRigEditor/Private/Editor/SRigCurveContainer.h
--- a/Engine/Plugins/Animation/ControlRig/Source/ControlRigEditor/Private/Editor/SRigCurveContainer.h
+++ b/Engine/Plugins/Animation/ControlRig/Source/ControlRigEditor/Private/Editor/SRigCurveContainer.h
@@ -243,6 +243,7 @@ private:
 	void OnSelectionChanged(FDisplayedRigCurveInfoPtr Selection, ESelectInfo::Type SelectInfo);
 
 	bool bIsChangingRigHierarchy;
+	void OnHierarchyModified_AnyThread(ERigHierarchyNotification InNotif, URigHierarchy* InHierarchy, const FRigNotificationSubject& InSubject);
 	void OnHierarchyModified(ERigHierarchyNotification InNotif, URigHierarchy* InHierarchy, const FRigNotificationSubject& InSubject);
 	void OnPostConstruction(UControlRig* InControlRig, const FName& InName);
 	void HandleSetObjectBeingDebugged(UObject* InObject);
