Control Rig: Fix outliner not always selecting all elements
#rnx

diff --git a/Engine/Plugins/Animation/ControlRig/Source/ControlRigEditor/Private/EditMode/SControlRigOutliner.cpp b/Engine/Plugins/Animation/ControlRig/Source/ControlRigEditor/Private/EditMode/SControlRigOutliner.cpp
--- a/Engine/Plugins/Animation/ControlRig/Source/ControlRigEditor/Private/EditMode/SControlRigOutliner.cpp
+++ b/Engine/Plugins/Animation/ControlRig/Source/ControlRigEditor/Private/EditMode/SControlRigOutliner.cpp
@@ -2177,12 +2177,6 @@ void SControlRigOutliner::HandleControlSelected(UControlRig* Subject, FRigContro
 	if (Found.IsValid())
 	{
 		TGuardValue<bool> GuardRigHierarchyChanges(bIsChangingRigHierarchy, true);
-
-		if (!FSlateApplication::Get().GetModifierKeys().IsShiftDown() && !FSlateApplication::Get().GetModifierKeys().IsControlDown())
-		{
-			HierarchyTreeView->GetTreeView()->ClearSelection();
-		}
-
 		HierarchyTreeView->GetTreeView()->SetItemSelection(Found, bSelected, ESelectInfo::Direct);
 	}
 	
