Control Rig: Fix linker issues on clang 18 with partially specialized, exported templated member functions.

#rnx

diff --git a/Engine/Plugins/Animation/ControlRig/Source/ControlRig/Private/Tracks/SampleTrack.cpp b/Engine/Plugins/Animation/ControlRig/Source/ControlRig/Private/Tracks/SampleTrack.cpp
--- a/Engine/Plugins/Animation/ControlRig/Source/ControlRig/Private/Tracks/SampleTrack.cpp
+++ b/Engine/Plugins/Animation/ControlRig/Source/ControlRig/Private/Tracks/SampleTrack.cpp
@@ -7,7 +7,7 @@
 //////////////////////////////////////////////////////////////////////////////////////////////////////////////////////////////
 
 template<>
-void FSampleTrack<FInstancedStruct>::Serialize(FArchive& InArchive)
+CONTROLRIG_API void FSampleTrack<FInstancedStruct>::Serialize(FArchive& InArchive)
 {
 	FSampleTrackBase::Serialize(InArchive);
 	
@@ -32,7 +32,7 @@ void FSampleTrack<FInstancedStruct>::Serialize(FArchive& InArchive)
 }
 
 template<>
-void FSampleTrack<TArray<FInstancedStruct>>::Serialize(FArchive& InArchive)
+CONTROLRIG_API void FSampleTrack<TArray<FInstancedStruct>>::Serialize(FArchive& InArchive)
 {
 	FSampleTrackBase::Serialize(InArchive);
 	
@@ -69,7 +69,7 @@ void FSampleTrack<TArray<FInstancedStruct>>::Serialize(FArchive& InArchive)
 }
 
 template<>
-int32 FSampleTrack<FInstancedStruct>::GetSizePerValue() const
+CONTROLRIG_API int32 FSampleTrack<FInstancedStruct>::GetSizePerValue() const
 {
 	if(ScriptStruct)
 	{
@@ -79,7 +79,7 @@ int32 FSampleTrack<FInstancedStruct>::GetSizePerValue() const
 }
 
 template<>
-void FSampleTrack<FInstancedStruct>::AddSampleFromProperty(const FProperty* InProperty, const uint8* InMemory)
+CONTROLRIG_API void FSampleTrack<FInstancedStruct>::AddSampleFromProperty(const FProperty* InProperty, const uint8* InMemory)
 {
 	const FStructProperty* StructProperty = CastFieldChecked<FStructProperty>(InProperty);
 	check(StructProperty->Struct == GetScriptStruct());
@@ -90,7 +90,7 @@ void FSampleTrack<FInstancedStruct>::AddSampleFromProperty(const FProperty* InPr
 }
 
 template<>
-void FSampleTrack<FInstancedStruct>::GetSampleForProperty(int32 InTimeIndex, FSampleTrackIndex& InOutSampleTrackIndex, const FProperty* InProperty, uint8* OutMemory) const
+CONTROLRIG_API void FSampleTrack<FInstancedStruct>::GetSampleForProperty(int32 InTimeIndex, FSampleTrackIndex& InOutSampleTrackIndex, const FProperty* InProperty, uint8* OutMemory) const
 {
 	const FStructProperty* StructProperty = CastFieldChecked<FStructProperty>(InProperty);
 	check(StructProperty->Struct == GetScriptStruct());
@@ -100,7 +100,7 @@ void FSampleTrack<FInstancedStruct>::GetSampleForProperty(int32 InTimeIndex, FSa
 }
 
 template<>
-void FSampleTrack<TArray<FInstancedStruct>>::AddSampleFromProperty(const FProperty* InProperty, const uint8* InMemory)
+CONTROLRIG_API void FSampleTrack<TArray<FInstancedStruct>>::AddSampleFromProperty(const FProperty* InProperty, const uint8* InMemory)
 {
 	const FArrayProperty* ArrayProperty = CastFieldChecked<FArrayProperty>(InProperty);
 	const FStructProperty* StructProperty = CastFieldChecked<FStructProperty>(ArrayProperty->Inner);
@@ -119,7 +119,7 @@ void FSampleTrack<TArray<FInstancedStruct>>::AddSampleFromProperty(const FProper
 }
 
 template<>
-void FSampleTrack<TArray<FInstancedStruct>>::GetSampleForProperty(int32 InTimeIndex, FSampleTrackIndex& InOutSampleTrackIndex, const FProperty* InProperty, uint8* OutMemory) const
+CONTROLRIG_API void FSampleTrack<TArray<FInstancedStruct>>::GetSampleForProperty(int32 InTimeIndex, FSampleTrackIndex& InOutSampleTrackIndex, const FProperty* InProperty, uint8* OutMemory) const
 {
 	const FArrayProperty* ArrayProperty = CastFieldChecked<FArrayProperty>(InProperty);
 	const FStructProperty* StructProperty = CastFieldChecked<FStructProperty>(ArrayProperty->Inner);
