Control Rig: Fix crash with pure virtual function call error
#rnx

diff --git a/Engine/Plugins/Animation/ControlRig/Source/ControlRigEditor/Private/EditMode/ControlRigEditMode.cpp b/Engine/Plugins/Animation/ControlRig/Source/ControlRigEditor/Private/EditMode/ControlRigEditMode.cpp
--- a/Engine/Plugins/Animation/ControlRig/Source/ControlRigEditor/Private/EditMode/ControlRigEditMode.cpp
+++ b/Engine/Plugins/Animation/ControlRig/Source/ControlRigEditor/Private/EditMode/ControlRigEditMode.cpp
@@ -1361,17 +1361,20 @@ bool FControlRigEditMode::UsesTransformWidget() const
 {
 	for (const auto& Pairs : ControlRigShapeActors)
 	{
-		for (const AControlRigShapeActor* ShapeActor : Pairs.Value)
+		if (TStrongObjectPtr<UControlRig> ControlRig = Pairs.Key.Pin())
 		{
-			if (ShapeActor->IsSelected())
+			for (const AControlRigShapeActor* ShapeActor : Pairs.Value)
+			{
+				if (ShapeActor->IsSelected())
+				{
+					return true;
+				}
+			}
+			if (AreRigElementSelectedAndMovable(ControlRig.Get()))
 			{
 				return true;
 			}
 		}
-		if (AreRigElementSelectedAndMovable(Pairs.Key))
-		{
-			return true;
-		}
 	}
 	return FEdMode::UsesTransformWidget();
 }
@@ -1380,16 +1383,19 @@ bool FControlRigEditMode::UsesTransformWidget(UE::Widget::EWidgetMode CheckMode)
 {
 	for (const auto& Pairs : ControlRigShapeActors)
 	{
-		for (const AControlRigShapeActor* ShapeActor : Pairs.Value)
+		if (TStrongObjectPtr<UControlRig> ControlRig = Pairs.Key.Pin())
 		{
-			if (ShapeActor->IsSelected())
+			for (const AControlRigShapeActor* ShapeActor : Pairs.Value)
 			{
-				return ModeSupportedByShapeActor(ShapeActor, CheckMode);
+				if (ShapeActor->IsSelected())
+				{
+					return ModeSupportedByShapeActor(ShapeActor, CheckMode);
+				}
+			}
+			if (AreRigElementSelectedAndMovable(ControlRig.Get()))
+			{
+				return true;
 			}
-		}
-		if (AreRigElementSelectedAndMovable(Pairs.Key))
-		{
-			return true;
 		}
 	}
 	return FEdMode::UsesTransformWidget(CheckMode);
@@ -1401,16 +1407,19 @@ FVector FControlRigEditMode::GetWidgetLocation() const
 	int NumSelected = 0;
 	for (const auto& Pairs : ControlRigShapeActors)
 	{
-		if (AreRigElementSelectedAndMovable(Pairs.Key))
+		if (TStrongObjectPtr<UControlRig> ControlRig = Pairs.Key.Pin())
 		{
-			if (const FTransform* PivotTransform = PivotTransforms.Find(Pairs.Key))
+			if (AreRigElementSelectedAndMovable(ControlRig.Get()))
 			{
-				// check that the cached pivot is up-to-date and update it if needed
-				FTransform Transform = *PivotTransform;
-				UpdatePivotTransformsIfNeeded(Pairs.Key, Transform);
-				const FTransform ComponentTransform = GetHostingSceneComponentTransform(Pairs.Key);
-				PivotLocation += ComponentTransform.TransformPosition(Transform.GetLocation());
-				++NumSelected;
+				if (const FTransform* PivotTransform = PivotTransforms.Find(ControlRig.Get()))
+				{
+					// check that the cached pivot is up-to-date and update it if needed
+					FTransform Transform = *PivotTransform;
+					UpdatePivotTransformsIfNeeded(ControlRig.Get(), Transform);
+					const FTransform ComponentTransform = GetHostingSceneComponentTransform(ControlRig.Get());
+					PivotLocation += ComponentTransform.TransformPosition(Transform.GetLocation());
+					++NumSelected;
+				}
 			}
 		}
 	}	
@@ -1440,15 +1449,18 @@ bool FControlRigEditMode::GetCustomDrawingCoordinateSystem(FMatrix& OutMatrix, v
 	//since we strip translation just want the first one
 	for (const auto& Pairs : ControlRigShapeActors)
 	{
-		if (AreRigElementSelectedAndMovable(Pairs.Key))
+		if (TStrongObjectPtr<UControlRig> ControlRig = Pairs.Key.Pin())
 		{
-			if (const FTransform* PivotTransform = PivotTransforms.Find(Pairs.Key))
+			if (AreRigElementSelectedAndMovable(ControlRig.Get()))
 			{
-				// check that the cached pivot is up-to-date and update it if needed
-				FTransform Transform = *PivotTransform;
-				UpdatePivotTransformsIfNeeded(Pairs.Key, Transform);
-				OutMatrix = Transform.ToMatrixNoScale().RemoveTranslation();
-				return true;
+				if (const FTransform* PivotTransform = PivotTransforms.Find(ControlRig.Get()))
+				{
+					// check that the cached pivot is up-to-date and update it if needed
+					FTransform Transform = *PivotTransform;
+					UpdatePivotTransformsIfNeeded(ControlRig.Get(), Transform);
+					OutMatrix = Transform.ToMatrixNoScale().RemoveTranslation();
+					return true;
+				}
 			}
 		}
 	}
@@ -1787,18 +1799,21 @@ bool FControlRigEditMode::BoxSelect(FBox& InBox, bool InSelect)
 	}
 
 	// Select all actors that are within the selection box area.  Be aware that certain modes do special processing below.	
-	TMap<UControlRig*, TArray<FRigElementKey>> ControlsToSelect;
-	for (const auto& [ControlRig, ShapeActors]: ControlRigShapeActors)
+	TMap<TWeakObjectPtr<UControlRig>, TArray<FRigElementKey>> ControlsToSelect;
+	for (const auto& [WeakControlRig, ShapeActors]: ControlRigShapeActors)
 	{
-		if (ControlRig && ControlRig->GetControlsVisible())
+		if (TStrongObjectPtr<UControlRig> ControlRig = WeakControlRig.Pin())
 		{
-			for (AControlRigShapeActor* ShapeActor : ShapeActors)
+			if (ControlRig && ControlRig->GetControlsVisible())
 			{
-				const bool bTreatShape = ShapeActor && ShapeActor->IsSelectable() && !ShapeActor->IsTemporarilyHiddenInEditor();
-				if (bTreatShape && IntersectsBox(*ShapeActor, InBox, LevelViewportClient, bStrictDragSelection))
+				for (AControlRigShapeActor* ShapeActor : ShapeActors)
 				{
-					TArray<FRigElementKey>& Controls = ControlsToSelect.FindOrAdd(ControlRig);
-					Controls.Add(ShapeActor->GetElementKey());
+					const bool bTreatShape = ShapeActor && ShapeActor->IsSelectable() && !ShapeActor->IsTemporarilyHiddenInEditor();
+					if (bTreatShape && IntersectsBox(*ShapeActor, InBox, LevelViewportClient, bStrictDragSelection))
+					{
+						TArray<FRigElementKey>& Controls = ControlsToSelect.FindOrAdd(ControlRig.Get());
+						Controls.Add(ShapeActor->GetElementKey());
+					}
 				}
 			}
 		}
@@ -1846,7 +1861,7 @@ bool FControlRigEditMode::FrustumSelect(const FConvexVolume& InFrustum, FEditorV
 		ClearRigElementSelection(ValidControlTypeMask());
 	}
 
-	TMap<UControlRig*, TArray<FRigElementKey>> ControlsToSelect;
+	TMap<TWeakObjectPtr<UControlRig>, TArray<FRigElementKey>> ControlsToSelect;
 	
 	FSelectionHelper SelectionHelper(InViewportClient, ControlRigShapeActors, ControlsToSelect);
 	SelectionHelper.GetFromFrustum(InFrustum);
@@ -2071,17 +2086,20 @@ bool FControlRigEditMode::InputDelta(FEditorViewportClient* InViewportClient, FV
 	auto GatherSelectedKeys = [&]() -> FMultiControlRigElementSelection
 	{
 		FMultiControlRigElementSelection MultiRigSelection;
-		for (TPair<UControlRig*, TArray<TObjectPtr<AControlRigShapeActor>>>& Pairs : ControlRigShapeActors)
+		for (TPair<TWeakObjectPtr<UControlRig>, TArray<TObjectPtr<AControlRigShapeActor>>>& Pairs : ControlRigShapeActors)
 		{
-			MultiRigSelection.Rigs.Add(Pairs.Key);
-			FRigElementKeyCollection& Keys = MultiRigSelection.KeysPerRig.Add_GetRef({});
-			if (AreRigElementsSelected(ValidControlTypeMask(), Pairs.Key))
+			if (TStrongObjectPtr<UControlRig> ControlRig = Pairs.Key.Pin())
 			{
-				for (TObjectPtr<AControlRigShapeActor> ShapeActor : Pairs.Value)
+				MultiRigSelection.Rigs.Add(Pairs.Key);
+				FRigElementKeyCollection& Keys = MultiRigSelection.KeysPerRig.Add_GetRef({});
+				if (AreRigElementsSelected(ValidControlTypeMask(), ControlRig.Get()))
 				{
-					if (ShapeActor->IsEnabled() && ShapeActor->IsSelected())
+					for (TObjectPtr<AControlRigShapeActor> ShapeActor : Pairs.Value)
 					{
-						Keys.Add(ShapeActor->GetElementKey());
+						if (ShapeActor->IsEnabled() && ShapeActor->IsSelected())
+						{
+							Keys.Add(ShapeActor->GetElementKey());
+						}
 					}
 				}
 			}
@@ -2107,175 +2125,176 @@ bool FControlRigEditMode::InputDelta(FEditorViewportClient* InViewportClient, FV
 		
 		for (auto& Pairs : ControlRigShapeActors)
 		{
-			UControlRig* ControlRig = Pairs.Key;
-			
-			if (AreRigElementsSelected(ValidControlTypeMask(), ControlRig))
+			if (TStrongObjectPtr<UControlRig> ControlRig = Pairs.Key.Pin())
 			{
-				FTransform ComponentTransform = GetHostingSceneComponentTransform(ControlRig);
-
-				if (bIsChangingControlShapeTransform)
+				if (AreRigElementsSelected(ValidControlTypeMask(), ControlRig.Get()))
 				{
-					for (AControlRigShapeActor* ShapeActor : Pairs.Value)
+					FTransform ComponentTransform = GetHostingSceneComponentTransform(ControlRig.Get());
+
+					if (bIsChangingControlShapeTransform)
 					{
-						if (ShapeActor->IsSelected())
+						for (AControlRigShapeActor* ShapeActor : Pairs.Value)
 						{
-							if (bManipulatorMadeChange == false)
+							if (ShapeActor->IsSelected())
 							{
-								GEditor->BeginTransaction(LOCTEXT("ChangeControlShapeTransaction", "Change Control Shape Transform"));
-							}
+								if (bManipulatorMadeChange == false)
+								{
+									GEditor->BeginTransaction(LOCTEXT("ChangeControlShapeTransaction", "Change Control Shape Transform"));
+								}
 
-							ChangeControlShapeTransform(ShapeActor, TransformContext, ComponentTransform);
-							bManipulatorMadeChange = true;
+								ChangeControlShapeTransform(ShapeActor, TransformContext, ComponentTransform);
+								bManipulatorMadeChange = true;
 
-							// break here since we only support changing shape transform of a single control at a time
-							break;
+								// break here since we only support changing shape transform of a single control at a time
+								break;
+							}
 						}
 					}
-				}
-				else
-				{
-					bool bDoLocal = EnableLocalTransform;
-					bool bUseLocal = false;
-					bool bCalcLocal = bDoLocal;
-					bool bFirstTime = true;
-					FTransform InOutLocal = FTransform::Identity;
+					else
+					{
+						bool bDoLocal = EnableLocalTransform;
+						bool bUseLocal = false;
+						bool bCalcLocal = bDoLocal;
+						bool bFirstTime = true;
+						FTransform InOutLocal = FTransform::Identity;
 					
-					bool const bJustStartedManipulation = !bManipulatorMadeChange;
-					const bool bAnyAdditiveRig = ControlRig->IsAdditive();
+						bool const bJustStartedManipulation = !bManipulatorMadeChange;
+						const bool bAnyAdditiveRig = ControlRig->IsAdditive();
 
-					TMap<AControlRigShapeActor*, TArray<TFunction<void()>>> TasksPerActor;
-					for (AControlRigShapeActor* ShapeActor : Pairs.Value)
-					{
-						if (ShapeActor->IsEnabled() && ShapeActor->IsSelected())
+						TMap<AControlRigShapeActor*, TArray<TFunction<void()>>> TasksPerActor;
+						for (AControlRigShapeActor* ShapeActor : Pairs.Value)
 						{
-							// test local vs global
-							if (bManipulatorMadeChange == false)
+							if (ShapeActor->IsEnabled() && ShapeActor->IsSelected())
 							{
-								GEditor->BeginTransaction(LOCTEXT("MoveControlTransaction", "Move Control"));
-							}
-
-							// cache interaction dependencies + evaluate the rig at least once before manipulating anything
-							if (ControlRig->ElementsBeingInteracted.IsEmpty() && InteractionScopes.Contains(ControlRig))
-							{
-								GetInteractionDependencies(ControlRig);
-								EvaluateRig(ControlRig);
-							}
+								// test local vs global
+								if (bManipulatorMadeChange == false)
+								{
+									GEditor->BeginTransaction(LOCTEXT("MoveControlTransaction", "Move Control"));
+								}
 
-							// Cannot benefit of same local transform when applying to additive rigs
-							if (!bAnyAdditiveRig)
-							{
-								if (bFirstTime)
+								// cache interaction dependencies + evaluate the rig at least once before manipulating anything
+								if (ControlRig->ElementsBeingInteracted.IsEmpty() && InteractionScopes.Contains(ControlRig.Get()))
 								{
-									bFirstTime = false;
+									GetInteractionDependencies(ControlRig.Get());
+									EvaluateRig(ControlRig.Get());
 								}
-								else
+
+								// Cannot benefit of same local transform when applying to additive rigs
+								if (!bAnyAdditiveRig)
 								{
-									if (bDoLocal)
+									if (bFirstTime)
+									{
+										bFirstTime = false;
+									}
+									else
 									{
-										bUseLocal = true;
-										bDoLocal = false;
+										if (bDoLocal)
+										{
+											bUseLocal = true;
+											bDoLocal = false;
+										}
 									}
 								}
-							}
 
-							if(bJustStartedManipulation)
-							{
-								if(const FRigControlElement* ControlElement = ControlRig->FindControl(ShapeActor->ControlName))
+								if(bJustStartedManipulation)
 								{
-									ShapeActor->OffsetTransform = ControlRig->GetHierarchy()->GetGlobalControlOffsetTransform(ControlElement->GetKey(), false);
+									if(const FRigControlElement* ControlElement = ControlRig->FindControl(ShapeActor->ControlName))
+									{
+										ShapeActor->OffsetTransform = ControlRig->GetHierarchy()->GetGlobalControlOffsetTransform(ControlElement->GetKey(), false);
+									}
 								}
-							}
 
-							TArray<TFunction<void()>> Tasks;
-							MoveControlShape(ShapeActor, TransformContext, ComponentTransform, bUseLocal, bDoLocal, &InOutLocal, Tasks);
-							TasksPerActor.Add(ShapeActor, Tasks);
-							bManipulatorMadeChange = true;
+								TArray<TFunction<void()>> Tasks;
+								MoveControlShape(ShapeActor, TransformContext, ComponentTransform, bUseLocal, bDoLocal, &InOutLocal, Tasks);
+								TasksPerActor.Add(ShapeActor, Tasks);
+								bManipulatorMadeChange = true;
+							}
 						}
-					}
 
-					{
-						// evaluate the rig(s) before the tasks
-						const FPendingControlRigEvaluator _(this);
-					}
-
-					// process remaining tasks
-					while(!TasksPerActor.IsEmpty())
-					{
-						// place another evaluator here which is going to run after the tasks for this phase
-						const FPendingControlRigEvaluator _(this);
+						{
+							// evaluate the rig(s) before the tasks
+							const FPendingControlRigEvaluator _(this);
+						}
 
-						// run one tasks for each control - until there are no tasks lefts
-						TArray<AControlRigShapeActor*> KeysToRemove;
-						for(TPair<AControlRigShapeActor*, TArray<TFunction<void()>>>& Pair : TasksPerActor)
+						// process remaining tasks
+						while(!TasksPerActor.IsEmpty())
 						{
-							if(!Pair.Value.IsEmpty())
+							// place another evaluator here which is going to run after the tasks for this phase
+							const FPendingControlRigEvaluator _(this);
+
+							// run one tasks for each control - until there are no tasks lefts
+							TArray<AControlRigShapeActor*> KeysToRemove;
+							for(TPair<AControlRigShapeActor*, TArray<TFunction<void()>>>& Pair : TasksPerActor)
 							{
-								TFunction<void()> Task = Pair.Value[0];
-								Pair.Value.RemoveAt(0);
-								Task();
+								if(!Pair.Value.IsEmpty())
+								{
+									TFunction<void()> Task = Pair.Value[0];
+									Pair.Value.RemoveAt(0);
+									Task();
+								}
+								else
+								{
+									KeysToRemove.Add(Pair.Key);
+								}
 							}
-							else
+
+							for(AControlRigShapeActor* KeyToRemove : KeysToRemove)
 							{
-								KeysToRemove.Add(Pair.Key);
+								TasksPerActor.Remove(KeyToRemove);
 							}
 						}
-
-						for(AControlRigShapeActor* KeyToRemove : KeysToRemove)
-						{
-							TasksPerActor.Remove(KeyToRemove);
-						}
 					}
 				}
-			}
-			else if (AreRigElementSelectedAndMovable(Pairs.Key))
-			{
-				FTransform ComponentTransform = GetHostingSceneComponentTransform(Pairs.Key);
-
-				// set Bone transform
-				// that will set initial Bone transform
-				TArray<FRigElementKey> SelectedRigElements = GetSelectedRigElements(Pairs.Key);
-
-				for (int32 Index = 0; Index < SelectedRigElements.Num(); ++Index)
+				else if (AreRigElementSelectedAndMovable(ControlRig.Get()))
 				{
-					const ERigElementType SelectedRigElementType = SelectedRigElements[Index].Type;
+					FTransform ComponentTransform = GetHostingSceneComponentTransform(ControlRig.Get());
 
-					if (SelectedRigElementType == ERigElementType::Control)
+					// set Bone transform
+					// that will set initial Bone transform
+					TArray<FRigElementKey> SelectedRigElements = GetSelectedRigElements(ControlRig.Get());
+
+					for (int32 Index = 0; Index < SelectedRigElements.Num(); ++Index)
 					{
-						FTransform NewWorldTransform = OnGetRigElementTransformDelegate.Execute(SelectedRigElements[Index], false, true) * ComponentTransform;
-						bool bTransformChanged = false;
-						if (TransformContext.bRotation)
-						{
-							FQuat CurrentRotation = NewWorldTransform.GetRotation();
-							CurrentRotation = (TransformContext.Rot.Quaternion() * CurrentRotation);
-							NewWorldTransform.SetRotation(CurrentRotation);
-							bTransformChanged = true;
-						}
+						const ERigElementType SelectedRigElementType = SelectedRigElements[Index].Type;
 
-						if (TransformContext.bTranslation)
+						if (SelectedRigElementType == ERigElementType::Control)
 						{
-							FVector CurrentLocation = NewWorldTransform.GetLocation();
-							CurrentLocation = CurrentLocation + TransformContext.Drag;
-							NewWorldTransform.SetLocation(CurrentLocation);
-							bTransformChanged = true;
-						}
+							FTransform NewWorldTransform = OnGetRigElementTransformDelegate.Execute(SelectedRigElements[Index], false, true) * ComponentTransform;
+							bool bTransformChanged = false;
+							if (TransformContext.bRotation)
+							{
+								FQuat CurrentRotation = NewWorldTransform.GetRotation();
+								CurrentRotation = (TransformContext.Rot.Quaternion() * CurrentRotation);
+								NewWorldTransform.SetRotation(CurrentRotation);
+								bTransformChanged = true;
+							}
 
-						if (TransformContext.bScale)
-						{
-							FVector CurrentScale = NewWorldTransform.GetScale3D();
-							CurrentScale = CurrentScale + TransformContext.Scale;
-							NewWorldTransform.SetScale3D(CurrentScale);
-							bTransformChanged = true;
-						}
+							if (TransformContext.bTranslation)
+							{
+								FVector CurrentLocation = NewWorldTransform.GetLocation();
+								CurrentLocation = CurrentLocation + TransformContext.Drag;
+								NewWorldTransform.SetLocation(CurrentLocation);
+								bTransformChanged = true;
+							}
 
-						if (bTransformChanged)
-						{
-							if (bManipulatorMadeChange == false)
+							if (TransformContext.bScale)
+							{
+								FVector CurrentScale = NewWorldTransform.GetScale3D();
+								CurrentScale = CurrentScale + TransformContext.Scale;
+								NewWorldTransform.SetScale3D(CurrentScale);
+								bTransformChanged = true;
+							}
+
+							if (bTransformChanged)
 							{
-								GEditor->BeginTransaction(LOCTEXT("MoveControlTransaction", "Move Control"));
+								if (bManipulatorMadeChange == false)
+								{
+									GEditor->BeginTransaction(LOCTEXT("MoveControlTransaction", "Move Control"));
+								}
+								FTransform NewComponentTransform = NewWorldTransform.GetRelativeTransform(ComponentTransform);
+								OnSetRigElementTransformDelegate.Execute(SelectedRigElements[Index], NewComponentTransform, false);
+								bManipulatorMadeChange = true;
 							}
-							FTransform NewComponentTransform = NewWorldTransform.GetRelativeTransform(ComponentTransform);
-							OnSetRigElementTransformDelegate.Execute(SelectedRigElements[Index], NewComponentTransform, false);
-							bManipulatorMadeChange = true;
 						}
 					}
 				}
@@ -2416,7 +2435,7 @@ void FControlRigEditMode::SetRigElementSelection(UControlRig* ControlRig, ERigEl
 {
 	if (!bSelecting && ControlRig)
 	{
-		TMap<UControlRig*, TArray<FRigElementKey>> RigElementsToSelect;
+		TMap<TWeakObjectPtr<UControlRig>, TArray<FRigElementKey>> RigElementsToSelect;
 		TArray<FRigElementKey>& ElementsToSelect = RigElementsToSelect.FindOrAdd(ControlRig);
 		ElementsToSelect.Reserve(InRigElementNames.Num());
 		Algo::Transform(InRigElementNames, ElementsToSelect, [Type](const FName& ElementName)
@@ -2428,7 +2447,7 @@ void FControlRigEditMode::SetRigElementSelection(UControlRig* ControlRig, ERigEl
 	}
 }
 
-void FControlRigEditMode::SetRigElementsSelectionInternal(const TMap<UControlRig*, TArray<FRigElementKey>>& InRigElementsToSelect, const bool bSelected)
+void FControlRigEditMode::SetRigElementsSelectionInternal(const TMap<TWeakObjectPtr<UControlRig>, TArray<FRigElementKey>>& InRigElementsToSelect, const bool bSelected)
 {
 	if (bSelecting || InRigElementsToSelect.IsEmpty())
 	{
@@ -2450,16 +2469,19 @@ void FControlRigEditMode::SetRigElementsSelectionInternal(const TMap<UControlRig
 
 	TGuardValue<bool> ReentrantGuard(bSelecting, true);
 	
-	for (const auto& [ControlRig, Elements]: InRigElementsToSelect)
+	for (const auto& [WeakControlRig, Elements]: InRigElementsToSelect)
 	{
-		URigHierarchy* Hierarchy = ControlRig ? ControlRig->GetHierarchy() : nullptr;
-		if (URigHierarchyController* Controller = Hierarchy ? Hierarchy->GetController() : nullptr)
+		if (TStrongObjectPtr<UControlRig> ControlRig = WeakControlRig.Pin())
 		{
-			for (const FRigElementKey& Element : Elements)
+			URigHierarchy* Hierarchy = ControlRig ? ControlRig->GetHierarchy() : nullptr;
+			if (URigHierarchyController* Controller = Hierarchy ? Hierarchy->GetController() : nullptr)
 			{
-				constexpr bool bSetupUndo = true;
-				Controller->SelectElement(Element, bSelected, false, bSetupUndo);
-			}					
+				for (const FRigElementKey& Element : Elements)
+				{
+					constexpr bool bSetupUndo = true;
+					Controller->SelectElement(Element, bSelected, false, bSetupUndo);
+				}
+			}
 		}
 	}
 	
@@ -3110,19 +3132,22 @@ bool FControlRigEditMode::IsControlSelected(const bool bUseShapes) const
 	
 	if (bUseShapes)
 	{
-		for (const auto& [ControlRig, Shapes] : ControlRigShapeActors)
+		for (const auto& [WeakControlRig, Shapes] : ControlRigShapeActors)
 		{
-			for (const AControlRigShapeActor* ShapeActor : Shapes)
+			if (TStrongObjectPtr<UControlRig> ControlRig = WeakControlRig.Pin())
 			{
-				if (ShapeActor && ShapeActor->IsSelected())
+				for (const AControlRigShapeActor* ShapeActor : Shapes)
 				{
-					return true;
+					if (ShapeActor && ShapeActor->IsSelected())
+					{
+						return true;
+					}
 				}
-			}
 
-			if (ControlRig && AreRigElementsSelected(ControlType,ControlRig))
-			{
-				return true;
+				if (ControlRig && AreRigElementsSelected(ControlType,ControlRig.Get()))
+				{
+					return true;
+				}
 			}
 		}
 	}
@@ -6278,13 +6303,13 @@ void FControlRigEditMode::PostPoseUpdate(const FEditorViewportClient* InViewport
 	}
 
 	const bool bAreEditingControlRigDirectly = AreEditingControlRigDirectly();
-	auto UpdateShapes = [this, bAreEditingControlRigDirectly](const UControlRig* ControlRig, const TArray<TObjectPtr<AControlRigShapeActor>>& ShapeActors)
+	auto UpdateShapes = [this, bAreEditingControlRigDirectly](const TWeakObjectPtr<UControlRig> WeakControlRig, const TArray<TObjectPtr<AControlRigShapeActor>>& ShapeActors)
 	{
-		if(IsValid(ControlRig))
+		if(TStrongObjectPtr<UControlRig> ControlRig = WeakControlRig.Pin())
 		{
-			const FTransform ComponentTransform = bAreEditingControlRigDirectly ? FTransform::Identity : GetHostingSceneComponentTransform(ControlRig); 
+			const FTransform ComponentTransform = bAreEditingControlRigDirectly ? FTransform::Identity : GetHostingSceneComponentTransform(ControlRig.Get()); 
 
-			FShapeUpdateParams Params(ControlRig, ComponentTransform, IsControlRigSkelMeshVisible(ControlRig), IsInLevelEditor());
+			FShapeUpdateParams Params(ControlRig.Get(), ComponentTransform, IsControlRigSkelMeshVisible(ControlRig.Get()), IsInLevelEditor());
 			for (AControlRigShapeActor* ShapeActor : ShapeActors)
 			{
 				UpdateControlShape(ShapeActor, ControlRig->FindControl(ShapeActor->ControlName), Params);
@@ -6295,7 +6320,7 @@ void FControlRigEditMode::PostPoseUpdate(const FEditorViewportClient* InViewport
 	if (InRigs.IsEmpty())
 	{
 		// updates all control shapes properties
-		for (const TPair<UControlRig*,TArray<TObjectPtr<AControlRigShapeActor>>>& ShapeActors : ControlRigShapeActors)
+		for (const TPair<TWeakObjectPtr<UControlRig>,TArray<TObjectPtr<AControlRigShapeActor>>>& ShapeActors : ControlRigShapeActors)
 		{
 			UpdateShapes(ShapeActors.Key, ShapeActors.Value);
 		}
@@ -6304,11 +6329,11 @@ void FControlRigEditMode::PostPoseUpdate(const FEditorViewportClient* InViewport
 	{
 		for (const TWeakObjectPtr<UControlRig>& RigPtr: InRigs)
 		{
-			if (const UControlRig* ControlRig = RigPtr.Get())
+			if (TStrongObjectPtr<UControlRig> ControlRig = RigPtr.Pin())
 			{
-				if (const TArray<TObjectPtr<AControlRigShapeActor>>* ShapeActors = ControlRigShapeActors.Find(ControlRig))
+				if (const TArray<TObjectPtr<AControlRigShapeActor>>* ShapeActors = ControlRigShapeActors.Find(ControlRig.Get()))
 				{
-					UpdateShapes(ControlRig, *ShapeActors);
+					UpdateShapes(RigPtr, *ShapeActors);
 				}
 			}
 		}
@@ -6628,35 +6653,38 @@ void FControlRigEditMode::UpdateRotationContext()
 		return;
 	}
 	
-	for (const auto& [ControlRig, Shapes] : ControlRigShapeActors)
+	for (const auto& [WeakControlRig, Shapes] : ControlRigShapeActors)
 	{
-		if (const URigHierarchy* Hierarchy = ControlRig ? ControlRig->GetHierarchy() : nullptr)
+		if (TStrongObjectPtr<UControlRig> ControlRig = WeakControlRig.Pin())
 		{
-			if (Hierarchy->UsesPreferredEulerAngles())
+			if (const URigHierarchy* Hierarchy = ControlRig ? ControlRig->GetHierarchy() : nullptr)
 			{
-				for (const AControlRigShapeActor* ShapeActor : Shapes)
+				if (Hierarchy->UsesPreferredEulerAngles())
 				{
-					if (ShapeActor->IsEnabled() && ShapeActor->IsSelected())
+					for (const AControlRigShapeActor* ShapeActor : Shapes)
 					{
-						FRigControlElement* Control = ControlRig->FindControl(ShapeActor->ControlName);
-						const bool bUsePreferredRotationOrder = Hierarchy->GetUsePreferredRotationOrder(Control);
+						if (ShapeActor->IsEnabled() && ShapeActor->IsSelected())
+						{
+							FRigControlElement* Control = ControlRig->FindControl(ShapeActor->ControlName);
+							const bool bUsePreferredRotationOrder = Hierarchy->GetUsePreferredRotationOrder(Control);
 
-						RotationContext.RotationOrder = bUsePreferredRotationOrder ? Hierarchy->GetControlPreferredEulerRotationOrder(Control) : EEulerRotationOrder::XYZ;
-						RotationContext.Rotation = Hierarchy->GetControlPreferredRotator(Control);
+							RotationContext.RotationOrder = bUsePreferredRotationOrder ? Hierarchy->GetControlPreferredEulerRotationOrder(Control) : EEulerRotationOrder::XYZ;
+							RotationContext.Rotation = Hierarchy->GetControlPreferredRotator(Control);
 
-						if (TOptional<FTransform> ConstraintSpace = GetConstraintParentTransform(ControlRig, ShapeActor->ControlName))
-						{
-							RotationContext.Offset = *ConstraintSpace;
-						}
-						else
-						{
-							const FTransform Offset = Hierarchy->GetControlOffsetTransform(Control, ERigTransformType::CurrentGlobal);
-							const FTransform ComponentTransform = GetHostingSceneComponentTransform(ControlRig);
-							RotationContext.Offset = Offset * ComponentTransform;
-						}
+							if (TOptional<FTransform> ConstraintSpace = GetConstraintParentTransform(ControlRig.Get(), ShapeActor->ControlName))
+							{
+								RotationContext.Offset = *ConstraintSpace;
+							}
+							else
+							{
+								const FTransform Offset = Hierarchy->GetControlOffsetTransform(Control, ERigTransformType::CurrentGlobal);
+								const FTransform ComponentTransform = GetHostingSceneComponentTransform(ControlRig.Get());
+								RotationContext.Offset = Offset * ComponentTransform;
+							}
 
-						// only get the first rotation order
-						return;
+							// only get the first rotation order
+							return;
+						}
 					}
 				}
 			}
diff --git a/Engine/Plugins/Animation/ControlRig/Source/ControlRigEditor/Private/EditMode/ControlRigEditModeUtil.cpp b/Engine/Plugins/Animation/ControlRig/Source/ControlRigEditor/Private/EditMode/ControlRigEditModeUtil.cpp
--- a/Engine/Plugins/Animation/ControlRig/Source/ControlRigEditor/Private/EditMode/ControlRigEditModeUtil.cpp
+++ b/Engine/Plugins/Animation/ControlRig/Source/ControlRigEditor/Private/EditMode/ControlRigEditModeUtil.cpp
@@ -117,8 +117,8 @@ void FExplicitRotationInteraction::Apply(
 
 FSelectionHelper::FSelectionHelper(
 	FEditorViewportClient* InViewportClient,
-	const TMap<UControlRig*, TArray<TObjectPtr<AControlRigShapeActor>>>& InControlRigShapeActors,
-	TMap<UControlRig*, TArray<FRigElementKey>>& OutElements)
+	const TMap<TWeakObjectPtr<UControlRig>, TArray<TObjectPtr<AControlRigShapeActor>>>& InControlRigShapeActors,
+	TMap<TWeakObjectPtr<UControlRig>, TArray<FRigElementKey>>& OutElements)
 	: ViewportClient(InViewportClient)
 	, ControlRigShapeActors(InControlRigShapeActors)
 	, Elements(OutElements)
@@ -162,24 +162,27 @@ void FSelectionHelper::GetFromFrustum(const FConvexVolume& InFrustum) const
 
 	if (bTransparentBoxSelection)
 	{
-		for (const auto& [ControlRig, ShapeActors]: ControlRigShapeActors)
+		for (const auto& [WeakControlRig, ShapeActors]: ControlRigShapeActors)
 		{
-			if (ControlRig && ControlRig->GetControlsVisible())
+			if (TStrongObjectPtr<UControlRig> ControlRig = WeakControlRig.Pin())
 			{
-				for (AControlRigShapeActor* ShapeActor : ShapeActors)
+				if (ControlRig && ControlRig->GetControlsVisible())
 				{
-					const bool bTreatShape = ShapeActor && ShapeActor->IsSelectable() && !ShapeActor->IsTemporarilyHiddenInEditor();
-					if (bTreatShape)
+					for (AControlRigShapeActor* ShapeActor : ShapeActors)
 					{
-						for (UActorComponent* Component : ShapeActor->GetComponents())
+						const bool bTreatShape = ShapeActor && ShapeActor->IsSelectable() && !ShapeActor->IsTemporarilyHiddenInEditor();
+						if (bTreatShape)
 						{
-							UPrimitiveComponent* PrimitiveComponent = Cast<UPrimitiveComponent>(Component);
-							if (PrimitiveComponent && PrimitiveComponent->IsRegistered() && PrimitiveComponent->IsVisibleInEditor())
+							for (UActorComponent* Component : ShapeActor->GetComponents())
 							{
-								if (PrimitiveComponent->IsShown(ViewportClient->EngineShowFlags) && PrimitiveComponent->ComponentIsTouchingSelectionFrustum(InFrustum, false /*only bsp*/, false/*encompass entire*/))
+								UPrimitiveComponent* PrimitiveComponent = Cast<UPrimitiveComponent>(Component);
+								if (PrimitiveComponent && PrimitiveComponent->IsRegistered() && PrimitiveComponent->IsVisibleInEditor())
 								{
-									TArray<FRigElementKey>& Controls = Elements.FindOrAdd(ControlRig);
-									Controls.Add(ShapeActor->GetElementKey());
+									if (PrimitiveComponent->IsShown(ViewportClient->EngineShowFlags) && PrimitiveComponent->ComponentIsTouchingSelectionFrustum(InFrustum, false /*only bsp*/, false/*encompass entire*/))
+									{
+										TArray<FRigElementKey>& Controls = Elements.FindOrAdd(ControlRig.Get());
+										Controls.Add(ShapeActor->GetElementKey());
+									}
 								}
 							}
 						}
diff --git a/Engine/Plugins/Animation/ControlRig/Source/ControlRigEditor/Public/EditMode/ControlRigEditMode.h b/Engine/Plugins/Animation/ControlRig/Source/ControlRigEditor/Public/EditMode/ControlRigEditMode.h
--- a/Engine/Plugins/Animation/ControlRig/Source/ControlRigEditor/Public/EditMode/ControlRigEditMode.h
+++ b/Engine/Plugins/Animation/ControlRig/Source/ControlRigEditor/Public/EditMode/ControlRigEditMode.h
@@ -502,7 +502,7 @@ protected:
 	bool bSuspendHierarchyNotifs;
 
 	/** Shape actors */
-	TMap<UControlRig*,TArray<TObjectPtr<AControlRigShapeActor>>> ControlRigShapeActors;
+	TMap<TWeakObjectPtr<UControlRig>,TArray<TObjectPtr<AControlRigShapeActor>>> ControlRigShapeActors;
 
 	/** Manager for anim details proxies */
 	TObjectPtr<UAnimDetailsProxyManager> AnimDetailsProxyManager;
@@ -563,7 +563,7 @@ private:
 	void SetRigElementSelectionInternal(UControlRig* ControlRig, ERigElementType Type, const FName& InRigElementName, bool bSelected);
 
 	/** Set multiple RigElements' selection states. */
-	void SetRigElementsSelectionInternal(const TMap<UControlRig*, TArray<FRigElementKey>>& InRigElementsToSelect, const bool bSelected);
+	void SetRigElementsSelectionInternal(const TMap<TWeakObjectPtr<UControlRig>, TArray<FRigElementKey>>& InRigElementsToSelect, bool bSelected);
 
 	/** Whether or not Pivot Transforms have changed, in which case we need to redraw viewport. */
 	bool HasPivotTransformsChanged() const;
diff --git a/Engine/Plugins/Animation/ControlRig/Source/ControlRigEditor/Public/EditMode/ControlRigEditModeUtil.h b/Engine/Plugins/Animation/ControlRig/Source/ControlRigEditor/Public/EditMode/ControlRigEditModeUtil.h
--- a/Engine/Plugins/Animation/ControlRig/Source/ControlRigEditor/Public/EditMode/ControlRigEditModeUtil.h
+++ b/Engine/Plugins/Animation/ControlRig/Source/ControlRigEditor/Public/EditMode/ControlRigEditModeUtil.h
@@ -157,8 +157,8 @@ struct FSelectionHelper
 {
 	FSelectionHelper(
 		FEditorViewportClient* InViewportClient,
-		const TMap<UControlRig*, TArray<TObjectPtr<AControlRigShapeActor>>>& InControlRigShapeActors,
-		TMap<UControlRig*, TArray<FRigElementKey>>& OutElements);
+		const TMap<TWeakObjectPtr<UControlRig>, TArray<TObjectPtr<AControlRigShapeActor>>>& InControlRigShapeActors,
+		TMap<TWeakObjectPtr<UControlRig>, TArray<FRigElementKey>>& OutElements);
 
 	/**
 	 * Gets the elements contained in the frustum.
@@ -181,10 +181,10 @@ private:
 	FEditorViewportClient* ViewportClient = nullptr;
 
 	/** A reference to the edit mode's control shapes. */
-	const TMap<UControlRig*, TArray<TObjectPtr<AControlRigShapeActor>>>& ControlRigShapeActors;
+	const TMap<TWeakObjectPtr<UControlRig>, TArray<TObjectPtr<AControlRigShapeActor>>>& ControlRigShapeActors;
 
 	/** The elements to be selected. */
-	TMap<UControlRig*, TArray<FRigElementKey>>& Elements;
+	TMap<TWeakObjectPtr<UControlRig>, TArray<FRigElementKey>>& Elements;
 
 	/** Current viewport settings. (used for strict box & transparent selection) */
 	const ULevelEditorViewportSettings* LevelEditorViewportSettings = nullptr;
