Fix for ControlRig PoseAdapter triggering at check during  Undo / Redo, due to the code that sets a new hiearchy binding expecting a hierarchy to be already binded.
Updated the code to use the base class weak pointer, rather than storing a local raw pointer.

diff --git a/Engine/Plugins/Animation/ControlRig/Source/ControlRig/Private/Tools/ControlRigPoseAdapter.cpp b/Engine/Plugins/Animation/ControlRig/Source/ControlRig/Private/Tools/ControlRigPoseAdapter.cpp
--- a/Engine/Plugins/Animation/ControlRig/Source/ControlRig/Private/Tools/ControlRigPoseAdapter.cpp
+++ b/Engine/Plugins/Animation/ControlRig/Source/ControlRig/Private/Tools/ControlRigPoseAdapter.cpp
@@ -12,17 +12,15 @@ void FControlRigPoseAdapter::PostLinked(URigHierarchy* InHierarchy)
 
 	FRigHierarchyPoseAdapter::PostLinked(InHierarchy);
 
-	Hierarchy = InHierarchy;
-
 	// 1. sort the hierarchies storage so that we can make sure initial and local is grouped correctly
 	SortHierarchyStorage();
 
 	// 2. make sure to compute all local transforms in initial and local
-	const TArray<FRigTransformElement*> TransformElements = Hierarchy->GetElementsOfType<FRigTransformElement>(false);
+	const TArray<FRigTransformElement*> TransformElements = InHierarchy->GetElementsOfType<FRigTransformElement>(false);
 	for (FRigTransformElement* TransformElement : TransformElements)
 	{
-		(void)Hierarchy->GetTransform(TransformElement, ERigTransformType::InitialLocal);
-		(void)Hierarchy->GetTransform(TransformElement, ERigTransformType::CurrentLocal);
+		(void)InHierarchy->GetTransform(TransformElement, ERigTransformType::InitialLocal);
+		(void)InHierarchy->GetTransform(TransformElement, ERigTransformType::CurrentLocal);
 	}
 
 	// 3. set up a list of dependents in the hierarchy of the rig to be reset to initial on execute
@@ -47,7 +45,7 @@ void FControlRigPoseAdapter::PostLinked(URigHierarchy* InHierarchy)
 				return true;
 			}
 
-			const FRigBaseElementParentArray& ParentElements = InPoseAdapter->Hierarchy->GetParents(InTransformElement);
+			const FRigBaseElementParentArray& ParentElements = InPoseAdapter->GetHierarchy()->GetParents(InTransformElement);
 			for (const FRigBaseElement* ParentElement : ParentElements)
 			{
 				if (ProcessTransformElement(InPoseAdapter, InDependentTransformElements, Cast<FRigTransformElement>(const_cast<FRigBaseElement*>(ParentElement))))
@@ -104,7 +102,7 @@ void FControlRigPoseAdapter::PostLinked(URigHierarchy* InHierarchy)
 		const int32& TransformElementIndex = PoseIndexToElementIndex[PoseIndex];
 		if (TransformElementIndex != INDEX_NONE)
 		{
-			const FRigElementKeyAndIndex& KeyAndIndex = Hierarchy->GetKeyAndIndex(TransformElementIndex);
+			const FRigElementKeyAndIndex& KeyAndIndex = InHierarchy->GetKeyAndIndex(TransformElementIndex);
 			RelinkTransformStorage(KeyAndIndex, ERigTransformType::CurrentLocal, ERigTransformStorageType::Pose, &LocalPose[PoseIndex], &LocalPoseIsDirty[PoseIndex]);
 			RelinkTransformStorage(KeyAndIndex, ERigTransformType::CurrentGlobal, ERigTransformStorageType::Pose, &GlobalPose[PoseIndex], &GlobalPoseIsDirty[PoseIndex]);
 		}
@@ -122,7 +120,6 @@ void FControlRigPoseAdapter::PreUnlinked(URigHierarchy* InHierarchy)
 	UnlinkTransformStorage();
 	UpdateHierarchyStorage();
 
-	Hierarchy = nullptr;
 	FRigHierarchyPoseAdapter::PreUnlinked(InHierarchy);
 }
 
@@ -368,16 +365,17 @@ void FControlRigPoseAdapter::UnlinkTransformStorage()
 {
 	DECLARE_SCOPE_HIERARCHICAL_COUNTER_FUNC()
 
-	check(Hierarchy);
-
-	for (int32 PoseIndex = 0; PoseIndex < PoseIndexToElementIndex.Num(); PoseIndex++)
+	if (URigHierarchy* Hierarchy = GetHierarchy())
 	{
-		const int32& TransformElementIndex = PoseIndexToElementIndex[PoseIndex];
-		if (TransformElementIndex != INDEX_NONE)
+		for (int32 PoseIndex = 0; PoseIndex < PoseIndexToElementIndex.Num(); PoseIndex++)
 		{
-			const FRigElementKeyAndIndex& KeyAndIndex = Hierarchy->GetKeyAndIndex(TransformElementIndex);
-			RestoreTransformStorage(KeyAndIndex, ERigTransformType::CurrentLocal, ERigTransformStorageType::Pose, false);
-			RestoreTransformStorage(KeyAndIndex, ERigTransformType::CurrentGlobal, ERigTransformStorageType::Pose, false);
+			const int32& TransformElementIndex = PoseIndexToElementIndex[PoseIndex];
+			if (TransformElementIndex != INDEX_NONE)
+			{
+				const FRigElementKeyAndIndex& KeyAndIndex = Hierarchy->GetKeyAndIndex(TransformElementIndex);
+				RestoreTransformStorage(KeyAndIndex, ERigTransformType::CurrentLocal, ERigTransformStorageType::Pose, false);
+				RestoreTransformStorage(KeyAndIndex, ERigTransformType::CurrentGlobal, ERigTransformStorageType::Pose, false);
+			}
 		}
 	}
 
@@ -392,7 +390,7 @@ void FControlRigPoseAdapter::ConvertToLocalPose()
 	check(ParentPoseIndices.Num() == GlobalPose.Num());
 	LocalPose.SetNum(GlobalPose.Num());
 
-	if (Hierarchy)
+	if (URigHierarchy* Hierarchy = GetHierarchy())
 	{
 		for (int32 Index = 0; Index < ParentPoseIndices.Num(); Index++)
 		{
@@ -407,7 +405,7 @@ void FControlRigPoseAdapter::ConvertToGlobalPose()
 	check(ParentPoseIndices.Num() == LocalPose.Num());
 	GlobalPose.SetNum(LocalPose.Num());
 
-	if (Hierarchy)
+	if (URigHierarchy* Hierarchy = GetHierarchy())
 	{
 		for (int32 Index = 0; Index < ParentPoseIndices.Num(); Index++)
 		{
@@ -425,6 +423,7 @@ const FTransform& FControlRigPoseAdapter::GetLocalTransform(int32 InIndex)
 	check(ParentPoseIndices.IsValidIndex(InIndex));
 	check(RequiresHierarchyForSpaceConversion.IsValidIndex(InIndex));
 
+	URigHierarchy* Hierarchy = GetHierarchy();
 	check(Hierarchy);
 
 	if (LocalPoseIsDirty[InIndex])
@@ -461,6 +460,7 @@ const FTransform& FControlRigPoseAdapter::GetGlobalTransform(int32 InIndex)
 	check(ParentPoseIndices.IsValidIndex(InIndex));
 	check(RequiresHierarchyForSpaceConversion.IsValidIndex(InIndex));
 
+	URigHierarchy* Hierarchy = GetHierarchy();
 	check(Hierarchy);
 
 	if (GlobalPoseIsDirty[InIndex])
@@ -537,6 +537,7 @@ void FControlRigPoseAdapter::ComputeDependentTransforms()
 {
 	DECLARE_SCOPE_HIERARCHICAL_COUNTER_FUNC()
 
+	URigHierarchy* Hierarchy = GetHierarchy();
 	check(Hierarchy);
 
 	// ensure to compute all local transforms
@@ -579,7 +580,7 @@ void FControlRigPoseAdapter::MarkDependentsDirty()
 
 	ComputeDependentTransforms();
 
-	check(Hierarchy);
+	check(GetHierarchy() != nullptr);
 
 	// mark the global dependent as dirty
 	for (const FDependentTransform& Dependent : Dependents)
diff --git a/Engine/Plugins/Animation/ControlRig/Source/ControlRig/Public/Tools/ControlRigPoseAdapter.h b/Engine/Plugins/Animation/ControlRig/Source/ControlRig/Public/Tools/ControlRigPoseAdapter.h
--- a/Engine/Plugins/Animation/ControlRig/Source/ControlRig/Public/Tools/ControlRigPoseAdapter.h
+++ b/Engine/Plugins/Animation/ControlRig/Source/ControlRig/Public/Tools/ControlRigPoseAdapter.h
@@ -123,8 +123,6 @@ protected:
 
 	TArray<FDependentTransform> Dependents;
 
-	URigHierarchy* Hierarchy = nullptr;
-
 	bool bTransferInLocalSpace = true;
 	bool bRequiresResetPoseToInitial = true;
 	bool bEnabled = true;
