Animator Kit Plugin: autokey on release enabled by default

#rnx

diff --git a/Engine/Plugins/Animation/ControlRig/Source/ControlRigEditor/Public/Sequencer/AnimationAuthoringSettings.h b/Engine/Plugins/Animation/ControlRig/Source/ControlRigEditor/Public/Sequencer/AnimationAuthoringSettings.h
--- a/Engine/Plugins/Animation/ControlRig/Source/ControlRigEditor/Public/Sequencer/AnimationAuthoringSettings.h
+++ b/Engine/Plugins/Animation/ControlRig/Source/ControlRigEditor/Public/Sequencer/AnimationAuthoringSettings.h
@@ -16,7 +16,7 @@ public:
 
 	// Whether controls' keyframes should be added on release only.
 	UPROPERTY(config, EditAnywhere, Category = Interaction)
-	bool bAutoKeyOnRelease = false;
+	bool bAutoKeyOnRelease = true;
     
 	DECLARE_MULTICAST_DELEGATE_OneParam(FOnUpdateSettings,  const UAnimationAuthoringSettings*);
 	static FOnUpdateSettings OnSettingsChange;
