Constraints: added missing UseSelectedDelegate implementation

#rnx

diff --git a/Engine/Plugins/Animation/ControlRig/Source/ControlRigEditor/Private/Editor/Constraints/SConstraintsWidget.cpp b/Engine/Plugins/Animation/ControlRig/Source/ControlRigEditor/Private/Editor/Constraints/SConstraintsWidget.cpp
--- a/Engine/Plugins/Animation/ControlRig/Source/ControlRigEditor/Private/Editor/Constraints/SConstraintsWidget.cpp
+++ b/Engine/Plugins/Animation/ControlRig/Source/ControlRigEditor/Private/Editor/Constraints/SConstraintsWidget.cpp
@@ -1535,7 +1535,11 @@ TSharedRef<SWidget> SConstraintsEditionWidget::OnGetParentMenuContent()
 		GetActorChooserWidget()->SetIsOpen(false);
 	});
 
-	FSimpleDelegate UseSelectedDelegate = FSimpleDelegate::CreateLambda([]{ });
+	FSimpleDelegate UseSelectedDelegate = FSimpleDelegate::CreateLambda([this]
+	{
+		const TArray<AActor*> Actors = UE::Private::GetCurrentSelection();
+		UpdateParent(Actors.IsEmpty() ? nullptr : Actors[0]);
+	});
 	
 	return PropertyCustomizationHelpers::MakeActorPickerWithMenu(
 		ConstrainableParent.Actor.Get(),
