Control Rig: Fix crash when postconstruction is refreshing the curve container
#rnx

diff --git a/Engine/Plugins/Animation/ControlRig/Source/ControlRigEditor/Private/Editor/SRigCurveContainer.cpp b/Engine/Plugins/Animation/ControlRig/Source/ControlRigEditor/Private/Editor/SRigCurveContainer.cpp
--- a/Engine/Plugins/Animation/ControlRig/Source/ControlRigEditor/Private/Editor/SRigCurveContainer.cpp
+++ b/Engine/Plugins/Animation/ControlRig/Source/ControlRigEditor/Private/Editor/SRigCurveContainer.cpp
@@ -267,14 +267,14 @@ void SRigCurveContainer::BindCommands()
 
 void SRigCurveContainer::OnPreviewMeshChanged(class USkeletalMesh* OldPreviewMesh, class USkeletalMesh* NewPreviewMesh)
 {
-	RefreshCurveList();
+	RefreshCurveList_AnyThread();
 }
 
 void SRigCurveContainer::OnFilterTextChanged( const FText& SearchText )
 {
 	FilterText = SearchText;
 
-	RefreshCurveList();
+	RefreshCurveList_AnyThread();
 }
 
 void SRigCurveContainer::OnFilterTextCommitted( const FText& SearchText, ETextCommit::Type CommitInfo )
@@ -456,6 +456,21 @@ void SRigCurveContainer::CreateRigCurveList( const FString& SearchText )
 
 }
 
+void SRigCurveContainer::RefreshCurveList_AnyThread()
+{
+	if(IsInGameThread())
+	{
+		RefreshCurveList();
+	}
+	else
+	{
+		FFunctionGraphTask::CreateAndDispatchWhenReady([this]()
+		{
+			RefreshCurveList();
+		}, TStatId(), nullptr, ENamedThreads::GameThread);
+	}
+}
+
 void SRigCurveContainer::RefreshCurveList()
 {
 	CreateRigCurveList(FilterText.ToString());
@@ -675,7 +690,7 @@ void SRigCurveContainer::OnHierarchyModified(ERigHierarchyNotification InNotif,
 		case ERigHierarchyNotification::ElementReordered:
 		case ERigHierarchyNotification::HierarchyReset:
 		{
-			RefreshCurveList();
+			RefreshCurveList_AnyThread();
 			break;
 		}
 		case ERigHierarchyNotification::ElementSelected:
@@ -704,7 +719,7 @@ void SRigCurveContainer::OnHierarchyModified(ERigHierarchyNotification InNotif,
 
 void SRigCurveContainer::OnPostConstruction(UControlRig* InControlRig, const FName& InName)
 {
-	RefreshCurveList();
+	RefreshCurveList_AnyThread();
 }
 
 void SRigCurveContainer::HandleSetObjectBeingDebugged(UObject* InObject)
@@ -739,12 +754,12 @@ void SRigCurveContainer::HandleSetObjectBeingDebugged(UObject* InObject)
 		ControlRig->OnPostConstruction_AnyThread().AddSP(this, &SRigCurveContainer::OnPostConstruction);
 	}
 
-	RefreshCurveList();
+	RefreshCurveList_AnyThread();
 }
 
 void SRigCurveContainer::HandleRefreshEditorFromBlueprint(URigVMBlueprint* InBlueprint)
 {
-	RefreshCurveList();
+	RefreshCurveList_AnyThread();
 }
 
 void SRigCurveContainer::CreateImportMenu(FMenuBuilder& MenuBuilder)
diff --git a/Engine/Plugins/Animation/ControlRig/Source/ControlRigEditor/Private/Editor/SRigCurveContainer.h b/Engine/Plugins/Animation/ControlRig/Source/ControlRigEditor/Private/Editor/SRigCurveContainer.h
--- a/Engine/Plugins/Animation/ControlRig/Source/ControlRigEditor/Private/Editor/SRigCurveContainer.h
+++ b/Engine/Plugins/Animation/ControlRig/Source/ControlRigEditor/Private/Editor/SRigCurveContainer.h
@@ -198,6 +198,7 @@ public:
 	*/
 	FText& GetFilterText() { return FilterText; }
 
+	void RefreshCurveList_AnyThread();
 	void RefreshCurveList();
 
 	// When a name is committed after being edited in the list
