Control Rig: Fix changes in controls from the details panel not being restored after construction
#rnx

diff --git a/Engine/Plugins/Animation/ControlRig/Source/ControlRigEditor/Private/ControlRigElementDetails.h b/Engine/Plugins/Animation/ControlRig/Source/ControlRigEditor/Private/ControlRigElementDetails.h
--- a/Engine/Plugins/Animation/ControlRig/Source/ControlRigEditor/Private/ControlRigElementDetails.h
+++ b/Engine/Plugins/Animation/ControlRig/Source/ControlRigEditor/Private/ControlRigElementDetails.h
@@ -26,6 +26,7 @@
 #include "ScopedTransaction.h"
 #include "Styling/AppStyle.h"
 #include "Algo/Transform.h"
+#include "Editor/ControlRigEditor.h"
 #include "Editor/SRigHierarchyTreeView.h"
 
 #include "ControlRigElementDetails.generated.h"
@@ -548,6 +549,19 @@ private:
 			HierarchyToChange = Hierarchy;
 		}
 
+		FControlRigEditorEditMode* EditMode = nullptr;
+		if (UControlRigBlueprint* Blueprint = PerElementInfos[0].GetBlueprint())
+		{
+			UAssetEditorSubsystem* AssetEditorSubsystem = GEditor->GetEditorSubsystem<UAssetEditorSubsystem>();
+			if (IAssetEditorInstance* EditorInstance = AssetEditorSubsystem->FindEditorForAsset(Blueprint, false))
+			{
+				if (IControlRigBaseEditor* Editor = FControlRigBaseEditor::GetFromAssetEditorInstance(EditorInstance))
+				{
+					EditMode = Editor->GetEditMode();
+				}
+			}
+		}
+
 		TSharedPtr<SNumericEntryBox<T>> NumericEntryBox;
 		
 		FDetailWidgetRow& WidgetRow = CategoryBuilder.AddCustomRow(Label);
@@ -615,7 +629,7 @@ private:
 				});
 		}
 
-		auto OnValueChanged = [ValueType, Keys, HierarchyToChange, this]
+		auto OnValueChanged = [ValueType, Keys, HierarchyToChange, EditMode, this]
 			(TOptional<T> InValue, ETextCommit::Type InCommitType, bool bSetupUndo)
 			{
 				if(!InValue.IsSet())
@@ -636,6 +650,10 @@ private:
 							HierarchyToChange->Modify();
 						}
 						HierarchyToChange->SetControlValue(Key, FRigControlValue::Make<T>(Value), ValueType, bSetupUndo, bSetupUndo);
+						if (EditMode)
+						{
+							EditMode->ModifiedRigElements.AddUnique(Key);
+						}
 					}
 				}
 
