Fix link errors when compiling EngineTest with verse-clang.

diff --git a/Engine/Plugins/Animation/ControlRig/Source/ControlRig/Public/Tracks/SampleTrack.h b/Engine/Plugins/Animation/ControlRig/Source/ControlRig/Public/Tracks/SampleTrack.h
--- a/Engine/Plugins/Animation/ControlRig/Source/ControlRig/Public/Tracks/SampleTrack.h
+++ b/Engine/Plugins/Animation/ControlRig/Source/ControlRig/Public/Tracks/SampleTrack.h
@@ -187,7 +187,7 @@ public:
 	}
 
 	// helper to store a sample based on a given property
-	virtual void AddSampleFromProperty(const FProperty* InProperty, const uint8* InMemory) override
+	CONTROLRIG_API virtual void AddSampleFromProperty(const FProperty* InProperty, const uint8* InMemory) override
 	{
 		T Value;
 		InProperty->CopyCompleteValue(&Value, InMemory);
@@ -195,7 +195,7 @@ public:
 	}
 
 	// helper to retrieve a sample based on a given property
-	virtual void GetSampleForProperty(int32 InTimeIndex, FSampleTrackIndex& InOutSampleTrackIndex, const FProperty* InProperty, uint8* OutMemory) const override
+	CONTROLRIG_API virtual void GetSampleForProperty(int32 InTimeIndex, FSampleTrackIndex& InOutSampleTrackIndex, const FProperty* InProperty, uint8* OutMemory) const override
 	{
 		const T& Value = GetValueAtTimeIndex(InTimeIndex, InOutSampleTrackIndex);
 		InProperty->CopyCompleteValue(OutMemory, &Value);
@@ -262,7 +262,7 @@ public:
 		return true;
 	}
 
-	virtual void Serialize(FArchive& InArchive) override
+	CONTROLRIG_API virtual void Serialize(FArchive& InArchive) override
 	{
 		FSampleTrackBase::Serialize(InArchive);
 		InArchive << ValuesStorage;
