Anim In Engine: Keep track of open anim layer ui on the anim layer asset data, defaults to show them, but if the user closes it, it will stay closed on re-open.

diff --git a/Engine/Plugins/Animation/ControlRig/Source/ControlRigEditor/Private/EditMode/ControlRigEditModeToolkit.cpp b/Engine/Plugins/Animation/ControlRig/Source/ControlRigEditor/Private/EditMode/ControlRigEditModeToolkit.cpp
--- a/Engine/Plugins/Animation/ControlRig/Source/ControlRigEditor/Private/EditMode/ControlRigEditModeToolkit.cpp
+++ b/Engine/Plugins/Animation/ControlRig/Source/ControlRigEditor/Private/EditMode/ControlRigEditModeToolkit.cpp
@@ -432,9 +432,12 @@ void FControlRigEditModeToolkit::InvokeUI()
 		else if (TSharedPtr<ISequencer> SequencerPtr = UAnimLayers::GetSequencerFromAsset())
 		{
 			//wasn't open but there are layers in open level sequence we try to open it
-			if (UAnimLayers::HasAnimLayers(SequencerPtr.Get()))
+			if (UAnimLayers* AnimLayers = UAnimLayers::GetAnimLayers(SequencerPtr.Get()))
 			{
-				TryInvokeToolkitUI(AnimLayerTabName);
+				if (AnimLayers && AnimLayers->bOpenUIOnOpen)
+				{
+					TryInvokeToolkitUI(AnimLayerTabName);
+				}
 			}
 		}
 		if (bSnapperTabOpen)
@@ -506,6 +509,13 @@ void FControlRigEditModeToolkit::UnregisterAndRemoveFloatingTabs()
 			{
 				bAnimLayerTabOpen = false;
 			}
+			if (TSharedPtr<ISequencer> SequencerPtr = UAnimLayers::GetSequencerFromAsset())
+			{
+				if (UAnimLayers* AnimLayers = UAnimLayers::GetAnimLayers(SequencerPtr.Get()))
+				{
+					AnimLayers->bOpenUIOnOpen = bAnimLayerTabOpen;
+				}
+			}
 			ModeUILayerPtr->GetTabManager()->UnregisterTabSpawner(AnimLayerTabName);
 
 			TSharedPtr<SDockTab> SnapperTab = ModeUILayerPtr->GetTabManager()->FindExistingLiveTab(FTabId(SnapperTabName));
diff --git a/Engine/Plugins/Animation/ControlRig/Source/ControlRigEditor/Private/Sequencer/AnimLayers/AnimLayers.h b/Engine/Plugins/Animation/ControlRig/Source/ControlRigEditor/Private/Sequencer/AnimLayers/AnimLayers.h
--- a/Engine/Plugins/Animation/ControlRig/Source/ControlRigEditor/Private/Sequencer/AnimLayers/AnimLayers.h
+++ b/Engine/Plugins/Animation/ControlRig/Source/ControlRigEditor/Private/Sequencer/AnimLayers/AnimLayers.h
@@ -326,6 +326,9 @@ public:
 	UPROPERTY(BlueprintReadWrite, Category = Layers)
 	TArray<TObjectPtr<UAnimLayer>> AnimLayers;
 
+	UPROPERTY()
+	bool bOpenUIOnOpen = true;
+
 	int32 GetAnimLayerIndex(UAnimLayer* InAnimLayer) const;
 	bool DeleteAnimLayer(ISequencer* InSequencer, int32 Index);
 	int32 DuplicateAnimLayer(ISequencer* InSequencer, int32 Index);
