Control Rig: Splitting up FSampleTrack to ensure we can DLL export type-specialized member functions.

#rnx

diff --git a/Engine/Plugins/Animation/ControlRig/Source/ControlRig/Private/Tracks/SampleTrack.cpp b/Engine/Plugins/Animation/ControlRig/Source/ControlRig/Private/Tracks/SampleTrack.cpp
--- a/Engine/Plugins/Animation/ControlRig/Source/ControlRig/Private/Tracks/SampleTrack.cpp
+++ b/Engine/Plugins/Animation/ControlRig/Source/ControlRig/Private/Tracks/SampleTrack.cpp
@@ -6,8 +6,7 @@
 /////// FSampleTrack
 //////////////////////////////////////////////////////////////////////////////////////////////////////////////////////////////
 
-template<>
-CONTROLRIG_API void FSampleTrack<FInstancedStruct>::Serialize(FArchive& InArchive)
+void FSampleTrack<FInstancedStruct>::Serialize(FArchive& InArchive)
 {
 	FSampleTrackBase::Serialize(InArchive);
 	
@@ -31,8 +30,36 @@ CONTROLRIG_API void FSampleTrack<FInstancedStruct>::Serialize(FArchive& InArchiv
 	}
 }
 
-template<>
-CONTROLRIG_API void FSampleTrack<TArray<FInstancedStruct>>::Serialize(FArchive& InArchive)
+int32 FSampleTrack<FInstancedStruct>::GetSizePerValue() const
+{
+	if(ScriptStruct)
+	{
+		return sizeof(FInstancedStruct) + ScriptStruct->GetStructureSize();
+	}
+	return sizeof(FInstancedStruct);
+}
+
+void FSampleTrack<FInstancedStruct>::AddSampleFromProperty(const FProperty* InProperty, const uint8* InMemory)
+{
+	const FStructProperty* StructProperty = CastFieldChecked<FStructProperty>(InProperty);
+	check(StructProperty->Struct == GetScriptStruct());
+
+	FInstancedStruct Value(GetScriptStruct());
+	StructProperty->CopyCompleteValue(Value.GetMutableMemory(), InMemory);
+	AddSample(Value);
+}
+
+void FSampleTrack<FInstancedStruct>::GetSampleForProperty(int32 InTimeIndex, FSampleTrackIndex& InOutSampleTrackIndex, const FProperty* InProperty, uint8* OutMemory) const
+{
+	const FStructProperty* StructProperty = CastFieldChecked<FStructProperty>(InProperty);
+	check(StructProperty->Struct == GetScriptStruct());
+
+	const FInstancedStruct& Value = GetValueAtTimeIndex(InTimeIndex, InOutSampleTrackIndex);
+	StructProperty->CopyCompleteValue(OutMemory, Value.GetMemory());
+}
+
+
+void FSampleTrack<TArray<FInstancedStruct>>::Serialize(FArchive& InArchive)
 {
 	FSampleTrackBase::Serialize(InArchive);
 	
@@ -68,39 +95,7 @@ CONTROLRIG_API void FSampleTrack<TArray<FInstancedStruct>>::Serialize(FArchive&
 	}
 }
 
-template<>
-CONTROLRIG_API int32 FSampleTrack<FInstancedStruct>::GetSizePerValue() const
-{
-	if(ScriptStruct)
-	{
-		return sizeof(FInstancedStruct) + ScriptStruct->GetStructureSize();
-	}
-	return sizeof(FInstancedStruct);
-}
-
-template<>
-CONTROLRIG_API void FSampleTrack<FInstancedStruct>::AddSampleFromProperty(const FProperty* InProperty, const uint8* InMemory)
-{
-	const FStructProperty* StructProperty = CastFieldChecked<FStructProperty>(InProperty);
-	check(StructProperty->Struct == GetScriptStruct());
-
-	FInstancedStruct Value(GetScriptStruct());
-	StructProperty->CopyCompleteValue(Value.GetMutableMemory(), InMemory);
-	AddSample(Value);
-}
-
-template<>
-CONTROLRIG_API void FSampleTrack<FInstancedStruct>::GetSampleForProperty(int32 InTimeIndex, FSampleTrackIndex& InOutSampleTrackIndex, const FProperty* InProperty, uint8* OutMemory) const
-{
-	const FStructProperty* StructProperty = CastFieldChecked<FStructProperty>(InProperty);
-	check(StructProperty->Struct == GetScriptStruct());
-
-	const FInstancedStruct& Value = GetValueAtTimeIndex(InTimeIndex, InOutSampleTrackIndex);
-	StructProperty->CopyCompleteValue(OutMemory, Value.GetMemory());
-}
-
-template<>
-CONTROLRIG_API void FSampleTrack<TArray<FInstancedStruct>>::AddSampleFromProperty(const FProperty* InProperty, const uint8* InMemory)
+void FSampleTrack<TArray<FInstancedStruct>>::AddSampleFromProperty(const FProperty* InProperty, const uint8* InMemory)
 {
 	const FArrayProperty* ArrayProperty = CastFieldChecked<FArrayProperty>(InProperty);
 	const FStructProperty* StructProperty = CastFieldChecked<FStructProperty>(ArrayProperty->Inner);
@@ -118,8 +113,7 @@ CONTROLRIG_API void FSampleTrack<TArray<FInstancedStruct>>::AddSampleFromPropert
 	AddSample(NewValues);
 }
 
-template<>
-CONTROLRIG_API void FSampleTrack<TArray<FInstancedStruct>>::GetSampleForProperty(int32 InTimeIndex, FSampleTrackIndex& InOutSampleTrackIndex, const FProperty* InProperty, uint8* OutMemory) const
+void FSampleTrack<TArray<FInstancedStruct>>::GetSampleForProperty(int32 InTimeIndex, FSampleTrackIndex& InOutSampleTrackIndex, const FProperty* InProperty, uint8* OutMemory) const
 {
 	const FArrayProperty* ArrayProperty = CastFieldChecked<FArrayProperty>(InProperty);
 	const FStructProperty* StructProperty = CastFieldChecked<FStructProperty>(ArrayProperty->Inner);
@@ -136,3 +130,6 @@ CONTROLRIG_API void FSampleTrack<TArray<FInstancedStruct>>::GetSampleForProperty
 		StructProperty->CopyCompleteValue(ArrayHelper.GetElementPtr(Index), Value.GetMemory());
 	}
 }
+
+template class CONTROLRIG_API FSampleTrack<FInstancedStruct>;
+template class CONTROLRIG_API FSampleTrack<TArray<FInstancedStruct>>;
diff --git a/Engine/Plugins/Animation/ControlRig/Source/ControlRig/Public/Tracks/SampleTrack.h b/Engine/Plugins/Animation/ControlRig/Source/ControlRig/Public/Tracks/SampleTrack.h
--- a/Engine/Plugins/Animation/ControlRig/Source/ControlRig/Public/Tracks/SampleTrack.h
+++ b/Engine/Plugins/Animation/ControlRig/Source/ControlRig/Public/Tracks/SampleTrack.h
@@ -6,21 +6,24 @@
 #include "Math/UnrealMath.h"
 #include "Math/UnrealMathUtility.h"
 
+namespace UE::ControlRig::Internal
+{
+
 template<typename T>
-class FSampleTrack : public FSampleTrackBase
+class FSampleTrackShared : public FSampleTrackBase
 {
 public:
 
-	FSampleTrack()
+	FSampleTrackShared()
 	{
 	}
 	
-	FSampleTrack(const FName& InName)
+	FSampleTrackShared(const FName& InName)
 	: FSampleTrackBase(InName)
 	{
 	}
 
-	virtual ~FSampleTrack() override {}
+	virtual ~FSampleTrackShared() override {}
 
 	// the number of stored values - which can be less than NumSamples if there is an atlas present
 	virtual int32 NumStoredValues() const override
@@ -186,21 +189,6 @@ public:
 		AddSample(InstancedStructArray, InTolerance);
 	}
 
-	// helper to store a sample based on a given property
-	CONTROLRIG_API virtual void AddSampleFromProperty(const FProperty* InProperty, const uint8* InMemory) override
-	{
-		T Value;
-		InProperty->CopyCompleteValue(&Value, InMemory);
-		AddSample(Value);
-	}
-
-	// helper to retrieve a sample based on a given property
-	CONTROLRIG_API virtual void GetSampleForProperty(int32 InTimeIndex, FSampleTrackIndex& InOutSampleTrackIndex, const FProperty* InProperty, uint8* OutMemory) const override
-	{
-		const T& Value = GetValueAtTimeIndex(InTimeIndex, InOutSampleTrackIndex);
-		InProperty->CopyCompleteValue(OutMemory, &Value);
-	}
-
 	// returns the type of this track
 	virtual ETrackType GetTrackType() const override
 	{
@@ -245,7 +233,7 @@ public:
 			return false;
 		}
 
-		const FSampleTrack<T>* Other = static_cast<const FSampleTrack<T>*>(InOther);
+		const FSampleTrackShared<T>* Other = static_cast<const FSampleTrackShared<T>*>(InOther);
 		if(NumSamples() != Other->NumSamples())
 		{
 			return false;
@@ -262,21 +250,14 @@ public:
 		return true;
 	}
 
-	CONTROLRIG_API virtual void Serialize(FArchive& InArchive) override
+	const FSampleTrackShared* GetCastReferencedTrack() const
 	{
-		FSampleTrackBase::Serialize(InArchive);
-		InArchive << ValuesStorage;
-		UpdateValueArrayView();
-	}
-
-	const FSampleTrack* GetCastReferencedTrack() const
-	{
-		return static_cast<const FSampleTrack*>(GetReferencedTrack());
+		return static_cast<const FSampleTrackShared*>(GetReferencedTrack());
 	}
 	
-	FSampleTrack* GetCastReferencedTrack()
+	FSampleTrackShared* GetCastReferencedTrack()
 	{
-		return static_cast<FSampleTrack*>(GetReferencedTrack());
+		return static_cast<FSampleTrackShared*>(GetReferencedTrack());
 	}
 
 protected:
@@ -343,17 +324,13 @@ protected:
 	}
 
 public:
+	virtual int32 GetSizePerValue() const = 0;
 	
 	static bool IsArrayValue()
 	{
 		return TIsArray<T>::Value;
 	}
 
-	int32 GetSizePerValue() const
-	{
-		return sizeof(T);
-	}
-
 	/**
 	 * Analyses the values stored and introduces the atlas index array,
 	 * an indirection to save memory. After this Values is going to store unique elements,
@@ -597,7 +574,7 @@ protected:
 	virtual int32 AppendValuesFromTrack(const FSampleTrackBase* InOtherTrack) override
 	{
 		const int32 FirstValueIndex = ValuesStorage.Num(); 
-		ValuesStorage.Append(static_cast<const FSampleTrack*>(InOtherTrack)->ValuesStorage);
+		ValuesStorage.Append(static_cast<const FSampleTrackShared*>(InOtherTrack)->ValuesStorage);
 		UpdateValueArrayView();
 		return FirstValueIndex;
 	}
@@ -612,7 +589,7 @@ protected:
 	{
 		if(IsReferenced())
 		{
-			FSampleTrack* ReferencedTrack = GetCastReferencedTrack();
+			FSampleTrackShared* ReferencedTrack = GetCastReferencedTrack();
 			if(ReferencedTrack == nullptr)
 			{
 				Values = TArrayView<T>();
@@ -647,178 +624,241 @@ protected:
 };
 
 template<>
-inline FSampleTrackBase::ETrackType FSampleTrack<bool>::GetTrackType() const
+inline FSampleTrackBase::ETrackType FSampleTrackShared<bool>::GetTrackType() const
 {
 	return ETrackType_Bool;
 }
 
 template<>
-inline FSampleTrackBase::ETrackType FSampleTrack<int32>::GetTrackType() const
+inline FSampleTrackBase::ETrackType FSampleTrackShared<int32>::GetTrackType() const
 {
 	return ETrackType_Int32;
 }
 
 template<>
-inline FSampleTrackBase::ETrackType FSampleTrack<uint32>::GetTrackType() const
+inline FSampleTrackBase::ETrackType FSampleTrackShared<uint32>::GetTrackType() const
 {
 	return ETrackType_Uint32;
 }
 
 template<>
-inline FSampleTrackBase::ETrackType FSampleTrack<float>::GetTrackType() const
+inline FSampleTrackBase::ETrackType FSampleTrackShared<float>::GetTrackType() const
 {
 	return ETrackType_Float;
 }
 
 template<>
-inline FSampleTrackBase::ETrackType FSampleTrack<FName>::GetTrackType() const
+inline FSampleTrackBase::ETrackType FSampleTrackShared<FName>::GetTrackType() const
 {
 	return ETrackType_Name;
 }
 
 template<>
-inline FSampleTrackBase::ETrackType FSampleTrack<FString>::GetTrackType() const
+inline FSampleTrackBase::ETrackType FSampleTrackShared<FString>::GetTrackType() const
 {
 	return ETrackType_String;
 }
 
 template<>
-inline FSampleTrackBase::ETrackType FSampleTrack<FVector3f>::GetTrackType() const
+inline FSampleTrackBase::ETrackType FSampleTrackShared<FVector3f>::GetTrackType() const
 {
 	return ETrackType_Vector3f;
 }
 
 template<>
-inline FSampleTrackBase::ETrackType FSampleTrack<FQuat4f>::GetTrackType() const
+inline FSampleTrackBase::ETrackType FSampleTrackShared<FQuat4f>::GetTrackType() const
 {
 	return ETrackType_Quatf;
 }
 
 template<>
-inline FSampleTrackBase::ETrackType FSampleTrack<FTransform3f>::GetTrackType() const
+inline FSampleTrackBase::ETrackType FSampleTrackShared<FTransform3f>::GetTrackType() const
 {
 	return ETrackType_Transformf;
 }
 
 template<>
-inline FSampleTrackBase::ETrackType FSampleTrack<FLinearColor>::GetTrackType() const
+inline FSampleTrackBase::ETrackType FSampleTrackShared<FLinearColor>::GetTrackType() const
 {
 	return ETrackType_LinearColor;
 }
 
 template<>
-inline FSampleTrackBase::ETrackType FSampleTrack<FRigElementKey>::GetTrackType() const
+inline FSampleTrackBase::ETrackType FSampleTrackShared<FRigElementKey>::GetTrackType() const
 {
 	return ETrackType_ElementKey;
 }
 
 template<>
-inline FSampleTrackBase::ETrackType FSampleTrack<FRigComponentKey>::GetTrackType() const
+inline FSampleTrackBase::ETrackType FSampleTrackShared<FRigComponentKey>::GetTrackType() const
 {
 	return ETrackType_ComponentKey;
 }
 
 template<>
-inline FSampleTrackBase::ETrackType FSampleTrack<FInstancedStruct>::GetTrackType() const
+inline FSampleTrackBase::ETrackType FSampleTrackShared<FInstancedStruct>::GetTrackType() const
 {
 	return ETrackType_Struct;
 }
 
 template<>
-inline FSampleTrackBase::ETrackType FSampleTrack<TArray<bool>>::GetTrackType() const
+inline FSampleTrackBase::ETrackType FSampleTrackShared<TArray<bool>>::GetTrackType() const
 {
 	return ETrackType_BoolArray;
 }
 
 template<>
-inline FSampleTrackBase::ETrackType FSampleTrack<TArray<int32>>::GetTrackType() const
+inline FSampleTrackBase::ETrackType FSampleTrackShared<TArray<int32>>::GetTrackType() const
 {
 	return ETrackType_Int32Array;
 }
 
 template<>
-inline FSampleTrackBase::ETrackType FSampleTrack<TArray<uint32>>::GetTrackType() const
+inline FSampleTrackBase::ETrackType FSampleTrackShared<TArray<uint32>>::GetTrackType() const
 {
 	return ETrackType_Uint32Array;
 }
 
 template<>
-inline FSampleTrackBase::ETrackType FSampleTrack<TArray<float>>::GetTrackType() const
+inline FSampleTrackBase::ETrackType FSampleTrackShared<TArray<float>>::GetTrackType() const
 {
 	return ETrackType_FloatArray;
 }
 
 template<>
-inline FSampleTrackBase::ETrackType FSampleTrack<TArray<FName>>::GetTrackType() const
+inline FSampleTrackBase::ETrackType FSampleTrackShared<TArray<FName>>::GetTrackType() const
 {
 	return ETrackType_NameArray;
 }
 
 template<>
-inline FSampleTrackBase::ETrackType FSampleTrack<TArray<FString>>::GetTrackType() const
+inline FSampleTrackBase::ETrackType FSampleTrackShared<TArray<FString>>::GetTrackType() const
 {
 	return ETrackType_StringArray;
 }
 
 template<>
-inline FSampleTrackBase::ETrackType FSampleTrack<TArray<FVector3f>>::GetTrackType() const
+inline FSampleTrackBase::ETrackType FSampleTrackShared<TArray<FVector3f>>::GetTrackType() const
 {
 	return ETrackType_Vector3fArray;
 }
 
 template<>
-inline FSampleTrackBase::ETrackType FSampleTrack<TArray<FQuat4f>>::GetTrackType() const
+inline FSampleTrackBase::ETrackType FSampleTrackShared<TArray<FQuat4f>>::GetTrackType() const
 {
 	return ETrackType_QuatfArray;
 }
 
 template<>
-inline FSampleTrackBase::ETrackType FSampleTrack<TArray<FTransform3f>>::GetTrackType() const
+inline FSampleTrackBase::ETrackType FSampleTrackShared<TArray<FTransform3f>>::GetTrackType() const
 {
 	return ETrackType_TransformfArray;
 }
 
 template<>
-inline FSampleTrackBase::ETrackType FSampleTrack<TArray<FLinearColor>>::GetTrackType() const
+inline FSampleTrackBase::ETrackType FSampleTrackShared<TArray<FLinearColor>>::GetTrackType() const
 {
 	return ETrackType_LinearColorArray;
 }
 
 template<>
-inline FSampleTrackBase::ETrackType FSampleTrack<TArray<FRigElementKey>>::GetTrackType() const
+inline FSampleTrackBase::ETrackType FSampleTrackShared<TArray<FRigElementKey>>::GetTrackType() const
 {
 	return ETrackType_ElementKeyArray;
 }
 
 template<>
-inline FSampleTrackBase::ETrackType FSampleTrack<TArray<FRigComponentKey>>::GetTrackType() const
+inline FSampleTrackBase::ETrackType FSampleTrackShared<TArray<FRigComponentKey>>::GetTrackType() const
 {
 	return ETrackType_ComponentKeyArray;
 }
 
 template<>
-inline FSampleTrackBase::ETrackType FSampleTrack<TArray<FInstancedStruct>>::GetTrackType() const
+inline FSampleTrackBase::ETrackType FSampleTrackShared<TArray<FInstancedStruct>>::GetTrackType() const
 {
 	return ETrackType_StructArray;
 }
+}
 
-template<>
-CONTROLRIG_API void FSampleTrack<FInstancedStruct>::Serialize(FArchive& InArchive);
 
-template<>
-CONTROLRIG_API void FSampleTrack<TArray<FInstancedStruct>>::Serialize(FArchive& InArchive);
+template<typename T>
+class FSampleTrack : public UE::ControlRig::Internal::FSampleTrackShared<T>
+{
+	using Base = UE::ControlRig::Internal::FSampleTrackShared<T>;
+public:
+	FSampleTrack() = default;
+	FSampleTrack(const FName& InName) : Base(InName) {}
+	virtual ~FSampleTrack() override {}
+	
+	virtual void Serialize(FArchive& InArchive) override
+	{
+		FSampleTrackBase::Serialize(InArchive);
+		InArchive << Base::ValuesStorage;
+		Base::UpdateValueArrayView();
+	}
 
-template<>
-CONTROLRIG_API int32 FSampleTrack<FInstancedStruct>::GetSizePerValue() const;
+	virtual int32 GetSizePerValue() const override
+	{
+		return sizeof(T);
+	}
 
-template<>
-CONTROLRIG_API void FSampleTrack<FInstancedStruct>::AddSampleFromProperty(const FProperty* InProperty, const uint8* InMemory);
+	// helper to store a sample based on a given property
+	virtual void AddSampleFromProperty(const FProperty* InProperty, const uint8* InMemory) override
+	{
+		T Value;
+		InProperty->CopyCompleteValue(&Value, InMemory);
+		Base::AddSample(Value);
+	}
 
-template<>
-CONTROLRIG_API void FSampleTrack<FInstancedStruct>::GetSampleForProperty(int32 InTimeIndex, FSampleTrackIndex& InOutSampleTrackIndex, const FProperty* InProperty, uint8* OutMemory) const;
+	// helper to retrieve a sample based on a given property
+	virtual void GetSampleForProperty(int32 InTimeIndex, FSampleTrackIndex& InOutSampleTrackIndex, const FProperty* InProperty, uint8* OutMemory) const override
+	{
+		const T& Value = Base::GetValueAtTimeIndex(InTimeIndex, InOutSampleTrackIndex);
+		InProperty->CopyCompleteValue(OutMemory, &Value);
+	}
+	
+	friend class FControlRigBasicsSampleTrackTest;
+};
 
 template<>
-CONTROLRIG_API void FSampleTrack<TArray<FInstancedStruct>>::AddSampleFromProperty(const FProperty* InProperty, const uint8* InMemory);
+class CONTROLRIG_API FSampleTrack<FInstancedStruct> : public UE::ControlRig::Internal::FSampleTrackShared<FInstancedStruct>
+{
+	using Base = FSampleTrackShared<FInstancedStruct>;
+public:
+	FSampleTrack() = default;
+	FSampleTrack(const FName& InName) : Base(InName) {}
+	virtual ~FSampleTrack() override {}
+	
+	// Specialized functions for this type.
+	virtual void Serialize(FArchive& InArchive) override;
+	virtual int32 GetSizePerValue() const override;
+	virtual void AddSampleFromProperty(const FProperty* InProperty, const uint8* InMemory) override;
+	virtual void GetSampleForProperty(int32 InTimeIndex, FSampleTrackIndex& InOutSampleTrackIndex, const FProperty* InProperty, uint8* OutMemory) const override;
+	
+	friend class FControlRigBasicsSampleTrackTest;
+};
 
 template<>
-CONTROLRIG_API void FSampleTrack<TArray<FInstancedStruct>>::GetSampleForProperty(int32 InTimeIndex, FSampleTrackIndex& InOutSampleTrackIndex, const FProperty* InProperty, uint8* OutMemory) const;
+class CONTROLRIG_API FSampleTrack<TArray<FInstancedStruct>> : public UE::ControlRig::Internal::FSampleTrackShared<TArray<FInstancedStruct>>
+{
+	using Base = FSampleTrackShared<TArray<FInstancedStruct>>;
+public:
+	FSampleTrack() = default;
+	FSampleTrack(const FName& InName) : Base(InName) {}
+	virtual ~FSampleTrack() override {}
+
+	// Specialized functions for this type.
+	virtual void Serialize(FArchive& InArchive) override;
+	virtual int32 GetSizePerValue() const override
+	{
+		return sizeof(TArray<FInstancedStruct>);
+	}
+	virtual void AddSampleFromProperty(const FProperty* InProperty, const uint8* InMemory) override;
+	virtual void GetSampleForProperty(int32 InTimeIndex, FSampleTrackIndex& InOutSampleTrackIndex, const FProperty* InProperty, uint8* OutMemory) const override;
+
+	friend class FControlRigBasicsSampleTrackTest;
+};
+
+
+extern template class CONTROLRIG_API FSampleTrack<FInstancedStruct>;
+extern template class CONTROLRIG_API FSampleTrack<TArray<FInstancedStruct>>;
