Control Rig / RigVM: Ignore transactions that don't affect the asset

#rnx

diff --git a/Engine/Plugins/Animation/ControlRig/Source/ControlRigEditor/Private/Editor/ControlRigEditor.cpp b/Engine/Plugins/Animation/ControlRig/Source/ControlRigEditor/Private/Editor/ControlRigEditor.cpp
--- a/Engine/Plugins/Animation/ControlRig/Source/ControlRigEditor/Private/Editor/ControlRigEditor.cpp
+++ b/Engine/Plugins/Animation/ControlRig/Source/ControlRigEditor/Private/Editor/ControlRigEditor.cpp
@@ -2138,38 +2138,66 @@ FReply FControlRigBaseEditor::OnSpawnGraphNodeByShortcutImpl(FInputChord InChord
 
 void FControlRigBaseEditor::PostTransactionImpl(bool bSuccess, const FTransaction* Transaction, bool bIsRedo)
 {
-	EnsureValidRigElementsInDetailPanel();
+	if(Transaction == nullptr)
+	{
+		return;
+	}
 
-	if (UControlRigBlueprint* RigBlueprint = Cast<UControlRigBlueprint>(GetRigVMBlueprint()))
+	UControlRigBlueprint* RigBlueprint = Cast<UControlRigBlueprint>(GetRigVMBlueprint());
+	if(RigBlueprint == nullptr)
 	{
-		// Do not compile here. ControlRigBlueprint::PostTransacted decides when it is necessary to compile depending
-		// on the properties that are affected.
-		//Compile();
+		return;
+	}
+	
+    const UPackage* Package = RigBlueprint->GetPackage();
 
-		UpdateRigVMHost();
-		
-		USkeletalMesh* PreviewMesh = GetPersonaToolkit()->GetPreviewScene()->GetPreviewMesh();
-		if (PreviewMesh != RigBlueprint->GetPreviewMesh())
+    TArray<UObject*> TransactedObjects;
+    Transaction->GetTransactionObjects(TransactedObjects);
+
+	bool bAffectsThisAsset = false;
+	for(const UObject* Object : TransactedObjects)
+	{
+		if(Object->GetPackage() == Package)
 		{
-			RigBlueprint->SetPreviewMesh(PreviewMesh);
-			GetPersonaToolkit()->SetPreviewMesh(PreviewMesh, true);
+			bAffectsThisAsset = true;
+			break;
 		}
+	}
 
-		if (UControlRig* DebuggedControlRig = Cast<UControlRig>(RigBlueprint->GetObjectBeingDebugged()))
+	if(!bAffectsThisAsset)
+	{
+		return;
+	}
+
+	EnsureValidRigElementsInDetailPanel();
+
+	// Do not compile here. ControlRigBlueprint::PostTransacted decides when it is necessary to compile depending
+	// on the properties that are affected.
+	//Compile();
+
+	UpdateRigVMHost();
+	
+	USkeletalMesh* PreviewMesh = GetPersonaToolkit()->GetPreviewScene()->GetPreviewMesh();
+	if (PreviewMesh != RigBlueprint->GetPreviewMesh())
+	{
+		RigBlueprint->SetPreviewMesh(PreviewMesh);
+		GetPersonaToolkit()->SetPreviewMesh(PreviewMesh, true);
+	}
+
+	if (UControlRig* DebuggedControlRig = Cast<UControlRig>(RigBlueprint->GetObjectBeingDebugged()))
+	{
+		if(URigHierarchy* Hierarchy = DebuggedControlRig->GetHierarchy())
 		{
-			if(URigHierarchy* Hierarchy = DebuggedControlRig->GetHierarchy())
+			if(Hierarchy->Num() == 0)
 			{
-				if(Hierarchy->Num() == 0)
-				{
-					OnHierarchyChanged();
-				}
+				OnHierarchyChanged();
 			}
 		}
+	}
 
-		if (FControlRigEditMode* EditMode = GetEditMode())
-		{
-			EditMode->RequestToRecreateControlShapeActors();
-		}
+	if (FControlRigEditMode* EditMode = GetEditMode())
+	{
+		EditMode->RequestToRecreateControlShapeActors();
 	}
 }
 
diff --git a/Engine/Plugins/Animation/ControlRig/Source/ControlRigEditor/Private/Editor/SModularRigTreeView.cpp b/Engine/Plugins/Animation/ControlRig/Source/ControlRigEditor/Private/Editor/SModularRigTreeView.cpp
--- a/Engine/Plugins/Animation/ControlRig/Source/ControlRigEditor/Private/Editor/SModularRigTreeView.cpp
+++ b/Engine/Plugins/Animation/ControlRig/Source/ControlRigEditor/Private/Editor/SModularRigTreeView.cpp
@@ -255,13 +255,16 @@ TSharedRef<SWidget> SModularRigModelItem::GenerateWidgetForColumn(const FName& C
 				{
 					if (const FRigModuleInstance* Module = ModularRig->FindModule(WeakRigTreeElement.Pin()->ModuleName))
 					{
-						if (const UControlRigBlueprint* ModuleBlueprint = Cast<UControlRigBlueprint>(Module->GetRig()->GetClass()->ClassGeneratedBy))
+						if (const UControlRig* ModuleRig = Module->GetRig())
 						{
-							for (const FRigVMTag& Tag : ModuleBlueprint->GetAssetVariant().Tags)
+							if (const UControlRigBlueprint* ModuleBlueprint = Cast<UControlRigBlueprint>(ModuleRig->GetClass()->ClassGeneratedBy))
 							{
-								if (Tag.bMarksSubjectAsInvalid)
+								for (const FRigVMTag& Tag : ModuleBlueprint->GetAssetVariant().Tags)
 								{
-									return EVisibility::Visible;
+									if (Tag.bMarksSubjectAsInvalid)
+									{
+										return EVisibility::Visible;
+									}
 								}
 							}
 						}
diff --git a/Engine/Plugins/Runtime/RigVM/Source/RigVMEditor/Private/Editor/RigVMEditor.cpp b/Engine/Plugins/Runtime/RigVM/Source/RigVMEditor/Private/Editor/RigVMEditor.cpp
--- a/Engine/Plugins/Runtime/RigVM/Source/RigVMEditor/Private/Editor/RigVMEditor.cpp
+++ b/Engine/Plugins/Runtime/RigVM/Source/RigVMEditor/Private/Editor/RigVMEditor.cpp
@@ -891,7 +891,30 @@ void FRigVMEditorBase::PostRedoImpl(bool bSuccess)
 
 void FRigVMEditorBase::PostTransaction(bool bSuccess, const FTransaction* Transaction, bool bIsRedo)
 {
-	ForceEditorRefresh(ERefreshRigVMEditorReason::PostUndo);
+	if(Transaction == nullptr)
+	{
+		return;
+	}
+
+	const URigVMBlueprint* Blueprint = GetRigVMBlueprint();
+	if(Blueprint == nullptr)
+	{
+			return;
+	}
+
+	const UPackage* Package = Blueprint->GetPackage();
+
+	TArray<UObject*> TransactedObjects;
+	Transaction->GetTransactionObjects(TransactedObjects);
+	
+	for(const UObject* Object : TransactedObjects)
+	{
+		if(Object->GetPackage() == Package)
+		{
+			ForceEditorRefresh(ERefreshRigVMEditorReason::PostUndo);
+			return;
+		}
+	}
 }
 
 void FRigVMEditorBase::OnStartWatchingPin()
