Control Rig: Fix swap module losing configuration overrides
#rxn

diff --git a/Engine/Plugins/Animation/ControlRig/Source/ControlRig/Private/ModularRigController.cpp b/Engine/Plugins/Animation/ControlRig/Source/ControlRig/Private/ModularRigController.cpp
--- a/Engine/Plugins/Animation/ControlRig/Source/ControlRig/Private/ModularRigController.cpp
+++ b/Engine/Plugins/Animation/ControlRig/Source/ControlRig/Private/ModularRigController.cpp
@@ -2002,7 +2002,25 @@ void UModularRigController::RefreshModuleVariables(const FRigModuleReference* In
 	Module->Bindings.Reset();
 	for (const FControlRigOverrideValue& Override : ConfigOverrides)
 	{
-		SetConfigValueInModule(ModuleName, Override, false);
+		// Make sure the Override properties belong to the module class
+		FControlRigOverrideValue CorrectedOverride = Override;
+		if(Override.GetRootProperty()->GetOwnerUObject() != Module->Class.Get())
+		{
+			for (int32 i=0; i<Override.Properties.Num(); i++)
+			{
+				const FControlRigOverrideValue::FPropertyInfo& Property = Override.Properties[i];
+				if (FProperty* TargetProperty = Module->Class->FindPropertyByName(Property.Property->GetFName()))
+				{
+					CorrectedOverride.Properties[i].Property = TargetProperty;
+				}
+				else
+				{
+					UE_LOG(LogControlRig, Error, TEXT("Provided path '%s' does not belong to the module instance class %s."), *ModuleName.ToString(), *Module->Class.ToString());
+					continue;
+				}
+			}
+		}
+		SetConfigValueInModule(ModuleName, CorrectedOverride, false);
 	}
 	for (const TPair<FName, FString>& Pair : Bindings)
 	{
