Ran UnrealCodeFixup script to fix inline/dllstorage non-compliant code

diff --git a/Engine/Plugins/Animation/ControlRig/Source/ControlRig/Public/Tracks/SampleTrack.h b/Engine/Plugins/Animation/ControlRig/Source/ControlRig/Public/Tracks/SampleTrack.h
--- a/Engine/Plugins/Animation/ControlRig/Source/ControlRig/Public/Tracks/SampleTrack.h
+++ b/Engine/Plugins/Animation/ControlRig/Source/ControlRig/Public/Tracks/SampleTrack.h
@@ -821,7 +821,7 @@ public:
 };
 
 template<>
-class CONTROLRIG_API FSampleTrack<FInstancedStruct> : public UE::ControlRig::Internal::FSampleTrackShared<FInstancedStruct>
+class FSampleTrack<FInstancedStruct> : public UE::ControlRig::Internal::FSampleTrackShared<FInstancedStruct>
 {
 	using Base = FSampleTrackShared<FInstancedStruct>;
 public:
@@ -839,7 +839,7 @@ public:
 };
 
 template<>
-class CONTROLRIG_API FSampleTrack<TArray<FInstancedStruct>> : public UE::ControlRig::Internal::FSampleTrackShared<TArray<FInstancedStruct>>
+class FSampleTrack<TArray<FInstancedStruct>> : public UE::ControlRig::Internal::FSampleTrackShared<TArray<FInstancedStruct>>
 {
 	using Base = FSampleTrackShared<TArray<FInstancedStruct>>;
 public:
