Control Rig: Fix sequencer grouping multiple controls when forcing short names
#rnx

diff --git a/Engine/Plugins/Animation/ControlRig/Source/ControlRig/Private/Sequencer/MovieSceneControlRigParameterSection.cpp b/Engine/Plugins/Animation/ControlRig/Source/ControlRig/Private/Sequencer/MovieSceneControlRigParameterSection.cpp
--- a/Engine/Plugins/Animation/ControlRig/Source/ControlRig/Private/Sequencer/MovieSceneControlRigParameterSection.cpp
+++ b/Engine/Plugins/Animation/ControlRig/Source/ControlRig/Private/Sequencer/MovieSceneControlRigParameterSection.cpp
@@ -2028,6 +2028,30 @@ EMovieSceneChannelProxyType UMovieSceneControlRigParameterSection::CacheChannelP
 #endif
 
 			URigHierarchy* Hierarchy = ControlRig->GetHierarchy();
+
+			TMap<const FRigControlElement*, FText> ControlToGroupName;
+			TSet<FString> GroupNames;
+			auto GetUniqueGroupName = [Hierarchy, &GroupNames, &ControlToGroupName, ElementNameDisplayMode](const FRigControlElement* Element) -> FText 
+			{
+				FString DisplayName = Hierarchy->GetDisplayNameForUI(Element, ElementNameDisplayMode).ToString();
+				if (GroupNames.Contains(DisplayName))
+				{
+					int32 i=1;
+					FString NewName;
+					do
+					{
+						NewName = FString::Printf(TEXT("%s_%d"), *DisplayName, i++);
+					}
+					while (GroupNames.Contains(NewName));
+					DisplayName = NewName;
+				}
+
+				GroupNames.Add(DisplayName);
+				const FText Result = FText::FromString(DisplayName);
+				ControlToGroupName.FindOrAdd(Element) = Result;
+				return Result;
+			};
+			
 			for (FRigControlElement* ControlElement : SortedControls)
 			{
 				if (!Hierarchy->IsAnimatable(ControlElement))
@@ -2045,7 +2069,14 @@ EMovieSceneChannelProxyType UMovieSceneControlRigParameterSection::CacheChannelP
 					if (const FRigControlElement* ParentControlElement = Cast<FRigControlElement>(Hierarchy->GetFirstParent(ControlElement)))
 					{
 						ParentControlName = ParentControlElement->GetFName();
-						Group = Hierarchy->GetDisplayNameForUI(ParentControlElement, ElementNameDisplayMode);
+						if (FText* Found = ControlToGroupName.Find(ParentControlElement))
+						{
+							Group = *Found;
+						}
+						else
+						{
+							Group = GetUniqueGroupName(ParentControlElement);
+						}
 					}
 				}
 
@@ -2064,7 +2095,7 @@ EMovieSceneChannelProxyType UMovieSceneControlRigParameterSection::CacheChannelP
 							if (Group.IsEmpty())
 							{
 								ControlChannelMap.Add(Scalar.ParameterName, FChannelMapInfo(ControlIndex, SortOrder, FloatChannelIndex, INDEX_NONE, NAME_None, MaskIndex, CategoryIndex));
-								Group = Hierarchy->GetDisplayNameForUI(ControlElement, ElementNameDisplayMode);
+								Group = GetUniqueGroupName(ControlElement);
 								if (bEnabled)
 								{
 									++CategoryIndex;
@@ -2099,7 +2130,7 @@ EMovieSceneChannelProxyType UMovieSceneControlRigParameterSection::CacheChannelP
 							if (Group.IsEmpty())
 							{
 								ControlChannelMap.Add(Bool.ParameterName, FChannelMapInfo(ControlIndex, SortOrder, BoolChannelIndex, INDEX_NONE, BoolChannelTypeName, MaskIndex, CategoryIndex));
-								Group = Hierarchy->GetDisplayNameForUI(ControlElement, ElementNameDisplayMode);
+								Group = GetUniqueGroupName(ControlElement);
 								if (bEnabled)
 								{
 									++CategoryIndex;
@@ -2136,7 +2167,7 @@ EMovieSceneChannelProxyType UMovieSceneControlRigParameterSection::CacheChannelP
 								if (Group.IsEmpty())
 								{
 									ControlChannelMap.Add(Enum.ParameterName, FChannelMapInfo(ControlIndex, SortOrder, EnumChannelIndex, INDEX_NONE, EnumChannelTypeName, MaskIndex, CategoryIndex));
-									Group = Hierarchy->GetDisplayNameForUI(ControlElement, ElementNameDisplayMode);
+									Group = GetUniqueGroupName(ControlElement);
 									if (bEnabled)
 									{
 										++CategoryIndex;
@@ -2170,7 +2201,7 @@ EMovieSceneChannelProxyType UMovieSceneControlRigParameterSection::CacheChannelP
 								if (Group.IsEmpty())
 								{
 									ControlChannelMap.Add(Integer.ParameterName, FChannelMapInfo(ControlIndex, SortOrder, IntegerChannelIndex, INDEX_NONE, IntegerChannelTypeName, MaskIndex, CategoryIndex));
-									Group = Hierarchy->GetDisplayNameForUI(ControlElement, ElementNameDisplayMode);
+									Group = GetUniqueGroupName(ControlElement);
 									if (bEnabled)
 									{
 										++CategoryIndex;
@@ -2213,7 +2244,7 @@ EMovieSceneChannelProxyType UMovieSceneControlRigParameterSection::CacheChannelP
 								{
 									++CategoryIndex;
 								}
-								Group = Hierarchy->GetDisplayNameForUI(ControlElement, ElementNameDisplayMode);
+								Group = GetUniqueGroupName(ControlElement);
 							}
 							else
 							{
@@ -2249,7 +2280,7 @@ EMovieSceneChannelProxyType UMovieSceneControlRigParameterSection::CacheChannelP
 								{
 									++CategoryIndex;
 								}
-								Group = Hierarchy->GetDisplayNameForUI(ControlElement, ElementNameDisplayMode);
+								Group = GetUniqueGroupName(ControlElement);
 							}
 							else
 							{
@@ -2328,7 +2359,7 @@ EMovieSceneChannelProxyType UMovieSceneControlRigParameterSection::CacheChannelP
 								{
 									++CategoryIndex;
 								}
-								Group = Hierarchy->GetDisplayNameForUI(ControlElement, ElementNameDisplayMode);
+								Group = GetUniqueGroupName(ControlElement);
 							}
 							else
 							{
