[ControlRig]
* Fixed up dllstorage on specializations of SampleTrack

diff --git a/Engine/Plugins/Animation/ControlRig/Source/ControlRig/Private/Tracks/SampleTrack.cpp b/Engine/Plugins/Animation/ControlRig/Source/ControlRig/Private/Tracks/SampleTrack.cpp
--- a/Engine/Plugins/Animation/ControlRig/Source/ControlRig/Private/Tracks/SampleTrack.cpp
+++ b/Engine/Plugins/Animation/ControlRig/Source/ControlRig/Private/Tracks/SampleTrack.cpp
@@ -130,6 +130,3 @@ void FSampleTrack<TArray<FInstancedStruct>>::GetSampleForProperty(int32 InTimeIn
 		StructProperty->CopyCompleteValue(ArrayHelper.GetElementPtr(Index), Value.GetMemory());
 	}
 }
-
-template class CONTROLRIG_API FSampleTrack<FInstancedStruct>;
-template class CONTROLRIG_API FSampleTrack<TArray<FInstancedStruct>>;
diff --git a/Engine/Plugins/Animation/ControlRig/Source/ControlRig/Public/Tracks/SampleTrack.h b/Engine/Plugins/Animation/ControlRig/Source/ControlRig/Public/Tracks/SampleTrack.h
--- a/Engine/Plugins/Animation/ControlRig/Source/ControlRig/Public/Tracks/SampleTrack.h
+++ b/Engine/Plugins/Animation/ControlRig/Source/ControlRig/Public/Tracks/SampleTrack.h
@@ -830,10 +830,10 @@ public:
 	virtual ~FSampleTrack() override {}
 	
 	// Specialized functions for this type.
-	virtual void Serialize(FArchive& InArchive) override;
-	virtual int32 GetSizePerValue() const override;
-	virtual void AddSampleFromProperty(const FProperty* InProperty, const uint8* InMemory) override;
-	virtual void GetSampleForProperty(int32 InTimeIndex, FSampleTrackIndex& InOutSampleTrackIndex, const FProperty* InProperty, uint8* OutMemory) const override;
+	CONTROLRIG_API virtual void Serialize(FArchive& InArchive) override;
+	CONTROLRIG_API virtual int32 GetSizePerValue() const override;
+	CONTROLRIG_API virtual void AddSampleFromProperty(const FProperty* InProperty, const uint8* InMemory) override;
+	CONTROLRIG_API virtual void GetSampleForProperty(int32 InTimeIndex, FSampleTrackIndex& InOutSampleTrackIndex, const FProperty* InProperty, uint8* OutMemory) const override;
 	
 	friend class FControlRigBasicsSampleTrackTest;
 };
@@ -848,17 +848,13 @@ public:
 	virtual ~FSampleTrack() override {}
 
 	// Specialized functions for this type.
-	virtual void Serialize(FArchive& InArchive) override;
+	CONTROLRIG_API virtual void Serialize(FArchive& InArchive) override;
 	virtual int32 GetSizePerValue() const override
 	{
 		return sizeof(TArray<FInstancedStruct>);
 	}
-	virtual void AddSampleFromProperty(const FProperty* InProperty, const uint8* InMemory) override;
-	virtual void GetSampleForProperty(int32 InTimeIndex, FSampleTrackIndex& InOutSampleTrackIndex, const FProperty* InProperty, uint8* OutMemory) const override;
+	CONTROLRIG_API virtual void AddSampleFromProperty(const FProperty* InProperty, const uint8* InMemory) override;
+	CONTROLRIG_API virtual void GetSampleForProperty(int32 InTimeIndex, FSampleTrackIndex& InOutSampleTrackIndex, const FProperty* InProperty, uint8* OutMemory) const override;
 
 	friend class FControlRigBasicsSampleTrackTest;
 };
-
-
-extern template class CONTROLRIG_API FSampleTrack<FInstancedStruct>;
-extern template class CONTROLRIG_API FSampleTrack<TArray<FInstancedStruct>>;
