Unshelved from pending changelist '43055921':

Fixed undo / redo crash in Control Rig hierarchy mappings, when undoing across a construction transaction

diff --git a/Engine/Plugins/Animation/ControlRig/Source/ControlRig/Private/Tools/ControlRigHierarchyMappings.cpp b/Engine/Plugins/Animation/ControlRig/Source/ControlRig/Private/Tools/ControlRigHierarchyMappings.cpp
--- a/Engine/Plugins/Animation/ControlRig/Source/ControlRig/Private/Tools/ControlRigHierarchyMappings.cpp
+++ b/Engine/Plugins/Animation/ControlRig/Source/ControlRig/Private/Tools/ControlRigHierarchyMappings.cpp
@@ -146,7 +146,7 @@ void FControlRigHierarchyMappings::UpdateInput(UControlRig* ControlRig
 	DECLARE_SCOPE_HIERARCHICAL_COUNTER_FUNC()
 
 	URigHierarchy* Hierarchy = ControlRig->GetHierarchy();
-	if (Hierarchy == nullptr)
+	if (!IsValid(Hierarchy))
 	{
 		return;
 	}
@@ -333,7 +333,7 @@ void FControlRigHierarchyMappings::UpdateOutput(UControlRig* ControlRig
 	DECLARE_SCOPE_HIERARCHICAL_COUNTER_FUNC()
 
 	URigHierarchy* Hierarchy = ControlRig->GetHierarchy();
-	if (Hierarchy == nullptr)
+	if (!IsValid(Hierarchy))
 	{
 		return;
 	}
diff --git a/Engine/Plugins/Animation/ControlRig/Source/ControlRig/Private/Tools/ControlRigPoseAdapter.cpp b/Engine/Plugins/Animation/ControlRig/Source/ControlRig/Private/Tools/ControlRigPoseAdapter.cpp
--- a/Engine/Plugins/Animation/ControlRig/Source/ControlRig/Private/Tools/ControlRigPoseAdapter.cpp
+++ b/Engine/Plugins/Animation/ControlRig/Source/ControlRig/Private/Tools/ControlRigPoseAdapter.cpp
@@ -537,39 +537,39 @@ void FControlRigPoseAdapter::ComputeDependentTransforms()
 {
 	DECLARE_SCOPE_HIERARCHICAL_COUNTER_FUNC()
 
-	URigHierarchy* Hierarchy = GetHierarchy();
-	check(Hierarchy);
-
-	// ensure to compute all local transforms
-	for (const FDependentTransform& Dependent : Dependents)
+	if (URigHierarchy* Hierarchy = GetHierarchy())
 	{
-		if (FRigTransformElement* TransformElement = Hierarchy->Get<FRigTransformElement>(Dependent.KeyAndIndex.Index))
+		// ensure to compute all local transforms
+		for (const FDependentTransform& Dependent : Dependents)
 		{
-			switch(Dependent.StorageType)
+			if (FRigTransformElement* TransformElement = Hierarchy->Get<FRigTransformElement>(Dependent.KeyAndIndex.Index))
 			{
-				case ERigTransformStorageType::Pose:
+				switch(Dependent.StorageType)
 				{
-					(void)Hierarchy->GetTransform(TransformElement, MakeLocal(Dependent.TransformType));
-					break;
-				}
-				case ERigTransformStorageType::Offset:
-				{
-					FRigControlElement* ControlElement = CastChecked<FRigControlElement>(TransformElement);
-					(void)Hierarchy->GetControlOffsetTransform(ControlElement, MakeLocal(Dependent.TransformType));
-					break;
-				}
-				case ERigTransformStorageType::Shape:
-				{
-					FRigControlElement* ControlElement = CastChecked<FRigControlElement>(TransformElement);
-					(void)Hierarchy->GetControlShapeTransform(ControlElement, MakeLocal(Dependent.TransformType));
-					break;
-				}
-				default:
-				{
-					break;
+					case ERigTransformStorageType::Pose:
+					{
+						(void)Hierarchy->GetTransform(TransformElement, MakeLocal(Dependent.TransformType));
+						break;
+					}
+					case ERigTransformStorageType::Offset:
+					{
+						FRigControlElement* ControlElement = CastChecked<FRigControlElement>(TransformElement);
+						(void)Hierarchy->GetControlOffsetTransform(ControlElement, MakeLocal(Dependent.TransformType));
+						break;
+					}
+					case ERigTransformStorageType::Shape:
+					{
+						FRigControlElement* ControlElement = CastChecked<FRigControlElement>(TransformElement);
+						(void)Hierarchy->GetControlShapeTransform(ControlElement, MakeLocal(Dependent.TransformType));
+						break;
+					}
+					default:
+					{
+						break;
+					}
 				}
+				check(Dependent.DirtyState->Local.Get() == false);
 			}
-			check(Dependent.DirtyState->Local.Get() == false);
 		}
 	}
 }
@@ -578,14 +578,15 @@ void FControlRigPoseAdapter::MarkDependentsDirty()
 {
 	DECLARE_SCOPE_HIERARCHICAL_COUNTER_FUNC()
 
-	ComputeDependentTransforms();
-
-	check(GetHierarchy() != nullptr);
-
-	// mark the global dependent as dirty
-	for (const FDependentTransform& Dependent : Dependents)
+	if (IsValid(GetHierarchy()))
 	{
-		check(Dependent.DirtyState->Local.Get() == false);
-		Dependent.DirtyState->Global.Set(true);
+		ComputeDependentTransforms();
+
+		// mark the global dependent as dirty
+		for (const FDependentTransform& Dependent : Dependents)
+		{
+			check(Dependent.DirtyState->Local.Get() == false);
+			Dependent.DirtyState->Global.Set(true);
+		}
 	}
-}
\ No newline at end of file
+}
