Remove legacy toolbar usage in the Animation Editor stack.

#rnx

diff --git a/Engine/Plugins/Animation/ControlRig/Source/ControlRigEditor/Private/Editor/ControlRigEditor.cpp b/Engine/Plugins/Animation/ControlRig/Source/ControlRigEditor/Private/Editor/ControlRigEditor.cpp
--- a/Engine/Plugins/Animation/ControlRig/Source/ControlRigEditor/Private/Editor/ControlRigEditor.cpp
+++ b/Engine/Plugins/Animation/ControlRig/Source/ControlRigEditor/Private/Editor/ControlRigEditor.cpp
@@ -243,41 +243,6 @@ void FillControlRigMenu(UToolMenu* InMenu)
 	}
 }
 
-TSharedRef<SWidget> GenerateControlRigMenuWidget(const TWeakPtr<SEditorViewport>& InViewport)
-{
-	if (TSharedPtr<SEditorViewport> Viewport = InViewport.Pin())
-	{
-		// We generate a menu via UToolMenus, so we can use FillControlRigMenu call from both old and new toolbar
-		FName OldFollowModeMenuName = "AnimationEditor.OldViewportToolbar.ControlRig";
-		if (!UToolMenus::Get()->IsMenuRegistered(OldFollowModeMenuName))
-		{
-			if (UToolMenu* const Menu = UToolMenus::Get()->RegisterMenu(OldFollowModeMenuName, NAME_None, EMultiBoxType::Menu, false))
-			{
-				Menu->AddDynamicSection(
-					"BaseSection",
-					FNewToolMenuDelegate::CreateLambda(
-						[](UToolMenu* InMenu)
-						{
-							UE::ControlRigEditor::Private::FillControlRigMenu(InMenu);
-						}
-					)
-				);
-			}
-		}
-
-		// Providing the UToolMenu with a context containing the data needed by its current entries
-		FToolMenuContext MenuContext;
-		{
-			UUnrealEdViewportToolbarContext* const ContextObject = UE::UnrealEd::CreateViewportToolbarDefaultContext(Viewport);
-			MenuContext.AddObject(ContextObject);
-
-			MenuContext.AppendCommandList(Viewport->GetCommandList());
-		}
-		return UToolMenus::Get()->GenerateWidget(OldFollowModeMenuName, MenuContext);
-	}
-
-	return SNullWidget::NullWidget;
-}
 } // namespace UE::ControlRigEditor::Private
 
 TAutoConsoleVariable<bool> CVarControlRigShowTestingToolbar(TEXT("ControlRig.Test.EnableTestingToolbar"), false, TEXT("When true we'll show the testing toolbar in Control Rig Editor."));
@@ -2852,28 +2817,6 @@ void FControlRigBaseEditor::HandleViewportCreated(const TSharedRef<class IPerson
 		ChangeShapeTransformNotificationOptions
 	);
 
-	if (const TSharedPtr<SEditorViewport>& EditorViewportWidget = InViewport->GetViewportClient().GetEditorViewportWidget())
-	{
-		TSharedPtr<FUICommandList> CommandList = EditorViewportWidget->GetCommandList();
-		InViewport->AddToolbarExtender(
-			TEXT("AnimViewportDefaultCamera"),
-			FMenuExtensionDelegate::CreateLambda(
-				[EditorViewportWidgetWeak = EditorViewportWidget.ToWeakPtr()](FMenuBuilder& InMenuBuilder)
-				{
-					if (TSharedPtr<SEditorViewport> ViewportPinned = EditorViewportWidgetWeak.Pin())
-					{
-						InMenuBuilder.AddWidget(
-							UE::ControlRigEditor::Private::GenerateControlRigMenuWidget(EditorViewportWidgetWeak), FText()
-						);
-					}
-				}
-			)
-		);
-	}
-
-	// Extending for new toolbar as well. Using Extenders setup for the old toolbar was convoluted, since they are currently
-	// stored in the old toolbar, making them difficult to access and handle from the new one. Also, they are being added
-	// from this very same callback (HandleViewportCreated) after the creation of the FToolMenuContext in the new toolbar, making things slightly more complicated.
 	if (UToolMenu* const Submenu = UToolMenus::Get()->ExtendMenu("AnimationEditor.ViewportToolbar.Show"))
 	{
 		FToolMenuSection& ControlRigSection = Submenu->FindOrAddSection(
diff --git a/Engine/Source/Editor/Persona/Private/SAnimViewportToolBar.cpp b/Engine/Source/Editor/Persona/Private/SAnimViewportToolBar.cpp
deleted file mode 100644
diff --git a/Engine/Source/Editor/Persona/Private/SAnimViewportToolBar.h b/Engine/Source/Editor/Persona/Private/SAnimViewportToolBar.h
deleted file mode 100644
diff --git a/Engine/Source/Editor/Persona/Private/SAnimationEditorViewport.cpp b/Engine/Source/Editor/Persona/Private/SAnimationEditorViewport.cpp
--- a/Engine/Source/Editor/Persona/Private/SAnimationEditorViewport.cpp
+++ b/Engine/Source/Editor/Persona/Private/SAnimationEditorViewport.cpp
@@ -17,7 +17,6 @@
 
 #include "SAnimationScrubPanel.h"
 #include "SAnimMontageScrubPanel.h"
-#include "SAnimViewportToolBar.h"
 #include "AnimViewportMenuCommands.h"
 #include "AnimViewportShowCommands.h"
 #include "AnimViewportLODCommands.h"
@@ -130,20 +129,10 @@ void SAnimationEditorViewport::PopulateViewportOverlays(TSharedRef<SOverlay> Ove
 	{
 		TSharedRef<SVerticalBox> VerticalBox = SNew(SVerticalBox);
 
-		const TAttribute<FMargin> Padding = TAttribute<FMargin>::CreateLambda(
-			[]() -> FMargin
-			{
-				// The old viewport toolbar draws inside the viewport, so we need to pad to not overlap it.
-				const float TopPadding = UE::UnrealEd::ShowOldViewportToolbars() ? 29.0f : 0;
-				return FMargin(0, TopPadding, 0, 0);
-			}
-		);
-
 		// clang-format off
 		Overlay->AddSlot()
 			.VAlign(VAlign_Top)
 			.HAlign(HAlign_Fill)
-			.Padding(Padding)
 		[
 			VerticalBox
 		];
@@ -200,30 +189,6 @@ TSharedRef<FEditorViewportClient> SAnimationEditorViewport::MakeEditorViewportCl
 	return LevelViewportClient.ToSharedRef();
 }
 
-TSharedPtr<SWidget> SAnimationEditorViewport::MakeViewportToolbar()
-{
-	// clang-format off
-	TSharedRef<SVerticalBox> ToolbarWidget =
-		SNew(SVerticalBox)
-		+ SVerticalBox::Slot()
-		[
-			SAssignNew(ViewportToolbar, SAnimViewportToolBar, TabBodyPtr.Pin(), SharedThis(this))
-			.Visibility(EVisibility::SelfHitTestInvisible)
-			.Cursor(EMouseCursor::Default)
-			.Extenders(Extenders)
-			.ContextName(ContextName)
-			.ShowShowMenu(bShowShowMenu)
-			.ShowLODMenu(bShowLODMenu)
-			.ShowPlaySpeedMenu(bShowPlaySpeedMenu)
-			.ShowFloorOptions(bShowFloorOptions)
-			.ShowTurnTable(bShowTurnTable)
-			.ShowPhysicsMenu(bShowPhysicsMenu)
-		];
-	// clang-format on
-
-	return ToolbarWidget;
-}
-
 TSharedPtr<SWidget> SAnimationEditorViewport::BuildViewportToolbar()
 {
 	// Register the viewport toolbar if another viewport hasn't already (it's shared).
@@ -705,12 +670,6 @@ void SAnimationEditorViewportTabBody::RemoveNotification(const TWeakPtr<SWidget>
 	}
 }
 
-
-void SAnimationEditorViewportTabBody::AddToolbarExtender(FName MenuToExtend, FMenuExtensionDelegate MenuBuilderDelegate)
-{
-	return ViewportWidget->ViewportToolbar->AddMenuExtender(MenuToExtend, MenuBuilderDelegate);
-}
-
 void SAnimationEditorViewportTabBody::AddOverlayWidget(TSharedRef<SWidget> InOverlaidWidget, int32 ZOrder)
 {
 	ViewportWidget->ViewportOverlay->AddSlot(ZOrder)
diff --git a/Engine/Source/Editor/Persona/Private/SAnimationEditorViewport.h b/Engine/Source/Editor/Persona/Private/SAnimationEditorViewport.h
--- a/Engine/Source/Editor/Persona/Private/SAnimationEditorViewport.h
+++ b/Engine/Source/Editor/Persona/Private/SAnimationEditorViewport.h
@@ -18,8 +18,8 @@
 #include "SEditorViewport.h"
 #include "PersonaModule.h"
 #include "SNameComboBox.h"
+#include "ScopedTransaction.h"
 
-class FScopedTransaction;
 class FUICommandList_Pinnable;
 class SAnimViewportToolBar;
 class SAnimationEditorViewportTabBody;
@@ -98,7 +98,6 @@ public:
 protected:
 	// SEditorViewport interface
 	virtual TSharedRef<FEditorViewportClient> MakeEditorViewportClient() override;
-	virtual TSharedPtr<SWidget> MakeViewportToolbar() override;
 	virtual TSharedPtr<SWidget> BuildViewportToolbar() override;
 	virtual TSharedPtr<IPreviewProfileController> CreatePreviewProfileController() override;
 	virtual void OnFocusViewportToSelection() override;
@@ -229,7 +228,6 @@ public:
 	virtual TSharedRef<IPinnedCommandList> GetPinnedCommandList() const override;
 	virtual TWeakPtr<SWidget> AddNotification(TAttribute<EMessageSeverity::Type> InSeverity, TAttribute<bool> InCanBeDismissed, const TSharedRef<SWidget>& InNotificationWidget, FPersonaViewportNotificationOptions InOptions) override;
 	virtual void RemoveNotification(const TWeakPtr<SWidget>& InContainingWidget) override;
-	virtual void AddToolbarExtender(FName MenuToExtend, FMenuExtensionDelegate MenuBuilderDelegate) override;
 	virtual FPersonaViewportKeyDownDelegate& GetKeyDownDelegate() override { return OnKeyDownDelegate; }
 	virtual void AddOverlayWidget( TSharedRef<SWidget> InOverlaidWidget, int32 ZOrder = INDEX_NONE ) override;
 	virtual void RemoveOverlayWidget( TSharedRef<SWidget> InOverlaidWidget ) override;
diff --git a/Engine/Source/Editor/Persona/Public/IPersonaViewport.h b/Engine/Source/Editor/Persona/Public/IPersonaViewport.h
--- a/Engine/Source/Editor/Persona/Public/IPersonaViewport.h
+++ b/Engine/Source/Editor/Persona/Public/IPersonaViewport.h
@@ -66,7 +66,8 @@ public:
 	 * @param MenuToExtend        The name of the toolbar menu to extend
 	 * @param MenuBuilderDelegate The delegate to use when filling the menu
 	 */
-	virtual void AddToolbarExtender(FName MenuToExtend, FMenuExtensionDelegate MenuBuilderDelegate) = 0;
+	 UE_DEPRECATED(5.7, "Please extend the menu using UToolMenus instead.")
+	virtual void AddToolbarExtender(FName MenuToExtend, FMenuExtensionDelegate MenuBuilderDelegate) {}
 
 	/**
 	 * Returns the delegate broadcasted two in case we have an unhandled key
