Updated headers using UnrealCodeFixup to make sure dllstorage is on methods/staticvars instead of types. Part 1/n

diff --git a/Engine/Plugins/2D/Paper2D/Source/Paper2DEditor/Classes/PaperFlipbookFactory.h b/Engine/Plugins/2D/Paper2D/Source/Paper2DEditor/Classes/PaperFlipbookFactory.h
--- a/Engine/Plugins/2D/Paper2D/Source/Paper2DEditor/Classes/PaperFlipbookFactory.h
+++ b/Engine/Plugins/2D/Paper2D/Source/Paper2DEditor/Classes/PaperFlipbookFactory.h
@@ -6,18 +6,22 @@
 #include "PaperFlipbook.h"
 #include "PaperFlipbookFactory.generated.h"
 
+#define UE_API PAPER2DEDITOR_API
+
 /**
  * Factory for flipbooks
  */
 
-UCLASS()
-class PAPER2DEDITOR_API UPaperFlipbookFactory : public UFactory
+UCLASS(MinimalAPI)
+class UPaperFlipbookFactory : public UFactory
 {
 	GENERATED_UCLASS_BODY()
 
 	TArray<FPaperFlipbookKeyFrame> KeyFrames;
 
 	// UFactory interface
-	virtual UObject* FactoryCreateNew(UClass* Class, UObject* InParent, FName Name, EObjectFlags Flags, UObject* Context, FFeedbackContext* Warn) override;
+	UE_API virtual UObject* FactoryCreateNew(UClass* Class, UObject* InParent, FName Name, EObjectFlags Flags, UObject* Context, FFeedbackContext* Warn) override;
 	// End of UFactory interface
 };
+
+#undef UE_API
diff --git a/Engine/Plugins/2D/Paper2D/Source/Paper2DEditor/Classes/PaperImporterSettings.h b/Engine/Plugins/2D/Paper2D/Source/Paper2DEditor/Classes/PaperImporterSettings.h
--- a/Engine/Plugins/2D/Paper2D/Source/Paper2DEditor/Classes/PaperImporterSettings.h
+++ b/Engine/Plugins/2D/Paper2D/Source/Paper2DEditor/Classes/PaperImporterSettings.h
@@ -4,6 +4,8 @@
 
 #include "PaperImporterSettings.generated.h"
 
+#define UE_API PAPER2DEDITOR_API
+
 class UTexture;
 enum TextureCompressionSettings : int;
 enum TextureGroup : int;
@@ -47,8 +49,8 @@ enum class ESpriteInitMaterialType : uint8
 /**
  * Implements the settings for imported Paper2D assets, such as sprite sheet textures.
  */
-UCLASS(config=Editor, defaultconfig)
-class PAPER2DEDITOR_API UPaperImporterSettings : public UObject
+UCLASS(MinimalAPI, config=Editor, defaultconfig)
+class UPaperImporterSettings : public UObject
 {
 	GENERATED_BODY()
 
@@ -115,45 +117,47 @@ protected:
 	FSoftObjectPath LitDefaultOpaqueMaterialName;
 
 public:
-	UPaperImporterSettings();
+	UE_API UPaperImporterSettings();
 
 	// Should the source texture be scanned when creating new tile maps (from a tile set or via importing) to determine the appropriate material? (if false, the Default Masked Material is always used)
 	bool ShouldPickBestMaterialWhenCreatingTileMaps() const { return bPickBestMaterialWhenCreatingTileMaps; }
 
 	// Removes the suffix from the specified name if it matches something in BaseMapTextureSuffixes
-	FString RemoveSuffixFromBaseMapName(const FString& InName) const;
+	UE_API FString RemoveSuffixFromBaseMapName(const FString& InName) const;
 
 	// Generates names to test for a normal map using NormalMapTextureSuffixes
-	void GenerateNormalMapNamesToTest(const FString& InRoot, TArray<FString>& InOutNames) const;
+	UE_API void GenerateNormalMapNamesToTest(const FString& InRoot, TArray<FString>& InOutNames) const;
 
 	// Applies the compression settings to the specified texture
-	void ApplyTextureSettings(UTexture2D* Texture) const;
+	UE_API void ApplyTextureSettings(UTexture2D* Texture) const;
 
 	// Fills out the sprite init parameters with the default settings given the desired material type and lighting mode (which can both be automatic)
 	// Note: This should be called after the texture has been set, as that will be analyzed if the lighting/material type flags are set to automatic
-	void ApplySettingsForSpriteInit(FSpriteAssetInitParameters& InitParams, ESpriteInitMaterialLightingMode LightingMode = ESpriteInitMaterialLightingMode::Automatic, ESpriteInitMaterialType MaterialTypeMode = ESpriteInitMaterialType::Automatic) const;
+	UE_API void ApplySettingsForSpriteInit(FSpriteAssetInitParameters& InitParams, ESpriteInitMaterialLightingMode LightingMode = ESpriteInitMaterialLightingMode::Automatic, ESpriteInitMaterialType MaterialTypeMode = ESpriteInitMaterialType::Automatic) const;
 
 	// Fills out the tile map with the default settings given the desired material type and lighting mode (which can both be automatic)
-	void ApplySettingsForTileMapInit(UPaperTileMap* TileMap, UPaperTileSet* DefaultTileSet = nullptr, ESpriteInitMaterialLightingMode LightingMode = ESpriteInitMaterialLightingMode::Automatic, ESpriteInitMaterialType MaterialTypeMode = ESpriteInitMaterialType::Automatic, bool bCreateEmptyLayer = true) const;
+	UE_API void ApplySettingsForTileMapInit(UPaperTileMap* TileMap, UPaperTileSet* DefaultTileSet = nullptr, ESpriteInitMaterialLightingMode LightingMode = ESpriteInitMaterialLightingMode::Automatic, ESpriteInitMaterialType MaterialTypeMode = ESpriteInitMaterialType::Automatic, bool bCreateEmptyLayer = true) const;
 
 	// Analyzes the specified texture in the Offset..Offset+Dimensions region and returns the best kind of material
 	// to represent the alpha content in the texture (typically masked or translucent, but can return opaque if bAnalysisCanUseOpaque is true
 	// Note: Will return Automatic if the texture is null.
-	ESpriteInitMaterialType AnalyzeTextureForDesiredMaterialType(UTexture* Texture, const FIntPoint& Offset, const FIntPoint& Dimensions) const;
+	UE_API ESpriteInitMaterialType AnalyzeTextureForDesiredMaterialType(UTexture* Texture, const FIntPoint& Offset, const FIntPoint& Dimensions) const;
 
 	// Returns the default translucent material
-	UMaterialInterface* GetDefaultTranslucentMaterial(bool bLit) const;
+	UE_API UMaterialInterface* GetDefaultTranslucentMaterial(bool bLit) const;
 
 	// Returns the default opaque material
-	UMaterialInterface* GetDefaultOpaqueMaterial(bool bLit) const;
+	UE_API UMaterialInterface* GetDefaultOpaqueMaterial(bool bLit) const;
 
 	// Returns the default masked material
-	UMaterialInterface* GetDefaultMaskedMaterial(bool bLit) const;
+	UE_API UMaterialInterface* GetDefaultMaskedMaterial(bool bLit) const;
 
 	// Returns the default material for the specified material type.
 	// Input should be Masked, Opaque, or Translucent.  Automatic and LeaveAsIs will be treated like masked!
-	UMaterialInterface* GetDefaultMaterial(ESpriteInitMaterialType MaterialType, bool bUseLitMaterial) const;
+	UE_API UMaterialInterface* GetDefaultMaterial(ESpriteInitMaterialType MaterialType, bool bUseLitMaterial) const;
 
 	// Returns the default pixels/uu setting
 	float GetDefaultPixelsPerUnrealUnit() const { return DefaultPixelsPerUnrealUnit; }
 };
+
+#undef UE_API
diff --git a/Engine/Plugins/2D/Paper2D/Source/Paper2DEditor/Classes/TileMapAssetImportData.h b/Engine/Plugins/2D/Paper2D/Source/Paper2DEditor/Classes/TileMapAssetImportData.h
--- a/Engine/Plugins/2D/Paper2D/Source/Paper2DEditor/Classes/TileMapAssetImportData.h
+++ b/Engine/Plugins/2D/Paper2D/Source/Paper2DEditor/Classes/TileMapAssetImportData.h
@@ -5,6 +5,8 @@
 #include "EditorFramework/AssetImportData.h"
 #include "TileMapAssetImportData.generated.h"
 
+#define UE_API PAPER2DEDITOR_API
+
 class UPaperTileSet;
 
 USTRUCT()
@@ -25,13 +27,15 @@ struct FTileSetImportMapping
 /**
  * Base class for import data and options used when importing a tile map
  */
-UCLASS()
-class PAPER2DEDITOR_API UTileMapAssetImportData : public UAssetImportData
+UCLASS(MinimalAPI)
+class UTileMapAssetImportData : public UAssetImportData
 {
 	GENERATED_UCLASS_BODY()
 
 	UPROPERTY()
 	TArray<FTileSetImportMapping> TileSetMap;
 
-	static UTileMapAssetImportData* GetImportDataForTileMap(class UPaperTileMap* TileMap);
+	static UE_API UTileMapAssetImportData* GetImportDataForTileMap(class UPaperTileMap* TileMap);
 };
+
+#undef UE_API
diff --git a/Engine/Plugins/2D/Paper2D/Source/Paper2DEditor/Public/PaperFlipbookHelpers.h b/Engine/Plugins/2D/Paper2D/Source/Paper2DEditor/Public/PaperFlipbookHelpers.h
--- a/Engine/Plugins/2D/Paper2D/Source/Paper2DEditor/Public/PaperFlipbookHelpers.h
+++ b/Engine/Plugins/2D/Paper2D/Source/Paper2DEditor/Public/PaperFlipbookHelpers.h
@@ -4,12 +4,14 @@
 
 #include "Containers/Map.h"
 
+#define UE_API PAPER2DEDITOR_API
+
 class UPaperSprite;
 
 //////////////////////////////////////////////////////////////////////////
 // FPaperFlipbookHelpers
 
-class PAPER2DEDITOR_API FPaperFlipbookHelpers
+class FPaperFlipbookHelpers
 {
 public:
 	// Try and best guess split up the sprites into multiple flipbooks based on names
@@ -20,12 +22,14 @@ public:
 	// Sprites array contains valid sprites
 	// SpriteNames contains an array of sprite names (parallel to Sprites)
 	// or, an empty array if Sprite->GetName() is to be used.
-	static void ExtractFlipbooksFromSprites(TMap<FString, TArray<UPaperSprite*> >& OutSpriteFlipbookMap, const TArray<UPaperSprite*>& Sprites, const TArray<FString>& SpriteNames);
+	static UE_API void ExtractFlipbooksFromSprites(TMap<FString, TArray<UPaperSprite*> >& OutSpriteFlipbookMap, const TArray<UPaperSprite*>& Sprites, const TArray<FString>& SpriteNames);
 
 private:
 
-	static FString GetCleanerSpriteName(const FString& Name);
-	static bool ExtractSpriteNumber(const FString& String, FString& BareString, int32& Number);
+	static UE_API FString GetCleanerSpriteName(const FString& Name);
+	static UE_API bool ExtractSpriteNumber(const FString& String, FString& BareString, int32& Number);
 
 	FPaperFlipbookHelpers() {}
 };
+
+#undef UE_API
diff --git a/Engine/Plugins/2D/Paper2D/Source/Paper2DEditor/Public/PaperJSONHelpers.h b/Engine/Plugins/2D/Paper2D/Source/Paper2DEditor/Public/PaperJSONHelpers.h
--- a/Engine/Plugins/2D/Paper2D/Source/Paper2DEditor/Public/PaperJSONHelpers.h
+++ b/Engine/Plugins/2D/Paper2D/Source/Paper2DEditor/Public/PaperJSONHelpers.h
@@ -5,47 +5,51 @@
 #include "Math/Vector2D.h"
 #include "Templates/SharedPointer.h"
 
+#define UE_API PAPER2DEDITOR_API
+
 class FJsonValue;
 
 //////////////////////////////////////////////////////////////////////////
 // FPaperJSONHelpers
 
-class PAPER2DEDITOR_API FPaperJSONHelpers
+class FPaperJSONHelpers
 {
 public:
 	// returns the value of Key in Item, or DefaultValue if the Key is missing or the wrong type
-	static FString ReadString(TSharedPtr<class FJsonObject> Item, const FString& Key, const FString& DefaultValue);
+	static UE_API FString ReadString(TSharedPtr<class FJsonObject> Item, const FString& Key, const FString& DefaultValue);
 
 	// Returns the object named Key or nullptr if it is missing or the wrong type
-	static TSharedPtr<class FJsonObject> ReadObject(TSharedPtr<class FJsonObject> Item, const FString& Key);
+	static UE_API TSharedPtr<class FJsonObject> ReadObject(TSharedPtr<class FJsonObject> Item, const FString& Key);
 
 	// Returns the array named Key or nullptr if it is missing or the wrong type
-	static const TArray< TSharedPtr<FJsonValue> >& ReadArray(TSharedPtr<class FJsonObject> Item, const FString& Key);
+	static UE_API const TArray< TSharedPtr<FJsonValue> >& ReadArray(TSharedPtr<class FJsonObject> Item, const FString& Key);
 
 	// Returns the bool named Key or bDefaultIfMissing if it is missing or the wrong type (note: no way to determine errors!)
-	static bool ReadBoolean(const TSharedPtr<class FJsonObject> Item, const FString& Key, bool bDefaultIfMissing);
+	static UE_API bool ReadBoolean(const TSharedPtr<class FJsonObject> Item, const FString& Key, bool bDefaultIfMissing);
 
 	// Returns true if the field named Key is a Number, populating Out_Value, or false if it is missing or the wrong type (Out_Value is set to 0.0f on failure)
-	static bool ReadFloatNoDefault(const TSharedPtr<class FJsonObject> Item, const FString& Key, float& Out_Value);
+	static UE_API bool ReadFloatNoDefault(const TSharedPtr<class FJsonObject> Item, const FString& Key, float& Out_Value);
 
 	// Returns true if the field named Key is a Number, populating Out_Value, or false if it is missing or the wrong type (Out_Value is set to 0.0f on failure)
-	static bool ReadFloatNoDefault(const TSharedPtr<class FJsonObject> Item, const FString& Key, double& Out_Value);
+	static UE_API bool ReadFloatNoDefault(const TSharedPtr<class FJsonObject> Item, const FString& Key, double& Out_Value);
 
 	// Returns true if the field named Key is a Number (truncating it to an integer), populating Out_Value, or false if it is missing or the wrong type (Out_Value is set to 0 on failure)
-	static bool ReadIntegerNoDefault(const TSharedPtr<class FJsonObject> Item, const FString& Key, int32& Out_Value);
+	static UE_API bool ReadIntegerNoDefault(const TSharedPtr<class FJsonObject> Item, const FString& Key, int32& Out_Value);
 
 	// Returns true if the object named Key is a struct containing four integers (x,y,w,h), populating XY and WH with the values)
-	static bool ReadRectangle(const TSharedPtr<class FJsonObject> Item, const FString& Key, FIntPoint& Out_XY, FIntPoint& Out_WH);
+	static UE_API bool ReadRectangle(const TSharedPtr<class FJsonObject> Item, const FString& Key, FIntPoint& Out_XY, FIntPoint& Out_WH);
 
 	// Returns true if the object named Key is a struct containing two floats (w,h), populating WH with the values)
-	static bool ReadSize(const TSharedPtr<class FJsonObject> Item, const FString& Key, FVector2D& Out_WH);
+	static UE_API bool ReadSize(const TSharedPtr<class FJsonObject> Item, const FString& Key, FVector2D& Out_WH);
 
 	// Returns true if the object named Key is a struct containing two floats (x,y), populating WH with the values)
-	static bool ReadXY(const TSharedPtr<class FJsonObject> Item, const FString& Key, FVector2D& Out_WH);
+	static UE_API bool ReadXY(const TSharedPtr<class FJsonObject> Item, const FString& Key, FVector2D& Out_WH);
 
 	// Returns true if the object named Key is a struct containing two floats (x,y), populating XY with the values)
-	static bool ReadIntPoint(const TSharedPtr<class FJsonObject> Item, const FString& Key, FIntPoint& Out_XY);
+	static UE_API bool ReadIntPoint(const TSharedPtr<class FJsonObject> Item, const FString& Key, FIntPoint& Out_XY);
 
 private:
 	FPaperJSONHelpers() {}
 };
+
+#undef UE_API
diff --git a/Engine/Plugins/AI/MassAI/Source/MassAIBehaviorEditor/Private/MassAIBehaviorEditorModule.h b/Engine/Plugins/AI/MassAI/Source/MassAIBehaviorEditor/Private/MassAIBehaviorEditorModule.h
--- a/Engine/Plugins/AI/MassAI/Source/MassAIBehaviorEditor/Private/MassAIBehaviorEditorModule.h
+++ b/Engine/Plugins/AI/MassAI/Source/MassAIBehaviorEditor/Private/MassAIBehaviorEditorModule.h
@@ -4,11 +4,15 @@
 
 #include "Modules/ModuleInterface.h"
 
-class MASSAIBEHAVIOREDITOR_API FMassAIBehaviorEditorModule : public IModuleInterface
+#define UE_API MASSAIBEHAVIOREDITOR_API
+
+class FMassAIBehaviorEditorModule : public IModuleInterface
 {
 public:
 	//~ Begin IModuleInterface
-	virtual void StartupModule() override;
-	virtual void ShutdownModule() override;
+	UE_API virtual void StartupModule() override;
+	UE_API virtual void ShutdownModule() override;
 	//~ End IModuleInterface
 };
+
+#undef UE_API
diff --git a/Engine/Plugins/AI/MassAI/Source/MassNavigationEditor/Private/MassNavigationTestingActor.h b/Engine/Plugins/AI/MassAI/Source/MassNavigationEditor/Private/MassNavigationTestingActor.h
--- a/Engine/Plugins/AI/MassAI/Source/MassNavigationEditor/Private/MassNavigationTestingActor.h
+++ b/Engine/Plugins/AI/MassAI/Source/MassNavigationEditor/Private/MassNavigationTestingActor.h
@@ -9,54 +9,56 @@
 #include "MassZoneGraphNavigationFragments.h"
 #include "MassNavigationTestingActor.generated.h"
 
+#define UE_API MASSNAVIGATIONEDITOR_API
+
 class AZoneGraphData;
 class UZoneGraphTestingComponent;
 class UZoneGraphSubsystem;
 class AMassNavigationTestingActor;
 
 #if UE_ENABLE_DEBUG_DRAWING
-class MASSNAVIGATIONEDITOR_API FMassNavigationTestingSceneProxy final : public FDebugRenderSceneProxy
+class FMassNavigationTestingSceneProxy final : public FDebugRenderSceneProxy
 {
 public:
-	FMassNavigationTestingSceneProxy(const UPrimitiveComponent& InComponent);
+	UE_API FMassNavigationTestingSceneProxy(const UPrimitiveComponent& InComponent);
 	
-	virtual SIZE_T GetTypeHash() const override;
-	virtual FPrimitiveViewRelevance GetViewRelevance(const FSceneView* View) const override;
-	virtual uint32 GetMemoryFootprint(void) const override;
+	UE_API virtual SIZE_T GetTypeHash() const override;
+	UE_API virtual FPrimitiveViewRelevance GetViewRelevance(const FSceneView* View) const override;
+	UE_API virtual uint32 GetMemoryFootprint(void) const override;
 };
 #endif // UE_ENABLE_DEBUG_DRAWING
 
 /** Component for testing MassMovement functionality. */
-UCLASS(ClassGroup = Debug)
-class MASSNAVIGATIONEDITOR_API UMassNavigationTestingComponent : public UDebugDrawComponent
+UCLASS(MinimalAPI, ClassGroup = Debug)
+class UMassNavigationTestingComponent : public UDebugDrawComponent
 {
 	GENERATED_BODY()
 public:
-	UMassNavigationTestingComponent(const FObjectInitializer& ObjectInitializer = FObjectInitializer::Get());
+	UE_API UMassNavigationTestingComponent(const FObjectInitializer& ObjectInitializer = FObjectInitializer::Get());
 
 #if WITH_EDITOR
-	virtual void PostEditChangeProperty(FPropertyChangedEvent& PropertyChangedEvent) override;
+	UE_API virtual void PostEditChangeProperty(FPropertyChangedEvent& PropertyChangedEvent) override;
 #endif // WITH_EDITOR
 
-	virtual void OnRegister() override;
-	virtual void OnUnregister() override;
+	UE_API virtual void OnRegister() override;
+	UE_API virtual void OnUnregister() override;
 
-	virtual FBoxSphereBounds CalcBounds(const FTransform& LocalToWorld) const override;
+	UE_API virtual FBoxSphereBounds CalcBounds(const FTransform& LocalToWorld) const override;
 
 #if UE_ENABLE_DEBUG_DRAWING
-	virtual FDebugRenderSceneProxy* CreateDebugSceneProxy() override;
+	UE_API virtual FDebugRenderSceneProxy* CreateDebugSceneProxy() override;
 #endif
 
-	void UpdateTests();
-	void PinLane();
-	void ClearPinnedLane();
+	UE_API void UpdateTests();
+	UE_API void PinLane();
+	UE_API void ClearPinnedLane();
 
 protected:
 
 #if WITH_EDITOR
-	void OnZoneGraphDataBuildDone(const struct FZoneGraphBuildData& BuildData);
+	UE_API void OnZoneGraphDataBuildDone(const struct FZoneGraphBuildData& BuildData);
 #endif
-	void OnZoneGraphDataChanged(const AZoneGraphData* ZoneGraphData);
+	UE_API void OnZoneGraphDataChanged(const AZoneGraphData* ZoneGraphData);
 
 #if WITH_EDITOR
 	FDelegateHandle OnDataChangedHandle;
@@ -98,24 +100,26 @@ protected:
 };
 
 /** Debug actor to visually test zone graph. */
-UCLASS(hidecategories = (Actor, Input, Collision, Rendering, Replication, Partition, HLOD, Cooking))
-class MASSNAVIGATIONEDITOR_API AMassNavigationTestingActor : public AActor
+UCLASS(MinimalAPI, hidecategories = (Actor, Input, Collision, Rendering, Replication, Partition, HLOD, Cooking))
+class AMassNavigationTestingActor : public AActor
 {
 	GENERATED_BODY()
 public:
-	AMassNavigationTestingActor(const FObjectInitializer& ObjectInitializer = FObjectInitializer::Get());
+	UE_API AMassNavigationTestingActor(const FObjectInitializer& ObjectInitializer = FObjectInitializer::Get());
 
 #if WITH_EDITOR
-	virtual void PostEditMove(bool bFinished) override;
+	UE_API virtual void PostEditMove(bool bFinished) override;
 #endif // WITH_EDITOR
 
 	UFUNCTION(BlueprintCallable, CallInEditor, Category = "Default")
-	void PinLane();
+	UE_API void PinLane();
 	
 	UFUNCTION(BlueprintCallable, CallInEditor, Category = "Default")
-	void ClearPinnedLane();
+	UE_API void ClearPinnedLane();
 
 protected:
 	UPROPERTY(Category = Default, VisibleAnywhere, meta = (AllowPrivateAccess = "true"))
 	TObjectPtr<UMassNavigationTestingComponent> DebugComp;
 };
+
+#undef UE_API
diff --git a/Engine/Plugins/AI/MassAI/Source/MassNavigationEditor/Public/IMassNavigationEditor.h b/Engine/Plugins/AI/MassAI/Source/MassNavigationEditor/Public/IMassNavigationEditor.h
--- a/Engine/Plugins/AI/MassAI/Source/MassNavigationEditor/Public/IMassNavigationEditor.h
+++ b/Engine/Plugins/AI/MassAI/Source/MassNavigationEditor/Public/IMassNavigationEditor.h
@@ -6,7 +6,7 @@
 #include "Toolkits/AssetEditorToolkit.h"
 
 /** MassNavigation Editor public interface */
-class MASSNAVIGATIONEDITOR_API IMassNavigationEditor : public FAssetEditorToolkit
+class IMassNavigationEditor : public FAssetEditorToolkit
 {
 public:
 
diff --git a/Engine/Plugins/AI/MassAI/Source/MassNavigationEditor/Public/MassNavigationEditorModule.h b/Engine/Plugins/AI/MassAI/Source/MassNavigationEditor/Public/MassNavigationEditorModule.h
--- a/Engine/Plugins/AI/MassAI/Source/MassNavigationEditor/Public/MassNavigationEditorModule.h
+++ b/Engine/Plugins/AI/MassAI/Source/MassNavigationEditor/Public/MassNavigationEditorModule.h
@@ -13,12 +13,12 @@
 
 class IMassNavigationEditor;
 
-class MASSNAVIGATIONEDITOR_API FMassNavigationEditorModule : public IModuleInterface
+class FMassNavigationEditorModule : public IModuleInterface
 {
 public:
 	// Begin IModuleInterface
-	virtual void StartupModule() override;
-	virtual void ShutdownModule() override;
+	MASSNAVIGATIONEDITOR_API virtual void StartupModule() override;
+	MASSNAVIGATIONEDITOR_API virtual void ShutdownModule() override;
 	// End IModuleInterface
 
 private:
diff --git a/Engine/Plugins/Animation/AnimationData/Source/AnimationData/Private/AnimSequencerDataModel.h b/Engine/Plugins/Animation/AnimationData/Source/AnimationData/Private/AnimSequencerDataModel.h
--- a/Engine/Plugins/Animation/AnimationData/Source/AnimationData/Private/AnimSequencerDataModel.h
+++ b/Engine/Plugins/Animation/AnimationData/Source/AnimationData/Private/AnimSequencerDataModel.h
@@ -9,6 +9,8 @@
 
 #include "AnimSequencerDataModel.generated.h"
 
+#define UE_API ANIMATIONDATA_API
+
 class UControlRig;
 class UMovieSceneControlRigParameterTrack;
 class UMovieSceneControlRigParameterSection;
@@ -30,26 +32,26 @@ struct FAnimationCurveMetaData
 	FString Comment;
 };
 
-UCLASS()
-class ANIMATIONDATA_API UAnimationSequencerDataModel : public UMovieSceneSequence, public IAnimationDataModel
+UCLASS(MinimalAPI)
+class UAnimationSequencerDataModel : public UMovieSceneSequence, public IAnimationDataModel
 {
 private:
 	GENERATED_BODY()
 public:
-	static int32 RetainFloatCurves;
-	static int32 ValidationMode;
-	static int32 UseDirectFKControlRigMode;
-	static int32 LazyRigHierarchyInitializationMode;
+	static UE_API int32 RetainFloatCurves;
+	static UE_API int32 ValidationMode;
+	static UE_API int32 UseDirectFKControlRigMode;
+	static UE_API int32 LazyRigHierarchyInitializationMode;
 	
 	/** Begin UObject overrides*/
-	virtual void GetPreloadDependencies(TArray<UObject*>& OutDeps) override;
-	virtual void PostDuplicate(bool bDuplicateForPIE) override;
-	virtual void PreSave(FObjectPreSaveContext ObjectSaveContext) override;
+	UE_API virtual void GetPreloadDependencies(TArray<UObject*>& OutDeps) override;
+	UE_API virtual void PostDuplicate(bool bDuplicateForPIE) override;
+	UE_API virtual void PreSave(FObjectPreSaveContext ObjectSaveContext) override;
 	virtual bool IsEditorOnly() const override { return true; }
 #if WITH_EDITOR
-	virtual void WillNeverCacheCookedPlatformDataAgain() override;
-	virtual void PreEditUndo() override;
-	virtual void PostEditUndo() override;
+	UE_API virtual void WillNeverCacheCookedPlatformDataAgain() override;
+	UE_API virtual void PreEditUndo() override;
+	UE_API virtual void PostEditUndo() override;
 #endif
 	/** End UObject overrides */
 
@@ -57,73 +59,73 @@ public:
 	virtual void BindPossessableObject(const FGuid& ObjectId, UObject& PossessedObject, UObject* Context) override {}
 	virtual bool CanPossessObject(UObject& Object, UObject* InPlaybackContext) const override { return false; }
 	virtual void LocateBoundObjects(const FGuid& ObjectId, UObject* Context, TArray<UObject*, TInlineAllocator<1>>& OutObjects) const override {}
-	virtual UMovieScene* GetMovieScene() const override;
-	virtual UObject* GetParentObject(UObject* Object) const override;
+	UE_API virtual UMovieScene* GetMovieScene() const override;
+	UE_API virtual UObject* GetParentObject(UObject* Object) const override;
 	virtual void UnbindPossessableObjects(const FGuid& ObjectId) override {}
 	virtual void UnbindObjects(const FGuid& ObjectId, const TArray<UObject*>& InObjects, UObject* Context) override {}
 	virtual void UnbindInvalidObjects(const FGuid& ObjectId, UObject* Context) override {}
-	virtual void PostLoad() override;
+	UE_API virtual void PostLoad() override;
 #if WITH_EDITORONLY_DATA
-	static void DeclareConstructClasses(TArray<FTopLevelAssetPath>& OutConstructClasses, const UClass* SpecificSubclass);
+	static UE_API void DeclareConstructClasses(TArray<FTopLevelAssetPath>& OutConstructClasses, const UClass* SpecificSubclass);
 #endif
 	/** End UMovieSceneSequence overrides */
 	
 	/** Begin IAnimationDataModel overrides*/
-	virtual double GetPlayLength() const override;
-	virtual int32 GetNumberOfFrames() const override;	
-	virtual int32 GetNumberOfKeys() const override;
-	virtual FFrameRate GetFrameRate() const override;
-	PRAGMA_DISABLE_DEPRECATION_WARNINGS	
+	UE_API virtual double GetPlayLength() const override;
+	UE_API virtual int32 GetNumberOfFrames() const override;	
+	UE_API virtual int32 GetNumberOfKeys() const override;
+	UE_API virtual FFrameRate GetFrameRate() const override;
+	UE_API PRAGMA_DISABLE_DEPRECATION_WARNINGS	
 	virtual const TArray<FBoneAnimationTrack>& GetBoneAnimationTracks() const override;
-	virtual const FBoneAnimationTrack& GetBoneTrackByIndex(int32 TrackIndex) const override;
-	virtual const FBoneAnimationTrack& GetBoneTrackByName(FName TrackName) const override;
-	virtual const FBoneAnimationTrack* FindBoneTrackByName(FName Name) const override;
-	virtual const FBoneAnimationTrack* FindBoneTrackByIndex(int32 BoneIndex) const override;
-	virtual int32 GetBoneTrackIndex(const FBoneAnimationTrack& Track) const override;
-	virtual int32 GetBoneTrackIndexByName(FName TrackName) const override;	
-	virtual bool IsValidBoneTrackIndex(int32 TrackIndex) const override;
-	virtual int32 GetNumBoneTracks() const override;
-	virtual void GetBoneTrackNames(TArray<FName>& OutNames) const override;
-	virtual const FAnimCurveBase* FindCurve(const FAnimationCurveIdentifier& CurveIdentifier) const override;
-	virtual const FFloatCurve* FindFloatCurve(const FAnimationCurveIdentifier& CurveIdentifier) const override;
-	virtual const FTransformCurve* FindTransformCurve(const FAnimationCurveIdentifier& CurveIdentifier) const override;
-	virtual const FRichCurve* FindRichCurve(const FAnimationCurveIdentifier& CurveIdentifier) const override;
-	PRAGMA_ENABLE_DEPRECATION_WARNINGS
+	UE_API virtual const FBoneAnimationTrack& GetBoneTrackByIndex(int32 TrackIndex) const override;
+	UE_API virtual const FBoneAnimationTrack& GetBoneTrackByName(FName TrackName) const override;
+	UE_API virtual const FBoneAnimationTrack* FindBoneTrackByName(FName Name) const override;
+	UE_API virtual const FBoneAnimationTrack* FindBoneTrackByIndex(int32 BoneIndex) const override;
+	UE_API virtual int32 GetBoneTrackIndex(const FBoneAnimationTrack& Track) const override;
+	UE_API virtual int32 GetBoneTrackIndexByName(FName TrackName) const override;	
+	UE_API virtual bool IsValidBoneTrackIndex(int32 TrackIndex) const override;
+	UE_API virtual int32 GetNumBoneTracks() const override;
+	UE_API virtual void GetBoneTrackNames(TArray<FName>& OutNames) const override;
+	UE_API virtual const FAnimCurveBase* FindCurve(const FAnimationCurveIdentifier& CurveIdentifier) const override;
+	UE_API virtual const FFloatCurve* FindFloatCurve(const FAnimationCurveIdentifier& CurveIdentifier) const override;
+	UE_API virtual const FTransformCurve* FindTransformCurve(const FAnimationCurveIdentifier& CurveIdentifier) const override;
+	UE_API virtual const FRichCurve* FindRichCurve(const FAnimationCurveIdentifier& CurveIdentifier) const override;
+	UE_API PRAGMA_ENABLE_DEPRECATION_WARNINGS
 	virtual bool IsValidBoneTrackName(const FName& TrackName) const override;
-	virtual FTransform GetBoneTrackTransform(FName TrackName, const FFrameNumber& FrameNumber) const override;
-	virtual void GetBoneTrackTransforms(FName TrackName, const TArray<FFrameNumber>& FrameNumbers, TArray<FTransform>& OutTransforms) const override;	
-	virtual void GetBoneTrackTransforms(FName TrackName, TArray<FTransform>& OutTransforms) const override;
-	virtual void GetBoneTracksTransform(const TArray<FName>& TrackNames, const FFrameNumber& FrameNumber, TArray<FTransform>& OutTransforms) const override;	
-	virtual FTransform EvaluateBoneTrackTransform(FName TrackName, const FFrameTime& FrameTime, const EAnimInterpolationType& Interpolation) const override;	
-	virtual const FAnimationCurveData& GetCurveData() const override;
-	virtual int32 GetNumberOfTransformCurves() const override;
-	virtual int32 GetNumberOfFloatCurves() const override;
-	virtual const TArray<struct FFloatCurve>& GetFloatCurves() const override;
-	virtual const TArray<struct FTransformCurve>& GetTransformCurves() const override;
-	virtual const FAnimCurveBase& GetCurve(const FAnimationCurveIdentifier& CurveIdentifier) const override;
-	virtual const FFloatCurve& GetFloatCurve(const FAnimationCurveIdentifier& CurveIdentifier) const override;
-	virtual const FTransformCurve& GetTransformCurve(const FAnimationCurveIdentifier& CurveIdentifier) const override;
-	virtual const FRichCurve& GetRichCurve(const FAnimationCurveIdentifier& CurveIdentifier) const override;
-	virtual TArrayView<const FAnimatedBoneAttribute> GetAttributes() const override;
-	virtual int32 GetNumberOfAttributes() const override;
-	virtual int32 GetNumberOfAttributesForBoneIndex(const int32 BoneIndex) const override;
-	virtual void GetAttributesForBone(const FName& BoneName, TArray<const FAnimatedBoneAttribute*>& OutBoneAttributes) const override;
-	virtual const FAnimatedBoneAttribute& GetAttribute(const FAnimationAttributeIdentifier& AttributeIdentifier) const override;
-	virtual const FAnimatedBoneAttribute* FindAttribute(const FAnimationAttributeIdentifier& AttributeIdentifier) const override;
-	virtual UAnimSequence* GetAnimationSequence() const override;
+	UE_API virtual FTransform GetBoneTrackTransform(FName TrackName, const FFrameNumber& FrameNumber) const override;
+	UE_API virtual void GetBoneTrackTransforms(FName TrackName, const TArray<FFrameNumber>& FrameNumbers, TArray<FTransform>& OutTransforms) const override;	
+	UE_API virtual void GetBoneTrackTransforms(FName TrackName, TArray<FTransform>& OutTransforms) const override;
+	UE_API virtual void GetBoneTracksTransform(const TArray<FName>& TrackNames, const FFrameNumber& FrameNumber, TArray<FTransform>& OutTransforms) const override;	
+	UE_API virtual FTransform EvaluateBoneTrackTransform(FName TrackName, const FFrameTime& FrameTime, const EAnimInterpolationType& Interpolation) const override;	
+	UE_API virtual const FAnimationCurveData& GetCurveData() const override;
+	UE_API virtual int32 GetNumberOfTransformCurves() const override;
+	UE_API virtual int32 GetNumberOfFloatCurves() const override;
+	UE_API virtual const TArray<struct FFloatCurve>& GetFloatCurves() const override;
+	UE_API virtual const TArray<struct FTransformCurve>& GetTransformCurves() const override;
+	UE_API virtual const FAnimCurveBase& GetCurve(const FAnimationCurveIdentifier& CurveIdentifier) const override;
+	UE_API virtual const FFloatCurve& GetFloatCurve(const FAnimationCurveIdentifier& CurveIdentifier) const override;
+	UE_API virtual const FTransformCurve& GetTransformCurve(const FAnimationCurveIdentifier& CurveIdentifier) const override;
+	UE_API virtual const FRichCurve& GetRichCurve(const FAnimationCurveIdentifier& CurveIdentifier) const override;
+	UE_API virtual TArrayView<const FAnimatedBoneAttribute> GetAttributes() const override;
+	UE_API virtual int32 GetNumberOfAttributes() const override;
+	UE_API virtual int32 GetNumberOfAttributesForBoneIndex(const int32 BoneIndex) const override;
+	UE_API virtual void GetAttributesForBone(const FName& BoneName, TArray<const FAnimatedBoneAttribute*>& OutBoneAttributes) const override;
+	UE_API virtual const FAnimatedBoneAttribute& GetAttribute(const FAnimationAttributeIdentifier& AttributeIdentifier) const override;
+	UE_API virtual const FAnimatedBoneAttribute* FindAttribute(const FAnimationAttributeIdentifier& AttributeIdentifier) const override;
+	UE_API virtual UAnimSequence* GetAnimationSequence() const override;
 	virtual FAnimDataModelModifiedEvent& GetModifiedEvent() override { return ModifiedEvent; }
-	virtual FGuid GenerateGuid(const FGuidGenerationSettings& InSettings) const override;
+	UE_API virtual FGuid GenerateGuid(const FGuidGenerationSettings& InSettings) const override;
 #if WITH_EDITOR
-	virtual FString GenerateDebugStateString() const override;
+	UE_API virtual FString GenerateDebugStateString() const override;
 #endif
-	virtual TScriptInterface<IAnimationDataController> GetController() override;
-	virtual IAnimationDataModel::FModelNotifier& GetNotifier() override;
+	UE_API virtual TScriptInterface<IAnimationDataController> GetController() override;
+	UE_API virtual IAnimationDataModel::FModelNotifier& GetNotifier() override;
 	virtual FAnimDataModelModifiedDynamicEvent& GetModifiedDynamicEvent() override { return ModifiedEventDynamic; }
-    virtual void Evaluate(FAnimationPoseData& InOutPoseData, const UE::Anim::DataModel::FEvaluationContext& EvaluationContext) const override;
+    UE_API virtual void Evaluate(FAnimationPoseData& InOutPoseData, const UE::Anim::DataModel::FEvaluationContext& EvaluationContext) const override;
 	virtual bool HasBeenPopulated() const override { return bPopulated; }
-	virtual void IterateBoneKeys(const FName& BoneName, TFunction<bool(const FVector3f& Pos, const FQuat4f&, const FVector3f, const FFrameNumber&)> IterationFunction) const override;
+	UE_API virtual void IterateBoneKeys(const FName& BoneName, TFunction<bool(const FVector3f& Pos, const FQuat4f&, const FVector3f, const FFrameNumber&)> IterationFunction) const override;
 protected:
-	virtual void OnNotify(const EAnimDataModelNotifyType& NotifyType, const FAnimDataModelNotifPayload& Payload) override final;
+	UE_API virtual void OnNotify(const EAnimDataModelNotifyType& NotifyType, const FAnimDataModelNotifPayload& Payload) override final;
 
 	virtual void LockEvaluationAndModification() const override final
 	{
@@ -147,40 +149,40 @@ protected:
 	/** End IAnimationDataModel overrides */
 	
 	/** Controller helper functionality */
-	FTransformCurve* FindMutableTransformCurveById(const FAnimationCurveIdentifier& CurveIdentifier);
-	FFloatCurve* FindMutableFloatCurveById(const FAnimationCurveIdentifier& CurveIdentifier);
-	FAnimCurveBase* FindMutableCurveById(const FAnimationCurveIdentifier& CurveIdentifier);
-	FRichCurve* GetMutableRichCurve(const FAnimationCurveIdentifier& CurveIdentifier);
-
-	void RegenerateLegacyCurveData();
-	void UpdateLegacyCurveData();
-
-	void ValidateData() const;
-	void ValidateSequencerData() const;
-	void ValidateControlRigData() const;
-	void ValidateLegacyAgainstControlRigData() const;
+	UE_API FTransformCurve* FindMutableTransformCurveById(const FAnimationCurveIdentifier& CurveIdentifier);
+	UE_API FFloatCurve* FindMutableFloatCurveById(const FAnimationCurveIdentifier& CurveIdentifier);
+	UE_API FAnimCurveBase* FindMutableCurveById(const FAnimationCurveIdentifier& CurveIdentifier);
+	UE_API FRichCurve* GetMutableRichCurve(const FAnimationCurveIdentifier& CurveIdentifier);
+
+	UE_API void RegenerateLegacyCurveData();
+	UE_API void UpdateLegacyCurveData();
+
+	UE_API void ValidateData() const;
+	UE_API void ValidateSequencerData() const;
+	UE_API void ValidateControlRigData() const;
+	UE_API void ValidateLegacyAgainstControlRigData() const;
 	
-	void GeneratePoseData(UControlRig* ControlRig, FAnimationPoseData& InOutPoseData, const UE::Anim::DataModel::FEvaluationContext& EvaluationContext) const;
-	void EvaluateTrack(UMovieSceneControlRigParameterTrack* CR_Track, const UE::Anim::DataModel::FEvaluationContext& EvaluationContext) const;
-
-	UMovieSceneControlRigParameterTrack* GetControlRigTrack() const;
-	UMovieSceneControlRigParameterSection* GetFKControlRigSection() const;
-	URigHierarchy* GetControlRigHierarchy() const;
-	USkeleton* GetSkeleton() const;
-	void InitializeFKControlRig(UFKControlRig* FKControlRig, USkeleton* Skeleton, bool bForceHierarchyInitialization=false) const;
-	UControlRig* GetControlRig() const;
+	UE_API void GeneratePoseData(UControlRig* ControlRig, FAnimationPoseData& InOutPoseData, const UE::Anim::DataModel::FEvaluationContext& EvaluationContext) const;
+	UE_API void EvaluateTrack(UMovieSceneControlRigParameterTrack* CR_Track, const UE::Anim::DataModel::FEvaluationContext& EvaluationContext) const;
+
+	UE_API UMovieSceneControlRigParameterTrack* GetControlRigTrack() const;
+	UE_API UMovieSceneControlRigParameterSection* GetFKControlRigSection() const;
+	UE_API URigHierarchy* GetControlRigHierarchy() const;
+	UE_API USkeleton* GetSkeleton() const;
+	UE_API void InitializeFKControlRig(UFKControlRig* FKControlRig, USkeleton* Skeleton, bool bForceHierarchyInitialization=false) const;
+	UE_API UControlRig* GetControlRig() const;
 	
-	void IterateTransformControlCurve(const FName& BoneName, TFunction<void(const FTransform&, const FFrameNumber&)> IterationFunction, const TArray<FFrameNumber>* InFrameNumbers = nullptr) const;
-	void GenerateTransformKeysForControl(const FName& BoneName, TArray<FTransform>& InOutTransforms, TArray<FFrameNumber>& InOutFrameNumbers) const;
-	void RemoveOutOfDateControls() const;
+	UE_API void IterateTransformControlCurve(const FName& BoneName, TFunction<void(const FTransform&, const FFrameNumber&)> IterationFunction, const TArray<FFrameNumber>* InFrameNumbers = nullptr) const;
+	UE_API void GenerateTransformKeysForControl(const FName& BoneName, TArray<FTransform>& InOutTransforms, TArray<FFrameNumber>& InOutFrameNumbers) const;
+	UE_API void RemoveOutOfDateControls() const;
 	
-	void GenerateTransformKeysForControl(const FName& BoneName, const TArray<FFrameNumber>& FrameNumbers, TArray<FTransform>& InOutTransforms) const;
-	void ClearControlRigData();
+	UE_API void GenerateTransformKeysForControl(const FName& BoneName, const TArray<FFrameNumber>& FrameNumbers, TArray<FTransform>& InOutTransforms) const;
+	UE_API void ClearControlRigData();
 private:	
 	template <typename HasherType>
 	void GenerateStateHash(HasherType& Hasher, const FGuidGenerationSettings& InSettings) const;
 
-	void InitializeRigHierarchy(UFKControlRig* FKControlRig, const USkeleton* Skeleton) const;
+	UE_API void InitializeRigHierarchy(UFKControlRig* FKControlRig, const USkeleton* Skeleton) const;
 	
 	/** Dynamic delegate event allows scripting to register to any broadcast-ed notify. */
 	UPROPERTY(BlueprintAssignable, Transient, Category = AnimationDataModel, meta = (ScriptName = "ModifiedEvent", AllowPrivateAccess = "true"))
@@ -220,3 +222,4 @@ private:
 };
 
 
+#undef UE_API
diff --git a/Engine/Plugins/Animation/AnimationModifierLibrary/Source/AnimationModifierLibrary/Public/MotionExtractorUtilities.h b/Engine/Plugins/Animation/AnimationModifierLibrary/Source/AnimationModifierLibrary/Public/MotionExtractorUtilities.h
--- a/Engine/Plugins/Animation/AnimationModifierLibrary/Source/AnimationModifierLibrary/Public/MotionExtractorUtilities.h
+++ b/Engine/Plugins/Animation/AnimationModifierLibrary/Source/AnimationModifierLibrary/Public/MotionExtractorUtilities.h
@@ -6,14 +6,16 @@
 #include "UObject/Package.h"
 #include "MotionExtractorUtilities.generated.h"
 
+#define UE_API ANIMATIONMODIFIERLIBRARY_API
+
 class UAnimSequence;
 enum class EMotionExtractor_Axis : uint8;
 enum class EMotionExtractor_MotionType : uint8;
 struct FBoneContainer;
 struct FCompactPoseBoneIndex;
 
-UCLASS(meta=(ScriptName="MotionExtractorUtilityLibrary"))
-class ANIMATIONMODIFIERLIBRARY_API UMotionExtractorUtilityLibrary : public UBlueprintFunctionLibrary
+UCLASS(MinimalAPI, meta=(ScriptName="MotionExtractorUtilityLibrary"))
+class UMotionExtractorUtilityLibrary : public UBlueprintFunctionLibrary
 {
 	GENERATED_BODY()
 
@@ -27,7 +29,7 @@ public:
 	* @param Axis				   Which axis/axes the motion for this curve was extracted from
 	*/
 	UFUNCTION(BlueprintCallable, Category = "Motion Extractor Utility")
-	static FName GenerateCurveName(
+	static UE_API FName GenerateCurveName(
 		FName BoneName,
 		EMotionExtractor_MotionType MotionType,
 		EMotionExtractor_Axis Axis);
@@ -42,7 +44,7 @@ public:
 	* @param Axis				   Which axis/axes to extract motion from
 	*/
 	UFUNCTION(BlueprintCallable, Category = "Motion Extractor Utility")
-	static float GetDesiredValue(
+	static UE_API float GetDesiredValue(
 		const FTransform& BoneTransform,
 		const FTransform& LastBoneTransform,
 		float DeltaTime,
@@ -57,7 +59,7 @@ public:
 	* @param SampleRate				Sample rate of the animation. It's recommended to use high values if the animation has very sudden direction changes.
 	*/
 	UFUNCTION(BlueprintCallable, Category = "Motion Extractor Utility")
-	static TArray<FVector2D> GetStoppedRangesFromRootMotion(const UAnimSequence* AnimSequence, float StopSpeedThreshold = 10.0f, float SampleRate = 120.0f);
+	static UE_API TArray<FVector2D> GetStoppedRangesFromRootMotion(const UAnimSequence* AnimSequence, float StopSpeedThreshold = 10.0f, float SampleRate = 120.0f);
 
 	/**
 	* Returns the ranges (X/Start to Y/End) in the specified animation sequence where the animation is considered moving.
@@ -67,14 +69,16 @@ public:
 	* @param SampleRate				Sample rate of the animation. It's recommended to use high values if the animation has very sudden direction changes.
 	*/
 	UFUNCTION(BlueprintCallable, Category = "Motion Extractor Utility")
-	static TArray<FVector2D> GetMovingRangesFromRootMotion(const UAnimSequence* AnimSequence, float StopSpeedThreshold = 10.0f, float SampleRate = 120.0f);
+	static UE_API TArray<FVector2D> GetMovingRangesFromRootMotion(const UAnimSequence* AnimSequence, float StopSpeedThreshold = 10.0f, float SampleRate = 120.0f);
 
 	/** 
 	* Helper function to extract the pose for a given bone at a given time
 	* IMPORTANT: This function expects you to add a MemMark (FMemMark Mark(FMemStack::Get());) at the correct scope if you are using it from outside world's tick
 	*/
-	static FTransform ExtractBoneTransform(UAnimSequence* Animation, const FBoneContainer& BoneContainer, FCompactPoseBoneIndex CompactPoseBoneIndex, float Time, bool bComponentSpace);
+	static UE_API FTransform ExtractBoneTransform(UAnimSequence* Animation, const FBoneContainer& BoneContainer, FCompactPoseBoneIndex CompactPoseBoneIndex, float Time, bool bComponentSpace);
 
 	/** Helper function to calculate the magnitude of a vector only considering a specific axis or axes */
-	static double CalculateMagnitude(const FVector& Vector, EMotionExtractor_Axis Axis);
+	static UE_API double CalculateMagnitude(const FVector& Vector, EMotionExtractor_Axis Axis);
 };
+
+#undef UE_API
diff --git a/Engine/Plugins/Animation/AnimationWarping/Source/Editor/Public/AnimGraph/AnimGraphNode_FootPlacement.h b/Engine/Plugins/Animation/AnimationWarping/Source/Editor/Public/AnimGraph/AnimGraphNode_FootPlacement.h
--- a/Engine/Plugins/Animation/AnimationWarping/Source/Editor/Public/AnimGraph/AnimGraphNode_FootPlacement.h
+++ b/Engine/Plugins/Animation/AnimationWarping/Source/Editor/Public/AnimGraph/AnimGraphNode_FootPlacement.h
@@ -5,10 +5,12 @@
 #include "BoneControllers/AnimNode_FootPlacement.h"
 #include "AnimGraphNode_FootPlacement.generated.h"
 
+#define UE_API ANIMATIONWARPINGEDITOR_API
+
 namespace ENodeTitleType { enum Type : int; }
 
-UCLASS(Experimental)
-class ANIMATIONWARPINGEDITOR_API UAnimGraphNode_FootPlacement : public UAnimGraphNode_SkeletalControlBase
+UCLASS(MinimalAPI, Experimental)
+class UAnimGraphNode_FootPlacement : public UAnimGraphNode_SkeletalControlBase
 {
 	GENERATED_UCLASS_BODY()
 
@@ -17,14 +19,16 @@ class ANIMATIONWARPINGEDITOR_API UAnimGraphNode_FootPlacement : public UAnimGrap
 
 public:
 	// UEdGraphNode interface
-	virtual FText GetNodeTitle(ENodeTitleType::Type TitleType) const override;
-	virtual FText GetTooltipText() const override;
-	virtual FLinearColor GetNodeTitleColor() const override;
+	UE_API virtual FText GetNodeTitle(ENodeTitleType::Type TitleType) const override;
+	UE_API virtual FText GetTooltipText() const override;
+	UE_API virtual FLinearColor GetNodeTitleColor() const override;
 	// End of UEdGraphNode interface
 
 protected:
 	// UAnimGraphNode_SkeletalControlBase interface
-	virtual FText GetControllerDescription() const override;
+	UE_API virtual FText GetControllerDescription() const override;
 	virtual const FAnimNode_SkeletalControlBase* GetNode() const override { return &Node; }
 	// End of UAnimGraphNode_SkeletalControlBase interface
 };
+
+#undef UE_API
diff --git a/Engine/Plugins/Animation/AnimationWarping/Source/Editor/Public/AnimGraph/AnimGraphNode_OffsetRootBone.h b/Engine/Plugins/Animation/AnimationWarping/Source/Editor/Public/AnimGraph/AnimGraphNode_OffsetRootBone.h
--- a/Engine/Plugins/Animation/AnimationWarping/Source/Editor/Public/AnimGraph/AnimGraphNode_OffsetRootBone.h
+++ b/Engine/Plugins/Animation/AnimationWarping/Source/Editor/Public/AnimGraph/AnimGraphNode_OffsetRootBone.h
@@ -7,10 +7,12 @@
 
 #include "AnimGraphNode_OffsetRootBone.generated.h"
 
+#define UE_API ANIMATIONWARPINGEDITOR_API
+
 namespace ENodeTitleType { enum Type : int; }
 
-UCLASS(Experimental)
-class ANIMATIONWARPINGEDITOR_API UAnimGraphNode_OffsetRootBone : public UAnimGraphNode_Base
+UCLASS(MinimalAPI, Experimental)
+class UAnimGraphNode_OffsetRootBone : public UAnimGraphNode_Base
 {
 	GENERATED_UCLASS_BODY()
 
@@ -20,18 +22,20 @@ class ANIMATIONWARPINGEDITOR_API UAnimGraphNode_OffsetRootBone : public UAnimGra
 
 public:
 	// UEdGraphNode interface
-	virtual FText GetNodeTitle(ENodeTitleType::Type TitleType) const override;
-	virtual FLinearColor GetNodeTitleColor() const override;
-	virtual FText GetTooltipText() const override;
-	virtual void CustomizeDetails(IDetailLayoutBuilder& DetailBuilder) override;
-	virtual void PostEditChangeProperty(struct FPropertyChangedEvent& PropertyChangedEvent) override;
-	virtual void GetInputLinkAttributes(FNodeAttributeArray& OutAttributes) const override;
-	virtual void GetOutputLinkAttributes(FNodeAttributeArray& OutAttributes) const override;
+	UE_API virtual FText GetNodeTitle(ENodeTitleType::Type TitleType) const override;
+	UE_API virtual FLinearColor GetNodeTitleColor() const override;
+	UE_API virtual FText GetTooltipText() const override;
+	UE_API virtual void CustomizeDetails(IDetailLayoutBuilder& DetailBuilder) override;
+	UE_API virtual void PostEditChangeProperty(struct FPropertyChangedEvent& PropertyChangedEvent) override;
+	UE_API virtual void GetInputLinkAttributes(FNodeAttributeArray& OutAttributes) const override;
+	UE_API virtual void GetOutputLinkAttributes(FNodeAttributeArray& OutAttributes) const override;
 	// End of UEdGraphNode interface
 
 protected:
 	// UAnimGraphNode_Base interface
-	virtual void CustomizePinData(UEdGraphPin* Pin, FName SourcePropertyName, int32 ArrayIndex) const override;
+	UE_API virtual void CustomizePinData(UEdGraphPin* Pin, FName SourcePropertyName, int32 ArrayIndex) const override;
 	// End of UAnimGraphNode_Base interface
 
 };
+
+#undef UE_API
diff --git a/Engine/Plugins/Animation/AnimationWarping/Source/Editor/Public/AnimGraph/AnimGraphNode_OrientationWarping.h b/Engine/Plugins/Animation/AnimationWarping/Source/Editor/Public/AnimGraph/AnimGraphNode_OrientationWarping.h
--- a/Engine/Plugins/Animation/AnimationWarping/Source/Editor/Public/AnimGraph/AnimGraphNode_OrientationWarping.h
+++ b/Engine/Plugins/Animation/AnimationWarping/Source/Editor/Public/AnimGraph/AnimGraphNode_OrientationWarping.h
@@ -5,10 +5,12 @@
 #include "BoneControllers/AnimNode_OrientationWarping.h"
 #include "AnimGraphNode_OrientationWarping.generated.h"
 
+#define UE_API ANIMATIONWARPINGEDITOR_API
+
 namespace ENodeTitleType { enum Type : int; }
 
-UCLASS()
-class ANIMATIONWARPINGEDITOR_API UAnimGraphNode_OrientationWarping : public UAnimGraphNode_SkeletalControlBase
+UCLASS(MinimalAPI)
+class UAnimGraphNode_OrientationWarping : public UAnimGraphNode_SkeletalControlBase
 {
 	GENERATED_UCLASS_BODY()
 
@@ -17,23 +19,25 @@ class ANIMATIONWARPINGEDITOR_API UAnimGraphNode_OrientationWarping : public UAni
 
 public:
 	// UEdGraphNode interface
-	virtual FText GetNodeTitle(ENodeTitleType::Type TitleType) const override;
-	virtual FLinearColor GetNodeTitleColor() const override;
-	virtual FText GetTooltipText() const override;
-	virtual void CustomizeDetails(IDetailLayoutBuilder& DetailBuilder) override;
-	virtual void PostEditChangeProperty(struct FPropertyChangedEvent& PropertyChangedEvent) override;
-	virtual void GetInputLinkAttributes(FNodeAttributeArray& OutAttributes) const override;
-	virtual void GetOutputLinkAttributes(FNodeAttributeArray& OutAttributes) const override;
-	virtual void ValidateAnimNodeDuringCompilation(USkeleton* ForSkeleton, FCompilerResultsLog& MessageLog) override;
+	UE_API virtual FText GetNodeTitle(ENodeTitleType::Type TitleType) const override;
+	UE_API virtual FLinearColor GetNodeTitleColor() const override;
+	UE_API virtual FText GetTooltipText() const override;
+	UE_API virtual void CustomizeDetails(IDetailLayoutBuilder& DetailBuilder) override;
+	UE_API virtual void PostEditChangeProperty(struct FPropertyChangedEvent& PropertyChangedEvent) override;
+	UE_API virtual void GetInputLinkAttributes(FNodeAttributeArray& OutAttributes) const override;
+	UE_API virtual void GetOutputLinkAttributes(FNodeAttributeArray& OutAttributes) const override;
+	UE_API virtual void ValidateAnimNodeDuringCompilation(USkeleton* ForSkeleton, FCompilerResultsLog& MessageLog) override;
 	// End of UEdGraphNode interface
 
 protected:
 	// UAnimGraphNode_Base interface
-	virtual void CustomizePinData(UEdGraphPin* Pin, FName SourcePropertyName, int32 ArrayIndex) const override;
+	UE_API virtual void CustomizePinData(UEdGraphPin* Pin, FName SourcePropertyName, int32 ArrayIndex) const override;
 	// End of UAnimGraphNode_Base interface
 
 	// UAnimGraphNode_SkeletalControlBase interface
-	virtual FText GetControllerDescription() const override;
+	UE_API virtual FText GetControllerDescription() const override;
 	virtual const FAnimNode_SkeletalControlBase* GetNode() const override { return &Node; }
 	// End of UAnimGraphNode_SkeletalControlBase interface
 };
+
+#undef UE_API
diff --git a/Engine/Plugins/Animation/AnimationWarping/Source/Editor/Public/AnimGraph/AnimGraphNode_OverrideRootMotion.h b/Engine/Plugins/Animation/AnimationWarping/Source/Editor/Public/AnimGraph/AnimGraphNode_OverrideRootMotion.h
--- a/Engine/Plugins/Animation/AnimationWarping/Source/Editor/Public/AnimGraph/AnimGraphNode_OverrideRootMotion.h
+++ b/Engine/Plugins/Animation/AnimationWarping/Source/Editor/Public/AnimGraph/AnimGraphNode_OverrideRootMotion.h
@@ -7,10 +7,12 @@
 
 #include "AnimGraphNode_OverrideRootMotion.generated.h"
 
+#define UE_API ANIMATIONWARPINGEDITOR_API
+
 namespace ENodeTitleType { enum Type : int; }
 
-UCLASS(Experimental)
-class ANIMATIONWARPINGEDITOR_API UAnimGraphNode_OverrideRootMotion : public UAnimGraphNode_Base
+UCLASS(MinimalAPI, Experimental)
+class UAnimGraphNode_OverrideRootMotion : public UAnimGraphNode_Base
 {
 	GENERATED_UCLASS_BODY()
 
@@ -20,11 +22,13 @@ class ANIMATIONWARPINGEDITOR_API UAnimGraphNode_OverrideRootMotion : public UAni
 
 public:
 	// UEdGraphNode interface
-	virtual FText GetNodeTitle(ENodeTitleType::Type TitleType) const override;
-	virtual FLinearColor GetNodeTitleColor() const override;
-	virtual FText GetTooltipText() const override;
-	virtual void GetInputLinkAttributes(FNodeAttributeArray& OutAttributes) const override;
-	virtual void GetOutputLinkAttributes(FNodeAttributeArray& OutAttributes) const override;
+	UE_API virtual FText GetNodeTitle(ENodeTitleType::Type TitleType) const override;
+	UE_API virtual FLinearColor GetNodeTitleColor() const override;
+	UE_API virtual FText GetTooltipText() const override;
+	UE_API virtual void GetInputLinkAttributes(FNodeAttributeArray& OutAttributes) const override;
+	UE_API virtual void GetOutputLinkAttributes(FNodeAttributeArray& OutAttributes) const override;
 	// End of UEdGraphNode interface
 
 };
+
+#undef UE_API
diff --git a/Engine/Plugins/Animation/AnimationWarping/Source/Editor/Public/AnimGraph/AnimGraphNode_SlopeWarping.h b/Engine/Plugins/Animation/AnimationWarping/Source/Editor/Public/AnimGraph/AnimGraphNode_SlopeWarping.h
--- a/Engine/Plugins/Animation/AnimationWarping/Source/Editor/Public/AnimGraph/AnimGraphNode_SlopeWarping.h
+++ b/Engine/Plugins/Animation/AnimationWarping/Source/Editor/Public/AnimGraph/AnimGraphNode_SlopeWarping.h
@@ -7,10 +7,12 @@
 
 #include "AnimGraphNode_SlopeWarping.generated.h"
 
+#define UE_API ANIMATIONWARPINGEDITOR_API
+
 namespace ENodeTitleType { enum Type : int; }
 
-UCLASS(Experimental)
-class ANIMATIONWARPINGEDITOR_API UAnimGraphNode_SlopeWarping : public UAnimGraphNode_SkeletalControlBase
+UCLASS(MinimalAPI, Experimental)
+class UAnimGraphNode_SlopeWarping : public UAnimGraphNode_SkeletalControlBase
 {
 	GENERATED_UCLASS_BODY()
 
@@ -19,9 +21,9 @@ class ANIMATIONWARPINGEDITOR_API UAnimGraphNode_SlopeWarping : public UAnimGraph
 
 public:
 	// UEdGraphNode interface
-	virtual FText GetNodeTitle(ENodeTitleType::Type TitleType) const override;
-	virtual FText GetTooltipText() const override;
-	virtual FLinearColor GetNodeTitleColor() const override;
+	UE_API virtual FText GetNodeTitle(ENodeTitleType::Type TitleType) const override;
+	UE_API virtual FText GetTooltipText() const override;
+	UE_API virtual FLinearColor GetNodeTitleColor() const override;
 	// End of UEdGraphNode interface
 
 protected:
@@ -29,7 +31,9 @@ protected:
 	// End of UAnimGraphNode_Base interface
 
 	// UAnimGraphNode_SkeletalControlBase interface
-	virtual FText GetControllerDescription() const override;
+	UE_API virtual FText GetControllerDescription() const override;
 	virtual const FAnimNode_SkeletalControlBase* GetNode() const override { return &Node; }
 	// End of UAnimGraphNode_SkeletalControlBase interface
 };
+
+#undef UE_API
diff --git a/Engine/Plugins/Animation/AnimationWarping/Source/Editor/Public/AnimGraph/AnimGraphNode_Steering.h b/Engine/Plugins/Animation/AnimationWarping/Source/Editor/Public/AnimGraph/AnimGraphNode_Steering.h
--- a/Engine/Plugins/Animation/AnimationWarping/Source/Editor/Public/AnimGraph/AnimGraphNode_Steering.h
+++ b/Engine/Plugins/Animation/AnimationWarping/Source/Editor/Public/AnimGraph/AnimGraphNode_Steering.h
@@ -5,10 +5,12 @@
 #include "BoneControllers/AnimNode_Steering.h"
 #include "AnimGraphNode_Steering.generated.h"
 
+#define UE_API ANIMATIONWARPINGEDITOR_API
+
 namespace ENodeTitleType { enum Type : int; }
 
-UCLASS(Experimental)
-class ANIMATIONWARPINGEDITOR_API UAnimGraphNode_Steering : public UAnimGraphNode_SkeletalControlBase
+UCLASS(MinimalAPI, Experimental)
+class UAnimGraphNode_Steering : public UAnimGraphNode_SkeletalControlBase
 {
 	GENERATED_UCLASS_BODY()
 
@@ -17,23 +19,25 @@ class ANIMATIONWARPINGEDITOR_API UAnimGraphNode_Steering : public UAnimGraphNode
 
 public:
 	// UEdGraphNode interface
-	virtual FText GetNodeTitle(ENodeTitleType::Type TitleType) const override;
-	virtual FLinearColor GetNodeTitleColor() const override;
-	virtual FText GetTooltipText() const override;
-	virtual void CustomizeDetails(IDetailLayoutBuilder& DetailBuilder) override;
-	virtual void PostEditChangeProperty(struct FPropertyChangedEvent& PropertyChangedEvent) override;
-	virtual void GetInputLinkAttributes(FNodeAttributeArray& OutAttributes) const override;
-	virtual void GetOutputLinkAttributes(FNodeAttributeArray& OutAttributes) const override;
-	virtual void ValidateAnimNodeDuringCompilation(USkeleton* ForSkeleton, FCompilerResultsLog& MessageLog) override;
+	UE_API virtual FText GetNodeTitle(ENodeTitleType::Type TitleType) const override;
+	UE_API virtual FLinearColor GetNodeTitleColor() const override;
+	UE_API virtual FText GetTooltipText() const override;
+	UE_API virtual void CustomizeDetails(IDetailLayoutBuilder& DetailBuilder) override;
+	UE_API virtual void PostEditChangeProperty(struct FPropertyChangedEvent& PropertyChangedEvent) override;
+	UE_API virtual void GetInputLinkAttributes(FNodeAttributeArray& OutAttributes) const override;
+	UE_API virtual void GetOutputLinkAttributes(FNodeAttributeArray& OutAttributes) const override;
+	UE_API virtual void ValidateAnimNodeDuringCompilation(USkeleton* ForSkeleton, FCompilerResultsLog& MessageLog) override;
 	// End of UEdGraphNode interface
 
 protected:
 	// UAnimGraphNode_Base interface
-	virtual void CustomizePinData(UEdGraphPin* Pin, FName SourcePropertyName, int32 ArrayIndex) const override;
+	UE_API virtual void CustomizePinData(UEdGraphPin* Pin, FName SourcePropertyName, int32 ArrayIndex) const override;
 	// End of UAnimGraphNode_Base interface
 
 	// UAnimGraphNode_SkeletalControlBase interface
-	virtual FText GetControllerDescription() const override;
+	UE_API virtual FText GetControllerDescription() const override;
 	virtual const FAnimNode_SkeletalControlBase* GetNode() const override { return &Node; }
 	// End of UAnimGraphNode_SkeletalControlBase interface
 };
+
+#undef UE_API
diff --git a/Engine/Plugins/Animation/AnimationWarping/Source/Editor/Public/AnimGraph/AnimGraphNode_StrideWarping.h b/Engine/Plugins/Animation/AnimationWarping/Source/Editor/Public/AnimGraph/AnimGraphNode_StrideWarping.h
--- a/Engine/Plugins/Animation/AnimationWarping/Source/Editor/Public/AnimGraph/AnimGraphNode_StrideWarping.h
+++ b/Engine/Plugins/Animation/AnimationWarping/Source/Editor/Public/AnimGraph/AnimGraphNode_StrideWarping.h
@@ -7,10 +7,12 @@
 
 #include "AnimGraphNode_StrideWarping.generated.h"
 
+#define UE_API ANIMATIONWARPINGEDITOR_API
+
 namespace ENodeTitleType { enum Type : int; }
 
-UCLASS()
-class ANIMATIONWARPINGEDITOR_API UAnimGraphNode_StrideWarping : public UAnimGraphNode_SkeletalControlBase
+UCLASS(MinimalAPI)
+class UAnimGraphNode_StrideWarping : public UAnimGraphNode_SkeletalControlBase
 {
 	GENERATED_UCLASS_BODY()
 
@@ -19,23 +21,25 @@ class ANIMATIONWARPINGEDITOR_API UAnimGraphNode_StrideWarping : public UAnimGrap
 
 public:
 	// UEdGraphNode interface
-	virtual FText GetNodeTitle(ENodeTitleType::Type TitleType) const override;
-	virtual FLinearColor GetNodeTitleColor() const override;
-	virtual FText GetTooltipText() const override;
-	virtual void CustomizeDetails(IDetailLayoutBuilder& DetailBuilder) override;
-	virtual void PostEditChangeProperty(struct FPropertyChangedEvent& PropertyChangedEvent) override;
-	virtual void GetInputLinkAttributes(FNodeAttributeArray& OutAttributes) const override;
-	virtual void GetOutputLinkAttributes(FNodeAttributeArray& OutAttributes) const override;
-	virtual void ValidateAnimNodeDuringCompilation(USkeleton* ForSkeleton, FCompilerResultsLog& MessageLog) override;
+	UE_API virtual FText GetNodeTitle(ENodeTitleType::Type TitleType) const override;
+	UE_API virtual FLinearColor GetNodeTitleColor() const override;
+	UE_API virtual FText GetTooltipText() const override;
+	UE_API virtual void CustomizeDetails(IDetailLayoutBuilder& DetailBuilder) override;
+	UE_API virtual void PostEditChangeProperty(struct FPropertyChangedEvent& PropertyChangedEvent) override;
+	UE_API virtual void GetInputLinkAttributes(FNodeAttributeArray& OutAttributes) const override;
+	UE_API virtual void GetOutputLinkAttributes(FNodeAttributeArray& OutAttributes) const override;
+	UE_API virtual void ValidateAnimNodeDuringCompilation(USkeleton* ForSkeleton, FCompilerResultsLog& MessageLog) override;
 	// End of UEdGraphNode interface
 
 protected:
 	// UAnimGraphNode_Base interface
-	virtual void CustomizePinData(UEdGraphPin* Pin, FName SourcePropertyName, int32 ArrayIndex) const override;
+	UE_API virtual void CustomizePinData(UEdGraphPin* Pin, FName SourcePropertyName, int32 ArrayIndex) const override;
 	// End of UAnimGraphNode_Base interface
 
 	// UAnimGraphNode_SkeletalControlBase interface
-	virtual FText GetControllerDescription() const override;
+	UE_API virtual FText GetControllerDescription() const override;
 	virtual const FAnimNode_SkeletalControlBase* GetNode() const override { return &Node; }
 	// End of UAnimGraphNode_SkeletalControlBase interface
 };
+
+#undef UE_API
diff --git a/Engine/Plugins/Animation/AnimationWarping/Source/Editor/Public/AnimGraph/AnimGraphNode_WarpTest.h b/Engine/Plugins/Animation/AnimationWarping/Source/Editor/Public/AnimGraph/AnimGraphNode_WarpTest.h
--- a/Engine/Plugins/Animation/AnimationWarping/Source/Editor/Public/AnimGraph/AnimGraphNode_WarpTest.h
+++ b/Engine/Plugins/Animation/AnimationWarping/Source/Editor/Public/AnimGraph/AnimGraphNode_WarpTest.h
@@ -6,10 +6,12 @@
 #include "BoneControllers/AnimNode_WarpTest.h"
 #include "AnimGraphNode_WarpTest.generated.h"
 
+#define UE_API ANIMATIONWARPINGEDITOR_API
+
 namespace ENodeTitleType { enum Type : int; }
 
-UCLASS(Experimental)
-class ANIMATIONWARPINGEDITOR_API UAnimGraphNode_WarpTest : public UAnimGraphNode_Base
+UCLASS(MinimalAPI, Experimental)
+class UAnimGraphNode_WarpTest : public UAnimGraphNode_Base
 {
 	GENERATED_UCLASS_BODY()
 
@@ -18,8 +20,10 @@ class ANIMATIONWARPINGEDITOR_API UAnimGraphNode_WarpTest : public UAnimGraphNode
 
 public:
 	// UEdGraphNode interface
-	virtual FText GetNodeTitle(ENodeTitleType::Type TitleType) const override;
-	virtual FLinearColor GetNodeTitleColor() const override;
-	virtual FText GetTooltipText() const override;
+	UE_API virtual FText GetNodeTitle(ENodeTitleType::Type TitleType) const override;
+	UE_API virtual FLinearColor GetNodeTitleColor() const override;
+	UE_API virtual FText GetTooltipText() const override;
 	// End of UEdGraphNode interface
 };
+
+#undef UE_API
diff --git a/Engine/Plugins/Animation/BlendSpaceMotionAnalysis/Source/BlendSpaceMotionAnalysis/Public/BlendSpaceMotionAnalysis.h b/Engine/Plugins/Animation/BlendSpaceMotionAnalysis/Source/BlendSpaceMotionAnalysis/Public/BlendSpaceMotionAnalysis.h
--- a/Engine/Plugins/Animation/BlendSpaceMotionAnalysis/Source/BlendSpaceMotionAnalysis/Public/BlendSpaceMotionAnalysis.h
+++ b/Engine/Plugins/Animation/BlendSpaceMotionAnalysis/Source/BlendSpaceMotionAnalysis/Public/BlendSpaceMotionAnalysis.h
@@ -7,6 +7,8 @@
 #include "LocomotionAnalysis.h"
 #include "BlendSpaceMotionAnalysis.generated.h"
 
+#define UE_API BLENDSPACEMOTIONANALYSIS_API
+
 //======================================================================================================================
 UCLASS()
 class UCachedMotionAnalysisProperties : public UCachedAnalysisProperties
@@ -18,9 +20,11 @@ public:
 };
 
 //======================================================================================================================
-class BLENDSPACEMOTIONANALYSIS_API FBlendSpaceMotionAnalysis : public IModuleInterface 
+class FBlendSpaceMotionAnalysis : public IModuleInterface 
 {
 public:
-	void StartupModule() override;
-	void ShutdownModule() override;
+	UE_API void StartupModule() override;
+	UE_API void ShutdownModule() override;
 };
+
+#undef UE_API
diff --git a/Engine/Plugins/Animation/BlendStack/Source/Editor/Public/AnimGraphNode_BlendStack.h b/Engine/Plugins/Animation/BlendStack/Source/Editor/Public/AnimGraphNode_BlendStack.h
--- a/Engine/Plugins/Animation/BlendStack/Source/Editor/Public/AnimGraphNode_BlendStack.h
+++ b/Engine/Plugins/Animation/BlendStack/Source/Editor/Public/AnimGraphNode_BlendStack.h
@@ -6,35 +6,37 @@
 #include "BlendStack/AnimNode_BlendStack.h"
 #include "AnimGraphNode_BlendStack.generated.h"
 
+#define UE_API BLENDSTACKEDITOR_API
+
 class UAnimGraphNode_BlendStackInput;
 
-UCLASS(Abstract)
-class BLENDSTACKEDITOR_API UAnimGraphNode_BlendStack_Base : public UAnimGraphNode_AssetPlayerBase
+UCLASS(MinimalAPI, Abstract)
+class UAnimGraphNode_BlendStack_Base : public UAnimGraphNode_AssetPlayerBase
 {
 	GENERATED_BODY()
 
 public:
-	virtual void GetOutputLinkAttributes(FNodeAttributeArray& OutAttributes) const override;
-	virtual void OnProcessDuringCompilation(IAnimBlueprintCompilationContext& InCompilationContext, IAnimBlueprintGeneratedClassCompiledData& OutCompiledData) override;
+	UE_API virtual void GetOutputLinkAttributes(FNodeAttributeArray& OutAttributes) const override;
+	UE_API virtual void OnProcessDuringCompilation(IAnimBlueprintCompilationContext& InCompilationContext, IAnimBlueprintGeneratedClassCompiledData& OutCompiledData) override;
 
-	virtual bool DoesSupportTimeForTransitionGetter() const override;
-	virtual UAnimationAsset* GetAnimationAsset() const override;
-	virtual const TCHAR* GetTimePropertyName() const override;
-	virtual UScriptStruct* GetTimePropertyStruct() const override;
+	UE_API virtual bool DoesSupportTimeForTransitionGetter() const override;
+	UE_API virtual UAnimationAsset* GetAnimationAsset() const override;
+	UE_API virtual const TCHAR* GetTimePropertyName() const override;
+	UE_API virtual UScriptStruct* GetTimePropertyStruct() const override;
 
-	virtual void PostPlacedNewNode() override;
-	virtual UObject* GetJumpTargetForDoubleClick() const override;
-	virtual void JumpToDefinition() const override;
-	virtual void DestroyNode() override;
-	virtual void PostPasteNode() override;
-	virtual void Serialize(FArchive& Ar) override;
+	UE_API virtual void PostPlacedNewNode() override;
+	UE_API virtual UObject* GetJumpTargetForDoubleClick() const override;
+	UE_API virtual void JumpToDefinition() const override;
+	UE_API virtual void DestroyNode() override;
+	UE_API virtual void PostPasteNode() override;
+	UE_API virtual void Serialize(FArchive& Ar) override;
 
-	virtual TArray<UEdGraph*> GetSubGraphs() const override;
-	virtual void BakeDataDuringCompilation(class FCompilerResultsLog& MessageLog) override;
+	UE_API virtual TArray<UEdGraph*> GetSubGraphs() const override;
+	UE_API virtual void BakeDataDuringCompilation(class FCompilerResultsLog& MessageLog) override;
 
 protected:
 	// Helper function for compilation
-	void ExpandGraphAndProcessNodes(
+	UE_API void ExpandGraphAndProcessNodes(
 		int GraphIndex,
 		UEdGraph* SourceGraph, 
 		UAnimGraphNode_Base* SourceRootNode, TArrayView<UAnimGraphNode_BlendStackInput*> SourceInputNode,
@@ -42,11 +44,11 @@ protected:
 		UAnimGraphNode_Base*& OutRootNode, TArrayView<UAnimGraphNode_BlendStackInput*> OutInputNodes);
 
 	virtual FAnimNode_BlendStack_Standalone* GetBlendStackNode() const PURE_VIRTUAL(UAnimGraphNode_BlendStack_Base::GetBlendStackNode, return nullptr;);
-	int32 GetMaxActiveBlends() const;
+	UE_API int32 GetMaxActiveBlends() const;
 
 private:
 
-	void CreateGraph();
+	UE_API void CreateGraph();
 
 	UPROPERTY()
 	TObjectPtr<UEdGraph> BoundGraph = nullptr;
@@ -67,4 +69,6 @@ class UAnimGraphNode_BlendStack : public UAnimGraphNode_BlendStack_Base
 
 protected:
 	virtual FAnimNode_BlendStack_Standalone* GetBlendStackNode() const override { return (FAnimNode_BlendStack_Standalone*)(&Node); }
-};
\ No newline at end of file
+};
+
+#undef UE_API
diff --git a/Engine/Plugins/Animation/ControlRig/Source/ControlRigDeveloper/Public/ControlRigBlueprint.h b/Engine/Plugins/Animation/ControlRig/Source/ControlRigDeveloper/Public/ControlRigBlueprint.h
--- a/Engine/Plugins/Animation/ControlRig/Source/ControlRigDeveloper/Public/ControlRigBlueprint.h
+++ b/Engine/Plugins/Animation/ControlRig/Source/ControlRigDeveloper/Public/ControlRigBlueprint.h
@@ -42,7 +42,7 @@ enum class EControlRigType : uint8
 };
 
 USTRUCT(BlueprintType)
-struct CONTROLRIGDEVELOPER_API FModuleReferenceData
+struct FModuleReferenceData
 {
 	GENERATED_BODY()
 
@@ -69,48 +69,48 @@ public:
 	FSoftClassPath ReferencedModule;
 };
 
-UCLASS(BlueprintType, meta=(IgnoreClassThumbnail))
-class CONTROLRIGDEVELOPER_API UControlRigBlueprint : public URigVMBlueprint, public IInterface_PreviewMeshProvider, public IRigHierarchyProvider
+UCLASS(MinimalAPI, BlueprintType, meta=(IgnoreClassThumbnail))
+class UControlRigBlueprint : public URigVMBlueprint, public IInterface_PreviewMeshProvider, public IRigHierarchyProvider
 {
 	GENERATED_UCLASS_BODY()
 
 public:
-	UControlRigBlueprint();
+	CONTROLRIGDEVELOPER_API UControlRigBlueprint();
 
 	//  --- IRigVMClientHost interface ---
 	virtual UClass* GetRigVMSchemaClass() const override { return UControlRigSchema::StaticClass(); }
 	virtual UScriptStruct* GetRigVMExecuteContextStruct() const override { return FControlRigExecuteContext::StaticStruct(); }
-	virtual UClass* GetRigVMEdGraphClass() const override;
-	virtual UClass* GetRigVMEdGraphNodeClass() const override;
-	virtual UClass* GetRigVMEdGraphSchemaClass() const override;
-	virtual UClass* GetRigVMEditorSettingsClass() const override;
+	CONTROLRIGDEVELOPER_API virtual UClass* GetRigVMEdGraphClass() const override;
+	CONTROLRIGDEVELOPER_API virtual UClass* GetRigVMEdGraphNodeClass() const override;
+	CONTROLRIGDEVELOPER_API virtual UClass* GetRigVMEdGraphSchemaClass() const override;
+	CONTROLRIGDEVELOPER_API virtual UClass* GetRigVMEditorSettingsClass() const override;
 
 	// URigVMBlueprint interface
 	virtual UClass* GetRigVMBlueprintGeneratedClassPrototype() const override { return UControlRigBlueprintGeneratedClass::StaticClass(); }
-	virtual TArray<FString> GeneratePythonCommands(const FString InNewBlueprintName) override;
-	virtual void GetPreloadDependencies(TArray<UObject*>& OutDeps) override;
+	CONTROLRIGDEVELOPER_API virtual TArray<FString> GeneratePythonCommands(const FString InNewBlueprintName) override;
+	CONTROLRIGDEVELOPER_API virtual void GetPreloadDependencies(TArray<UObject*>& OutDeps) override;
 #if WITH_EDITOR
-	virtual const FLazyName& GetPanelPinFactoryName() const override;
+	CONTROLRIGDEVELOPER_API virtual const FLazyName& GetPanelPinFactoryName() const override;
 	static inline const FLazyName ControlRigPanelNodeFactoryName = FLazyName(TEXT("FControlRigGraphPanelPinFactory"));
-	virtual IRigVMEditorModule* GetEditorModule() const override;
+	CONTROLRIGDEVELOPER_API virtual IRigVMEditorModule* GetEditorModule() const override;
 #endif
 
-	virtual void Serialize(FArchive& Ar) override;
+	CONTROLRIGDEVELOPER_API virtual void Serialize(FArchive& Ar) override;
 
 #if WITH_EDITOR
 
 	// UBlueprint interface
-	virtual UClass* RegenerateClass(UClass* ClassToRegenerate, UObject* PreviousCDO) override;
+	CONTROLRIGDEVELOPER_API virtual UClass* RegenerateClass(UClass* ClassToRegenerate, UObject* PreviousCDO) override;
 	virtual bool SupportedByDefaultBlueprintFactory() const override { return false; }
 	virtual bool IsValidForBytecodeOnlyRecompile() const override { return false; }
-	virtual void GetTypeActions(FBlueprintActionDatabaseRegistrar& ActionRegistrar) const override;
-	virtual void GetInstanceActions(FBlueprintActionDatabaseRegistrar& ActionRegistrar) const override;
-	virtual void PreSave(FObjectPreSaveContext ObjectSaveContext) override;
-	virtual void PostLoad() override;
-	virtual void PostTransacted(const FTransactionObjectEvent& TransactionEvent) override;
-	virtual void PostDuplicate(bool bDuplicateForPIE) override;
-	virtual void PostRename(UObject* OldOuter, const FName OldName) override;
-	virtual bool RequiresForceLoadMembers(UObject* InObject) const override;
+	CONTROLRIGDEVELOPER_API virtual void GetTypeActions(FBlueprintActionDatabaseRegistrar& ActionRegistrar) const override;
+	CONTROLRIGDEVELOPER_API virtual void GetInstanceActions(FBlueprintActionDatabaseRegistrar& ActionRegistrar) const override;
+	CONTROLRIGDEVELOPER_API virtual void PreSave(FObjectPreSaveContext ObjectSaveContext) override;
+	CONTROLRIGDEVELOPER_API virtual void PostLoad() override;
+	CONTROLRIGDEVELOPER_API virtual void PostTransacted(const FTransactionObjectEvent& TransactionEvent) override;
+	CONTROLRIGDEVELOPER_API virtual void PostDuplicate(bool bDuplicateForPIE) override;
+	CONTROLRIGDEVELOPER_API virtual void PostRename(UObject* OldOuter, const FName OldName) override;
+	CONTROLRIGDEVELOPER_API virtual bool RequiresForceLoadMembers(UObject* InObject) const override;
 
 	virtual bool SupportsGlobalVariables() const override { return true; }
 	virtual bool SupportsLocalVariables() const override { return !IsModularRig(); }
@@ -119,15 +119,15 @@ public:
 
 
 	// UObject interface
-	virtual void PostEditChangeProperty(FPropertyChangedEvent& PropertyChangedEvent) override;
-	virtual void PostEditChangeChainProperty(struct FPropertyChangedChainEvent& PropertyChangedEvent) override;
+	CONTROLRIGDEVELOPER_API virtual void PostEditChangeProperty(FPropertyChangedEvent& PropertyChangedEvent) override;
+	CONTROLRIGDEVELOPER_API virtual void PostEditChangeChainProperty(struct FPropertyChangedChainEvent& PropertyChangedEvent) override;
 
 #endif	// #if WITH_EDITOR
 
 	UFUNCTION(BlueprintCallable, Category = "VM")
-	UClass* GetControlRigClass() const;
+	CONTROLRIGDEVELOPER_API UClass* GetControlRigClass() const;
 
-	bool IsModularRig() const;
+	CONTROLRIGDEVELOPER_API bool IsModularRig() const;
 
 	UFUNCTION(BlueprintCallable, Category = "Control Rig Blueprint")
 	UControlRig* CreateControlRig() { return Cast<UControlRig>(CreateRigVMHost()); }
@@ -137,42 +137,42 @@ public:
 
 	/** IInterface_PreviewMeshProvider interface */
 	UFUNCTION(BlueprintCallable, Category = "Control Rig Blueprint")
-	virtual void SetPreviewMesh(USkeletalMesh* PreviewMesh, bool bMarkAsDirty = true) override;
+	CONTROLRIGDEVELOPER_API virtual void SetPreviewMesh(USkeletalMesh* PreviewMesh, bool bMarkAsDirty = true) override;
 	
 	UFUNCTION(BlueprintCallable, Category = "Control Rig Blueprint")
-	virtual USkeletalMesh* GetPreviewMesh() const override;
+	CONTROLRIGDEVELOPER_API virtual USkeletalMesh* GetPreviewMesh() const override;
 
 	UFUNCTION(BlueprintPure, Category = "Control Rig Blueprint")
-	bool IsControlRigModule() const;
+	CONTROLRIGDEVELOPER_API bool IsControlRigModule() const;
 
 #if WITH_EDITORONLY_DATA
 	
 	bool CanTurnIntoControlRigModule_Blueprint(bool InAutoConvertHierarchy = false) const { return CanTurnIntoControlRigModule(InAutoConvertHierarchy); }
 
-	bool CanTurnIntoControlRigModule(bool InAutoConvertHierarchy, FString* OutErrorMessage = nullptr) const;
+	CONTROLRIGDEVELOPER_API bool CanTurnIntoControlRigModule(bool InAutoConvertHierarchy, FString* OutErrorMessage = nullptr) const;
 
 	UFUNCTION(BlueprintCallable, meta = (DisplayName = "TurnIntoControlRigModule", ScriptName = "TurnIntoControlRigModule"), Category = "Control Rig Blueprint")
 	bool TurnIntoControlRigModule_Blueprint() { return TurnIntoControlRigModule(); }
 
-	bool TurnIntoControlRigModule(bool InAutoConvertHierarchy = false, FString* OutErrorMessage = nullptr);
+	CONTROLRIGDEVELOPER_API bool TurnIntoControlRigModule(bool InAutoConvertHierarchy = false, FString* OutErrorMessage = nullptr);
 
 	UFUNCTION(BlueprintPure, meta = (DisplayName = "CanTurnIntoStandaloneRig", ScriptName = "CanTurnIntoStandaloneRig"), Category = "Control Rig Blueprint")
 	bool CanTurnIntoStandaloneRig_Blueprint() const { return CanTurnIntoStandaloneRig(); }
 
-	bool CanTurnIntoStandaloneRig(FString* OutErrorMessage = nullptr) const;
+	CONTROLRIGDEVELOPER_API bool CanTurnIntoStandaloneRig(FString* OutErrorMessage = nullptr) const;
 
 	UFUNCTION(BlueprintCallable, meta = (DisplayName = "TurnIntoStandaloneRig", ScriptName = "TurnIntoStandaloneRig"), Category = "Control Rig Blueprint")
 	bool TurnIntoStandaloneRig_Blueprint() { return TurnIntoStandaloneRig(); }
 
-	bool TurnIntoStandaloneRig(FString* OutErrorMessage = nullptr);
+	CONTROLRIGDEVELOPER_API bool TurnIntoStandaloneRig(FString* OutErrorMessage = nullptr);
 
 	UFUNCTION(BlueprintCallable, Category = "Control Rig Blueprint")
-	TArray<URigVMNode*> ConvertHierarchyElementsToSpawnerNodes(URigHierarchy* InHierarchy, TArray<FRigElementKey> InKeys, bool bRemoveElements = true);
+	CONTROLRIGDEVELOPER_API TArray<URigVMNode*> ConvertHierarchyElementsToSpawnerNodes(URigHierarchy* InHierarchy, TArray<FRigElementKey> InKeys, bool bRemoveElements = true);
 
 #endif // WITH_EDITORONLY_DATA
 
 	UFUNCTION(BlueprintPure, Category = "Control Rig Blueprint")
-	UTexture2D* GetRigModuleIcon() const;
+	CONTROLRIGDEVELOPER_API UTexture2D* GetRigModuleIcon() const;
 
 	DECLARE_EVENT_OneParam(UControlRigBlueprint, FOnRigTypeChanged, UControlRigBlueprint*);
 
@@ -210,42 +210,42 @@ public:
 	TMap<FRigElementKey, FRigElementKeyCollection> ArrayConnectionMap;
 
 	UFUNCTION(BlueprintPure, Category = "Control Rig Blueprint")
-	TArray<FModuleReferenceData> FindReferencesToModule() const;
+	CONTROLRIGDEVELOPER_API TArray<FModuleReferenceData> FindReferencesToModule() const;
 
-	static EControlRigType GetRigType(const FAssetData& InAsset);
-	static TArray<FSoftObjectPath> GetReferencesToRigModule(const FAssetData& InModuleAsset);
+	static CONTROLRIGDEVELOPER_API EControlRigType GetRigType(const FAssetData& InAsset);
+	static CONTROLRIGDEVELOPER_API TArray<FSoftObjectPath> GetReferencesToRigModule(const FAssetData& InModuleAsset);
 
 	UFUNCTION(BlueprintCallable, Category = "Control Rig Blueprint")
-	void UpdateExposedModuleConnectors() const;
+	CONTROLRIGDEVELOPER_API void UpdateExposedModuleConnectors() const;
 
 #if WITH_EDITOR
-	TArray<FOverrideStatusSubject> GetOverrideSubjects() const;
-	uint32 GetOverrideSubjectsHash() const;
+	CONTROLRIGDEVELOPER_API TArray<FOverrideStatusSubject> GetOverrideSubjects() const;
+	CONTROLRIGDEVELOPER_API uint32 GetOverrideSubjectsHash() const;
 #endif
 	
 protected:
 
-	TArray<FModuleReferenceData> GetModuleReferenceData() const;
+	CONTROLRIGDEVELOPER_API TArray<FModuleReferenceData> GetModuleReferenceData() const;
 
 	FOnRigTypeChanged OnRigTypeChangedDelegate;
 	
-	bool ResolveConnector(const FRigElementKey& DraggedKey, const FRigElementKey& TargetKey, bool bSetupUndoRedo = true);
-	bool ResolveConnectorToArray(const FRigElementKey& DraggedKey, const TArray<FRigElementKey>& TargetKeys, bool bSetupUndoRedo = true);
+	CONTROLRIGDEVELOPER_API bool ResolveConnector(const FRigElementKey& DraggedKey, const FRigElementKey& TargetKey, bool bSetupUndoRedo = true);
+	CONTROLRIGDEVELOPER_API bool ResolveConnectorToArray(const FRigElementKey& DraggedKey, const TArray<FRigElementKey>& TargetKeys, bool bSetupUndoRedo = true);
 
-	void UpdateConnectionMapFromModel();
+	CONTROLRIGDEVELOPER_API void UpdateConnectionMapFromModel();
 
 	/** Asset searchable information about exposed public functions on this rig */
 	UPROPERTY(AssetRegistrySearchable)
 	TArray<FRigVMOldPublicFunctionData> PublicFunctions_DEPRECATED;
 
-	virtual void SetupDefaultObjectDuringCompilation(URigVMHost* InCDO) override;
+	CONTROLRIGDEVELOPER_API virtual void SetupDefaultObjectDuringCompilation(URigVMHost* InCDO) override;
 
 public:
 
-	virtual void SetupPinRedirectorsForBackwardsCompatibility() override;
+	CONTROLRIGDEVELOPER_API virtual void SetupPinRedirectorsForBackwardsCompatibility() override;
 
 	UFUNCTION(BlueprintCallable, Category = "VM")
-	static TArray<UControlRigBlueprint*> GetCurrentlyOpenRigBlueprints();
+	static CONTROLRIGDEVELOPER_API TArray<UControlRigBlueprint*> GetCurrentlyOpenRigBlueprints();
 
 #if WITH_EDITORONLY_DATA
 	UPROPERTY()
@@ -254,7 +254,7 @@ public:
 	UPROPERTY(EditAnywhere, Category = Shapes)
 	TArray<TSoftObjectPtr<UControlRigShapeLibrary>> ShapeLibraries;
 
-	const FControlRigShapeDefinition* GetControlShapeByName(const FName& InName) const;
+	CONTROLRIGDEVELOPER_API const FControlRigShapeDefinition* GetControlShapeByName(const FName& InName) const;
 
 	UPROPERTY(transient, DuplicateTransient, meta = (DisplayName = "VM Statistics", DisplayAfter = "VMCompileSettings"))
 	FRigVMStatistics Statistics_DEPRECATED;
@@ -265,19 +265,19 @@ public:
 
 #if WITH_EDITOR
 	/** Remove a transient / temporary control used to interact with a pin */
-	FName AddTransientControl(const URigVMUnitNode* InNode, const FRigDirectManipulationTarget& InTarget);
+	CONTROLRIGDEVELOPER_API FName AddTransientControl(const URigVMUnitNode* InNode, const FRigDirectManipulationTarget& InTarget);
 
 	/** Remove a transient / temporary control used to interact with a pin */
-	FName RemoveTransientControl(const URigVMUnitNode* InNode, const FRigDirectManipulationTarget& InTarget);
+	CONTROLRIGDEVELOPER_API FName RemoveTransientControl(const URigVMUnitNode* InNode, const FRigDirectManipulationTarget& InTarget);
 
 	/** Remove a transient / temporary control used to interact with a bone */
-	FName AddTransientControl(const FRigElementKey& InElement);
+	CONTROLRIGDEVELOPER_API FName AddTransientControl(const FRigElementKey& InElement);
 
 	/** Remove a transient / temporary control used to interact with a bone */
-	FName RemoveTransientControl(const FRigElementKey& InElement);
+	CONTROLRIGDEVELOPER_API FName RemoveTransientControl(const FRigElementKey& InElement);
 
 	/** Removes all  transient / temporary control used to interact with pins */
-	void ClearTransientControls();
+	CONTROLRIGDEVELOPER_API void ClearTransientControls();
 
 #endif
 
@@ -301,10 +301,10 @@ public:
 	FModularRigModel ModularRigModel;
 
 	UFUNCTION(BlueprintCallable, Category = "Modules")
-	UModularRigController* GetModularRigController();
+	CONTROLRIGDEVELOPER_API UModularRigController* GetModularRigController();
 
 	UFUNCTION(BlueprintCallable, Category = "Control Rig Blueprint")
-	void RecompileModularRig();
+	CONTROLRIGDEVELOPER_API void RecompileModularRig();
 
 	UPROPERTY(AssetRegistrySearchable)
 	EControlRigType ControlRigType;
@@ -347,34 +347,34 @@ public:
 
 private:
 
-	static TArray<UControlRigBlueprint*> sCurrentlyOpenedRigBlueprints;
+	static CONTROLRIGDEVELOPER_API TArray<UControlRigBlueprint*> sCurrentlyOpenedRigBlueprints;
 
-	virtual void PathDomainSpecificContentOnLoad() override;
-	virtual void GetBackwardsCompatibilityPublicFunctions(TArray<FName>& BackwardsCompatiblePublicFunctions, TMap<URigVMLibraryNode*, FRigVMGraphFunctionHeader>& OldHeaders) override;
-	void PatchRigElementKeyCacheOnLoad();
-	void PatchPropagateToChildren();
+	CONTROLRIGDEVELOPER_API virtual void PathDomainSpecificContentOnLoad() override;
+	CONTROLRIGDEVELOPER_API virtual void GetBackwardsCompatibilityPublicFunctions(TArray<FName>& BackwardsCompatiblePublicFunctions, TMap<URigVMLibraryNode*, FRigVMGraphFunctionHeader>& OldHeaders) override;
+	CONTROLRIGDEVELOPER_API void PatchRigElementKeyCacheOnLoad();
+	CONTROLRIGDEVELOPER_API void PatchPropagateToChildren();
 
 protected:
-	virtual void CreateMemberVariablesOnLoad() override;
-	virtual void PatchVariableNodesOnLoad() override;
+	CONTROLRIGDEVELOPER_API virtual void CreateMemberVariablesOnLoad() override;
+	CONTROLRIGDEVELOPER_API virtual void PatchVariableNodesOnLoad() override;
 
 public:
-	void UpdateElementKeyRedirector(UControlRig* InControlRig) const;
-	void PropagatePoseFromInstanceToBP(UControlRig* InControlRig) const;
-	void PropagatePoseFromBPToInstances() const;
-	void PropagateHierarchyFromBPToInstances() const;
-	void PropagateDrawInstructionsFromBPToInstances() const;
-	void PropagatePropertyFromBPToInstances(FRigElementKey InRigElement, const FProperty* InProperty) const;
-	void PropagatePropertyFromInstanceToBP(FRigElementKey InRigElement, const FProperty* InProperty, UControlRig* InInstance) const;
-	void PropagateModuleHierarchyFromBPToInstances() const;
-	void UpdateModularDependencyDelegates();
-	void OnModularDependencyVMCompiled(UObject* InBlueprint, URigVM* InVM, FRigVMExtendedExecuteContext& InExecuteContext);
-	void OnModularDependencyChanged(URigVMBlueprint* InBlueprint);
-	void RequestConstructionOnAllModules();
-	void RefreshModuleVariables();
-	void RefreshModuleVariables(const FRigModuleReference* InModule);
-	void RefreshModuleConnectors();
-	void RefreshModuleConnectors(const FRigModuleReference* InModule, bool bPropagateHierarchy = true);
+	CONTROLRIGDEVELOPER_API void UpdateElementKeyRedirector(UControlRig* InControlRig) const;
+	CONTROLRIGDEVELOPER_API void PropagatePoseFromInstanceToBP(UControlRig* InControlRig) const;
+	CONTROLRIGDEVELOPER_API void PropagatePoseFromBPToInstances() const;
+	CONTROLRIGDEVELOPER_API void PropagateHierarchyFromBPToInstances() const;
+	CONTROLRIGDEVELOPER_API void PropagateDrawInstructionsFromBPToInstances() const;
+	CONTROLRIGDEVELOPER_API void PropagatePropertyFromBPToInstances(FRigElementKey InRigElement, const FProperty* InProperty) const;
+	CONTROLRIGDEVELOPER_API void PropagatePropertyFromInstanceToBP(FRigElementKey InRigElement, const FProperty* InProperty, UControlRig* InInstance) const;
+	CONTROLRIGDEVELOPER_API void PropagateModuleHierarchyFromBPToInstances() const;
+	CONTROLRIGDEVELOPER_API void UpdateModularDependencyDelegates();
+	CONTROLRIGDEVELOPER_API void OnModularDependencyVMCompiled(UObject* InBlueprint, URigVM* InVM, FRigVMExtendedExecuteContext& InExecuteContext);
+	CONTROLRIGDEVELOPER_API void OnModularDependencyChanged(URigVMBlueprint* InBlueprint);
+	CONTROLRIGDEVELOPER_API void RequestConstructionOnAllModules();
+	CONTROLRIGDEVELOPER_API void RefreshModuleVariables();
+	CONTROLRIGDEVELOPER_API void RefreshModuleVariables(const FRigModuleReference* InModule);
+	CONTROLRIGDEVELOPER_API void RefreshModuleConnectors();
+	CONTROLRIGDEVELOPER_API void RefreshModuleConnectors(const FRigModuleReference* InModule, bool bPropagateHierarchy = true);
 
 	/**
 	* Returns the modified event, which can be used to 
@@ -399,27 +399,27 @@ private:
 	int32 ModulesRecompilationBracket = 0;
 
 
-	void HandleHierarchyModified(ERigHierarchyNotification InNotification, URigHierarchy* InHierarchy, const FRigNotificationSubject& InSubject);
-	void HandleHierarchyElementKeyChanged(const FRigElementKey& InOldKey, const FRigElementKey& InNewKey);
-	void HandleHierarchyComponentKeyChanged(const FRigComponentKey& InOldKey, const FRigComponentKey& InNewKey);
+	CONTROLRIGDEVELOPER_API void HandleHierarchyModified(ERigHierarchyNotification InNotification, URigHierarchy* InHierarchy, const FRigNotificationSubject& InSubject);
+	CONTROLRIGDEVELOPER_API void HandleHierarchyElementKeyChanged(const FRigElementKey& InOldKey, const FRigElementKey& InNewKey);
+	CONTROLRIGDEVELOPER_API void HandleHierarchyComponentKeyChanged(const FRigComponentKey& InOldKey, const FRigComponentKey& InNewKey);
 
-	void HandleRigModulesModified(EModularRigNotification InNotification, const FRigModuleReference* InModule);
+	CONTROLRIGDEVELOPER_API void HandleRigModulesModified(EModularRigNotification InNotification, const FRigModuleReference* InModule);
 
 #if WITH_EDITOR
-	virtual void HandlePackageDone() override;
-	virtual void HandleConfigureRigVMController(const FRigVMClient* InClient, URigVMController* InControllerToConfigure) override;
+	CONTROLRIGDEVELOPER_API virtual void HandlePackageDone() override;
+	CONTROLRIGDEVELOPER_API virtual void HandleConfigureRigVMController(const FRigVMClient* InClient, URigVMController* InControllerToConfigure) override;
 #endif
 
-	void UpdateConnectionMapAfterRename(const FString& InOldModuleName);
+	CONTROLRIGDEVELOPER_API void UpdateConnectionMapAfterRename(const FString& InOldModuleName);
 
 	// Class used to temporarily cache all 
 	// current control values and reapply them
 	// on destruction
-	class CONTROLRIGDEVELOPER_API FControlValueScope
+	class FControlValueScope
 	{
 	public: 
-		FControlValueScope(UControlRigBlueprint* InBlueprint);
-		~FControlValueScope();
+		CONTROLRIGDEVELOPER_API FControlValueScope(UControlRigBlueprint* InBlueprint);
+		CONTROLRIGDEVELOPER_API ~FControlValueScope();
 
 	private:
 
diff --git a/Engine/Plugins/Animation/ControlRig/Source/ControlRigDeveloper/Public/ControlRigBlueprintCompiler.h b/Engine/Plugins/Animation/ControlRig/Source/ControlRigDeveloper/Public/ControlRigBlueprintCompiler.h
--- a/Engine/Plugins/Animation/ControlRig/Source/ControlRigDeveloper/Public/ControlRigBlueprintCompiler.h
+++ b/Engine/Plugins/Animation/ControlRig/Source/ControlRigDeveloper/Public/ControlRigBlueprintCompiler.h
@@ -4,15 +4,17 @@
 
 #include "RigVMBlueprintCompiler.h"
 
-class CONTROLRIGDEVELOPER_API FControlRigBlueprintCompiler : public FRigVMBlueprintCompiler
+#define UE_API CONTROLRIGDEVELOPER_API
+
+class FControlRigBlueprintCompiler : public FRigVMBlueprintCompiler
 {
 public:
 	/** IBlueprintCompiler interface */
-	virtual bool CanCompile(const UBlueprint* Blueprint) override;
-	virtual void Compile(UBlueprint* Blueprint, const FKismetCompilerOptions& CompileOptions, FCompilerResultsLog& Results) override;
+	UE_API virtual bool CanCompile(const UBlueprint* Blueprint) override;
+	UE_API virtual void Compile(UBlueprint* Blueprint, const FKismetCompilerOptions& CompileOptions, FCompilerResultsLog& Results) override;
 };
 
-class CONTROLRIGDEVELOPER_API FControlRigBlueprintCompilerContext : public FRigVMBlueprintCompilerContext
+class FControlRigBlueprintCompilerContext : public FRigVMBlueprintCompilerContext
 {
 public:
 	FControlRigBlueprintCompilerContext(UBlueprint* SourceSketch, FCompilerResultsLog& InMessageLog, const FKismetCompilerOptions& InCompilerOptions)
@@ -21,6 +23,8 @@ public:
 	}
 
 	// FKismetCompilerContext interface
-	virtual void SpawnNewClass(const FString& NewClassName) override;
-	virtual void CopyTermDefaultsToDefaultObject(UObject* DefaultObject) override;
-};
\ No newline at end of file
+	UE_API virtual void SpawnNewClass(const FString& NewClassName) override;
+	UE_API virtual void CopyTermDefaultsToDefaultObject(UObject* DefaultObject) override;
+};
+
+#undef UE_API
diff --git a/Engine/Plugins/Animation/ControlRig/Source/ControlRigDeveloper/Public/ControlRigIOMapping.h b/Engine/Plugins/Animation/ControlRig/Source/ControlRigDeveloper/Public/ControlRigIOMapping.h
--- a/Engine/Plugins/Animation/ControlRig/Source/ControlRigDeveloper/Public/ControlRigIOMapping.h
+++ b/Engine/Plugins/Animation/ControlRig/Source/ControlRigDeveloper/Public/ControlRigIOMapping.h
@@ -7,6 +7,8 @@
 #include "RigVMCore/RigVMExternalVariable.h"
 #include "Styling/SlateTypes.h"
 
+#define UE_API CONTROLRIGDEVELOPER_API
+
 struct FRigVMVariableMappingInfo;
 struct FOptionalPinFromProperty;
 class IDetailLayoutBuilder;
@@ -14,7 +16,7 @@ class UClass;
 class USkeleton;
 class UControlRig;
 
-struct CONTROLRIGDEVELOPER_API FControlRigIOMapping : public TSharedFromThis<FControlRigIOMapping>
+struct FControlRigIOMapping : public TSharedFromThis<FControlRigIOMapping>
 {
 	DECLARE_DELEGATE_RetVal(USkeleton*, FOnGetTargetSkeleton);
 	DECLARE_DELEGATE_RetVal(UClass*, FOnGetTargetClass);
@@ -22,10 +24,10 @@ struct CONTROLRIGDEVELOPER_API FControlRigIOMapping : public TSharedFromThis<FCo
 	DECLARE_DELEGATE_ThreeParams(FOnVariableMappingChanged, const FName& /*PathName*/, const FName& /*Curve*/, bool /*bInput*/);
 
 	FControlRigIOMapping() = delete;
-	FControlRigIOMapping(TMap<FName, FName>& InInputMapping, TMap<FName, FName>& InOutputMapping, TArray<FOptionalPinFromProperty>& InCustomPinProperties);
-	virtual ~FControlRigIOMapping();
+	UE_API FControlRigIOMapping(TMap<FName, FName>& InInputMapping, TMap<FName, FName>& InOutputMapping, TArray<FOptionalPinFromProperty>& InCustomPinProperties);
+	UE_API virtual ~FControlRigIOMapping();
 
-	struct CONTROLRIGDEVELOPER_API FControlsInfo
+	struct FControlsInfo
 	{
 		FName Name;
 		FString DisplayName;
@@ -35,11 +37,11 @@ struct CONTROLRIGDEVELOPER_API FControlRigIOMapping : public TSharedFromThis<FCo
 	};
 
 	// Helper enabling quering controls in editor from a Control Rig class
-	struct CONTROLRIGDEVELOPER_API FRigControlsData
+	struct FRigControlsData
 	{
 		FRigControlsData() = default;
 
-		const TArray<FControlsInfo>& GetControls(const UClass* ControlRigClass, USkeleton* TargetSkeleton) const;
+		UE_API const TArray<FControlsInfo>& GetControls(const UClass* ControlRigClass, USkeleton* TargetSkeleton) const;
 
 	private:
 		// Controls data
@@ -47,7 +49,7 @@ struct CONTROLRIGDEVELOPER_API FControlRigIOMapping : public TSharedFromThis<FCo
 		mutable TArray<FControlsInfo> ControlsInfo;
 	};
 
-	bool CreateVariableMappingWidget(IDetailLayoutBuilder& DetailBuilder);
+	UE_API bool CreateVariableMappingWidget(IDetailLayoutBuilder& DetailBuilder);
 
 	const TMap<FName, FRigVMExternalVariable>& GetInputVariables() const
 	{
@@ -58,9 +60,9 @@ struct CONTROLRIGDEVELOPER_API FControlRigIOMapping : public TSharedFromThis<FCo
 		return OutputVariables;
 	}
 
-	const TArray<FControlRigIOMapping::FControlsInfo>& GetControls() const;
+	UE_API const TArray<FControlRigIOMapping::FControlsInfo>& GetControls() const;
 
-	void RebuildExposedProperties();
+	UE_API void RebuildExposedProperties();
 
 	FOnGetTargetSkeleton& GetOnGetTargetSkeletonDelegate()
 	{
@@ -82,10 +84,10 @@ struct CONTROLRIGDEVELOPER_API FControlRigIOMapping : public TSharedFromThis<FCo
 		return OnVariableMappingChangedDelegate;
 	}
 
-	bool IsInputProperty(const FName& PropertyName) const;
+	UE_API bool IsInputProperty(const FName& PropertyName) const;
 
-	void SetIOMapping(bool bInput, const FName& SourceProperty, const FName& TargetCurve);
-	FName GetIOMapping(bool bInput, const FName& SourceProperty) const;
+	UE_API void SetIOMapping(bool bInput, const FName& SourceProperty, const FName& TargetCurve);
+	UE_API FName GetIOMapping(bool bInput, const FName& SourceProperty) const;
 
 	void SetIgnoreVariablesWithNoMemory(bool bIgnore) { bIgnoreVariablesWithNoMemory = bIgnore; }
 
@@ -105,27 +107,29 @@ private:
 	bool bIgnoreVariablesWithNoMemory = false;
 
 	// property related things
-	static void GetVariables(const UClass* TargetClass, bool bInput, bool bIgnoreVariablesWithNoMemory, TMap<FName, FRigVMExternalVariable>& OutParameters);
+	static UE_API void GetVariables(const UClass* TargetClass, bool bInput, bool bIgnoreVariablesWithNoMemory, TMap<FName, FRigVMExternalVariable>& OutParameters);
 
 	TMap<FName, FRigVMExternalVariable> InputVariables;
 	TMap<FName, FRigVMExternalVariable> OutputVariables;
 
 	// pin option related
-	bool IsPropertyExposeEnabled(FName PropertyName) const;
-	ECheckBoxState IsPropertyExposed(FName PropertyName) const;
-	void OnPropertyExposeCheckboxChanged(ECheckBoxState NewState, FName PropertyName);
+	UE_API bool IsPropertyExposeEnabled(FName PropertyName) const;
+	UE_API ECheckBoxState IsPropertyExposed(FName PropertyName) const;
+	UE_API void OnPropertyExposeCheckboxChanged(ECheckBoxState NewState, FName PropertyName);
 
 #if WITH_EDITOR
 	// SRigVMVariableMappingWidget related
-	void OnVariableMappingChanged(const FName& PathName, const FName& Curve, bool bInput);
-	FName GetVariableMapping(const FName& PathName, bool bInput);
-	void GetAvailableMapping(const FName& PathName, TArray<FName>& OutArray, bool bInput);
-	void CreateVariableMapping(const FString& FilteredText, TArray< TSharedPtr<FRigVMVariableMappingInfo> >& OutArray, bool bInput);
+	UE_API void OnVariableMappingChanged(const FName& PathName, const FName& Curve, bool bInput);
+	UE_API FName GetVariableMapping(const FName& PathName, bool bInput);
+	UE_API void GetAvailableMapping(const FName& PathName, TArray<FName>& OutArray, bool bInput);
+	UE_API void CreateVariableMapping(const FString& FilteredText, TArray< TSharedPtr<FRigVMVariableMappingInfo> >& OutArray, bool bInput);
 #endif
 
-	bool IsAvailableToMapToCurve(const FName& PropertyName, bool bInput) const;
-	const FControlsInfo* FindControlElement(const FName& InControlName) const;
+	UE_API bool IsAvailableToMapToCurve(const FName& PropertyName, bool bInput) const;
+	UE_API const FControlsInfo* FindControlElement(const FName& InControlName) const;
 
-	const UClass* GetTargetClass() const;
-	USkeleton* GetTargetSkeleton() const;
+	UE_API const UClass* GetTargetClass() const;
+	UE_API USkeleton* GetTargetSkeleton() const;
 };
+
+#undef UE_API
diff --git a/Engine/Plugins/Animation/ControlRig/Source/ControlRigDeveloper/Public/ControlRigSchema.h b/Engine/Plugins/Animation/ControlRig/Source/ControlRigDeveloper/Public/ControlRigSchema.h
--- a/Engine/Plugins/Animation/ControlRig/Source/ControlRigDeveloper/Public/ControlRigSchema.h
+++ b/Engine/Plugins/Animation/ControlRig/Source/ControlRigDeveloper/Public/ControlRigSchema.h
@@ -5,14 +5,18 @@
 #include "RigVMModel/RigVMSchema.h"
 #include "ControlRigSchema.generated.h"
 
-UCLASS(BlueprintType)
-class CONTROLRIGDEVELOPER_API UControlRigSchema : public URigVMSchema
+#define UE_API CONTROLRIGDEVELOPER_API
+
+UCLASS(MinimalAPI, BlueprintType)
+class UControlRigSchema : public URigVMSchema
 {
 	GENERATED_UCLASS_BODY()
 
 public:
 
-	virtual bool ShouldUnfoldStruct(URigVMController* InController, const UStruct* InStruct) const override;
+	UE_API virtual bool ShouldUnfoldStruct(URigVMController* InController, const UStruct* InStruct) const override;
+
+	UE_API virtual bool SupportsUnitFunction(URigVMController* InController, const FRigVMFunction* InUnitFunction) const override;
+};
 
-	virtual bool SupportsUnitFunction(URigVMController* InController, const FRigVMFunction* InUnitFunction) const override;
-};
\ No newline at end of file
+#undef UE_API
diff --git a/Engine/Plugins/Animation/ControlRig/Source/ControlRigDeveloper/Public/ControlRigVisualGraphUtils.h b/Engine/Plugins/Animation/ControlRig/Source/ControlRigDeveloper/Public/ControlRigVisualGraphUtils.h
--- a/Engine/Plugins/Animation/ControlRig/Source/ControlRigDeveloper/Public/ControlRigVisualGraphUtils.h
+++ b/Engine/Plugins/Animation/ControlRig/Source/ControlRigDeveloper/Public/ControlRigVisualGraphUtils.h
@@ -6,7 +6,11 @@
 #include "VisualGraphUtils.h"
 #include "Rigs/RigHierarchy.h"
 
-struct CONTROLRIGDEVELOPER_API FControlRigVisualGraphUtils
+#define UE_API CONTROLRIGDEVELOPER_API
+
+struct FControlRigVisualGraphUtils
 {
-	static FString DumpRigHierarchyToDotGraph(URigHierarchy* InHierarchy, const FName& InEventName = NAME_None);
+	static UE_API FString DumpRigHierarchyToDotGraph(URigHierarchy* InHierarchy, const FName& InEventName = NAME_None);
 };
+
+#undef UE_API
