Updated headers using UnrealCodeFixup to make sure dllstorage is on methods/staticvars instead of types. Part 1/n

diff --git a/Engine/Plugins/Animation/ControlRig/Source/ControlRigDeveloper/Public/Graph/ControlRigGraph.h b/Engine/Plugins/Animation/ControlRig/Source/ControlRigDeveloper/Public/Graph/ControlRigGraph.h
--- a/Engine/Plugins/Animation/ControlRig/Source/ControlRigDeveloper/Public/Graph/ControlRigGraph.h
+++ b/Engine/Plugins/Animation/ControlRig/Source/ControlRigDeveloper/Public/Graph/ControlRigGraph.h
@@ -12,24 +12,26 @@
 #include "EdGraph/RigVMEdGraph.h"
 #include "ControlRigGraph.generated.h"
 
+#define UE_API CONTROLRIGDEVELOPER_API
+
 class UControlRigBlueprint;
 class UControlRigGraphSchema;
 class UControlRig;
 class URigVMController;
 struct FRigCurveContainer;
 
-UCLASS()
-class CONTROLRIGDEVELOPER_API UControlRigGraph : public URigVMEdGraph
+UCLASS(MinimalAPI)
+class UControlRigGraph : public URigVMEdGraph
 {
 	GENERATED_BODY()
 
 public:
-	UControlRigGraph();
+	UE_API UControlRigGraph();
 
 	/** Set up this graph */
-	virtual void InitializeFromBlueprint(URigVMBlueprint* InBlueprint) override;
+	UE_API virtual void InitializeFromBlueprint(URigVMBlueprint* InBlueprint) override;
 
-	virtual bool HandleModifiedEvent_Internal(ERigVMGraphNotifType InNotifType, URigVMGraph* InGraph, UObject* InSubject) override;
+	UE_API virtual bool HandleModifiedEvent_Internal(ERigVMGraphNotifType InNotifType, URigVMGraph* InGraph, UObject* InSubject) override;
 
 #if WITH_EDITOR
 
@@ -62,17 +64,17 @@ public:
 		return GetElementNameList(ERigElementType::Socket);
 	}
 
-	virtual const TArray<TSharedPtr<FRigVMStringWithTag>>* GetNameListForWidget(const FString& InWidgetName) const override;
+	UE_API virtual const TArray<TSharedPtr<FRigVMStringWithTag>>* GetNameListForWidget(const FString& InWidgetName) const override;
 
-	void CacheNameLists(URigHierarchy* InHierarchy, const FRigVMDrawContainer* DrawContainer, TArray<TSoftObjectPtr<UControlRigShapeLibrary>> ShapeLibraries);
-	const TArray<TSharedPtr<FRigVMStringWithTag>>* GetElementNameList(ERigElementType InElementType = ERigElementType::Bone) const;
-	const TArray<TSharedPtr<FRigVMStringWithTag>>* GetElementNameList(URigVMPin* InPin = nullptr) const;
-	const TArray<TSharedPtr<FRigVMStringWithTag>> GetSelectedElementsNameList() const;
-	const TArray<TSharedPtr<FRigVMStringWithTag>>* GetDrawingNameList(URigVMPin* InPin = nullptr) const;
-	const TArray<TSharedPtr<FRigVMStringWithTag>>* GetShapeNameList(URigVMPin* InPin = nullptr) const;
+	UE_API void CacheNameLists(URigHierarchy* InHierarchy, const FRigVMDrawContainer* DrawContainer, TArray<TSoftObjectPtr<UControlRigShapeLibrary>> ShapeLibraries);
+	UE_API const TArray<TSharedPtr<FRigVMStringWithTag>>* GetElementNameList(ERigElementType InElementType = ERigElementType::Bone) const;
+	UE_API const TArray<TSharedPtr<FRigVMStringWithTag>>* GetElementNameList(URigVMPin* InPin = nullptr) const;
+	UE_API const TArray<TSharedPtr<FRigVMStringWithTag>> GetSelectedElementsNameList() const;
+	UE_API const TArray<TSharedPtr<FRigVMStringWithTag>>* GetDrawingNameList(URigVMPin* InPin = nullptr) const;
+	UE_API const TArray<TSharedPtr<FRigVMStringWithTag>>* GetShapeNameList(URigVMPin* InPin = nullptr) const;
 
-	FReply HandleGetSelectedClicked(URigVMEdGraph* InEdGraph, URigVMPin* InPin, FString InDefaultValue);
-	FReply HandleBrowseClicked(URigVMEdGraph* InEdGraph, URigVMPin* InPin, FString InDefaultValue);
+	UE_API FReply HandleGetSelectedClicked(URigVMEdGraph* InEdGraph, URigVMPin* InPin, FString InDefaultValue);
+	UE_API FReply HandleBrowseClicked(URigVMEdGraph* InEdGraph, URigVMPin* InPin, FString InDefaultValue);
 
 private:
 
@@ -155,7 +157,7 @@ private:
 	TArray<TSharedPtr<FRigVMStringWithTag>> ShapeNameList;
 	int32 LastHierarchyTopologyVersion;
 
-	static TArray<TSharedPtr<FRigVMStringWithTag>> EmptyElementNameList;
+	static UE_API TArray<TSharedPtr<FRigVMStringWithTag>> EmptyElementNameList;
 
 #endif
 
@@ -169,3 +171,5 @@ inline bool UControlRigGraph::IncludeElementInNameList<FRigControlElement>(const
 {
 	return !InElement->IsAnimationChannel();
 }
+
+#undef UE_API
diff --git a/Engine/Plugins/Animation/ControlRig/Source/ControlRigDeveloper/Public/Graph/ControlRigGraphNode.h b/Engine/Plugins/Animation/ControlRig/Source/ControlRigDeveloper/Public/Graph/ControlRigGraphNode.h
--- a/Engine/Plugins/Animation/ControlRig/Source/ControlRigDeveloper/Public/Graph/ControlRigGraphNode.h
+++ b/Engine/Plugins/Animation/ControlRig/Source/ControlRigDeveloper/Public/Graph/ControlRigGraphNode.h
@@ -14,6 +14,8 @@
 #include "EdGraph/RigVMEdGraphNode.h"
 #include "ControlRigGraphNode.generated.h"
 
+#define UE_API CONTROLRIGDEVELOPER_API
+
 class FBlueprintActionDatabaseRegistrar;
 class UEdGraph;
 struct FSlateIcon;
@@ -23,8 +25,8 @@ struct FSearchTagDataPair;
 #endif
 
 /** Base class for animation ControlRig-related nodes */
-UCLASS()
-class CONTROLRIGDEVELOPER_API UControlRigGraphNode : public URigVMEdGraphNode
+UCLASS(MinimalAPI)
+class UControlRigGraphNode : public URigVMEdGraphNode
 {
 	GENERATED_BODY()
 
@@ -39,13 +41,13 @@ class CONTROLRIGDEVELOPER_API UControlRigGraphNode : public URigVMEdGraphNode
 
 public:
 
-	UControlRigGraphNode();
+	UE_API UControlRigGraphNode();
 
 #if WITH_EDITOR
 #if WITH_RIGVMLEGACYEDITOR
-	void AddPinSearchMetaDataInfo(const UEdGraphPin* Pin, TArray<FSearchTagDataPair>& OutTaggedMetaData) const;
+	UE_API void AddPinSearchMetaDataInfo(const UEdGraphPin* Pin, TArray<FSearchTagDataPair>& OutTaggedMetaData) const;
 #else
-	void AddRigVMPinSearchMetaDataInfo(const UEdGraphPin* Pin, TArray<UBlueprintExtension::FSearchTagDataPair>& OutTaggedMetaData) const;
+	UE_API void AddRigVMPinSearchMetaDataInfo(const UEdGraphPin* Pin, TArray<UBlueprintExtension::FSearchTagDataPair>& OutTaggedMetaData) const;
 #endif
 #endif
 
@@ -59,3 +61,5 @@ private:
 	friend class UControlRigIfNodeSpawner;
 	friend class UControlRigSelectNodeSpawner;
 };
+
+#undef UE_API
diff --git a/Engine/Plugins/Animation/ControlRig/Source/ControlRigDeveloper/Public/Graph/ControlRigGraphSchema.h b/Engine/Plugins/Animation/ControlRig/Source/ControlRigDeveloper/Public/Graph/ControlRigGraphSchema.h
--- a/Engine/Plugins/Animation/ControlRig/Source/ControlRigDeveloper/Public/Graph/ControlRigGraphSchema.h
+++ b/Engine/Plugins/Animation/ControlRig/Source/ControlRigDeveloper/Public/Graph/ControlRigGraphSchema.h
@@ -10,14 +10,16 @@
 
 #include "ControlRigGraphSchema.generated.h"
 
+#define UE_API CONTROLRIGDEVELOPER_API
+
 class UControlRigBlueprint;
 class UControlRigGraph;
 class UControlRigGraphNode;
 class UControlRigGraphNode_Unit;
 class UControlRigGraphNode_Property;
 
-UCLASS()
-class CONTROLRIGDEVELOPER_API UControlRigGraphSchema : public URigVMEdGraphSchema
+UCLASS(MinimalAPI)
+class UControlRigGraphSchema : public URigVMEdGraphSchema
 {
 	GENERATED_BODY()
 
@@ -26,12 +28,14 @@ public:
 	static inline const FLazyName GraphName_ControlRig = FLazyName(TEXT("Rig"));
 
 public:
-	UControlRigGraphSchema();
+	UE_API UControlRigGraphSchema();
 
 	// UEdGraphSchema interface
-	virtual FLinearColor GetPinTypeColor(const FEdGraphPinType& PinType) const override;
+	UE_API virtual FLinearColor GetPinTypeColor(const FEdGraphPinType& PinType) const override;
 
 	// URigVMEdGraphSchema interface
 	virtual const FLazyName& GetRootGraphName() const override { return GraphName_ControlRig; }
-	virtual bool IsRigVMDefaultEvent(const FName& InEventName) const override;
+	UE_API virtual bool IsRigVMDefaultEvent(const FName& InEventName) const override;
 };
+
+#undef UE_API
diff --git a/Engine/Plugins/Animation/ControlRig/Source/ControlRigEditor/Private/EditMode/SControlRigOutliner.h b/Engine/Plugins/Animation/ControlRig/Source/ControlRigEditor/Private/EditMode/SControlRigOutliner.h
--- a/Engine/Plugins/Animation/ControlRig/Source/ControlRigEditor/Private/EditMode/SControlRigOutliner.h
+++ b/Engine/Plugins/Animation/ControlRig/Source/ControlRigEditor/Private/EditMode/SControlRigOutliner.h
@@ -14,6 +14,8 @@
 #include "Editor/SRigHierarchyTreeView.h"
 #include "Widgets/SBoxPanel.h"
 
+#define UE_API CONTROLRIGEDITOR_API
+
 
 class ISequencer;
 class SExpandableArea;
@@ -34,7 +36,7 @@ typedef STreeView<TSharedPtr<FMultiRigTreeElement>>::FOnSetExpansionRecursive FO
 
 uint32 GetTypeHash(const FMultiRigData& Data);
 
-struct CONTROLRIGEDITOR_API FMultiRigTreeDelegates
+struct FMultiRigTreeDelegates
 {
 	FOnGetRigTreeDisplaySettings OnGetDisplaySettings;
 	FOnMultiRigTreeSelectionChanged OnSelectionChanged;
@@ -79,7 +81,7 @@ struct CONTROLRIGEDITOR_API FMultiRigTreeDelegates
 		OnSelectionChanged.ExecuteIfBound(Selection, SelectInfo);
 	}
 
-	static FRigTreeDisplaySettings DefaultDisplaySettings;
+	static UE_API FRigTreeDisplaySettings DefaultDisplaySettings;
 	bool bIsChangingRigHierarchy;
 };
 
@@ -377,3 +379,5 @@ private:
 	//set of control rigs we are bound too and need to clear delegates from
 	TArray<TWeakObjectPtr<UControlRig>> BoundControlRigs;
 };
+
+#undef UE_API
diff --git a/Engine/Plugins/Animation/ControlRig/Source/ControlRigEditor/Private/Editor/Constraints/SConstraintsWidget.h b/Engine/Plugins/Animation/ControlRig/Source/ControlRigEditor/Private/Editor/Constraints/SConstraintsWidget.h
--- a/Engine/Plugins/Animation/ControlRig/Source/ControlRigEditor/Private/Editor/Constraints/SConstraintsWidget.h
+++ b/Engine/Plugins/Animation/ControlRig/Source/ControlRigEditor/Private/Editor/Constraints/SConstraintsWidget.h
@@ -12,6 +12,8 @@
 #include "ConstraintsManager.h"
 #include "Misc/NotifyHook.h"
 
+#define UE_API CONTROLRIGEDITOR_API
+
 class UTickableTransformConstraint;
 class AActor;
 class SConstraintsCreationWidget;
@@ -217,7 +219,7 @@ private:
 * Constraint List Shared by Edit Widget and Bake Widget
 */
 
-class CONTROLRIGEDITOR_API FBaseConstraintListWidget : public FEditorUndoClient
+class FBaseConstraintListWidget : public FEditorUndoClient
 {
 public:
 
@@ -233,20 +235,20 @@ public:
 	};
 public:
 
-	virtual ~FBaseConstraintListWidget() override;
+	UE_API virtual ~FBaseConstraintListWidget() override;
 	/* FEditorUndoClient interface */
-	virtual void PostUndo(bool bSuccess);
-	virtual void PostRedo(bool bSuccess);
+	UE_API virtual void PostUndo(bool bSuccess);
+	UE_API virtual void PostRedo(bool bSuccess);
 	/* End FEditorUndoClient interface */
 
 	/** Invalidates the constraint list for further rebuild. */
-	virtual void InvalidateConstraintList();
+	UE_API virtual void InvalidateConstraintList();
 
 	/** Rebuild the constraint list based on the current selection. */
-	virtual int32 RefreshConstraintList();
+	UE_API virtual int32 RefreshConstraintList();
 
 	/** Triggers a constraint list invalidation when selection in the level viewport. */
-	void OnActorSelectionChanged(const TArray<UObject*>& NewSelection, bool bForceRefresh);
+	UE_API void OnActorSelectionChanged(const TArray<UObject*>& NewSelection, bool bForceRefresh);
 
 	/** Which constraints to show  */
 	EShowConstraints GetShowConstraints() const { return ShowConstraints; }
@@ -258,8 +260,8 @@ public:
 			InvalidateConstraintList();
 		}
 	}
-	FText GetShowConstraintsText(EShowConstraints Index) const;
-	FText GetShowConstraintsTooltip(EShowConstraints Index) const;
+	UE_API FText GetShowConstraintsText(EShowConstraints Index) const;
+	UE_API FText GetShowConstraintsTooltip(EShowConstraints Index) const;
 
 protected:
 	/** Types */
@@ -267,8 +269,8 @@ protected:
 	using ConstraintItemListView = SListView< ItemSharedPtr >;
 
 
-	void RegisterSelectionChanged();
-	void UnregisterSelectionChanged();
+	UE_API void RegisterSelectionChanged();
+	UE_API void UnregisterSelectionChanged();
 
 	/** List view that shows constraint types*/
 	TSharedPtr< ConstraintItemListView > ListView;
@@ -283,7 +285,7 @@ protected:
 
 public:
 
-	static EShowConstraints ShowConstraints;
+	static UE_API EShowConstraints ShowConstraints;
 
 };
 
@@ -291,80 +293,80 @@ public:
  * SConstraintsEditionWidget
  */
 
-class CONTROLRIGEDITOR_API SConstraintsEditionWidget : public SCompoundWidget, public FBaseConstraintListWidget
+class SConstraintsEditionWidget : public SCompoundWidget, public FBaseConstraintListWidget
 {
 public:
 	SLATE_BEGIN_ARGS(SConstraintsEditionWidget)	{}
 	SLATE_END_ARGS()
 	/** Constructs this widget with InArgs */
-	void Construct(const FArguments& InArgs);
+	UE_API void Construct(const FArguments& InArgs);
 
 	/** Override for SWidget::Tick */
-	virtual void Tick(const FGeometry& AllottedGeometry, const double InCurrentTime, const float InDeltaTime) override;
+	UE_API virtual void Tick(const FGeometry& AllottedGeometry, const double InCurrentTime, const float InDeltaTime) override;
 
 	/**  */
-	bool CanMoveUp(const TSharedPtr<FEditableConstraintItem>& Item) const;
-	bool CanMoveDown(const TSharedPtr<FEditableConstraintItem>& Item) const;
+	UE_API bool CanMoveUp(const TSharedPtr<FEditableConstraintItem>& Item) const;
+	UE_API bool CanMoveDown(const TSharedPtr<FEditableConstraintItem>& Item) const;
 
 	/**  */
-	void MoveItemUp(const TSharedPtr<FEditableConstraintItem>& Item);
-	void MoveItemDown(const TSharedPtr<FEditableConstraintItem>& Item);
+	UE_API void MoveItemUp(const TSharedPtr<FEditableConstraintItem>& Item);
+	UE_API void MoveItemDown(const TSharedPtr<FEditableConstraintItem>& Item);
 
 	/**  */
-	void RemoveItem(const TSharedPtr<FEditableConstraintItem>& Item);
+	UE_API void RemoveItem(const TSharedPtr<FEditableConstraintItem>& Item);
 
 	// FBaseConstraintListWidget overrides
-	virtual void InvalidateConstraintList() override;
+	UE_API virtual void InvalidateConstraintList() override;
 
 	/** Notify from sequencer changed. */
-	void SequencerChanged(const TWeakPtr<ISequencer>& InNewSequencer);
+	UE_API void SequencerChanged(const TWeakPtr<ISequencer>& InNewSequencer);
 
 	/** Resets the selected parent data. */
-	void ResetParent();
+	UE_API void ResetParent();
 
 	/** Returns the parent that has been set using the chooser widget. */
-	const FConstrainable& GetParent() const;
-	bool IsParentValid() const;
+	UE_API const FConstrainable& GetParent() const;
+	UE_API bool IsParentValid() const;
 
 private:
 
 	/** Updates the current constrainable parent data. */
-	void UpdateParent(AActor* InActor);
+	UE_API void UpdateParent(AActor* InActor);
 
 	/** Constrainable parent data to be sued to create a new constraint. */
 	FConstrainable ConstrainableParent;
 
 	/** Widgets to set the constrainable parent without using the picker. */
-	TSharedRef<SWidget> OnGetParentMenuContent();
-	TSharedPtr<SComboButton> GetActorChooserWidget();
-	TSharedRef<SWidget> GetUseSelectedActorWidget();
-	void OnParentMenuOpenChanged(bool bOpened);
+	UE_API TSharedRef<SWidget> OnGetParentMenuContent();
+	UE_API TSharedPtr<SComboButton> GetActorChooserWidget();
+	UE_API TSharedRef<SWidget> GetUseSelectedActorWidget();
+	UE_API void OnParentMenuOpenChanged(bool bOpened);
 	
 	/** Generates a widget for the specified item */
-	TSharedRef<ITableRow> OnGenerateWidgetForItem(ItemSharedPtr InItem, const TSharedRef<STableViewBase>& OwnerTable);
+	UE_API TSharedRef<ITableRow> OnGenerateWidgetForItem(ItemSharedPtr InItem, const TSharedRef<STableViewBase>& OwnerTable);
 
 	/** Types */
 	using ItemSharedPtr = TSharedPtr<FEditableConstraintItem>;
 	using ConstraintItemListView = SListView< ItemSharedPtr >;
 
 	/** @todo documentation. */
-	TSharedPtr< SWidget > CreateContextMenu();
-	void OnItemDoubleClicked(ItemSharedPtr InItem);
+	UE_API TSharedPtr< SWidget > CreateContextMenu();
+	UE_API void OnItemDoubleClicked(ItemSharedPtr InItem);
 
-	FReply OnBakeClicked();
+	UE_API FReply OnBakeClicked();
 
-	void UpdateSequencer();
+	UE_API void UpdateSequencer();
 
 	//sequencer and it's time
 	TWeakPtr<ISequencer> WeakSequencer;
-	bool SequencerTimeChanged();
+	UE_API bool SequencerTimeChanged();
 	FQualifiedFrameTime SequencerLastTime;
 
 	TSharedPtr<SComboButton> ActorSelectionButtons;
 };
 
 /** Widget allowing baking of constraints */
-class CONTROLRIGEDITOR_API SConstraintBakeWidget : public SCompoundWidget, public FBaseConstraintListWidget
+class SConstraintBakeWidget : public SCompoundWidget, public FBaseConstraintListWidget
 {
 public:
 
@@ -375,27 +377,27 @@ public:
 	SLATE_ARGUMENT(TSharedPtr<ISequencer>, Sequencer)
 	SLATE_END_ARGS()
 
-	void Construct(const FArguments& InArgs);
+	UE_API void Construct(const FArguments& InArgs);
 	virtual ~SConstraintBakeWidget() override {}
 
-	FReply OpenDialog(bool bModal = true);
-	void CloseDialog();
+	UE_API FReply OpenDialog(bool bModal = true);
+	UE_API void CloseDialog();
 
 	//SWidget overrides
-	virtual void Tick(const FGeometry& AllottedGeometry, const double InCurrentTime, const float InDeltaTime) override;
+	UE_API virtual void Tick(const FGeometry& AllottedGeometry, const double InCurrentTime, const float InDeltaTime) override;
 
 	//FBaseConstraintListWidget overrides
-	virtual int32 RefreshConstraintList() override;
+	UE_API virtual int32 RefreshConstraintList() override;
 
 private:
-	void BakeSelected();
+	UE_API void BakeSelected();
 
 	/** Generates a widget for the specified item */
-	TSharedRef<ITableRow> OnGenerateWidgetForItem(ItemSharedPtr InItem, const TSharedRef<STableViewBase>& OwnerTable);
+	UE_API TSharedRef<ITableRow> OnGenerateWidgetForItem(ItemSharedPtr InItem, const TSharedRef<STableViewBase>& OwnerTable);
 
 	TSharedPtr<ISequencer> Sequencer;
 	//static to be reused
-	static TOptional<FBakingAnimationKeySettings> BakeConstraintSettings;
+	static UE_API TOptional<FBakingAnimationKeySettings> BakeConstraintSettings;
 	//structonscope for details panel
 	TSharedPtr < TStructOnScope<FBakingAnimationKeySettings>> Settings;
 	TWeakPtr<SWindow> DialogWindow;
@@ -423,4 +425,6 @@ private:
 	/** TSharedPtr to the tree item. */
 	TSharedPtr<FEditableConstraintItem> ConstraintItem;
 	TWeakPtr<SConstraintBakeWidget> ConstraintsWidget;
-};
\ No newline at end of file
+};
+
+#undef UE_API
diff --git a/Engine/Plugins/Animation/ControlRig/Source/ControlRigEditor/Private/Editor/ControlRigWrapperObject.h b/Engine/Plugins/Animation/ControlRig/Source/ControlRigEditor/Private/Editor/ControlRigWrapperObject.h
--- a/Engine/Plugins/Animation/ControlRig/Source/ControlRigEditor/Private/Editor/ControlRigWrapperObject.h
+++ b/Engine/Plugins/Animation/ControlRig/Source/ControlRigEditor/Private/Editor/ControlRigWrapperObject.h
@@ -6,16 +6,20 @@
 #include "Rigs/RigHierarchyDefines.h"
 #include "ControlRigWrapperObject.generated.h"
 
-UCLASS()
-class CONTROLRIGEDITOR_API UControlRigWrapperObject : public URigVMDetailsViewWrapperObject
+#define UE_API CONTROLRIGEDITOR_API
+
+UCLASS(MinimalAPI)
+class UControlRigWrapperObject : public URigVMDetailsViewWrapperObject
 {
 public:
 	GENERATED_BODY()
 
-	virtual UClass* GetClassForStruct(UScriptStruct* InStruct, bool bCreateIfNeeded) const override;
+	UE_API virtual UClass* GetClassForStruct(UScriptStruct* InStruct, bool bCreateIfNeeded) const override;
 
-	virtual void SetContent(const uint8* InStructMemory, const UStruct* InStruct) override;
-	virtual void GetContent(uint8* OutStructMemory, const UStruct* InStruct) const override;;
+	UE_API virtual void SetContent(const uint8* InStructMemory, const UStruct* InStruct) override;
+	UE_API virtual void GetContent(uint8* OutStructMemory, const UStruct* InStruct) const override;;
 
 	FRigHierarchyKey HierarchyKey;
 };
+
+#undef UE_API
diff --git a/Engine/Plugins/Animation/ControlRig/Source/ControlRigEditor/Private/Editor/SModularRigTreeView.h b/Engine/Plugins/Animation/ControlRig/Source/ControlRigEditor/Private/Editor/SModularRigTreeView.h
--- a/Engine/Plugins/Animation/ControlRig/Source/ControlRigEditor/Private/Editor/SModularRigTreeView.h
+++ b/Engine/Plugins/Animation/ControlRig/Source/ControlRigEditor/Private/Editor/SModularRigTreeView.h
@@ -27,7 +27,7 @@ typedef STableRow<TSharedPtr<FModularRigTreeElement>>::FOnCanAcceptDrop FOnModul
 typedef STableRow<TSharedPtr<FModularRigTreeElement>>::FOnAcceptDrop FOnModularRigTreeAcceptDrop;
 typedef STreeView<TSharedPtr<FModularRigTreeElement>>::FOnSelectionChanged FOnModularRigTreeSelectionChanged;
 
-struct CONTROLRIGEDITOR_API FModularRigTreeDelegates
+struct FModularRigTreeDelegates
 {
 public:
 	
diff --git a/Engine/Plugins/Animation/ControlRig/Source/ControlRigEditor/Private/Editor/SRigHierarchyTreeView.h b/Engine/Plugins/Animation/ControlRig/Source/ControlRigEditor/Private/Editor/SRigHierarchyTreeView.h
--- a/Engine/Plugins/Animation/ControlRig/Source/ControlRigEditor/Private/Editor/SRigHierarchyTreeView.h
+++ b/Engine/Plugins/Animation/ControlRig/Source/ControlRigEditor/Private/Editor/SRigHierarchyTreeView.h
@@ -8,12 +8,14 @@
 #include "SRigHierarchyTagWidget.h"
 #include "Settings/ControlRigSettings.h"
 
+#define UE_API CONTROLRIGEDITOR_API
+
 class SSearchBox;
 class SRigHierarchyTreeView;
 class SRigHierarchyItem;
 class FRigTreeElement;
 
-struct CONTROLRIGEDITOR_API FRigTreeDisplaySettings
+struct FRigTreeDisplaySettings
 {
 	FRigTreeDisplaySettings()
 	{
@@ -105,7 +107,7 @@ typedef STreeView<TSharedPtr<FRigTreeElement>>::FOnMouseButtonClick FOnRigTreeMo
 typedef STreeView<TSharedPtr<FRigTreeElement>>::FOnMouseButtonDoubleClick FOnRigTreeMouseButtonDoubleClick;
 typedef STreeView<TSharedPtr<FRigTreeElement>>::FOnSetExpansionRecursive FOnRigTreeSetExpansionRecursive;
 
-struct CONTROLRIGEDITOR_API FRigTreeDelegates
+struct FRigTreeDelegates
 {
 	FOnGetRigTreeHierarchy OnGetHierarchy;
 	FOnGetRigTreeDisplaySettings OnGetDisplaySettings;
@@ -217,7 +219,7 @@ struct CONTROLRIGEDITOR_API FRigTreeDelegates
 		}
 	}
 
-	static FRigTreeDisplaySettings DefaultDisplaySettings;
+	static UE_API FRigTreeDisplaySettings DefaultDisplaySettings;
 	bool bIsChangingRigHierarchy;
 };
 
@@ -439,3 +441,5 @@ private:
 	TSharedPtr<SRigHierarchyTreeView> TreeView;
 	float MaxHeight = 0.f;
 };
+
+#undef UE_API
diff --git a/Engine/Plugins/Animation/ControlRig/Source/ControlRigEditor/Private/Editor/SRigSpacePickerWidget.h b/Engine/Plugins/Animation/ControlRig/Source/ControlRigEditor/Private/Editor/SRigSpacePickerWidget.h
--- a/Engine/Plugins/Animation/ControlRig/Source/ControlRigEditor/Private/Editor/SRigSpacePickerWidget.h
+++ b/Engine/Plugins/Animation/ControlRig/Source/ControlRigEditor/Private/Editor/SRigSpacePickerWidget.h
@@ -12,6 +12,8 @@
 #include "EditorUndoClient.h"
 #include "RigSpacePickerBakeSettings.h"
 
+#define UE_API CONTROLRIGEDITOR_API
+
 class IMenu;
 
 DECLARE_DELEGATE_RetVal_TwoParams(FRigElementKey, FRigSpacePickerGetActiveSpace, URigHierarchy*, const FRigElementKey&);
@@ -22,7 +24,7 @@ DECLARE_DELEGATE_RetVal_TwoParams(TArray<FRigElementKeyWithLabel>, FRigSpacePick
 DECLARE_DELEGATE_RetVal_ThreeParams(FReply, SRigSpacePickerOnBake, URigHierarchy*, TArray<FRigElementKey> /* Controls */, FRigSpacePickerBakeSettings);
 
 /** Widget allowing picking of a space source for space switching */
-class CONTROLRIGEDITOR_API SRigSpacePickerWidget : public SCompoundWidget, public FEditorUndoClient
+class SRigSpacePickerWidget : public SCompoundWidget, public FEditorUndoClient
 {
 public:
 
@@ -61,17 +63,17 @@ public:
 		SLATE_EVENT(FOnClicked, OnBakeButtonClicked )
 	SLATE_END_ARGS()
 
-	void Construct(const FArguments& InArgs);
-	virtual ~SRigSpacePickerWidget() override;
+	UE_API void Construct(const FArguments& InArgs);
+	UE_API virtual ~SRigSpacePickerWidget() override;
 
-	void SetControls(URigHierarchy* InHierarchy, const TArray<FRigElementKey>& InControls);
+	UE_API void SetControls(URigHierarchy* InHierarchy, const TArray<FRigElementKey>& InControls);
 
-	FReply OpenDialog(bool bModal = true);
-	void CloseDialog();
+	UE_API FReply OpenDialog(bool bModal = true);
+	UE_API void CloseDialog();
 	
-	virtual FReply OnKeyDown( const FGeometry& MyGeometry, const FKeyEvent& InKeyEvent ) override;
-	virtual bool SupportsKeyboardFocus() const override;
-	virtual void Tick(const FGeometry& AllottedGeometry, const double InCurrentTime, const float InDeltaTime) override;
+	UE_API virtual FReply OnKeyDown( const FGeometry& MyGeometry, const FKeyEvent& InKeyEvent ) override;
+	UE_API virtual bool SupportsKeyboardFocus() const override;
+	UE_API virtual void Tick(const FGeometry& AllottedGeometry, const double InCurrentTime, const float InDeltaTime) override;
 
 	URigHierarchy* GetHierarchy() const
 	{
@@ -86,16 +88,16 @@ public:
 	const URigHierarchy* GetHierarchyConst() const { return GetHierarchy(); }
 	
 	const TArray<FRigElementKey>& GetControls() const { return ControlKeys; }
-	const TArray<FRigElementKey>& GetActiveSpaces() const;
-	const TArray<FRigElementKeyWithLabel>& GetDefaultSpaces() const;
-	TArray<FRigElementKeyWithLabel> GetSpaceList(bool bIncludeDefaultSpaces = false) const;
+	UE_API const TArray<FRigElementKey>& GetActiveSpaces() const;
+	UE_API const TArray<FRigElementKeyWithLabel>& GetDefaultSpaces() const;
+	UE_API TArray<FRigElementKeyWithLabel> GetSpaceList(bool bIncludeDefaultSpaces = false) const;
 	FRigSpacePickerActiveSpaceChanged& OnActiveSpaceChanged() { return ActiveSpaceChangedEvent; }
 	FRigSpacePickerSpaceListChanged& OnSpaceListChanged() { return SpaceListChangedEvent; }
-	void RefreshContents();
+	UE_API void RefreshContents();
 
 	// FEditorUndoClient interface
-	virtual void PostUndo(bool bSuccess);
-	virtual void PostRedo(bool bSuccess);
+	UE_API virtual void PostUndo(bool bSuccess);
+	UE_API virtual void PostRedo(bool bSuccess);
 	// End FEditorUndoClient interface
 	
 private:
@@ -107,7 +109,7 @@ private:
 		ESpacePickerType_Item
 	};
 
-	void AddSpacePickerRow(
+	UE_API void AddSpacePickerRow(
 		TSharedPtr<SVerticalBox> InListBox,
 		ESpacePickerType InType,
 		const FRigElementKey& InKey,
@@ -116,34 +118,34 @@ private:
 		const FText& InTitle,
 		FOnClicked OnClickedDelegate);
 
-	void RepopulateItemSpaces();
-	void ClearListBox(TSharedPtr<SVerticalBox> InListBox);
-	void UpdateActiveSpaces();
-	bool IsValidKey(const FRigElementKey& InKey) const;
-	bool IsDefaultSpace(const FRigElementKey& InKey) const;
+	UE_API void RepopulateItemSpaces();
+	UE_API void ClearListBox(TSharedPtr<SVerticalBox> InListBox);
+	UE_API void UpdateActiveSpaces();
+	UE_API bool IsValidKey(const FRigElementKey& InKey) const;
+	UE_API bool IsDefaultSpace(const FRigElementKey& InKey) const;
 
-	FReply HandleParentSpaceClicked();
-	FReply HandleWorldSpaceClicked();
-	FReply HandleElementSpaceClicked(FRigElementKey InKey);
-	FReply HandleSpaceMoveUp(FRigElementKey InKey);
-	FReply HandleSpaceMoveDown(FRigElementKey InKey);
-	void HandleSpaceDelete(FRigElementKey InKey);
+	UE_API FReply HandleParentSpaceClicked();
+	UE_API FReply HandleWorldSpaceClicked();
+	UE_API FReply HandleElementSpaceClicked(FRigElementKey InKey);
+	UE_API FReply HandleSpaceMoveUp(FRigElementKey InKey);
+	UE_API FReply HandleSpaceMoveDown(FRigElementKey InKey);
+	UE_API void HandleSpaceDelete(FRigElementKey InKey);
 
 public:
 	
-	FReply HandleAddElementClicked();
-	bool IsRestricted() const;
+	UE_API FReply HandleAddElementClicked();
+	UE_API bool IsRestricted() const;
 
 private:
 
-	bool IsSpaceMoveUpEnabled(FRigElementKey InKey) const;
-	bool IsSpaceMoveDownEnabled(FRigElementKey InKey) const;
+	UE_API bool IsSpaceMoveUpEnabled(FRigElementKey InKey) const;
+	UE_API bool IsSpaceMoveDownEnabled(FRigElementKey InKey) const;
 
-	void OnHierarchyModified(ERigHierarchyNotification InNotif, URigHierarchy* InHierarchy, const FRigNotificationSubject& InSubject);
+	UE_API void OnHierarchyModified(ERigHierarchyNotification InNotif, URigHierarchy* InHierarchy, const FRigNotificationSubject& InSubject);
 
-	FSlateColor GetButtonColor(ESpacePickerType InType, FRigElementKey InKey) const;
-	FRigElementKey GetActiveSpace_Private(URigHierarchy* InHierarchy, const FRigElementKey& InControlKey) const;
-	TArray<FRigElementKeyWithLabel> GetCurrentParents_Private(URigHierarchy* InHierarchy, const FRigElementKey& InControlKey) const;
+	UE_API FSlateColor GetButtonColor(ESpacePickerType InType, FRigElementKey InKey) const;
+	UE_API FRigElementKey GetActiveSpace_Private(URigHierarchy* InHierarchy, const FRigElementKey& InControlKey) const;
+	UE_API TArray<FRigElementKeyWithLabel> GetCurrentParents_Private(URigHierarchy* InHierarchy, const FRigElementKey& InControlKey) const;
 	
 	FRigSpacePickerActiveSpaceChanged ActiveSpaceChangedEvent;
 	FRigSpacePickerSpaceListChanged SpaceListChangedEvent;
@@ -177,13 +179,13 @@ private:
 	FDelegateHandle ActiveSpaceChangedWindowHandle;
 	FRigTreeDisplaySettings HierarchyDisplaySettings;
 
-	static FRigElementKey InValidKey;
+	static UE_API FRigElementKey InValidKey;
 
 	/** Unregisters any current pending selection active timer. */
-	void UnregisterPendingSelection();
+	UE_API void UnregisterPendingSelection();
 	
 	/** Updates the selected controls array and rebuilds the spaces list. */
-	void UpdateSelection(URigHierarchy* InHierarchy, const TArray<FRigElementKey>& InControls);
+	UE_API void UpdateSelection(URigHierarchy* InHierarchy, const TArray<FRigElementKey>& InControls);
 
 	/** Controls selection currently pending. */
 	TWeakPtr<FActiveTimerHandle> PendingSelectionHandle;
@@ -192,7 +194,7 @@ private:
 class ISequencer;
 
 /** Widget allowing baking controls from one space to another */
-class CONTROLRIGEDITOR_API SRigSpacePickerBakeWidget : public SCompoundWidget
+class SRigSpacePickerBakeWidget : public SCompoundWidget
 {
 public:
 
@@ -210,11 +212,11 @@ public:
 
 	SLATE_END_ARGS()
 
-	void Construct(const FArguments& InArgs);
+	UE_API void Construct(const FArguments& InArgs);
 	virtual ~SRigSpacePickerBakeWidget() override {}
 
-	FReply OpenDialog(bool bModal = true);
-	void CloseDialog();
+	UE_API FReply OpenDialog(bool bModal = true);
+	UE_API void CloseDialog();
 
 private:
 
@@ -228,3 +230,5 @@ private:
 	TSharedPtr<SRigSpacePickerWidget> SpacePickerWidget;
 	TSharedPtr<IStructureDetailsView> DetailsView;
 };
+
+#undef UE_API
diff --git a/Engine/Plugins/Animation/ControlRig/Source/ControlRigEditor/Public/ControlRigBlueprintEditorLibrary.h b/Engine/Plugins/Animation/ControlRig/Source/ControlRigEditor/Public/ControlRigBlueprintEditorLibrary.h
--- a/Engine/Plugins/Animation/ControlRig/Source/ControlRigEditor/Public/ControlRigBlueprintEditorLibrary.h
+++ b/Engine/Plugins/Animation/ControlRig/Source/ControlRigEditor/Public/ControlRigBlueprintEditorLibrary.h
@@ -13,6 +13,8 @@
 #include "RigVMEditorBlueprintLibrary.h"
 #include "ControlRigBlueprintEditorLibrary.generated.h"
 
+#define UE_API CONTROLRIGEDITOR_API
+
 UENUM(BlueprintType)
 enum class ECastToControlRigBlueprintCases : uint8
 {
@@ -20,44 +22,45 @@ enum class ECastToControlRigBlueprintCases : uint8
 	CastFailed
 };
 
-UCLASS(BlueprintType, meta=(ScriptName="ControlRigBlueprintLibrary"))
-class CONTROLRIGEDITOR_API UControlRigBlueprintEditorLibrary : public URigVMEditorBlueprintLibrary
+UCLASS(MinimalAPI, BlueprintType, meta=(ScriptName="ControlRigBlueprintLibrary"))
+class UControlRigBlueprintEditorLibrary : public URigVMEditorBlueprintLibrary
 {
 	GENERATED_BODY()
 
 public:
 
 	UFUNCTION(BlueprintCallable, Category = "Control Rig Blueprint", Meta = (ExpandEnumAsExecs = "Branches"))
-	static void CastToControlRigBlueprint(
+	static UE_API void CastToControlRigBlueprint(
 		UObject* Object, 
 		ECastToControlRigBlueprintCases& Branches,
 		UControlRigBlueprint*& AsControlRigBlueprint);
 
 	UFUNCTION(BlueprintCallable, Category = "Control Rig Blueprint")
-	static void SetPreviewMesh(UControlRigBlueprint* InRigBlueprint, USkeletalMesh* PreviewMesh, bool bMarkAsDirty = true);
+	static UE_API void SetPreviewMesh(UControlRigBlueprint* InRigBlueprint, USkeletalMesh* PreviewMesh, bool bMarkAsDirty = true);
 	
 	UFUNCTION(BlueprintCallable, BlueprintPure, Category = "Control Rig Blueprint")
-	static USkeletalMesh* GetPreviewMesh(UControlRigBlueprint* InRigBlueprint);
+	static UE_API USkeletalMesh* GetPreviewMesh(UControlRigBlueprint* InRigBlueprint);
 
 	UFUNCTION(BlueprintCallable, Category = "Control Rig Blueprint")
-	static void RequestControlRigInit(UControlRigBlueprint* InRigBlueprint);
+	static UE_API void RequestControlRigInit(UControlRigBlueprint* InRigBlueprint);
 
 	UFUNCTION(BlueprintCallable, BlueprintPure, Category = "Control Rig Blueprint")
-	static TArray<UControlRigBlueprint*> GetCurrentlyOpenRigBlueprints();
+	static UE_API TArray<UControlRigBlueprint*> GetCurrentlyOpenRigBlueprints();
 
 	UFUNCTION(BlueprintCallable, BlueprintPure, Category = "Control Rig Blueprint")
-	static TArray<UStruct*> GetAvailableRigUnits();
+	static UE_API TArray<UStruct*> GetAvailableRigUnits();
 
 	UFUNCTION(BlueprintCallable, BlueprintPure, Category = "Control Rig Blueprint")
-	static URigHierarchy* GetHierarchy(UControlRigBlueprint* InRigBlueprint);
+	static UE_API URigHierarchy* GetHierarchy(UControlRigBlueprint* InRigBlueprint);
 
 	UFUNCTION(BlueprintCallable, BlueprintPure, Category = "Control Rig Blueprint")
-	static URigHierarchyController* GetHierarchyController(UControlRigBlueprint* InRigBlueprint);
+	static UE_API URigHierarchyController* GetHierarchyController(UControlRigBlueprint* InRigBlueprint);
 
 	UFUNCTION(BlueprintCallable, BlueprintCallable, Category = "Control Rig Blueprint")
-	static void SetupAllEditorMenus();
+	static UE_API void SetupAllEditorMenus();
 
 	UFUNCTION(BlueprintCallable, BlueprintPure, Category = "Control Rig Blueprint")
-	static TArray<FRigModuleDescription> GetAvailableRigModules();
+	static UE_API TArray<FRigModuleDescription> GetAvailableRigModules();
 };
 
+#undef UE_API
diff --git a/Engine/Plugins/Animation/ControlRig/Source/ControlRigEditor/Public/ControlRigBlueprintFactory.h b/Engine/Plugins/Animation/ControlRig/Source/ControlRigEditor/Public/ControlRigBlueprintFactory.h
--- a/Engine/Plugins/Animation/ControlRig/Source/ControlRigEditor/Public/ControlRigBlueprintFactory.h
+++ b/Engine/Plugins/Animation/ControlRig/Source/ControlRigEditor/Public/ControlRigBlueprintFactory.h
@@ -11,22 +11,24 @@
 #include "ControlRig.h"
 #include "ControlRigBlueprintFactory.generated.h"
 
-UCLASS(HideCategories=Object)
-class CONTROLRIGEDITOR_API UControlRigBlueprintFactory : public UFactory
+#define UE_API CONTROLRIGEDITOR_API
+
+UCLASS(MinimalAPI, HideCategories=Object)
+class UControlRigBlueprintFactory : public UFactory
 {
 	GENERATED_BODY()
 
 public:
-	UControlRigBlueprintFactory();
+	UE_API UControlRigBlueprintFactory();
 
 	// The parent class of the created blueprint
 	UPROPERTY(EditAnywhere, Category="Control Rig Factory", meta=(AllowAbstract = ""))
 	TSubclassOf<UControlRig> ParentClass;
 
 	// UFactory Interface
-	virtual bool ConfigureProperties() override;
-	virtual UObject* FactoryCreateNew(UClass* Class, UObject* InParent, FName Name, EObjectFlags Flags, UObject* Context, FFeedbackContext* Warn, FName CallingContext) override;
-	virtual UObject* FactoryCreateNew(UClass* Class, UObject* InParent, FName Name, EObjectFlags Flags, UObject* Context, FFeedbackContext* Warn) override;
+	UE_API virtual bool ConfigureProperties() override;
+	UE_API virtual UObject* FactoryCreateNew(UClass* Class, UObject* InParent, FName Name, EObjectFlags Flags, UObject* Context, FFeedbackContext* Warn, FName CallingContext) override;
+	UE_API virtual UObject* FactoryCreateNew(UClass* Class, UObject* InParent, FName Name, EObjectFlags Flags, UObject* Context, FFeedbackContext* Warn) override;
 
 	/**
 	 * Create a new control rig asset within the contents space of the project.
@@ -34,7 +36,7 @@ public:
 	 * @param bModularRig If true the rig will be created as a modular rig
 	 */
 	UFUNCTION(BlueprintCallable, Category = "Control Rig")
-	static UControlRigBlueprint* CreateNewControlRigAsset(const FString& InDesiredPackagePath, const bool bModularRig = false);
+	static UE_API UControlRigBlueprint* CreateNewControlRigAsset(const FString& InDesiredPackagePath, const bool bModularRig = false);
 
 	/**
 	 * Create a new control rig asset within the contents space of the project
@@ -43,6 +45,7 @@ public:
 	 * @param bModularRig If true the rig will be created as a modular rig
 	 */
 	UFUNCTION(BlueprintCallable, Category = "Control Rig")
-	static UControlRigBlueprint* CreateControlRigFromSkeletalMeshOrSkeleton(UObject* InSelectedObject, const bool bModularRig = false);
+	static UE_API UControlRigBlueprint* CreateControlRigFromSkeletalMeshOrSkeleton(UObject* InSelectedObject, const bool bModularRig = false);
 };
 
+#undef UE_API
diff --git a/Engine/Plugins/Animation/ControlRig/Source/ControlRigEditor/Public/ControlRigSequencerEditorLibrary.h b/Engine/Plugins/Animation/ControlRig/Source/ControlRigEditor/Public/ControlRigSequencerEditorLibrary.h
--- a/Engine/Plugins/Animation/ControlRig/Source/ControlRigEditor/Public/ControlRigSequencerEditorLibrary.h
+++ b/Engine/Plugins/Animation/ControlRig/Source/ControlRigEditor/Public/ControlRigSequencerEditorLibrary.h
@@ -20,6 +20,8 @@
 #include "Filters/CurveEditorSmartReduceFilter.h"
 #include "ControlRigSequencerEditorLibrary.generated.h"
 
+#define UE_API CONTROLRIGEDITOR_API
+
 class ULevelSequence;
 class UTickableConstraint;
 class UTickableTransformConstraint;
@@ -69,8 +71,8 @@ enum class EAnimToolBlendOperation : uint8
 /**
 * This is a set of helper functions to access various parts of the Sequencer and Control Rig API via Python and Blueprints.
 */
-UCLASS(meta=(Transient, ScriptName="ControlRigSequencerLibrary"))
-class CONTROLRIGEDITOR_API UControlRigSequencerEditorLibrary : public UBlueprintFunctionLibrary
+UCLASS(MinimalAPI, meta=(Transient, ScriptName="ControlRigSequencerLibrary"))
+class UControlRigSequencerEditorLibrary : public UBlueprintFunctionLibrary
 {
 
 	public:
@@ -83,7 +85,7 @@ public:
 	* @return returns list of visible Control Rigs
 	*/
 	UFUNCTION(BlueprintCallable, Category = "Editor Scripting | Sequencer Tools | Control Rig")
-	static TArray<UControlRig*> GetVisibleControlRigs();
+	static UE_API TArray<UControlRig*> GetVisibleControlRigs();
 
 	/**
 	* Get all of the control rigs and their bindings in the level sequence
@@ -91,7 +93,7 @@ public:
 	* @return returns list of Control Rigs in the level sequence.
 	*/
 	UFUNCTION(BlueprintCallable, Category = "Editor Scripting | Sequencer Tools | Control Rig")
-	static TArray<FControlRigSequencerBindingProxy> GetControlRigs(ULevelSequence* LevelSequence);
+	static UE_API TArray<FControlRigSequencerBindingProxy> GetControlRigs(ULevelSequence* LevelSequence);
 
 	/**
 	* Find or create a Control Rig track of a specific class based upon the binding
@@ -102,7 +104,7 @@ public:
 	* @return returns Return the found or created track
 	*/
 	UFUNCTION(BlueprintCallable, Category = "Editor Scripting | Sequencer Tools | Control Rig")
-	static UMovieSceneTrack* FindOrCreateControlRigTrack(UWorld* World, ULevelSequence* LevelSequence, const UClass* ControlRigClass, const FMovieSceneBindingProxy& InBinding, bool bIsLayeredControlRig = false);
+	static UE_API UMovieSceneTrack* FindOrCreateControlRigTrack(UWorld* World, ULevelSequence* LevelSequence, const UClass* ControlRigClass, const FMovieSceneBindingProxy& InBinding, bool bIsLayeredControlRig = false);
 
 	/**
 	* Find or create a Control Rig Component
@@ -112,7 +114,7 @@ public:
 	* @return returns Find array of component Control Rigs that were found or created
 	*/
 	UFUNCTION(BlueprintCallable, Category = "Editor Scripting | Sequencer Tools | Control Rig")
-	static TArray<UMovieSceneTrack*> FindOrCreateControlRigComponentTrack(UWorld* World, ULevelSequence* LevelSequence, const FMovieSceneBindingProxy& InBinding);
+	static UE_API TArray<UMovieSceneTrack*> FindOrCreateControlRigComponentTrack(UWorld* World, ULevelSequence* LevelSequence, const FMovieSceneBindingProxy& InBinding);
 	
 	/*
 	 * Export Passed in Binding as an Anim Sequence, using open Level Sequence in Sequencer. Use this instead of 
@@ -126,7 +128,7 @@ public:
 	 * @bCreateLink If true will create a link between the animation sequence and the level sequence
 	 */
 	UFUNCTION(BlueprintCallable, Category = "Editor Scripting | Sequencer Tools | Animation")
-	static bool ExportAnimSequenceFromSequencer(UAnimSequence* AnimSequence, UAnimSeqExportOption* ExportOption, const FMovieSceneBindingProxy& Binding, bool bCreateLink);
+	static UE_API bool ExportAnimSequenceFromSequencer(UAnimSequence* AnimSequence, UAnimSeqExportOption* ExportOption, const FMovieSceneBindingProxy& Binding, bool bCreateLink);
 
 	/**
 	* Load anim sequence into this control rig section
@@ -147,7 +149,7 @@ public:
 	* @return returns True if successful, False otherwise
 	*/
 	UFUNCTION(BlueprintCallable, Category = "Editor Scripting | Sequencer Tools | Control Rig")
-	static bool LoadAnimSequenceIntoControlRigSectionWithRange(UMovieSceneSection* MovieSceneSection, UAnimSequence* AnimSequence, USkeletalMeshComponent* SkelMeshComp,
+	static UE_API bool LoadAnimSequenceIntoControlRigSectionWithRange(UMovieSceneSection* MovieSceneSection, UAnimSequence* AnimSequence, USkeletalMeshComponent* SkelMeshComp,
 		FFrameNumber InStartFrame, bool bUseCustomAnimRange, FFrameNumber AnimStartRange, FFrameNumber AnimEndRange,
 		EMovieSceneTimeUnit TimeUnit = EMovieSceneTimeUnit::DisplayRate, bool bKeyReduce = false, float Tolerance = 0.001f,
 		EMovieSceneKeyInterpolation Interpolation = EMovieSceneKeyInterpolation::SmartAuto, bool bResetControls = true,
@@ -155,7 +157,7 @@ public:
 
 	UE_DEPRECATED(5.6, "Use LoadAnimSequenceIntoControlRigSectionWithRange that takes UseCustomTimeRange")
 	UFUNCTION(BlueprintCallable, Category = "Editor Scripting | Sequencer Tools | Control Rig", meta = (DeprecatedFunction, DeprecationMessage = "Use LoadAnimSequenceIntoControlRigSectionWithRange that takes UseCustomTimeRange"))
-	static bool LoadAnimSequenceIntoControlRigSection(UMovieSceneSection* MovieSceneSection, UAnimSequence* AnimSequence, USkeletalMeshComponent* SkelMeshComp,
+	static UE_API bool LoadAnimSequenceIntoControlRigSection(UMovieSceneSection* MovieSceneSection, UAnimSequence* AnimSequence, USkeletalMeshComponent* SkelMeshComp,
 		FFrameNumber InStartFrame, EMovieSceneTimeUnit TimeUnit = EMovieSceneTimeUnit::DisplayRate, bool bKeyReduce = false, float Tolerance = 0.001f,
 		EMovieSceneKeyInterpolation Interpolation = EMovieSceneKeyInterpolation::SmartAuto, bool bResetControls = true);
 
@@ -172,7 +174,7 @@ public:
 	* @return returns True if successful, False otherwise
 	*/
 	UFUNCTION(BlueprintCallable, Category = "Editor Scripting | Sequencer Tools | Control Rig")
-	static bool BakeToControlRig(UWorld* World, ULevelSequence* LevelSequence, UClass* ControlRigClass, UAnimSeqExportOption* ExportOptions, bool bReduceKeys, float Tolerance,
+	static UE_API bool BakeToControlRig(UWorld* World, ULevelSequence* LevelSequence, UClass* ControlRigClass, UAnimSeqExportOption* ExportOptions, bool bReduceKeys, float Tolerance,
 			const FMovieSceneBindingProxy& Binding, bool bResetControls = true);
 
 	/**
@@ -185,7 +187,7 @@ public:
 	* @return returns True if successful, False otherwise
 	*/
 	UFUNCTION(BlueprintCallable, Category = "Editor Scripting | Sequencer Tools | Control Rig")
-	static bool SmartReduce(FSmartReduceParams& ReduceParams, UMovieSceneSection* MovieSceneSection);
+	static UE_API bool SmartReduce(FSmartReduceParams& ReduceParams, UMovieSceneSection* MovieSceneSection);
 
 
 	/**
@@ -197,7 +199,7 @@ public:
 	* @return Returns True if successful, False otherwise
 	*/
 	UFUNCTION(BlueprintCallable, Category = "Editor Scripting | Sequencer Tools | Control Rig")
-	static bool BakeConstraint(UWorld* World, UTickableConstraint* Constraint, const TArray<FFrameNumber>& Frames, EMovieSceneTimeUnit TimeUnit = EMovieSceneTimeUnit::DisplayRate);
+	static UE_API bool BakeConstraint(UWorld* World, UTickableConstraint* Constraint, const TArray<FFrameNumber>& Frames, EMovieSceneTimeUnit TimeUnit = EMovieSceneTimeUnit::DisplayRate);
 
 	/**
 	* Bake the constraint to keys based on the passed in settings. This will use the open sequencer to bake. See ConstraintsScriptingLibrary to get the list of available constraints
@@ -207,7 +209,7 @@ public:
 	* @return Returns True if successful, False otherwise
 	*/
 	UFUNCTION(BlueprintCallable, Category = "Editor Scripting | Sequencer Tools | Control Rig")
-	static bool BakeConstraints(UWorld* World, TArray<UTickableConstraint*>& InConstraints,const FBakingAnimationKeySettings& InSettings);
+	static UE_API bool BakeConstraints(UWorld* World, TArray<UTickableConstraint*>& InConstraints,const FBakingAnimationKeySettings& InSettings);
 	/**
 	* Add a constraint possibly adding to sequencer also if one is open.
 	* @param World The active world
@@ -218,7 +220,7 @@ public:
 	* @return Returns the constraint if created all nullptr if not
 	*/
 	UFUNCTION(BlueprintCallable, Category = "Editor Scripting | Sequencer Tools | Control Rig")
-	static UTickableConstraint* AddConstraint(UWorld* World, ETransformConstraintType InType, UTransformableHandle* InChild, UTransformableHandle* InParent, const bool bMaintainOffset);
+	static UE_API UTickableConstraint* AddConstraint(UWorld* World, ETransformConstraintType InType, UTransformableHandle* InChild, UTransformableHandle* InParent, const bool bMaintainOffset);
 	
 	/**
 	* Get the constraint keys for the specified constraint
@@ -230,7 +232,7 @@ public:
 	* @return Returns true if we got the keys from this constraint
 	*/
 	UFUNCTION(BlueprintCallable, Category = "Editor Scripting | Sequencer Tools | Control Rig")
-	static bool GetConstraintKeys(UTickableConstraint* InConstraint, UMovieSceneSection* ConstraintSection,TArray<bool>& OutBools, TArray<FFrameNumber>& OutFrames, EMovieSceneTimeUnit TimeUnit = EMovieSceneTimeUnit::DisplayRate);
+	static UE_API bool GetConstraintKeys(UTickableConstraint* InConstraint, UMovieSceneSection* ConstraintSection,TArray<bool>& OutBools, TArray<FFrameNumber>& OutFrames, EMovieSceneTimeUnit TimeUnit = EMovieSceneTimeUnit::DisplayRate);
 	
 	/**
 	* Set the constraint active key in the current open Sequencer
@@ -241,7 +243,7 @@ public:
 	* @return Returns true if we set the constraint to be the passed in value, false if not. We may not do so if the value is the same.
 	*/
 	UFUNCTION(BlueprintCallable, Category = "Editor Scripting | Sequencer Tools | Control Rig")
-	static bool SetConstraintActiveKey(UTickableConstraint* InConstraint, bool bActive, FFrameNumber InFrame, EMovieSceneTimeUnit TimeUnit = EMovieSceneTimeUnit::DisplayRate);
+	static UE_API bool SetConstraintActiveKey(UTickableConstraint* InConstraint, bool bActive, FFrameNumber InFrame, EMovieSceneTimeUnit TimeUnit = EMovieSceneTimeUnit::DisplayRate);
 
 	/**
 	* Get all constraints for this object, which is described by a transformable handle
@@ -249,7 +251,7 @@ public:
 	* @return Returns array of Constraints this handle is constrained to.
 	*/
 	UFUNCTION(BlueprintCallable, Category = "Editor Scripting | Sequencer Tools | Control Rig")
-	static TArray <UTickableConstraint*> GetConstraintsForHandle(UWorld* InWorld, const UTransformableHandle* InChild);
+	static UE_API TArray <UTickableConstraint*> GetConstraintsForHandle(UWorld* InWorld, const UTransformableHandle* InChild);
 
 	/** Move the constraint active key in the current open Sequencer
 	* @param InConstraint The constraint whose key to move
@@ -260,7 +262,7 @@ public:
 	* @return Will return false if function fails, for example if there is no key at this time it will fail, or if the new time is invalid it could fail also
 	*/
 	UFUNCTION(BlueprintCallable, Category = "Editor Scripting | Sequencer Tools | Control Rig")
-	static bool MoveConstraintKey(UTickableConstraint* Constraint, UMovieSceneSection* ConstraintSection, FFrameNumber InTime, FFrameNumber InNewTime, EMovieSceneTimeUnit TimeUnit = EMovieSceneTimeUnit::DisplayRate);
+	static UE_API bool MoveConstraintKey(UTickableConstraint* Constraint, UMovieSceneSection* ConstraintSection, FFrameNumber InTime, FFrameNumber InNewTime, EMovieSceneTimeUnit TimeUnit = EMovieSceneTimeUnit::DisplayRate);
 
 	/** Delete the Key for the Constraint at the specified time. 
 	* @param InConstraint The constraint whose key to move
@@ -270,7 +272,7 @@ public:
 	* @return Will return false if function fails,  for example if there is no key at this time it will fail.
 	*/
 	UFUNCTION(BlueprintCallable, Category = "Editor Scripting | Sequencer Tools | Control Rig")
-	static bool DeleteConstraintKey(UTickableConstraint* Constraint, UMovieSceneSection* ConstraintSection, FFrameNumber InTime, EMovieSceneTimeUnit TimeUnit = EMovieSceneTimeUnit::DisplayRate);
+	static UE_API bool DeleteConstraintKey(UTickableConstraint* Constraint, UMovieSceneSection* ConstraintSection, FFrameNumber InTime, EMovieSceneTimeUnit TimeUnit = EMovieSceneTimeUnit::DisplayRate);
 	/**
 	* Compensate constraint at the specfied time 
 	* @param InConstraint The constraint to compensate
@@ -279,7 +281,7 @@ public:
 	* @return Returns true if it can compensate
 	*/
 	UFUNCTION(BlueprintCallable, Category = "Editor Scripting | Sequencer Tools | Control Rig")
-	static bool Compensate(UTickableConstraint* InConstraint,  FFrameNumber InTime, EMovieSceneTimeUnit TimeUnit = EMovieSceneTimeUnit::DisplayRate);
+	static UE_API bool Compensate(UTickableConstraint* InConstraint,  FFrameNumber InTime, EMovieSceneTimeUnit TimeUnit = EMovieSceneTimeUnit::DisplayRate);
 
 	/**
 	* Compensate constraint at all keys
@@ -287,7 +289,7 @@ public:
 	* @return Returns true if it can compensate
 	*/
 	UFUNCTION(BlueprintCallable, Category = "Editor Scripting | Sequencer Tools | Control Rig")
-	static bool CompensateAll(UTickableConstraint* InConstraint);
+	static UE_API bool CompensateAll(UTickableConstraint* InConstraint);
 
 	/**
 	* Peform a Tween operation on the current active sequencer time(must be visible).
@@ -296,7 +298,7 @@ public:
 	* @param TweenValue The tween value to use, range from -1(blend to previous) to 1(blend to next)
 	*/
 	UFUNCTION(BlueprintCallable, Category = "Editor Scripting | Sequencer Tools | Control Rig")
-	static bool TweenControlRig(ULevelSequence* LevelSequence, UControlRig* ControlRig, float TweenValue);
+	static UE_API bool TweenControlRig(ULevelSequence* LevelSequence, UControlRig* ControlRig, float TweenValue);
 	
 	/**
 	* Peform specified blend operation based upon selected keys in the curve editor or selected control rig controls
@@ -305,7 +307,7 @@ public:
 	* @param BlendValue The blend value to use, range from -1(blend to previous) to 1(blend to next)
 	*/
 	UFUNCTION(BlueprintCallable, Category = "Editor Scripting | Sequencer Tools | Control Rig")
-	static bool BlendValuesOnSelected(ULevelSequence* LevelSequence, EAnimToolBlendOperation BlendOperation, 
+	static UE_API bool BlendValuesOnSelected(ULevelSequence* LevelSequence, EAnimToolBlendOperation BlendOperation, 
 		float BlendValue);
 
 	/**
@@ -320,7 +322,7 @@ public:
 	* @param Returns True if successful
 	*/
 	UFUNCTION(BlueprintCallable, Category = "Editor Scripting | Sequencer Tools | Control Rig")
-	static bool SnapControlRig(ULevelSequence* LevelSequence, FFrameNumber StartFrame, FFrameNumber EndFrame, const FControlRigSnapperSelection& ChildrenToSnap,
+	static UE_API bool SnapControlRig(ULevelSequence* LevelSequence, FFrameNumber StartFrame, FFrameNumber EndFrame, const FControlRigSnapperSelection& ChildrenToSnap,
 		const FControlRigSnapperSelection& ParentToSnap, const UControlRigSnapSettings* SnapSettings, EMovieSceneTimeUnit TimeUnit = EMovieSceneTimeUnit::DisplayRate);
 	
 	/**
@@ -332,7 +334,7 @@ public:
 	* @return Returns World Transform
 	*/
 	UFUNCTION(BlueprintCallable, Category = "Editor Scripting | Sequencer Tools | Control Rig")
-	static FTransform GetActorWorldTransform(ULevelSequence* LevelSequence,AActor* Actor, FFrameNumber Frame, EMovieSceneTimeUnit TimeUnit = EMovieSceneTimeUnit::DisplayRate);
+	static UE_API FTransform GetActorWorldTransform(ULevelSequence* LevelSequence,AActor* Actor, FFrameNumber Frame, EMovieSceneTimeUnit TimeUnit = EMovieSceneTimeUnit::DisplayRate);
 
 	/**
 	* Get Actors World Transforms at specific times
@@ -343,7 +345,7 @@ public:
 	* @return Returns World Transforms
 	*/
 	UFUNCTION(BlueprintCallable, Category = "Editor Scripting | Sequencer Tools | Control Rig")
-	static TArray<FTransform> GetActorWorldTransforms(ULevelSequence* LevelSequence, AActor* Actor, const TArray<FFrameNumber>& Frames, EMovieSceneTimeUnit TimeUnit = EMovieSceneTimeUnit::DisplayRate);
+	static UE_API TArray<FTransform> GetActorWorldTransforms(ULevelSequence* LevelSequence, AActor* Actor, const TArray<FFrameNumber>& Frames, EMovieSceneTimeUnit TimeUnit = EMovieSceneTimeUnit::DisplayRate);
 
 	/**
 	* Get SkeletalMeshComponent World Transform at a specific time
@@ -355,7 +357,7 @@ public:
 	* @return Returns World Transform
 	*/
 	UFUNCTION(BlueprintCallable, Category = "Editor Scripting | Sequencer Tools | Control Rig")
-	static FTransform GetSkeletalMeshComponentWorldTransform(ULevelSequence* LevelSequence, USkeletalMeshComponent* SkeletalMeshComponent, FFrameNumber Frame,
+	static UE_API FTransform GetSkeletalMeshComponentWorldTransform(ULevelSequence* LevelSequence, USkeletalMeshComponent* SkeletalMeshComponent, FFrameNumber Frame,
 		EMovieSceneTimeUnit TimeUnit = EMovieSceneTimeUnit::DisplayRate,FName ReferenceName = NAME_None);
 
 	/**
@@ -368,7 +370,7 @@ public:
 	* @return Returns World Transforms
 	*/
 	UFUNCTION(BlueprintCallable, Category = "Editor Scripting | Sequencer Tools | Control Rig")
-	static TArray<FTransform> GetSkeletalMeshComponentWorldTransforms(ULevelSequence* LevelSequence, USkeletalMeshComponent* SkeletalMeshComponent, const TArray<FFrameNumber>& Frames, 
+	static UE_API TArray<FTransform> GetSkeletalMeshComponentWorldTransforms(ULevelSequence* LevelSequence, USkeletalMeshComponent* SkeletalMeshComponent, const TArray<FFrameNumber>& Frames, 
 		EMovieSceneTimeUnit TimeUnit = EMovieSceneTimeUnit::DisplayRate,FName ReferenceName = NAME_None);
 
 	/**
@@ -388,7 +390,7 @@ public:
 	* @return Returns World Transform
 	*/
 	UFUNCTION(BlueprintCallable, Category = "Editor Scripting | Sequencer Tools | Control Rig")
-	static FTransform GetControlRigWorldTransform(ULevelSequence* LevelSequence, UControlRig* ControlRig, FName ControlName, FFrameNumber Frame,
+	static UE_API FTransform GetControlRigWorldTransform(ULevelSequence* LevelSequence, UControlRig* ControlRig, FName ControlName, FFrameNumber Frame,
 		EMovieSceneTimeUnit TimeUnit = EMovieSceneTimeUnit::DisplayRate);
 
 	/**
@@ -401,7 +403,7 @@ public:
 	* @return Returns World Transforms
 	*/
 	UFUNCTION(BlueprintCallable, Category = "Editor Scripting | Sequencer Tools | Control Rig")
-	static TArray<FTransform> GetControlRigWorldTransforms(ULevelSequence* LevelSequence, UControlRig* ControlRig, FName ControlName, const TArray<FFrameNumber>& Frames,
+	static UE_API TArray<FTransform> GetControlRigWorldTransforms(ULevelSequence* LevelSequence, UControlRig* ControlRig, FName ControlName, const TArray<FFrameNumber>& Frames,
 		EMovieSceneTimeUnit TimeUnit = EMovieSceneTimeUnit::DisplayRate);
 
 	/**
@@ -415,7 +417,7 @@ public:
 	* @param bSetKey Whether or not to set a key.
 	*/
 	UFUNCTION(BlueprintCallable, Category = "Editor Scripting | Sequencer Tools | Control Rig")
-	static void SetControlRigWorldTransform(ULevelSequence* LevelSequence, UControlRig* ControlRig, FName ControlName, FFrameNumber Frame, const FTransform& WorldTransform,
+	static UE_API void SetControlRigWorldTransform(ULevelSequence* LevelSequence, UControlRig* ControlRig, FName ControlName, FFrameNumber Frame, const FTransform& WorldTransform,
 		EMovieSceneTimeUnit TimeUnit = EMovieSceneTimeUnit::DisplayRate, bool bSetKey = true);
 
 	/**
@@ -428,7 +430,7 @@ public:
 	* @param TimeUnit Unit for frame values, either in display rate or tick resolution
 	*/
 	UFUNCTION(BlueprintCallable, Category = "Editor Scripting | Sequencer Tools | Control Rig")
-	static void SetControlRigWorldTransforms(ULevelSequence* LevelSequence, UControlRig* ControlRig, FName ControlName, const TArray<FFrameNumber>& Frames, 
+	static UE_API void SetControlRigWorldTransforms(ULevelSequence* LevelSequence, UControlRig* ControlRig, FName ControlName, const TArray<FFrameNumber>& Frames, 
 		 const TArray<FTransform>& WorldTransforms, EMovieSceneTimeUnit TimeUnit = EMovieSceneTimeUnit::DisplayRate);
 	
 	/**
@@ -441,7 +443,7 @@ public:
 	* @return Returns Value at that time
 	*/
 	UFUNCTION(BlueprintCallable, Category = "Editor Scripting | Sequencer Tools | Control Rig")
-	static float GetLocalControlRigFloat(ULevelSequence* LevelSequence, UControlRig* ControlRig, FName ControlName, FFrameNumber Frame,
+	static UE_API float GetLocalControlRigFloat(ULevelSequence* LevelSequence, UControlRig* ControlRig, FName ControlName, FFrameNumber Frame,
 		EMovieSceneTimeUnit TimeUnit = EMovieSceneTimeUnit::DisplayRate);
 
 	/**
@@ -454,7 +456,7 @@ public:
 	* @return Returns Values at those times
 	*/
 	UFUNCTION(BlueprintCallable, Category = "Editor Scripting | Sequencer Tools | Control Rig")
-	static TArray<float> GetLocalControlRigFloats(ULevelSequence* LevelSequence, UControlRig* ControlRig, FName ControlName, const TArray<FFrameNumber>& Frames,
+	static UE_API TArray<float> GetLocalControlRigFloats(ULevelSequence* LevelSequence, UControlRig* ControlRig, FName ControlName, const TArray<FFrameNumber>& Frames,
 		EMovieSceneTimeUnit TimeUnit = EMovieSceneTimeUnit::DisplayRate);
 
 	/**
@@ -468,7 +470,7 @@ public:
 	* @param bSetKey If True set a key, if not just set the value
 	*/
 	UFUNCTION(BlueprintCallable, Category = "Editor Scripting | Sequencer Tools | Control Rig")
-	static void SetLocalControlRigFloat(ULevelSequence* LevelSequence, UControlRig* ControlRig, FName ControlName, FFrameNumber Frame, float Value,
+	static UE_API void SetLocalControlRigFloat(ULevelSequence* LevelSequence, UControlRig* ControlRig, FName ControlName, FFrameNumber Frame, float Value,
 		EMovieSceneTimeUnit TimeUnit = EMovieSceneTimeUnit::DisplayRate,bool bSetKey = true);
 
 	/**
@@ -481,7 +483,7 @@ public:
 	* @param TimeUnit Unit for frame values, either in display rate or tick resolution
 	*/
 	UFUNCTION(BlueprintCallable, Category = "Editor Scripting | Sequencer Tools | Control Rig")
-	static void SetLocalControlRigFloats(ULevelSequence* LevelSequence, UControlRig* ControlRig, FName ControlName, const TArray<FFrameNumber>& Frames, const TArray<float> Values,
+	static UE_API void SetLocalControlRigFloats(ULevelSequence* LevelSequence, UControlRig* ControlRig, FName ControlName, const TArray<FFrameNumber>& Frames, const TArray<float> Values,
 		EMovieSceneTimeUnit TimeUnit = EMovieSceneTimeUnit::DisplayRate);
 
 	/**
@@ -494,7 +496,7 @@ public:
 	* @return Returns Value at that time
 	*/
 	UFUNCTION(BlueprintCallable, Category = "Editor Scripting | Sequencer Tools | Control Rig")
-	static bool GetLocalControlRigBool(ULevelSequence* LevelSequence, UControlRig* ControlRig, FName ControlName, FFrameNumber Frame,
+	static UE_API bool GetLocalControlRigBool(ULevelSequence* LevelSequence, UControlRig* ControlRig, FName ControlName, FFrameNumber Frame,
 		EMovieSceneTimeUnit TimeUnit = EMovieSceneTimeUnit::DisplayRate);
 
 	/**
@@ -507,7 +509,7 @@ public:
 	* @return Returns Values at those times
 	*/
 	UFUNCTION(BlueprintCallable, Category = "Editor Scripting | Sequencer Tools | Control Rig")
-	static TArray<bool> GetLocalControlRigBools(ULevelSequence* LevelSequence, UControlRig* ControlRig, FName ControlName, const TArray<FFrameNumber>& Frames,
+	static UE_API TArray<bool> GetLocalControlRigBools(ULevelSequence* LevelSequence, UControlRig* ControlRig, FName ControlName, const TArray<FFrameNumber>& Frames,
 		EMovieSceneTimeUnit TimeUnit = EMovieSceneTimeUnit::DisplayRate);
 
 	/**
@@ -521,7 +523,7 @@ public:
 	* @param bSetKey If True set a key, if not just set the value
 	*/
 	UFUNCTION(BlueprintCallable, Category = "Editor Scripting | Sequencer Tools | Control Rig")
-	static void SetLocalControlRigBool(ULevelSequence* LevelSequence, UControlRig* ControlRig, FName ControlName, FFrameNumber Frame, bool Value,
+	static UE_API void SetLocalControlRigBool(ULevelSequence* LevelSequence, UControlRig* ControlRig, FName ControlName, FFrameNumber Frame, bool Value,
 		EMovieSceneTimeUnit TimeUnit = EMovieSceneTimeUnit::DisplayRate, bool bSetKey = true);
 
 	/**
@@ -534,7 +536,7 @@ public:
 	* @param TimeUnit Unit for frame values, either in display rate or tick resolution
 */
 	UFUNCTION(BlueprintCallable, Category = "Editor Scripting | Sequencer Tools | Control Rig")
-	static void SetLocalControlRigBools(ULevelSequence* LevelSequence, UControlRig* ControlRig, FName ControlName, const TArray<FFrameNumber>& Frames,
+	static UE_API void SetLocalControlRigBools(ULevelSequence* LevelSequence, UControlRig* ControlRig, FName ControlName, const TArray<FFrameNumber>& Frames,
 		const TArray<bool> Values, EMovieSceneTimeUnit TimeUnit = EMovieSceneTimeUnit::DisplayRate);
 
 	/**
@@ -547,7 +549,7 @@ public:
 	* @return Returns Value at that time
 	*/
 	UFUNCTION(BlueprintCallable, Category = "Editor Scripting | Sequencer Tools | Control Rig")
-	static int32 GetLocalControlRigInt(ULevelSequence* LevelSequence, UControlRig* ControlRig, FName ControlName, FFrameNumber Frame, 
+	static UE_API int32 GetLocalControlRigInt(ULevelSequence* LevelSequence, UControlRig* ControlRig, FName ControlName, FFrameNumber Frame, 
 		EMovieSceneTimeUnit TimeUnit = EMovieSceneTimeUnit::DisplayRate);
 
 	/**
@@ -560,7 +562,7 @@ public:
 	* @return Returns Values at those times
 	*/
 	UFUNCTION(BlueprintCallable, Category = "Editor Scripting | Sequencer Tools | Control Rig")
-	static TArray<int32> GetLocalControlRigInts(ULevelSequence* LevelSequence, UControlRig* ControlRig, FName ControlName, const TArray<FFrameNumber>& Frames,
+	static UE_API TArray<int32> GetLocalControlRigInts(ULevelSequence* LevelSequence, UControlRig* ControlRig, FName ControlName, const TArray<FFrameNumber>& Frames,
 		EMovieSceneTimeUnit TimeUnit = EMovieSceneTimeUnit::DisplayRate);
 
 	/**
@@ -574,7 +576,7 @@ public:
 	* @param bSetKey If True set a key, if not just set the value
 	*/
 	UFUNCTION(BlueprintCallable, Category = "Editor Scripting | Sequencer Tools | Control Rig")
-	static void SetLocalControlRigInt(ULevelSequence* LevelSequence, UControlRig* ControlRig, FName ControlName, FFrameNumber Frame, int32 Value,
+	static UE_API void SetLocalControlRigInt(ULevelSequence* LevelSequence, UControlRig* ControlRig, FName ControlName, FFrameNumber Frame, int32 Value,
 		EMovieSceneTimeUnit TimeUnit = EMovieSceneTimeUnit::DisplayRate,bool bSetKey = true);
 
 
@@ -588,7 +590,7 @@ public:
 	* @param TimeUnit Unit for frame values, either in display rate or tick resolution
 	*/
 	UFUNCTION(BlueprintCallable, Category = "Editor Scripting | Sequencer Tools | Control Rig")
-	static void SetLocalControlRigInts(ULevelSequence* LevelSequence, UControlRig* ControlRig, FName ControlName, const TArray<FFrameNumber>& Frames, const TArray<int32> Values,
+	static UE_API void SetLocalControlRigInts(ULevelSequence* LevelSequence, UControlRig* ControlRig, FName ControlName, const TArray<FFrameNumber>& Frames, const TArray<int32> Values,
 		EMovieSceneTimeUnit TimeUnit = EMovieSceneTimeUnit::DisplayRate);
 
 	/**
@@ -601,7 +603,7 @@ public:
 	* @return Returns Value at that time
 	*/
 	UFUNCTION(BlueprintCallable, Category = "Editor Scripting | Sequencer Tools | Control Rig")
-	static FVector2D GetLocalControlRigVector2D(ULevelSequence* LevelSequence, UControlRig* ControlRig, FName ControlName, FFrameNumber Frame,
+	static UE_API FVector2D GetLocalControlRigVector2D(ULevelSequence* LevelSequence, UControlRig* ControlRig, FName ControlName, FFrameNumber Frame,
 		EMovieSceneTimeUnit TimeUnit = EMovieSceneTimeUnit::DisplayRate);
 
 	/**
@@ -614,7 +616,7 @@ public:
 	* @return Returns Values at those times
 	*/
 	UFUNCTION(BlueprintCallable, Category = "Editor Scripting | Sequencer Tools | Control Rig")
-	static TArray<FVector2D> GetLocalControlRigVector2Ds(ULevelSequence* LevelSequence, UControlRig* ControlRig, FName ControlName, const TArray<FFrameNumber>& Frames,
+	static UE_API TArray<FVector2D> GetLocalControlRigVector2Ds(ULevelSequence* LevelSequence, UControlRig* ControlRig, FName ControlName, const TArray<FFrameNumber>& Frames,
 		EMovieSceneTimeUnit TimeUnit = EMovieSceneTimeUnit::DisplayRate);
 
 	/**
@@ -628,7 +630,7 @@ public:
 	* @param bSetKey If True set a key, if not just set the value
 	*/
 	UFUNCTION(BlueprintCallable, Category = "Editor Scripting | Sequencer Tools | Control Rig")
-	static void SetLocalControlRigVector2D(ULevelSequence* LevelSequence, UControlRig* ControlRig, FName ControlName, FFrameNumber Frame, FVector2D Value,
+	static UE_API void SetLocalControlRigVector2D(ULevelSequence* LevelSequence, UControlRig* ControlRig, FName ControlName, FFrameNumber Frame, FVector2D Value,
 		EMovieSceneTimeUnit TimeUnit = EMovieSceneTimeUnit::DisplayRate,bool bSetKey = true);
 
 
@@ -642,7 +644,7 @@ public:
 	* @param TimeUnit Unit for frame values, either in display rate or tick resolution
 	*/
 	UFUNCTION(BlueprintCallable, Category = "Editor Scripting | Sequencer Tools | Control Rig")
-	static void SetLocalControlRigVector2Ds(ULevelSequence* LevelSequence, UControlRig* ControlRig, FName ControlName, const TArray<FFrameNumber>& Frames, const TArray<FVector2D> Values,
+	static UE_API void SetLocalControlRigVector2Ds(ULevelSequence* LevelSequence, UControlRig* ControlRig, FName ControlName, const TArray<FFrameNumber>& Frames, const TArray<FVector2D> Values,
 		EMovieSceneTimeUnit TimeUnit = EMovieSceneTimeUnit::DisplayRate);
 
 	/**
@@ -655,7 +657,7 @@ public:
 	* @return Returns Value at that time
 	*/
 	UFUNCTION(BlueprintCallable, Category = "Editor Scripting | Sequencer Tools | Control Rig")
-	static FVector GetLocalControlRigPosition(ULevelSequence* LevelSequence, UControlRig* ControlRig, FName ControlName, FFrameNumber Frame,
+	static UE_API FVector GetLocalControlRigPosition(ULevelSequence* LevelSequence, UControlRig* ControlRig, FName ControlName, FFrameNumber Frame,
 		EMovieSceneTimeUnit TimeUnit = EMovieSceneTimeUnit::DisplayRate);
 
 	/**
@@ -668,7 +670,7 @@ public:
 	* @return Returns Values at those times
 	*/
 	UFUNCTION(BlueprintCallable, Category = "Editor Scripting | Sequencer Tools | Control Rig")
-	static TArray<FVector> GetLocalControlRigPositions(ULevelSequence* LevelSequence, UControlRig* ControlRig, FName ControlName, const TArray<FFrameNumber>& Frames,
+	static UE_API TArray<FVector> GetLocalControlRigPositions(ULevelSequence* LevelSequence, UControlRig* ControlRig, FName ControlName, const TArray<FFrameNumber>& Frames,
 		EMovieSceneTimeUnit TimeUnit = EMovieSceneTimeUnit::DisplayRate);
 
 	/**
@@ -682,7 +684,7 @@ public:
 	* @param bSetKey If True set a key, if not just set the value
 	*/
 	UFUNCTION(BlueprintCallable, Category = "Editor Scripting | Sequencer Tools | Control Rig")
-	static void SetLocalControlRigPosition(ULevelSequence* LevelSequence, UControlRig* ControlRig, FName ControlName, FFrameNumber Frame, FVector Value,
+	static UE_API void SetLocalControlRigPosition(ULevelSequence* LevelSequence, UControlRig* ControlRig, FName ControlName, FFrameNumber Frame, FVector Value,
 		EMovieSceneTimeUnit TimeUnit = EMovieSceneTimeUnit::DisplayRate, bool bSetKey = true);
 
 	/**
@@ -695,7 +697,7 @@ public:
 	* @param TimeUnit Unit for frame values, either in display rate or tick resolution
 	*/
 	UFUNCTION(BlueprintCallable, Category = "Editor Scripting | Sequencer Tools | Control Rig")
-	static void SetLocalControlRigPositions(ULevelSequence* LevelSequence, UControlRig* ControlRig, FName ControlName, const TArray<FFrameNumber>& Frames, const TArray<FVector> Values,
+	static UE_API void SetLocalControlRigPositions(ULevelSequence* LevelSequence, UControlRig* ControlRig, FName ControlName, const TArray<FFrameNumber>& Frames, const TArray<FVector> Values,
 		EMovieSceneTimeUnit TimeUnit = EMovieSceneTimeUnit::DisplayRate);
 
 	/**
@@ -708,7 +710,7 @@ public:
 	* @return Returns Value at that time
 	*/
 	UFUNCTION(BlueprintCallable, Category = "Editor Scripting | Sequencer Tools | Control Rig")
-	static FRotator GetLocalControlRigRotator(ULevelSequence* LevelSequence, UControlRig* ControlRig, FName ControlName, FFrameNumber Frame,
+	static UE_API FRotator GetLocalControlRigRotator(ULevelSequence* LevelSequence, UControlRig* ControlRig, FName ControlName, FFrameNumber Frame,
 		EMovieSceneTimeUnit TimeUnit = EMovieSceneTimeUnit::DisplayRate);
 
 	/**
@@ -721,7 +723,7 @@ public:
 	* @return Returns Values at those times
 	*/
 	UFUNCTION(BlueprintCallable, Category = "Editor Scripting | Sequencer Tools | Control Rig")
-	static TArray<FRotator> GetLocalControlRigRotators(ULevelSequence* LevelSequence, UControlRig* ControlRig, FName ControlName, const TArray<FFrameNumber>& Frames,
+	static UE_API TArray<FRotator> GetLocalControlRigRotators(ULevelSequence* LevelSequence, UControlRig* ControlRig, FName ControlName, const TArray<FFrameNumber>& Frames,
 		EMovieSceneTimeUnit TimeUnit = EMovieSceneTimeUnit::DisplayRate);
 
 	/**
@@ -735,7 +737,7 @@ public:
 	* @param bSetKey If True set a key, if not just set the value
 	*/
 	UFUNCTION(BlueprintCallable, Category = "Editor Scripting | Sequencer Tools | Control Rig")
-	static void SetLocalControlRigRotator(ULevelSequence* LevelSequence, UControlRig* ControlRig, FName ControlName, FFrameNumber Frame, FRotator Value, 
+	static UE_API void SetLocalControlRigRotator(ULevelSequence* LevelSequence, UControlRig* ControlRig, FName ControlName, FFrameNumber Frame, FRotator Value, 
 		EMovieSceneTimeUnit TimeUnit = EMovieSceneTimeUnit::DisplayRate, bool bSetKey = true);
 
 	/**
@@ -748,7 +750,7 @@ public:
 	* @param TimeUnit Unit for frame values, either in display rate or tick resolution
 	*/
 	UFUNCTION(BlueprintCallable, Category = "Editor Scripting | Sequencer Tools | Control Rig")
-	static void SetLocalControlRigRotators(ULevelSequence* LevelSequence, UControlRig* ControlRig, FName ControlName, const TArray<FFrameNumber>& Frames, const TArray<FRotator> Values,
+	static UE_API void SetLocalControlRigRotators(ULevelSequence* LevelSequence, UControlRig* ControlRig, FName ControlName, const TArray<FFrameNumber>& Frames, const TArray<FRotator> Values,
 		EMovieSceneTimeUnit TimeUnit = EMovieSceneTimeUnit::DisplayRate);
 
 	/**
@@ -761,7 +763,7 @@ public:
 	* @return Returns Value at that time
 	*/
 	UFUNCTION(BlueprintCallable, Category = "Editor Scripting | Sequencer Tools | Control Rig")
-	static FVector GetLocalControlRigScale(ULevelSequence* LevelSequence, UControlRig* ControlRig, FName ControlName, FFrameNumber Frame,
+	static UE_API FVector GetLocalControlRigScale(ULevelSequence* LevelSequence, UControlRig* ControlRig, FName ControlName, FFrameNumber Frame,
 		EMovieSceneTimeUnit TimeUnit = EMovieSceneTimeUnit::DisplayRate);
 
 	/**
@@ -774,7 +776,7 @@ public:
 	* @return Returns Values at those times
 	*/
 	UFUNCTION(BlueprintCallable, Category = "Editor Scripting | Sequencer Tools | Control Rig")
-	static TArray<FVector> GetLocalControlRigScales(ULevelSequence* LevelSequence, UControlRig* ControlRig, FName ControlName, const TArray<FFrameNumber>& Frames,
+	static UE_API TArray<FVector> GetLocalControlRigScales(ULevelSequence* LevelSequence, UControlRig* ControlRig, FName ControlName, const TArray<FFrameNumber>& Frames,
 		EMovieSceneTimeUnit TimeUnit = EMovieSceneTimeUnit::DisplayRate);
 
 	/**
@@ -788,7 +790,7 @@ public:
 	* @param bSetKey If True set a key, if not just set the value
 */
 	UFUNCTION(BlueprintCallable, Category = "Editor Scripting | Sequencer Tools | Control Rig")
-	static void SetLocalControlRigScale(ULevelSequence* LevelSequence, UControlRig* ControlRig, FName ControlName, FFrameNumber Frame, FVector Value, 
+	static UE_API void SetLocalControlRigScale(ULevelSequence* LevelSequence, UControlRig* ControlRig, FName ControlName, FFrameNumber Frame, FVector Value, 
 		EMovieSceneTimeUnit TimeUnit = EMovieSceneTimeUnit::DisplayRate, bool bSetKey = true);
 
 	/**
@@ -801,7 +803,7 @@ public:
 	* @param TimeUnit Unit for frame values, either in display rate or tick resolution
 	*/
 	UFUNCTION(BlueprintCallable, Category = "Editor Scripting | Sequencer Tools | Control Rig")
-	static void SetLocalControlRigScales(ULevelSequence* LevelSequence, UControlRig* ControlRig, FName ControlName, const TArray<FFrameNumber>& Frames, const TArray<FVector> Values,
+	static UE_API void SetLocalControlRigScales(ULevelSequence* LevelSequence, UControlRig* ControlRig, FName ControlName, const TArray<FFrameNumber>& Frames, const TArray<FVector> Values,
 		EMovieSceneTimeUnit TimeUnit = EMovieSceneTimeUnit::DisplayRate);
 
 
@@ -815,7 +817,7 @@ public:
 	* @return Returns Value at that time
 	*/
 	UFUNCTION(BlueprintCallable, Category = "Editor Scripting | Sequencer Tools | Control Rig")
-	static FEulerTransform GetLocalControlRigEulerTransform(ULevelSequence* LevelSequence, UControlRig* ControlRig, FName ControlName, FFrameNumber Frame,
+	static UE_API FEulerTransform GetLocalControlRigEulerTransform(ULevelSequence* LevelSequence, UControlRig* ControlRig, FName ControlName, FFrameNumber Frame,
 		EMovieSceneTimeUnit TimeUnit = EMovieSceneTimeUnit::DisplayRate);
 
 	/**
@@ -828,7 +830,7 @@ public:
 	* @return Returns Values at those times
 	*/
 	UFUNCTION(BlueprintCallable, Category = "Editor Scripting | Sequencer Tools | Control Rig")
-	static TArray<FEulerTransform> GetLocalControlRigEulerTransforms(ULevelSequence* LevelSequence, UControlRig* ControlRig, FName ControlName, const TArray<FFrameNumber>& Frames,
+	static UE_API TArray<FEulerTransform> GetLocalControlRigEulerTransforms(ULevelSequence* LevelSequence, UControlRig* ControlRig, FName ControlName, const TArray<FFrameNumber>& Frames,
 		EMovieSceneTimeUnit TimeUnit = EMovieSceneTimeUnit::DisplayRate);
 
 	/**
@@ -842,7 +844,7 @@ public:
 	* @param bSetKey If True set a key, if not just set the value
 	*/
 	UFUNCTION(BlueprintCallable, Category = "Editor Scripting | Sequencer Tools | Control Rig")
-	static void SetLocalControlRigEulerTransform(ULevelSequence* LevelSequence, UControlRig* ControlRig, FName ControlName, FFrameNumber Frame, FEulerTransform Value,
+	static UE_API void SetLocalControlRigEulerTransform(ULevelSequence* LevelSequence, UControlRig* ControlRig, FName ControlName, FFrameNumber Frame, FEulerTransform Value,
 		EMovieSceneTimeUnit TimeUnit = EMovieSceneTimeUnit::DisplayRate, bool bSetKey = true);
 
 	/**
@@ -855,7 +857,7 @@ public:
 	* @param TimeUnit Unit for frame values, either in display rate or tick resolution
 	*/
 	UFUNCTION(BlueprintCallable, Category = "Editor Scripting | Sequencer Tools | Control Rig")
-	static void SetLocalControlRigEulerTransforms(ULevelSequence* LevelSequence, UControlRig* ControlRig, FName ControlName, const TArray<FFrameNumber>& Frames, const TArray<FEulerTransform> Values,
+	static UE_API void SetLocalControlRigEulerTransforms(ULevelSequence* LevelSequence, UControlRig* ControlRig, FName ControlName, const TArray<FFrameNumber>& Frames, const TArray<FEulerTransform> Values,
 		EMovieSceneTimeUnit TimeUnit = EMovieSceneTimeUnit::DisplayRate);
 
 	/**
@@ -868,7 +870,7 @@ public:
 	* @return Returns Value at that time
 	*/
 	UFUNCTION(BlueprintCallable, Category = "Editor Scripting | Sequencer Tools | Control Rig")
-	static FTransformNoScale GetLocalControlRigTransformNoScale(ULevelSequence* LevelSequence, UControlRig* ControlRig, FName ControlName, FFrameNumber Frame,
+	static UE_API FTransformNoScale GetLocalControlRigTransformNoScale(ULevelSequence* LevelSequence, UControlRig* ControlRig, FName ControlName, FFrameNumber Frame,
 		EMovieSceneTimeUnit TimeUnit = EMovieSceneTimeUnit::DisplayRate);
 
 	/**
@@ -881,7 +883,7 @@ public:
 	* @return Returns Values at those times
 	*/
 	UFUNCTION(BlueprintCallable, Category = "Editor Scripting | Sequencer Tools | Control Rig")
-	static TArray<FTransformNoScale> GetLocalControlRigTransformNoScales(ULevelSequence* LevelSequence, UControlRig* ControlRig, FName ControlName, const TArray<FFrameNumber>& Frames,
+	static UE_API TArray<FTransformNoScale> GetLocalControlRigTransformNoScales(ULevelSequence* LevelSequence, UControlRig* ControlRig, FName ControlName, const TArray<FFrameNumber>& Frames,
 		EMovieSceneTimeUnit TimeUnit = EMovieSceneTimeUnit::DisplayRate);
 
 	/**
@@ -895,7 +897,7 @@ public:
 	* @param bSetKey If True set a key, if not just set the value
 	*/
 	UFUNCTION(BlueprintCallable, Category = "Editor Scripting | Sequencer Tools | Control Rig")
-	static void SetLocalControlRigTransformNoScale(ULevelSequence* LevelSequence, UControlRig* ControlRig, FName ControlName, FFrameNumber Frame, FTransformNoScale Value, 
+	static UE_API void SetLocalControlRigTransformNoScale(ULevelSequence* LevelSequence, UControlRig* ControlRig, FName ControlName, FFrameNumber Frame, FTransformNoScale Value, 
 		EMovieSceneTimeUnit TimeUnit = EMovieSceneTimeUnit::DisplayRate, bool bSetKey = true);
 
 	/**
@@ -908,7 +910,7 @@ public:
 	* @param TimeUnit Unit for frame values, either in display rate or tick resolution
 	*/
 	UFUNCTION(BlueprintCallable, Category = "Editor Scripting | Sequencer Tools | Control Rig")
-	static void SetLocalControlRigTransformNoScales(ULevelSequence* LevelSequence, UControlRig* ControlRig, FName ControlName, const TArray<FFrameNumber>& Frames, const TArray<FTransformNoScale> Values,
+	static UE_API void SetLocalControlRigTransformNoScales(ULevelSequence* LevelSequence, UControlRig* ControlRig, FName ControlName, const TArray<FFrameNumber>& Frames, const TArray<FTransformNoScale> Values,
 		EMovieSceneTimeUnit TimeUnit = EMovieSceneTimeUnit::DisplayRate);
 
 	/**
@@ -921,7 +923,7 @@ public:
 	* @return Returns Value at that time
 	*/
 	UFUNCTION(BlueprintCallable, Category = "Editor Scripting | Sequencer Tools | Control Rig")
-	static FTransform GetLocalControlRigTransform(ULevelSequence* LevelSequence, UControlRig* ControlRig, FName ControlName, FFrameNumber Frame, 
+	static UE_API FTransform GetLocalControlRigTransform(ULevelSequence* LevelSequence, UControlRig* ControlRig, FName ControlName, FFrameNumber Frame, 
 		EMovieSceneTimeUnit TimeUnit = EMovieSceneTimeUnit::DisplayRate);
 
 	/**
@@ -934,7 +936,7 @@ public:
 	* @return Returns Values at those times
 	*/
 	UFUNCTION(BlueprintCallable, Category = "Editor Scripting | Sequencer Tools | Control Rig")
-	static TArray<FTransform> GetLocalControlRigTransforms(ULevelSequence* LevelSequence, UControlRig* ControlRig, FName ControlName, const TArray<FFrameNumber>& Frames,
+	static UE_API TArray<FTransform> GetLocalControlRigTransforms(ULevelSequence* LevelSequence, UControlRig* ControlRig, FName ControlName, const TArray<FFrameNumber>& Frames,
 		EMovieSceneTimeUnit TimeUnit = EMovieSceneTimeUnit::DisplayRate);
 
 	/**
@@ -948,7 +950,7 @@ public:
 	* @param bSetKey If True set a key, if not just set the value
 	*/
 	UFUNCTION(BlueprintCallable, Category = "Editor Scripting | Sequencer Tools | Control Rig")
-	static void SetLocalControlRigTransform(ULevelSequence* LevelSequence, UControlRig* ControlRig, FName ControlName, FFrameNumber Frame, FTransform Value, 
+	static UE_API void SetLocalControlRigTransform(ULevelSequence* LevelSequence, UControlRig* ControlRig, FName ControlName, FFrameNumber Frame, FTransform Value, 
 		EMovieSceneTimeUnit TimeUnit = EMovieSceneTimeUnit::DisplayRate, bool bSetKey = true);
 
 	/**
@@ -961,7 +963,7 @@ public:
 	* @param TimeUnit Unit for frame values, either in display rate or tick resolution
 	*/
 	UFUNCTION(BlueprintCallable, Category = "Editor Scripting | Sequencer Tools | Control Rig")
-	static void SetLocalControlRigTransforms(ULevelSequence* LevelSequence, UControlRig* ControlRig, FName ControlName, const TArray<FFrameNumber>& Frames, const TArray<FTransform> Values,
+	static UE_API void SetLocalControlRigTransforms(ULevelSequence* LevelSequence, UControlRig* ControlRig, FName ControlName, const TArray<FFrameNumber>& Frames, const TArray<FTransform> Values,
 		EMovieSceneTimeUnit TimeUnit = EMovieSceneTimeUnit::DisplayRate);
 
 	/*
@@ -976,14 +978,14 @@ public:
 	 * @param InImportFileName Path to fbx file to create
 	 */
 	UFUNCTION(BlueprintCallable, Category = "Editor Scripting | Sequencer Tools | FBX")
-	static bool ImportFBXToControlRigTrack(UWorld* World, ULevelSequence* InSequence, UMovieSceneControlRigParameterTrack* InTrack, UMovieSceneControlRigParameterSection* InSection,
+	static UE_API bool ImportFBXToControlRigTrack(UWorld* World, ULevelSequence* InSequence, UMovieSceneControlRigParameterTrack* InTrack, UMovieSceneControlRigParameterSection* InSection,
 			const TArray<FString>& SelectedControlRigNames,
 			UMovieSceneUserImportFBXControlRigSettings* ImportFBXControlRigSettings,
 			const FString& ImportFilename);
 
 	/** Exports an FBX from the given control rig section. */
 	UFUNCTION(BlueprintCallable, Category = "Editor Scripting | Sequencer Tools | FBX")
-	static bool ExportFBXFromControlRigSection(ULevelSequence* Sequence, const UMovieSceneControlRigParameterSection* Section,
+	static UE_API bool ExportFBXFromControlRigSection(ULevelSequence* Sequence, const UMovieSceneControlRigParameterSection* Section,
 	                                    const UMovieSceneUserExportFBXControlRigSettings* ExportFBXControlRigSettings);
 
 	/*
@@ -995,7 +997,7 @@ public:
 	 * @param Tolerance If reducing keys, tolerance about which keys will be removed, smaller tolerance, more keys usually.
 	 */
 	UFUNCTION(BlueprintCallable, Category = "Editor Scripting | Sequencer Tools | Control Rig")
-	static bool CollapseControlRigAnimLayers(ULevelSequence* InSequence,UMovieSceneControlRigParameterTrack* InTrack, bool bKeyReduce = false, float Tolerance = 0.001f);
+	static UE_API bool CollapseControlRigAnimLayers(ULevelSequence* InSequence,UMovieSceneControlRigParameterTrack* InTrack, bool bKeyReduce = false, float Tolerance = 0.001f);
 
 	/*
 	 * Collapse and bake all sections and layers on a control rig track to just one section using passed in settings.
@@ -1005,21 +1007,21 @@ public:
 	 * @param InSettings Settings that determine how to collapse
 	 */
 	UFUNCTION(BlueprintCallable, Category = "Editor Scripting | Sequencer Tools | Control Rig")
-	static bool CollapseControlRigAnimLayersWithSettings(ULevelSequence* InSequence, UMovieSceneControlRigParameterTrack* InTrack, const FBakingAnimationKeySettings& InSettings);
+	static UE_API bool CollapseControlRigAnimLayersWithSettings(ULevelSequence* InSequence, UMovieSceneControlRigParameterTrack* InTrack, const FBakingAnimationKeySettings& InSettings);
 
 	/*
 	 * Get the default parent key, can be used a parent space.
 	 *
 	*/
 	UFUNCTION(BlueprintCallable, Category = "Editor Scripting | Sequencer Tools | Control Rig")
-	static FRigElementKey GetDefaultParentKey();
+	static UE_API FRigElementKey GetDefaultParentKey();
 
 	/*
 	 * Get the default world space key, can be used a world space.
 	 *
 	*/
 	UFUNCTION(BlueprintCallable, Category = "Editor Scripting | Sequencer Tools | Control Rig")
-	static FRigElementKey GetWorldSpaceReferenceKey();
+	static UE_API FRigElementKey GetWorldSpaceReferenceKey();
 
 	/*
 	 * Set the a key for the Control Rig Space for the Control at the specified time. If space is the same as the current no key witll be set.
@@ -1032,7 +1034,7 @@ public:
 	 * @param TimeUnit Unit for the InTime, either in display rate or tick resolution
 	 */
 	UFUNCTION(BlueprintCallable, Category = "Editor Scripting | Sequencer Tools | Control Rig")
-	static bool SetControlRigSpace(ULevelSequence* InSequence, UControlRig* InControlRig, FName InControlName, const FRigElementKey& InSpaceKey, FFrameNumber InTime,  EMovieSceneTimeUnit TimeUnit = EMovieSceneTimeUnit::DisplayRate);
+	static UE_API bool SetControlRigSpace(ULevelSequence* InSequence, UControlRig* InControlRig, FName InControlName, const FRigElementKey& InSpaceKey, FFrameNumber InTime,  EMovieSceneTimeUnit TimeUnit = EMovieSceneTimeUnit::DisplayRate);
 
 	/** Bake specified Control Rig Controls to a specified Space based upon the current settings
 	* @param InSequence Sequence to bake
@@ -1042,7 +1044,7 @@ public:
 	* @param TimeUnit Unit for the start and end times in the InSettings parameter.
 	*/
 	UFUNCTION(BlueprintCallable, Category = "Editor Scripting | Sequencer Tools | Control Rig")
-	static bool BakeControlRigSpace(ULevelSequence* InSequence, UControlRig* InControlRig, const TArray<FName>& InControlNames, FRigSpacePickerBakeSettings InSettings, EMovieSceneTimeUnit TimeUnit = EMovieSceneTimeUnit::DisplayRate);
+	static UE_API bool BakeControlRigSpace(ULevelSequence* InSequence, UControlRig* InControlRig, const TArray<FName>& InControlNames, FRigSpacePickerBakeSettings InSettings, EMovieSceneTimeUnit TimeUnit = EMovieSceneTimeUnit::DisplayRate);
 	
 	/** Perform compensation for any spaces at the specified time for the specified control rig
 	* @param InControlRig Control Rig to compensate
@@ -1051,13 +1053,13 @@ public:
 	* @return Will return false if function fails
 	*/
 	UFUNCTION(BlueprintCallable, Category = "Editor Scripting | Sequencer Tools | Control Rig")
-	static bool SpaceCompensate(UControlRig* InControlRig, FFrameNumber InTime, EMovieSceneTimeUnit TimeUnit = EMovieSceneTimeUnit::DisplayRate);
+	static UE_API bool SpaceCompensate(UControlRig* InControlRig, FFrameNumber InTime, EMovieSceneTimeUnit TimeUnit = EMovieSceneTimeUnit::DisplayRate);
 
 	/** Perform compensation for all spaces for the specified control rig
 	* @param InControlRig Control Rig to compensate
 	* @return Will return false if function fails
 	*/
-	static bool SpaceCompensateAll(UControlRig* InControlRig);
+	static UE_API bool SpaceCompensateAll(UControlRig* InControlRig);
 
 	/** Delete the Control Rig Space Key for the Control at the specified time. This will delete any attached Control Rig keys at this time and will perform any needed compensation to the new space.
 	*
@@ -1069,7 +1071,7 @@ public:
 	* @return Will return false if function fails,  for example if there is no key at this time it will fail.
 	*/
 	UFUNCTION(BlueprintCallable, Category = "Editor Scripting | Sequencer Tools | Control Rig")
-	static bool DeleteControlRigSpace(ULevelSequence* InSequence, UControlRig* InControlRig, FName InControlName, FFrameNumber InTime, EMovieSceneTimeUnit TimeUnit = EMovieSceneTimeUnit::DisplayRate);
+	static UE_API bool DeleteControlRigSpace(ULevelSequence* InSequence, UControlRig* InControlRig, FName InControlName, FFrameNumber InTime, EMovieSceneTimeUnit TimeUnit = EMovieSceneTimeUnit::DisplayRate);
 
 	/** Move the Control Rig Space Key for the Control at the specified time to the new time. This will also move any Control Rig keys at this space switch boundary.
 	*
@@ -1082,7 +1084,7 @@ public:
 	* @return Will return false if function fails, for example if there is no key at this time it will fail, or if the new time is invalid it could fail also
 	*/
 	UFUNCTION(BlueprintCallable, Category = "Editor Scripting | Sequencer Tools | Control Rig")
-	static bool MoveControlRigSpace(ULevelSequence* InSequence, UControlRig* InControlRig, FName InControlName, FFrameNumber InTime, FFrameNumber InNewTime, EMovieSceneTimeUnit TimeUnit = EMovieSceneTimeUnit::DisplayRate);
+	static UE_API bool MoveControlRigSpace(ULevelSequence* InSequence, UControlRig* InControlRig, FName InControlName, FFrameNumber InTime, FFrameNumber InNewTime, EMovieSceneTimeUnit TimeUnit = EMovieSceneTimeUnit::DisplayRate);
 
 	/** Rename the Control Rig Channels in Sequencer to the specified new control names, which should be present on the Control Rig
 	* @param InSequence Sequence to rename controls
@@ -1092,49 +1094,49 @@ public:
 	* @return Return true if the function succeeds, false if it doesn't which can happen if the name arrays don't match in size or any of the new Control Names aren't valid
 	*/
 	UFUNCTION(BlueprintCallable, Category = "Editor Scripting | Sequencer Tools | Control Rig")
-	static bool RenameControlRigControlChannels(ULevelSequence* InSequence, UControlRig* InControlRig, const TArray<FName>& InOldControlNames, const TArray<FName>& InNewControlNames);
+	static UE_API bool RenameControlRigControlChannels(ULevelSequence* InSequence, UControlRig* InControlRig, const TArray<FName>& InOldControlNames, const TArray<FName>& InNewControlNames);
 
 	/** Get the controls mask for the given ControlName */
 	UFUNCTION(BlueprintPure, Category = "Editor Scripting | Sequencer Tools | Control Rig")
-	static bool GetControlsMask(UMovieSceneSection* InSection, FName ControlName);
+	static UE_API bool GetControlsMask(UMovieSceneSection* InSection, FName ControlName);
 
 	/** Set the controls mask for the given ControlNames */
 	UFUNCTION(BlueprintCallable, Category = "Editor Scripting | Sequencer Tools | Control Rig")
-	static void SetControlsMask(UMovieSceneSection* InSection, const TArray<FName>& ControlNames, bool bVisible);
+	static UE_API void SetControlsMask(UMovieSceneSection* InSection, const TArray<FName>& ControlNames, bool bVisible);
 
 	/** Shows all of the controls for the given section */
 	UFUNCTION(BlueprintCallable, Category = "Editor Scripting | Sequencer Tools | Control Rig")
-	static void ShowAllControls(UMovieSceneSection* InSection);
+	static UE_API void ShowAllControls(UMovieSceneSection* InSection);
 
 	/** Hides all of the controls for the given section */
 	UFUNCTION(BlueprintCallable, Category = "Editor Scripting | Sequencer Tools | Control Rig")
-	static void HideAllControls(UMovieSceneSection* InSection);
+	static UE_API void HideAllControls(UMovieSceneSection* InSection);
 
 	/** Set Control Rig priority order */
 	UFUNCTION(BlueprintCallable, Category = "Editor Scripting | Sequencer Tools | Control Rig")
-	static void SetControlRigPriorityOrder(UMovieSceneTrack* InSection,int32 PriorityOrder);
+	static UE_API void SetControlRigPriorityOrder(UMovieSceneTrack* InSection,int32 PriorityOrder);
 
 	/** Get Control Rig prirority order */
 	UFUNCTION(BlueprintCallable, Category = "Editor Scripting | Sequencer Tools | Control Rig")
-	static int32 GetControlRigPriorityOrder(UMovieSceneTrack* InSection);
+	static UE_API int32 GetControlRigPriorityOrder(UMovieSceneTrack* InSection);
 
 	/**	Whether or not the control rig is an FK Control Rig.
 	@param InControlRig Rig to test to see if FK Control Rig
 	**/
 	UFUNCTION(BlueprintCallable, Category = "Editor Scripting | Sequencer Tools | Control Rig")
-	static bool IsFKControlRig(UControlRig* InControlRig);
+	static UE_API bool IsFKControlRig(UControlRig* InControlRig);
 
 	/**	Whether or not the control rig is an Layered Control Rig.
 	@param InControlRig Rig to test to see if Layered Control Rig
 	**/
 	UFUNCTION(BlueprintCallable, Category = "Editor Scripting | Sequencer Tools | Control Rig")
-	static bool IsLayeredControlRig(UControlRig* InControlRig);
+	static UE_API bool IsLayeredControlRig(UControlRig* InControlRig);
 
 	/**
 	 * Mark the layered mode of a control rig on the color and display name of a track
 	 * @param InTrack The track to modify the color and display name if necessary
 	 */
-	static bool MarkLayeredModeOnTrackDisplay(UMovieSceneControlRigParameterTrack* InTrack);
+	static UE_API bool MarkLayeredModeOnTrackDisplay(UMovieSceneControlRigParameterTrack* InTrack);
 
 	/*
 	 * Convert the control rig track into absolute or layered rig
@@ -1143,14 +1145,14 @@ public:
 	 * @param bSetIsLayered Convert to layered rig if true, or absolute if false
 	 */
 	UFUNCTION(BlueprintCallable, Category = "Editor Scripting | Sequencer Tools | Control Rig")
-	static bool SetControlRigLayeredMode(UMovieSceneControlRigParameterTrack* InTrack, bool bSetIsLayered);
+	static UE_API bool SetControlRigLayeredMode(UMovieSceneControlRigParameterTrack* InTrack, bool bSetIsLayered);
 
 	/**	Get FKControlRig Apply Mode.
 	@param InControlRig Rig to test
 	@return The EControlRigFKRigExecuteMode mode it is in, either Replace,Additive or Direct
 	**/
 	UFUNCTION(BlueprintCallable, Category = "Editor Scripting | Sequencer Tools | Control Rig")
-	static EControlRigFKRigExecuteMode  GetFKControlRigApplyMode(UControlRig* InControlRig);
+	static UE_API EControlRigFKRigExecuteMode  GetFKControlRigApplyMode(UControlRig* InControlRig);
 
 	/**	Set the FK Control Rig to apply mode
 	@param InControlRig Rig to set 
@@ -1158,7 +1160,7 @@ public:
 	@return returns True if the mode was set, may not be set if the Control Rig doesn't support these modes currently only FKControlRig's do.
 	**/
 	UFUNCTION(BlueprintCallable, Category = "Editor Scripting | Sequencer Tools | Control Rig")
-	static bool SetControlRigApplyMode(UControlRig* InControlRig, EControlRigFKRigExecuteMode InApplyMode);
+	static UE_API bool SetControlRigApplyMode(UControlRig* InControlRig, EControlRigFKRigExecuteMode InApplyMode);
 
 	/**
 	* Delete anim layer at specified index
@@ -1166,7 +1168,7 @@ public:
 	* @return Returns true if successful, false otherwise
 	*/
 	UFUNCTION(BlueprintCallable, Category = "Editor Scripting | Sequencer Tools | Control Rig | Animation Layers")
-	static bool DeleteAnimLayer(int32 Index);
+	static UE_API bool DeleteAnimLayer(int32 Index);
 
 	/**
 	* Duplicate anim layer at specified index
@@ -1174,19 +1176,19 @@ public:
 	* @return Returns index of new layer, -1 if none created
 	*/
 	UFUNCTION(BlueprintCallable, Category = "Editor Scripting | Sequencer Tools | Control Rig | Animation Layers")
-	static int32 DuplicateAnimLayer(int32 Index);
+	static UE_API int32 DuplicateAnimLayer(int32 Index);
 
 	/**
 	* Add anim layer from objects selected in Sequencer
 	* @return Returns Index of created anim layer
 	*/
 	UFUNCTION(BlueprintCallable, Category = "Editor Scripting | Sequencer Tools | Control Rig | Animation Layers")
-	static int32 AddAnimLayerFromSelection();
+	static UE_API int32 AddAnimLayerFromSelection();
 
 
 	UE_DEPRECATED(5.6, "Use MergeAnimLayersWithSettings that takes FMergeAnimLayerSettings")
 	UFUNCTION(BlueprintCallable, Category = "Editor Scripting | Sequencer Tools | Control Rig | AnimationLayers", meta = (DeprecatedFunction, DeprecationMessage = "Use MergeAnimLayersWithSettings that takes FMergeAnimLayerSettings"))
-	static bool MergeAnimLayers(const TArray<int32>& Indices);
+	static UE_API bool MergeAnimLayers(const TArray<int32>& Indices);
 
 	/**
 	* Merge specified anim layers into one layer. Will merge onto the anim layer with the lowest index
@@ -1196,14 +1198,14 @@ public:
 	*/
 
 	UFUNCTION(BlueprintCallable, Category = "Editor Scripting | Sequencer Tools | Control Rig | Animation Layers")
-	static bool MergeAnimLayersWithSettings(const TArray<int32>& Indices, const FMergeAnimLayerSettings& MergeAnimLayerSettings);
+	static UE_API bool MergeAnimLayersWithSettings(const TArray<int32>& Indices, const FMergeAnimLayerSettings& MergeAnimLayerSettings);
 
 	/**
 	* Get the animation layer objects
 	* @return Returns array of anim layer objects if they exist on active Sequencer
 	*/
 	UFUNCTION(BlueprintCallable, Category = "Editor Scripting | Sequencer Tools | Control Rig | Animation Layers")
-	static TArray<UAnimLayer*> GetAnimLayers();
+	static UE_API TArray<UAnimLayer*> GetAnimLayers();
 	
 	/**
 	* Helper function to get the index in the anim layer array from the anim layer
@@ -1211,6 +1213,8 @@ public:
 	* @return Returns index for the anim layer or INDEX_NONE(-1) if it doesn't exist
 	*/
 	UFUNCTION(BlueprintCallable, Category = "Editor Scripting | Sequencer Tools | Control Rig | Animation Layers")
-	static int32 GetAnimLayerIndex(UAnimLayer* AnimLayer);
+	static UE_API int32 GetAnimLayerIndex(UAnimLayer* AnimLayer);
 
 };
+
+#undef UE_API
