Updated headers using UnrealCodeFixup to make sure dllstorage is on methods/staticvars instead of types. Part 1/n

diff --git a/Engine/Plugins/Animation/ControlRig/Source/ControlRigEditor/Public/EditMode/ControlRigEditMode.h b/Engine/Plugins/Animation/ControlRig/Source/ControlRigEditor/Public/EditMode/ControlRigEditMode.h
--- a/Engine/Plugins/Animation/ControlRig/Source/ControlRigEditor/Public/EditMode/ControlRigEditMode.h
+++ b/Engine/Plugins/Animation/ControlRig/Source/ControlRigEditor/Public/EditMode/ControlRigEditMode.h
@@ -25,6 +25,8 @@
 
 #include "ControlRigEditMode.generated.h"
 
+#define UE_API CONTROLRIGEDITOR_API
+
 class UAnimDetailsProxyManager;
 class UTickableConstraint;
 class FEditorViewportClient;
@@ -132,7 +134,7 @@ private:
 };
 
 USTRUCT(BlueprintType)
-struct CONTROLRIGEDITOR_API FMultiControlRigElementSelection
+struct FMultiControlRigElementSelection
 {
 	GENERATED_BODY()
 
@@ -147,61 +149,61 @@ DECLARE_DYNAMIC_MULTICAST_DELEGATE_TwoParams(FControlRigEditModeInteractionStart
 DECLARE_DYNAMIC_MULTICAST_DELEGATE_OneParam(FControlRigEditModeInteractionEndedEvent, FMultiControlRigElementSelection, InteractionKeys);
 DECLARE_DYNAMIC_MULTICAST_DELEGATE_TwoParams(FControlRigEditModeInteractionUpdatedEvent, FMultiControlRigElementSelection, InteractionKeys, FControlRigInteractionTransformContext, TransformContext);
 
-class CONTROLRIGEDITOR_API FControlRigEditMode : public IPersonaEditMode
+class FControlRigEditMode : public IPersonaEditMode
 {
 public:
 	static inline const FLazyName ModeName = FLazyName(TEXT("EditMode.ControlRig"));
 
-	FControlRigEditMode();
-	~FControlRigEditMode();
+	UE_API FControlRigEditMode();
+	UE_API ~FControlRigEditMode();
 
 	/** Set the Control Rig Object to be active in the edit mode. You set both the Control Rig and a possible binding together with an optional Sequencer
 	 This will remove all other control rigs present and should be called for stand alone editors, like the Control Rig Editor*/
-	void SetObjects(UControlRig* InControlRig, UObject* BindingObject, const TWeakPtr<ISequencer>& InSequencer);
+	UE_API void SetObjects(UControlRig* InControlRig, UObject* BindingObject, const TWeakPtr<ISequencer>& InSequencer);
 
 	/** Add a Control Rig object if it doesn't exist, will return true if it was added, false if it wasn't since it's already there. You can also set the Sequencer.*/
-	bool AddControlRigObject(UControlRig* InControlRig, const TWeakPtr<ISequencer>& InSequencer);
+	UE_API bool AddControlRigObject(UControlRig* InControlRig, const TWeakPtr<ISequencer>& InSequencer);
 
 	/* Remove control rig */
-	void RemoveControlRig(UControlRig* InControlRig);
+	UE_API void RemoveControlRig(UControlRig* InControlRig);
 
 	/*Replace old Control Rig with the New Control Rig, perhaps from a recompile in the level editor*/
-	void ReplaceControlRig(UControlRig* OldControlRig, UControlRig* NewControlRig);
+	UE_API void ReplaceControlRig(UControlRig* OldControlRig, UControlRig* NewControlRig);
 
 	/** This edit mode is re-used between the level editor and the asset editors (control rig editor etc.). Calling this indicates which context we are in */
-	virtual bool IsInLevelEditor() const;
+	UE_API virtual bool IsInLevelEditor() const;
 	/** This is used to differentiate between the control rig editor and any other (asset/level) editors in which this edit mode is used */
 	virtual bool AreEditingControlRigDirectly() const { return false; }
 
 	// FEdMode interface
-	virtual bool UsesToolkits() const override;
-	virtual void Enter() override;
-	virtual void Exit() override;
-	virtual void Tick(FEditorViewportClient* ViewportClient, float DeltaTime) override;
-	virtual void Render(const FSceneView* View, FViewport* Viewport, FPrimitiveDrawInterface* PDI) override;
-	virtual void DrawHUD(FEditorViewportClient* ViewportClient, FViewport* Viewport, const FSceneView* View, FCanvas* Canvas) override;
-	virtual bool InputKey(FEditorViewportClient* InViewportClient, FViewport* InViewport, FKey InKey, EInputEvent InEvent) override;
-	virtual bool EndTracking(FEditorViewportClient* InViewportClient, FViewport* InViewport) override;
-	virtual bool StartTracking(FEditorViewportClient* InViewportClient, FViewport* InViewport) override;
-	virtual bool BeginTransform(const FGizmoState& InState) override;
-	virtual bool EndTransform(const FGizmoState& InState) override;
-	virtual bool HandleClick(FEditorViewportClient* InViewportClient, HHitProxy *HitProxy, const FViewportClick &Click) override;
-	virtual bool BoxSelect(FBox& InBox, bool InSelect = true) override;
-	virtual bool FrustumSelect(const FConvexVolume& InFrustum, FEditorViewportClient* InViewportClient, bool InSelect = true) override;
-	virtual void SelectNone() override;
-	virtual bool InputDelta(FEditorViewportClient* InViewportClient, FViewport* InViewport, FVector& InDrag, FRotator& InRot, FVector& InScale) override;
-	virtual bool UsesTransformWidget() const override;
-	virtual bool GetPivotForOrbit(FVector& OutPivot) const override;
-	virtual bool UsesTransformWidget(UE::Widget::EWidgetMode CheckMode) const override;
-	virtual FVector GetWidgetLocation() const override;
-	virtual bool GetCustomDrawingCoordinateSystem(FMatrix& OutMatrix, void* InData) override;
-	virtual bool GetCustomInputCoordinateSystem(FMatrix& OutMatrix, void* InData) override;
-	virtual bool ShouldDrawWidget() const override;
-	virtual bool IsCompatibleWith(FEditorModeID OtherModeID) const override;
-	virtual bool MouseMove(FEditorViewportClient* ViewportClient, FViewport* Viewport, int32 x, int32 y) override;
-	virtual bool MouseEnter( FEditorViewportClient* ViewportClient,FViewport* Viewport,int32 x, int32 y ) override;
-	virtual bool MouseLeave(FEditorViewportClient* ViewportClient, FViewport* Viewport) override;
-	virtual void PostUndo() override;
+	UE_API virtual bool UsesToolkits() const override;
+	UE_API virtual void Enter() override;
+	UE_API virtual void Exit() override;
+	UE_API virtual void Tick(FEditorViewportClient* ViewportClient, float DeltaTime) override;
+	UE_API virtual void Render(const FSceneView* View, FViewport* Viewport, FPrimitiveDrawInterface* PDI) override;
+	UE_API virtual void DrawHUD(FEditorViewportClient* ViewportClient, FViewport* Viewport, const FSceneView* View, FCanvas* Canvas) override;
+	UE_API virtual bool InputKey(FEditorViewportClient* InViewportClient, FViewport* InViewport, FKey InKey, EInputEvent InEvent) override;
+	UE_API virtual bool EndTracking(FEditorViewportClient* InViewportClient, FViewport* InViewport) override;
+	UE_API virtual bool StartTracking(FEditorViewportClient* InViewportClient, FViewport* InViewport) override;
+	UE_API virtual bool BeginTransform(const FGizmoState& InState) override;
+	UE_API virtual bool EndTransform(const FGizmoState& InState) override;
+	UE_API virtual bool HandleClick(FEditorViewportClient* InViewportClient, HHitProxy *HitProxy, const FViewportClick &Click) override;
+	UE_API virtual bool BoxSelect(FBox& InBox, bool InSelect = true) override;
+	UE_API virtual bool FrustumSelect(const FConvexVolume& InFrustum, FEditorViewportClient* InViewportClient, bool InSelect = true) override;
+	UE_API virtual void SelectNone() override;
+	UE_API virtual bool InputDelta(FEditorViewportClient* InViewportClient, FViewport* InViewport, FVector& InDrag, FRotator& InRot, FVector& InScale) override;
+	UE_API virtual bool UsesTransformWidget() const override;
+	UE_API virtual bool GetPivotForOrbit(FVector& OutPivot) const override;
+	UE_API virtual bool UsesTransformWidget(UE::Widget::EWidgetMode CheckMode) const override;
+	UE_API virtual FVector GetWidgetLocation() const override;
+	UE_API virtual bool GetCustomDrawingCoordinateSystem(FMatrix& OutMatrix, void* InData) override;
+	UE_API virtual bool GetCustomInputCoordinateSystem(FMatrix& OutMatrix, void* InData) override;
+	UE_API virtual bool ShouldDrawWidget() const override;
+	UE_API virtual bool IsCompatibleWith(FEditorModeID OtherModeID) const override;
+	UE_API virtual bool MouseMove(FEditorViewportClient* ViewportClient, FViewport* Viewport, int32 x, int32 y) override;
+	UE_API virtual bool MouseEnter( FEditorViewportClient* ViewportClient,FViewport* Viewport,int32 x, int32 y ) override;
+	UE_API virtual bool MouseLeave(FEditorViewportClient* ViewportClient, FViewport* Viewport) override;
+	UE_API virtual void PostUndo() override;
 
 	/* IPersonaEditMode interface */
 	virtual bool GetCameraTarget(FSphere& OutTarget) const override { return false; }
@@ -209,13 +211,13 @@ public:
 	virtual void GetOnScreenDebugInfo(TArray<FText>& OutDebugInfo) const override {}
 
 	/** FGCObject interface */
-	virtual void AddReferencedObjects( FReferenceCollector& Collector ) override;
+	UE_API virtual void AddReferencedObjects( FReferenceCollector& Collector ) override;
 
 	/** Refresh our internal object list (they may have changed) */
-	void RefreshObjects();
+	UE_API void RefreshObjects();
 
 	/** Find the edit mode corresponding to the specified world context */
-	static FControlRigEditMode* GetEditModeFromWorldContext(UWorld* InWorldContext);
+	static UE_API FControlRigEditMode* GetEditModeFromWorldContext(UWorld* InWorldContext);
 
 	/** Bone Manipulation Delegates */
 	FOnGetRigElementTransform& OnGetRigElementTransform() { return OnGetRigElementTransformDelegate; }
@@ -250,21 +252,21 @@ public:
 	FControlRigEditModeInteractionUpdatedEvent& OnGizmoInteractionPostUpdated() { return OnGizmoInteractionPostUpdatedDelegate; }
 
 	// callback that gets called when rig element is selected in other view
-	void OnHierarchyModified(ERigHierarchyNotification InNotif, URigHierarchy* InHierarchy, const FRigNotificationSubject& InSubject);
-	void OnHierarchyModified_AnyThread(ERigHierarchyNotification InNotif, URigHierarchy* InHierarchy, const FRigNotificationSubject& InSubject);
-	void OnControlModified(UControlRig* Subject, FRigControlElement* InControlElement, const FRigControlModifiedContext& Context);
-	void OnPreConstruction_AnyThread(UControlRig* InRig, const FName& InEventName);
-	void OnPostConstruction_AnyThread(UControlRig* InRig, const FName& InEventName);
+	UE_API void OnHierarchyModified(ERigHierarchyNotification InNotif, URigHierarchy* InHierarchy, const FRigNotificationSubject& InSubject);
+	UE_API void OnHierarchyModified_AnyThread(ERigHierarchyNotification InNotif, URigHierarchy* InHierarchy, const FRigNotificationSubject& InSubject);
+	UE_API void OnControlModified(UControlRig* Subject, FRigControlElement* InControlElement, const FRigControlModifiedContext& Context);
+	UE_API void OnPreConstruction_AnyThread(UControlRig* InRig, const FName& InEventName);
+	UE_API void OnPostConstruction_AnyThread(UControlRig* InRig, const FName& InEventName);
 
 	/** return true if it can be removed from preview scene 
 	- this is to ensure preview scene doesn't remove shape actors */
-	bool CanRemoveFromPreviewScene(const USceneComponent* InComponent);
+	UE_API bool CanRemoveFromPreviewScene(const USceneComponent* InComponent);
 
 	TSharedPtr<FUICommandList> GetCommandBindings() const { return CommandBindings; }
 
 	/** Requests to recreate the shape actors in the next tick. Will recreate only the ones for the specified
 	Control Rig, otherwise will recreate all of them*/
-	void RequestToRecreateControlShapeActors(UControlRig* ControlRig = nullptr); 
+	UE_API void RequestToRecreateControlShapeActors(UControlRig* ControlRig = nullptr); 
 
 	static uint32 ValidControlTypeMask()
 	{
@@ -275,150 +277,150 @@ public:
 
 protected:
 
-	EControlRigInteractionTransformSpace GetTransformSpace() const;
+	UE_API EControlRigInteractionTransformSpace GetTransformSpace() const;
 	
 	// shape related functions wrt enable/selection
 	/** Get the node name from the property path */
-	AControlRigShapeActor* GetControlShapeFromControlName(UControlRig* InControlRig,const FName& InControlName) const;
+	UE_API AControlRigShapeActor* GetControlShapeFromControlName(UControlRig* InControlRig,const FName& InControlName) const;
 
 	/** Helper function: set ControlRigs array to the details panel */
-	void SetObjects_Internal();
+	UE_API void SetObjects_Internal();
 
 	/** Set up Details Panel based upon Selected Objects*/
-	void SetUpDetailPanel() const;
+	UE_API void SetUpDetailPanel() const;
 
 	/** Updates cached pivot transforms */
-	void UpdatePivotTransforms();
-	bool ComputePivotFromEditedShape(UControlRig* InControlRig, FTransform& OutTransform) const;
-	bool ComputePivotFromShapeActors(UControlRig* InControlRig, const bool bEachLocalSpace, const EControlRigInteractionTransformSpace InSpace, FTransform& OutTransform) const;
-	bool ComputePivotFromElements(UControlRig* InControlRig, FTransform& OutTransform) const;
-	FTransform GetPivotOrientation(const FRigElementKey& InControlKey, const UControlRig* InControlRig, URigHierarchy* InHierarchy, const EControlRigInteractionTransformSpace InSpace, const FTransform& InComponentTransform) const;
+	UE_API void UpdatePivotTransforms();
+	UE_API bool ComputePivotFromEditedShape(UControlRig* InControlRig, FTransform& OutTransform) const;
+	UE_API bool ComputePivotFromShapeActors(UControlRig* InControlRig, const bool bEachLocalSpace, const EControlRigInteractionTransformSpace InSpace, FTransform& OutTransform) const;
+	UE_API bool ComputePivotFromElements(UControlRig* InControlRig, FTransform& OutTransform) const;
+	UE_API FTransform GetPivotOrientation(const FRigElementKey& InControlKey, const UControlRig* InControlRig, URigHierarchy* InHierarchy, const EControlRigInteractionTransformSpace InSpace, const FTransform& InComponentTransform) const;
 	
 	/** Get the current coordinate system space */
-	ECoordSystem GetCoordSystemSpace() const;
+	UE_API ECoordSystem GetCoordSystemSpace() const;
 	
 	/** Handle selection internally */
-	void HandleSelectionChanged();
+	UE_API void HandleSelectionChanged();
 
 	/** Toggles visibility of acive control rig shapes in the viewport */
-	void ToggleManipulators();
+	UE_API void ToggleManipulators();
 
 	/** Toggles visibility of all  control rig shapes in the viewport */
-	void ToggleAllManipulators();
+	UE_API void ToggleAllManipulators();
 
 	/** Returns true if all control rig shapes are visible in the viewport */
-	bool AreControlsVisible() const;
+	UE_API bool AreControlsVisible() const;
 
-	virtual bool HandleBeginTransform(const FEditorViewportClient* InViewportClient);
-	virtual bool HandleEndTransform(FEditorViewportClient* InViewportClient);
+	UE_API virtual bool HandleBeginTransform(const FEditorViewportClient* InViewportClient);
+	UE_API virtual bool HandleEndTransform(FEditorViewportClient* InViewportClient);
 
 public:
 
 	/** Toggle controls as overlay*/
-	void ToggleControlsAsOverlay();
+	UE_API void ToggleControlsAsOverlay();
 
 	/** Toggles visibility of acive control rig shapes inside the selected module in the viewport */
-	void ToggleModuleManipulators();
+	UE_API void ToggleModuleManipulators();
 	
 	/** Clear Selection*/
-	void ClearSelection();
+	UE_API void ClearSelection();
 
 	/** Frame to current Control Selection*/
-	void FrameSelection();
+	UE_API void FrameSelection();
 
 	/** Frame a list of provided items*/
-   	void FrameItems(const TArray<FRigElementKey>& InItems);
+   	UE_API void FrameItems(const TArray<FRigElementKey>& InItems);
 
 	/** Sets Passthrough Key on selected anim layers */
-	void SetAnimLayerPassthroughKey();
+	UE_API void SetAnimLayerPassthroughKey();
 
 	/** Select Mirrored Controls on Current Selection*/
-	void SelectMirroredControls();
+	UE_API void SelectMirroredControls();
 
 	/** Select Mirrored Controls on Current Selection, keeping current selection*/
-	void AddMirroredControlsToSelection();
+	UE_API void AddMirroredControlsToSelection();
 
 	/** Put Selected Controls To Mirrored Pose*/
-	void MirrorSelectedControls();
+	UE_API void MirrorSelectedControls();
 
 	/** Put Unselected Controls To Mirrored Pose*/
-	void MirrorUnselectedControls();
+	UE_API void MirrorUnselectedControls();
 
 	/** Select All Controls*/
-	void SelectAllControls();
+	UE_API void SelectAllControls();
 
 	//for the following pose functions we only support one pose(PoseNum = 0)
 	//but may support more later
 	/** Save a pose of selected controls*/
-	void SavePose(int32 PoseNum = 0);
+	UE_API void SavePose(int32 PoseNum = 0);
 
 	/** Select controls in saved pose*/
-	void SelectPose(bool bDoMirror, int32 PoseNum = 0 );
+	UE_API void SelectPose(bool bDoMirror, int32 PoseNum = 0 );
 
 	/** Paste saved pose */
-	void PastePose(bool bDoMirror, int32 PoseNum = 0);
+	UE_API void PastePose(bool bDoMirror, int32 PoseNum = 0);
 
 	/** Opens up the space picker widget */
-	void OpenSpacePickerWidget();
+	UE_API void OpenSpacePickerWidget();
 
 	/** Reset Transforms */
-	void ZeroTransforms(bool bSelectionOnly, bool bIncludeChannels = true);
+	UE_API void ZeroTransforms(bool bSelectionOnly, bool bIncludeChannels = true);
 
 	/** Invert Input Pose */
-	void InvertInputPose(bool bSelectionOnly, bool bIncludeChannels = true);
+	UE_API void InvertInputPose(bool bSelectionOnly, bool bIncludeChannels = true);
 
 	/** Reset Transforms for this Control Rig Controls based upon selection and channel states*/
-	static void ZeroTransforms(UControlRig* ControlRig, const FRigControlModifiedContext& Context,  bool bSelectionOnly, bool bIncludeChannels);
+	static UE_API void ZeroTransforms(UControlRig* ControlRig, const FRigControlModifiedContext& Context,  bool bSelectionOnly, bool bIncludeChannels);
 
 	/** Invert Input Pose for this Control Rig Controls based upon selection and channel states*/
-	static void InvertInputPose(UControlRig* ControlRig, const FRigControlModifiedContext& Context, bool bSelectionOnly, bool bIncludeChannels);
+	static UE_API void InvertInputPose(UControlRig* ControlRig, const FRigControlModifiedContext& Context, bool bSelectionOnly, bool bIncludeChannels);
 
 private:
 	
 	/** Whether or not we should Frame Selection or not*/
-	bool CanFrameSelection();
+	UE_API bool CanFrameSelection();
 
 	/** Increase Shape Size */
-	void IncreaseShapeSize();
+	UE_API void IncreaseShapeSize();
 
 	/** Decrease Shape Size */
-	void DecreaseShapeSize();
+	UE_API void DecreaseShapeSize();
 
 	/** Reset Shape Size */
-	void ResetControlShapeSize();
+	UE_API void ResetControlShapeSize();
 
 	/** Pending focus handler */
 	FPendingWidgetFocus PendingFocus;
 
 	/** Pending focus cvar binding functions to enable/disable pending focus mode */
-	void RegisterPendingFocusMode();
-	void UnregisterPendingFocusMode();
+	UE_API void RegisterPendingFocusMode();
+	UE_API void UnregisterPendingFocusMode();
 	FDelegateHandle PendingFocusHandle;
 
 	/** Listen to sequencer and optimize notifications and performances when playing. */
-	void SetSequencerDelegates(const TWeakPtr<ISequencer>& InWeakSequencer);
-	void UnsetSequencerDelegates() const;
-	void UpdateSequencerStatus();
+	UE_API void SetSequencerDelegates(const TWeakPtr<ISequencer>& InWeakSequencer);
+	UE_API void UnsetSequencerDelegates() const;
+	UE_API void UpdateSequencerStatus();
 	bool bSequencerPlaying = false;
 	
 public:
 	
 	/** Toggle Shape Transform Edit*/
-	void ToggleControlShapeTransformEdit();
+	UE_API void ToggleControlShapeTransformEdit();
 
 private:
 	
 	/** The hotkey text is passed to a viewport notification to inform users how to toggle shape edit*/
-	FText GetToggleControlShapeTransformEditHotKey() const;
+	UE_API FText GetToggleControlShapeTransformEditHotKey() const;
 
 	/** Bind our keyboard commands */
-	void BindCommands();
+	UE_API void BindCommands();
 
 	/** It creates if it doesn't have it */
-	void RecreateControlShapeActors();
+	UE_API void RecreateControlShapeActors();
 
 	/** Let the preview scene know how we want to select components */
-	bool ShapeSelectionOverride(const UPrimitiveComponent* InComponent) const;
+	UE_API bool ShapeSelectionOverride(const UPrimitiveComponent* InComponent) const;
 
 	/** Enable editing of control's shape transform instead of control's transform*/
 	bool bIsChangingControlShapeTransform;
@@ -454,13 +456,13 @@ protected:
 	TSharedPtr<FUICommandList> CommandBindings;
 
 	/** Called from the editor when a blueprint object replacement has occurred */
-	void OnObjectsReplaced(const TMap<UObject*, UObject*>& OldToNewInstanceMap);
+	UE_API void OnObjectsReplaced(const TMap<UObject*, UObject*>& OldToNewInstanceMap);
 
 	/** Copy control visibility from the source rig (and its modules if necessary) to the target rig */ 
-	void CopyControlsVisibility(const UControlRig* SourceRig, UControlRig* TargetRig);
+	UE_API void CopyControlsVisibility(const UControlRig* SourceRig, UControlRig* TargetRig);
 
 	/** Return true if transform setter/getter delegates are available */
-	bool IsTransformDelegateAvailable() const;
+	UE_API bool IsTransformDelegateAvailable() const;
 
 	FOnGetRigElementTransform OnGetRigElementTransformDelegate;
 	FOnSetRigElementTransform OnSetRigElementTransformDelegate;
@@ -487,11 +489,11 @@ protected:
 	FSimpleMulticastDelegate OnControlRigShapeActorsRecreatedDelegate;
 
 	/** GetSelectedRigElements, if InControlRig is nullptr get the first one */
-	TArray<FRigElementKey> GetSelectedRigElements() const;
-	static TArray<FRigElementKey> GetSelectedRigElements(UControlRig* InControlRig);
+	UE_API TArray<FRigElementKey> GetSelectedRigElements() const;
+	static UE_API TArray<FRigElementKey> GetSelectedRigElements(UControlRig* InControlRig);
 
 	/** Get the rig elements, based upon the selection or if it's a channel, also will do internal additive filters(like no bools)  */
-	static TArray<FRigElementKey> GetRigElementsForSettingTransforms(UControlRig* InControlRig, bool bSelectionOnly, bool bIncludeChannels);
+	static UE_API TArray<FRigElementKey> GetRigElementsForSettingTransforms(UControlRig* InControlRig, bool bSelectionOnly, bool bIncludeChannels);
 
 	/* Flag to recreate shapes during tick */
 	ERecreateControlRigShape RecreateControlShapesRequired;
@@ -508,11 +510,11 @@ protected:
 	TObjectPtr<UAnimDetailsProxyManager> AnimDetailsProxyManager;
 
 	/** Utility functions for UI/Some other viewport manipulation*/
-	bool IsControlSelected(const bool bUseShapes = false) const;
-	bool AreRigElementSelectedAndMovable(UControlRig* InControlRig) const;
+	UE_API bool IsControlSelected(const bool bUseShapes = false) const;
+	UE_API bool AreRigElementSelectedAndMovable(UControlRig* InControlRig) const;
 	
 	/** Set initial transform handlers */
-	void OpenContextMenu(FEditorViewportClient* InViewportClient);
+	UE_API void OpenContextMenu(FEditorViewportClient* InViewportClient);
 
 	/** previous Gizmo(Widget) scale before we enter this mode, used to set it back*/
 	float PreviousGizmoScale = 1.0f;
@@ -521,30 +523,30 @@ protected:
 	TMap<UControlRig*, UE::ControlRigEditMode::FInteractionDependencyCache> InteractionDependencies;
 	
 	/** Returns the interaction dependencies of that ControlRig. */
-	UE::ControlRigEditMode::FInteractionDependencyCache& GetInteractionDependencies(UControlRig* InControlRig);
+	UE_API UE::ControlRigEditMode::FInteractionDependencyCache& GetInteractionDependencies(UControlRig* InControlRig);
 	
 public: 
 	/** Clear all selected RigElements */
-	void ClearRigElementSelection(uint32 InTypes);
+	UE_API void ClearRigElementSelection(uint32 InTypes);
 
 	/** Set a RigElement's selection state */
-	void SetRigElementSelection(UControlRig* ControlRig, ERigElementType Type, const FName& InRigElementName, bool bSelected);
+	UE_API void SetRigElementSelection(UControlRig* ControlRig, ERigElementType Type, const FName& InRigElementName, bool bSelected);
 
 	/** Set multiple RigElement's selection states */
-	void SetRigElementSelection(UControlRig* ControlRig, ERigElementType Type, const TArray<FName>& InRigElementNames, bool bSelected);
+	UE_API void SetRigElementSelection(UControlRig* ControlRig, ERigElementType Type, const TArray<FName>& InRigElementNames, bool bSelected);
 
 	/** Check if any RigElements are selected */
-	bool AreRigElementsSelected(uint32 InTypes, UControlRig* InControlRig) const;
+	UE_API bool AreRigElementsSelected(uint32 InTypes, UControlRig* InControlRig) const;
 
 	/** Get all of the selected Controls*/
-	void GetAllSelectedControls(TMap<UControlRig*, TArray<FRigElementKey>>& OutSelectedControls) const;
+	UE_API void GetAllSelectedControls(TMap<UControlRig*, TArray<FRigElementKey>>& OutSelectedControls) const;
 
 	/** Get all of the ControlRigs, maybe not valid anymore */
-	TArrayView<const TWeakObjectPtr<UControlRig>> GetControlRigs() const;
-	TArrayView<TWeakObjectPtr<UControlRig>> GetControlRigs();
+	UE_API TArrayView<const TWeakObjectPtr<UControlRig>> GetControlRigs() const;
+	UE_API TArrayView<TWeakObjectPtr<UControlRig>> GetControlRigs();
 	/* Get valid  Control Rigs possibly just visible*/
-	TArray<UControlRig*> GetControlRigsArray(bool bIsVisible);
-	TArray<const UControlRig*> GetControlRigsArray(bool bIsVisible) const;
+	UE_API TArray<UControlRig*> GetControlRigsArray(bool bIsVisible);
+	UE_API TArray<const UControlRig*> GetControlRigsArray(bool bIsVisible) const;
 
 	/** Get the detail proxies control rig*/
 	UAnimDetailsProxyManager* GetAnimDetailsProxyManager() const { return AnimDetailsProxyManager; }
@@ -556,90 +558,90 @@ public:
 	void SuspendHierarchyNotifs(bool bVal) { bSuspendHierarchyNotifs = bVal; }
 
 	/** Request a certain transform widget for the next update */
-	void RequestTransformWidgetMode(UE::Widget::EWidgetMode InWidgetMode);
+	UE_API void RequestTransformWidgetMode(UE::Widget::EWidgetMode InWidgetMode);
 	
 private:
 	/** Set a RigElement's selection state. */
-	void SetRigElementSelectionInternal(UControlRig* ControlRig, ERigElementType Type, const FName& InRigElementName, bool bSelected);
+	UE_API void SetRigElementSelectionInternal(UControlRig* ControlRig, ERigElementType Type, const FName& InRigElementName, bool bSelected);
 
 	/** Set multiple RigElements' selection states. */
-	void SetRigElementsSelectionInternal(const TMap<TWeakObjectPtr<UControlRig>, TArray<FRigElementKey>>& InRigElementsToSelect, bool bSelected);
+	UE_API void SetRigElementsSelectionInternal(const TMap<TWeakObjectPtr<UControlRig>, TArray<FRigElementKey>>& InRigElementsToSelect, bool bSelected);
 
 	/** Whether or not Pivot Transforms have changed, in which case we need to redraw viewport. */
-	bool HasPivotTransformsChanged() const;
+	UE_API bool HasPivotTransformsChanged() const;
 
 	/** Updates the pivot transforms before ticking to ensure that they are up-to-date when needed. */
-	void UpdatePivotTransformsIfNeeded(UControlRig* InControlRig, FTransform& InOutTransform) const;
+	UE_API void UpdatePivotTransformsIfNeeded(UControlRig* InControlRig, FTransform& InOutTransform) const;
 	
 	FEditorViewportClient* CurrentViewportClient;
 	TArray<UE::Widget::EWidgetMode> RequestedWidgetModes;
 
 /* store coordinate system per widget mode*/
 private:
-	void OnWidgetModeChanged(UE::Widget::EWidgetMode InWidgetMode);
-	void OnCoordSystemChanged(ECoordSystem InCoordSystem);
+	UE_API void OnWidgetModeChanged(UE::Widget::EWidgetMode InWidgetMode);
+	UE_API void OnCoordSystemChanged(ECoordSystem InCoordSystem);
 	TArray<ECoordSystem> CoordSystemPerWidgetMode;
 	bool bIsChangingCoordSystem;
 
-	bool CanChangeControlShapeTransform();
+	UE_API bool CanChangeControlShapeTransform();
 
-	void OnSettingsChanged(const UControlRigEditModeSettings* InSettings);
+	UE_API void OnSettingsChanged(const UControlRigEditModeSettings* InSettings);
 	
 public:
 	//Toolbar functions
-	void SetOnlySelectRigControls(bool val);
-	bool GetOnlySelectRigControls()const;
-	bool SetSequencer(TWeakPtr<ISequencer> InSequencer);
+	UE_API void SetOnlySelectRigControls(bool val);
+	UE_API bool GetOnlySelectRigControls()const;
+	UE_API bool SetSequencer(TWeakPtr<ISequencer> InSequencer);
 
 private:
-	TSet<FName> GetActiveControlsFromSequencer(UControlRig* ControlRig);
+	UE_API TSet<FName> GetActiveControlsFromSequencer(UControlRig* ControlRig);
 
 	/** Create/Delete/Update shape actors for the specified ControlRig */
-	void CreateShapeActors(UControlRig* InControlRig);
-	void DestroyShapesActors(UControlRig* InControlRig);
-	bool TryUpdatingControlsShapes(UControlRig* InControlRig);
+	UE_API void CreateShapeActors(UControlRig* InControlRig);
+	UE_API void DestroyShapesActors(UControlRig* InControlRig);
+	UE_API bool TryUpdatingControlsShapes(UControlRig* InControlRig);
 
 	/*Internal function for adding ControlRig*/
-	void AddControlRigInternal(UControlRig* InControlRig);
+	UE_API void AddControlRigInternal(UControlRig* InControlRig);
 
 	/**
 	 * Updates the bound components (skeletal meshes or control rig component) and the control shapes so that all transform data are updated.
 	 * if InRig is empty, all rigs / control shapes will be updated.
 	 */
-	void TickManipulatableObjects(const TArray<TWeakObjectPtr<UControlRig>>& InRigs = TArray<TWeakObjectPtr<UControlRig>>()) const;
+	UE_API void TickManipulatableObjects(const TArray<TWeakObjectPtr<UControlRig>>& InRigs = TArray<TWeakObjectPtr<UControlRig>>()) const;
 
 	/* Check on tick to see if movie scene has changed, returns true if it has*/
-	bool CheckMovieSceneSig();
-	void SetControlShapeTransform( const AControlRigShapeActor* InShapeActor, const FTransform& InGlobalTransform,
+	UE_API bool CheckMovieSceneSig();
+	UE_API void SetControlShapeTransform( const AControlRigShapeActor* InShapeActor, const FTransform& InGlobalTransform,
 		const FTransform& InToWorldTransform, const FRigControlModifiedContext& InContext, const bool bPrintPython,
 		const FControlRigInteractionTransformContext& InTransformContext, const bool bFixEulerFlips = true) const;
-	static FTransform GetControlShapeTransform(const AControlRigShapeActor* ShapeActor);
+	static UE_API FTransform GetControlShapeTransform(const AControlRigShapeActor* ShapeActor);
 
-	static void ChangeControlShapeTransform(AControlRigShapeActor* ShapeActor, const FControlRigInteractionTransformContext& InContext, const FTransform& ToWorldTransform);
+	static UE_API void ChangeControlShapeTransform(AControlRigShapeActor* ShapeActor, const FControlRigInteractionTransformContext& InContext, const FTransform& ToWorldTransform);
 
-	bool ModeSupportedByShapeActor(const AControlRigShapeActor* ShapeActor, UE::Widget::EWidgetMode InMode) const;
+	UE_API bool ModeSupportedByShapeActor(const AControlRigShapeActor* ShapeActor, UE::Widget::EWidgetMode InMode) const;
 
 public:
 	//notify driven controls, should this be inside CR instead?
-	static void NotifyDrivenControls(UControlRig* InControlRig, const FRigElementKey& InKey, const FRigControlModifiedContext& InContext);
+	static UE_API void NotifyDrivenControls(UControlRig* InControlRig, const FRigElementKey& InKey, const FRigControlModifiedContext& InContext);
 
 protected:
 
-	bool MoveControlShapeLocally( AControlRigShapeActor* ShapeActor,
+	UE_API bool MoveControlShapeLocally( AControlRigShapeActor* ShapeActor,
 		const FControlRigInteractionTransformContext& InTransformContext, const FTransform& ToWorldTransform,
 		const FTransform& InLocal);
 	
-	void MoveControlShape(AControlRigShapeActor* ShapeActor,
+	UE_API void MoveControlShape(AControlRigShapeActor* ShapeActor,
 		const FControlRigInteractionTransformContext& InContext, const FTransform& ToWorldTransform,
 		const bool bUseLocal, const bool bCalcLocal, FTransform* InOutLocal, TArray<TFunction<void()>>& OutTasks);
 
 	/** Get bindings to a runtime object */
 	//If the passed in ControlRig is nullptr we use the first Control Rig(this can happen from the BP Editors).
-	USceneComponent* GetHostingSceneComponent(const UControlRig* ControlRig = nullptr) const;
-	FTransform	GetHostingSceneComponentTransform(const UControlRig* ControlRig =  nullptr) const;
+	UE_API USceneComponent* GetHostingSceneComponent(const UControlRig* ControlRig = nullptr) const;
+	UE_API FTransform	GetHostingSceneComponentTransform(const UControlRig* ControlRig =  nullptr) const;
 
 	//Get if the hosted component is visible
-	bool IsControlRigSkelMeshVisible(const UControlRig* InControlRig) const;
+	UE_API bool IsControlRigSkelMeshVisible(const UControlRig* InControlRig) const;
 	
 public:  
 		TSharedPtr<FDetailKeyFrameCacheAndHandler> DetailKeyFrameCache;
@@ -647,25 +649,25 @@ public:
 private:
 
 	// Post pose update handler
-	void OnPoseInitialized();
+	UE_API void OnPoseInitialized();
 
 	/**
 	 * Updates the control shapes properties (transform, visibility, ...) depending on the current viewport state.
 	 * If InRig is empty, all control shapes will be updated.
 	 */
-	void PostPoseUpdate(const FEditorViewportClient* InViewportClient = nullptr, const TArray<TWeakObjectPtr<UControlRig>>& InRigs = TArray<TWeakObjectPtr<UControlRig>>()) const;
+	UE_API void PostPoseUpdate(const FEditorViewportClient* InViewportClient = nullptr, const TArray<TWeakObjectPtr<UControlRig>>& InRigs = TArray<TWeakObjectPtr<UControlRig>>()) const;
 	
-	void UpdateSelectabilityOnSkeletalMeshes(UControlRig* InControlRig, bool bEnabled);
+	UE_API void UpdateSelectabilityOnSkeletalMeshes(UControlRig* InControlRig, bool bEnabled);
 
-	bool IsMovingCamera(const FViewport* InViewport) const;
-	bool IsDoingDrag(const FViewport* InViewport) const;
+	UE_API bool IsMovingCamera(const FViewport* InViewport) const;
+	UE_API bool IsDoingDrag(const FViewport* InViewport) const;
 
 	// world clean up handlers
 	FDelegateHandle OnWorldCleanupHandle;
-	void OnWorldCleanup(UWorld* World, bool bSessionEnded, bool bCleanupResources);
+	UE_API void OnWorldCleanup(UWorld* World, bool bSessionEnded, bool bCleanupResources);
 	UWorld* WorldPtr = nullptr;
 
-	void OnEditorClosed();
+	UE_API void OnEditorClosed();
 
 	struct FMarqueeDragTool
     {
@@ -706,7 +708,7 @@ private:
 	TArray<FRigElementKey> DeferredItemsToFrame;
 
 	/** Computes the current interaction types based on the widget mode */
-	static uint8 GetInteractionType(const FEditorViewportClient* InViewportClient);
+	static UE_API uint8 GetInteractionType(const FEditorViewportClient* InViewportClient);
 	uint8 InteractionType;
 	bool bShowControlsAsOverlay;
 
@@ -737,28 +739,28 @@ private:
 		FControlRigEditMode* EditMode;
 	};
 
-	static void EvaluateRig(UControlRig* InControlRig);
+	static UE_API void EvaluateRig(UControlRig* InControlRig);
 	
 	//to disable post pose update, needed for offfline evlauations
-	static bool bDoPostPoseUpdate;
+	static UE_API bool bDoPostPoseUpdate;
 
 private:
 
 	// get default/mutable settings
 	// todo: have a local setting object that listen to property changes to send updates
 	// instead of getting data directly from the CDO
-	const UControlRigEditModeSettings* GetSettings() const;
-	UControlRigEditModeSettings* GetMutableSettings() const;
+	UE_API const UControlRigEditModeSettings* GetSettings() const;
+	UE_API UControlRigEditModeSettings* GetMutableSettings() const;
 	mutable TWeakObjectPtr<UControlRigEditModeSettings> WeakSettings;
 
 	//pose used by the hotkeys
 	TObjectPtr<UControlRigPoseAsset> StoredPose;
 	
 	TWeakObjectPtr<UEditorTransformGizmoContextObject> WeakGizmoContext;
-	FRotationContext& GetRotationContext() const;
-	void UpdateRotationContext();
+	UE_API FRotationContext& GetRotationContext() const;
+	UE_API void UpdateRotationContext();
 
-	TOptional<FTransform> GetConstraintParentTransform(const UControlRig* InControlRig, const FName& InControlName) const;
+	UE_API TOptional<FTransform> GetConstraintParentTransform(const UControlRig* InControlRig, const FName& InControlName) const;
 	mutable UE::TransformConstraintUtil::FConstraintsInteractionCache ConstraintsCache;
 
 	// Used to store and apply keyframes (if deferred)  
@@ -792,3 +794,4 @@ public:
 	};
 };
 
+#undef UE_API
diff --git a/Engine/Plugins/Animation/ControlRig/Source/ControlRigEditor/Public/IControlRigEditorModule.h b/Engine/Plugins/Animation/ControlRig/Source/ControlRigEditor/Public/IControlRigEditorModule.h
--- a/Engine/Plugins/Animation/ControlRig/Source/ControlRigEditor/Public/IControlRigEditorModule.h
+++ b/Engine/Plugins/Animation/ControlRig/Source/ControlRigEditor/Public/IControlRigEditorModule.h
@@ -10,6 +10,8 @@
 #include "ClassViewerModule.h"
 #include "ClassViewerFilter.h"
 
+#define UE_API CONTROLRIGEDITOR_API
+
 DECLARE_LOG_CATEGORY_EXTERN(LogControlRigEditor, Log, All);
 
 class IToolkitHost;
@@ -18,15 +20,15 @@ class UControlRigGraphNode;
 class UControlRigGraphSchema;
 class FConnectionDrawingPolicy;
 
-class CONTROLRIGEDITOR_API FControlRigClassFilter : public IClassViewerFilter
+class FControlRigClassFilter : public IClassViewerFilter
 {
 public:
-	FControlRigClassFilter(bool bInCheckSkeleton, bool bInCheckAnimatable, bool bInCheckInversion, USkeleton* InSkeleton);
-	virtual bool IsClassAllowed(const FClassViewerInitializationOptions& InInitOptions, const UClass* InClass, TSharedRef< FClassViewerFilterFuncs > InFilterFuncs) override;
-	virtual bool IsUnloadedClassAllowed(const FClassViewerInitializationOptions& InInitOptions, const TSharedRef< const IUnloadedBlueprintData > InUnloadedClassData, TSharedRef< FClassViewerFilterFuncs > InFilterFuncs) override;	
+	UE_API FControlRigClassFilter(bool bInCheckSkeleton, bool bInCheckAnimatable, bool bInCheckInversion, USkeleton* InSkeleton);
+	UE_API virtual bool IsClassAllowed(const FClassViewerInitializationOptions& InInitOptions, const UClass* InClass, TSharedRef< FClassViewerFilterFuncs > InFilterFuncs) override;
+	UE_API virtual bool IsUnloadedClassAllowed(const FClassViewerInitializationOptions& InInitOptions, const TSharedRef< const IUnloadedBlueprintData > InUnloadedClassData, TSharedRef< FClassViewerFilterFuncs > InFilterFuncs) override;	
 
 private:
-	bool MatchesFilter(const FAssetData& AssetData);
+	UE_API bool MatchesFilter(const FAssetData& AssetData);
 
 private:
 	bool bFilterAssetBySkeleton;
@@ -58,3 +60,5 @@ public:
 	virtual TSharedRef<IControlRigEditor> CreateControlRigEditor(const EToolkitMode::Type Mode, const TSharedPtr<IToolkitHost>& InitToolkitHost, UControlRigBlueprint* Blueprint) = 0;
 
 };
+
+#undef UE_API
diff --git a/Engine/Plugins/Animation/ControlRig/Source/ControlRigEditor/Public/RigSpacePickerBakeSettings.h b/Engine/Plugins/Animation/ControlRig/Source/ControlRigEditor/Public/RigSpacePickerBakeSettings.h
--- a/Engine/Plugins/Animation/ControlRig/Source/ControlRigEditor/Public/RigSpacePickerBakeSettings.h
+++ b/Engine/Plugins/Animation/ControlRig/Source/ControlRigEditor/Public/RigSpacePickerBakeSettings.h
@@ -9,7 +9,7 @@
 
 
 USTRUCT(BlueprintType)
-struct CONTROLRIGEDITOR_API FRigSpacePickerBakeSettings
+struct FRigSpacePickerBakeSettings
 {
 	GENERATED_BODY();
 
@@ -44,4 +44,4 @@ struct CONTROLRIGEDITOR_API FRigSpacePickerBakeSettings
 	/** DEPRECATED 5.3 */
 	UPROPERTY(meta = (DeprecatedProperty, DeprecationMessage = "Use Settings.Tolerance instead"))
 	float Tolerance_DEPRECATED;
-};
\ No newline at end of file
+};
diff --git a/Engine/Plugins/Animation/ControlRig/Source/ControlRigEditor/Public/Sequencer/AnimationAuthoringSettings.h b/Engine/Plugins/Animation/ControlRig/Source/ControlRigEditor/Public/Sequencer/AnimationAuthoringSettings.h
--- a/Engine/Plugins/Animation/ControlRig/Source/ControlRigEditor/Public/Sequencer/AnimationAuthoringSettings.h
+++ b/Engine/Plugins/Animation/ControlRig/Source/ControlRigEditor/Public/Sequencer/AnimationAuthoringSettings.h
@@ -7,8 +7,10 @@
 
 #include "AnimationAuthoringSettings.generated.h"
 
-UCLASS(config=EditorPerProjectUserSettings, meta = (DisplayName = "Animation Authoring Settings"))
-class CONTROLRIGEDITOR_API UAnimationAuthoringSettings : public UDeveloperSettings
+#define UE_API CONTROLRIGEDITOR_API
+
+UCLASS(MinimalAPI, config=EditorPerProjectUserSettings, meta = (DisplayName = "Animation Authoring Settings"))
+class UAnimationAuthoringSettings : public UDeveloperSettings
 {
 	GENERATED_UCLASS_BODY()
 
@@ -19,7 +21,9 @@ public:
 	bool bAutoKeyOnRelease = true;
     
 	DECLARE_MULTICAST_DELEGATE_OneParam(FOnUpdateSettings,  const UAnimationAuthoringSettings*);
-	static FOnUpdateSettings OnSettingsChange;
+	static UE_API FOnUpdateSettings OnSettingsChange;
+
+	UE_API virtual void PostEditChangeProperty(FPropertyChangedEvent& PropertyChangedEvent) override;
+};
 
-	virtual void PostEditChangeProperty(FPropertyChangedEvent& PropertyChangedEvent) override;
-};
\ No newline at end of file
+#undef UE_API
diff --git a/Engine/Plugins/Animation/ControlRig/Source/ControlRigEditor/Public/Tools/ControlRigSnapSettings.h b/Engine/Plugins/Animation/ControlRig/Source/ControlRigEditor/Public/Tools/ControlRigSnapSettings.h
--- a/Engine/Plugins/Animation/ControlRig/Source/ControlRigEditor/Public/Tools/ControlRigSnapSettings.h
+++ b/Engine/Plugins/Animation/ControlRig/Source/ControlRigEditor/Public/Tools/ControlRigSnapSettings.h
@@ -9,15 +9,17 @@
 #include "BakingAnimationKeySettings.h"
 #include "ControlRigSnapSettings.generated.h"
 
+#define UE_API CONTROLRIGEDITOR_API
 
-UCLASS(config = EditorPerProjectUserSettings)
-class CONTROLRIGEDITOR_API UControlRigSnapSettings : public UObject
+
+UCLASS(MinimalAPI, config = EditorPerProjectUserSettings)
+class UControlRigSnapSettings : public UObject
 {
 
 public:
 	GENERATED_BODY()
 
-	UControlRigSnapSettings();
+	UE_API UControlRigSnapSettings();
 
 	/** When snapping keep offset, If false snap completely */
 	UPROPERTY(config, EditAnywhere, BlueprintReadWrite, Category = "Snap Settings", meta = (ToolTip = "When snapping keep offset, if false snap completely"))
@@ -50,3 +52,4 @@ public:
 };
 
 
+#undef UE_API
diff --git a/Engine/Plugins/Animation/ControlRig/Source/ControlRigEditor/Public/Tools/EvaluateSequencerTools.h b/Engine/Plugins/Animation/ControlRig/Source/ControlRigEditor/Public/Tools/EvaluateSequencerTools.h
--- a/Engine/Plugins/Animation/ControlRig/Source/ControlRigEditor/Public/Tools/EvaluateSequencerTools.h
+++ b/Engine/Plugins/Animation/ControlRig/Source/ControlRigEditor/Public/Tools/EvaluateSequencerTools.h
@@ -5,6 +5,8 @@
 #include "ActorForWorldTransforms.h"
 #include "MovieSceneBinding.h"
 
+#define UE_API CONTROLRIGEDITOR_API
+
 class UControlRig;
 class UWorld;
 class ISequencer;
@@ -82,7 +84,7 @@ struct  FFrameTimeByIndex
 };
 
 //Array of transforms, that maybe be sparse with only some transforms set
-struct CONTROLRIGEDITOR_API FArrayOfTransforms
+struct FArrayOfTransforms
 {
 	TArray<FTransform> Transforms;
 	void SetNum(int32 InNum)
@@ -98,7 +100,7 @@ struct CONTROLRIGEDITOR_API FArrayOfTransforms
 		return Transforms[Index];
 	}
 	//interpolate with sparse set of indices
-	FTransform Interp(const FFrameNumber& InTime, const TArray<int32>& ValidIndices, const TArray<FFrameNumber>& ValidTimes);
+	UE_API FTransform Interp(const FFrameNumber& InTime, const TArray<int32>& ValidIndices, const TArray<FFrameNumber>& ValidTimes);
 };
 
 // An 'Actor'(which may be an AActor/SceneComponent/Socket) and a set of sparse world transforms with corrensonding parent transforms
@@ -157,14 +159,14 @@ struct FControlRigAndWorldTransforms
 	TMap<FName, TSharedPtr<FArrayOfTransforms>> ControlAndWorldTransforms; //one for each control in the FControlRigAndWorldTransform
 };
 
-struct CONTROLRIGEDITOR_API FSetTransformHelpers
+struct FSetTransformHelpers
 {
 	//get a copy of a transform and possibly use it to set a constrainted transform on a control, if this happens it returns true
-	static bool SetConstrainedTransform(FTransform LocalTransform, UControlRig* ControlRig, FRigControlElement* ControlElement, const FRigControlModifiedContext& Context);
+	static UE_API bool SetConstrainedTransform(FTransform LocalTransform, UControlRig* ControlRig, FRigControlElement* ControlElement, const FRigControlModifiedContext& Context);
 	//set a control transform, will set a constrainted transform if one is needed, otherwise will just set a normal one.
-	static void SetControlTransform(const FEulerTransform& EulerTransform, UControlRig* ControlRig, FRigControlElement* ControlElement, const FRigControlModifiedContext& Context);
+	static UE_API void SetControlTransform(const FEulerTransform& EulerTransform, UControlRig* ControlRig, FRigControlElement* ControlElement, const FRigControlModifiedContext& Context);
 	//set a set of tranfsorms on an scene component, will handle constraints
-	static bool SetActorTransform(ISequencer* Sequencer, USceneComponent* SceneComponent, UMovieScene3DTransformSection* TransformSection, const TArray<FFrameNumber>& Frames,
+	static UE_API bool SetActorTransform(ISequencer* Sequencer, USceneComponent* SceneComponent, UMovieScene3DTransformSection* TransformSection, const TArray<FFrameNumber>& Frames,
 		const TArray<FTransform>& WorldTransformsToSnapTo, const TArray<FTransform>& ParentWorldTransforms);
 
 };
@@ -173,18 +175,18 @@ struct CONTROLRIGEDITOR_API FSetTransformHelpers
 //This is somewhat similar to what you would calculate with a DAG to figure out dependencies but since we don't have a DAG we need to do this manually by walking parents/tracks.
 //It will also calcualate the corresponding UMovieSceneSignedObject FGuid signatures, which will then be used to determine if we need to recalculate an objects tranforms, like
 //with motion trails since some downstream dependency has changed.
-struct CONTROLRIGEDITOR_API FSequencerTransformDependencies
+struct FSequencerTransformDependencies
 {
 	//Get the list of tracks in the moviescene with this binding that may effect the final transform of something in the binding.
 	//Currently it will return a set of trnasform, control rig, attachment, path and  skeletal anim tracks.
-	static TArray<UMovieSceneTrack*> GetTransformAffectingTracks(UMovieScene* InMovieScene, const FMovieSceneBinding& Binding);
+	static UE_API TArray<UMovieSceneTrack*> GetTransformAffectingTracks(UMovieScene* InMovieScene, const FMovieSceneBinding& Binding);
 
 	//Main function to calculate the dependencies for an actor and a set of tracks under that actor
-	void CalculateDependencies(ISequencer* InSequencer, AActor* InActor, TArray<UMovieSceneTrack*>& InTracks);
+	UE_API void CalculateDependencies(ISequencer* InSequencer, AActor* InActor, TArray<UMovieSceneTrack*>& InTracks);
 	//Compare 2 set of dependencies to see if they have changed.
-	bool Compare(FSequencerTransformDependencies& CopyInIfDifferent) const;
+	UE_API bool Compare(FSequencerTransformDependencies& CopyInIfDifferent) const;
 	//Copy other set of dependencies to this one
-	void CopyFrom(const FSequencerTransformDependencies& Other);
+	UE_API void CopyFrom(const FSequencerTransformDependencies& Other);
 	//If setup or not
 	bool IsEmpty() const { return Tracks.Num() == 0; }
 public:
@@ -196,17 +198,17 @@ public:
 	//Non sequencer actors that this object depends upon
 	TSet<AActor*> NonSequencerActors;
 private:
-	FMovieSceneBinding GetBindingFromTrack(UMovieScene* InMovieScene, UMovieSceneTrack* Track);
-	void CalculateTrackDependents(ISequencer* InSequencer, UMovieSceneTrack* InTrack);
+	UE_API FMovieSceneBinding GetBindingFromTrack(UMovieScene* InMovieScene, UMovieSceneTrack* Track);
+	UE_API void CalculateTrackDependents(ISequencer* InSequencer, UMovieSceneTrack* InTrack);
 	//Add a track to the set of calculated onces and recursively get it's dependencies
-	void AddTrack(ISequencer* InSequencer, UMovieSceneTrack* InTrack);
+	UE_API void AddTrack(ISequencer* InSequencer, UMovieSceneTrack* InTrack);
 
 };
 
-struct CONTROLRIGEDITOR_API FEvalHelpers
+struct FEvalHelpers
 {
 	//The main evaluation function to use a sparse set of times and caclculate their world transforms for a set of sequencer actors/control rigs.
-	static bool CalculateWorldTransforms(UWorld* World, ISequencer* Sequencer, const FFrameTimeByIndex& FrameTimeByIndex,
+	static UE_API bool CalculateWorldTransforms(UWorld* World, ISequencer* Sequencer, const FFrameTimeByIndex& FrameTimeByIndex,
 		const TArray<int32>& Indices, TArray<FActorAndWorldTransforms>& Actors, TMap<UControlRig*, FControlRigAndWorldTransforms>& ControlRigs,
 		const TSet<const UMovieSceneTrack*>& DependentTracks);
 };
@@ -217,3 +219,4 @@ struct CONTROLRIGEDITOR_API FEvalHelpers
 } // namespace AIE
 } // namespace UE
 
+#undef UE_API
diff --git a/Engine/Plugins/Animation/DNACalib/Source/DNACalibModule/Public/DNACalibDNAReader.h b/Engine/Plugins/Animation/DNACalib/Source/DNACalibModule/Public/DNACalibDNAReader.h
--- a/Engine/Plugins/Animation/DNACalib/Source/DNACalibModule/Public/DNACalibDNAReader.h
+++ b/Engine/Plugins/Animation/DNACalib/Source/DNACalibModule/Public/DNACalibDNAReader.h
@@ -9,180 +9,184 @@
 #include "UObject/NoExportTypes.h"
 #include "UObject/ObjectMacros.h"
 
+#define UE_API DNACALIBMODULE_API
+
 /**
 	@brief A UE wrapper for a special purpose DNA Reader type which serves as the output parameter of DNACalib.
 */
-class DNACALIBMODULE_API FDNACalibDNAReader : public IDNAReader
+class FDNACalibDNAReader : public IDNAReader
 {
 public:
-	explicit FDNACalibDNAReader(IDNAReader* Source);
+	UE_API explicit FDNACalibDNAReader(IDNAReader* Source);
 
 	// Header
-	uint16 GetFileFormatGeneration() const override;
-	uint16 GetFileFormatVersion() const override;
+	UE_API uint16 GetFileFormatGeneration() const override;
+	UE_API uint16 GetFileFormatVersion() const override;
 	// Descriptor
-	FString GetName() const override;
-	EArchetype GetArchetype() const override;
-	EGender GetGender() const override;
-	uint16 GetAge() const override;
-	uint32 GetMetaDataCount() const override;
-	FString GetMetaDataKey(uint32 Index) const override;
-	FString GetMetaDataValue(const FString& Key) const override;
-	ETranslationUnit GetTranslationUnit() const override;
-	ERotationUnit GetRotationUnit() const override;
-	FCoordinateSystem GetCoordinateSystem() const override;
-	uint16 GetLODCount() const override;
-	uint16 GetDBMaxLOD() const override;
-	FString GetDBComplexity() const override;
-	FString GetDBName() const override;
+	UE_API FString GetName() const override;
+	UE_API EArchetype GetArchetype() const override;
+	UE_API EGender GetGender() const override;
+	UE_API uint16 GetAge() const override;
+	UE_API uint32 GetMetaDataCount() const override;
+	UE_API FString GetMetaDataKey(uint32 Index) const override;
+	UE_API FString GetMetaDataValue(const FString& Key) const override;
+	UE_API ETranslationUnit GetTranslationUnit() const override;
+	UE_API ERotationUnit GetRotationUnit() const override;
+	UE_API FCoordinateSystem GetCoordinateSystem() const override;
+	UE_API uint16 GetLODCount() const override;
+	UE_API uint16 GetDBMaxLOD() const override;
+	UE_API FString GetDBComplexity() const override;
+	UE_API FString GetDBName() const override;
 	// Definition
-	uint16 GetGUIControlCount() const override;
-	FString GetGUIControlName(uint16 Index) const override;
-	uint16 GetRawControlCount() const override;
-	FString GetRawControlName(uint16 Index) const override;
-	uint16 GetJointCount() const override;
-	FString GetJointName(uint16 Index) const override;
-	uint16 GetJointIndexListCount() const override;
-	TArrayView<const uint16> GetJointIndicesForLOD(uint16 LOD) const override;
-	uint16 GetJointParentIndex(uint16 Index) const override;
-	uint16 GetBlendShapeChannelCount() const override;
-	FString GetBlendShapeChannelName(uint16 Index) const override;
-	uint16 GetBlendShapeChannelIndexListCount() const override;
-	TArrayView<const uint16> GetBlendShapeChannelIndicesForLOD(uint16 LOD) const override;
-	uint16 GetAnimatedMapCount() const override;
-	FString GetAnimatedMapName(uint16 Index) const override;
-	uint16 GetAnimatedMapIndexListCount() const override;
-	TArrayView<const uint16> GetAnimatedMapIndicesForLOD(uint16 LOD) const override;
-	uint16 GetMeshCount() const override;
-	FString GetMeshName(uint16 Index) const override;
-	uint16 GetMeshIndexListCount() const override;
-	TArrayView<const uint16> GetMeshIndicesForLOD(uint16 LOD) const override;
-	uint16 GetMeshBlendShapeChannelMappingCount() const override;
-	FMeshBlendShapeChannelMapping GetMeshBlendShapeChannelMapping(uint16 Index) const override;
-	TArrayView<const uint16> GetMeshBlendShapeChannelMappingIndicesForLOD(uint16 LOD) const override;
-	FVector GetNeutralJointTranslation(uint16 Index) const override;
-	FVector GetNeutralJointRotation(uint16 Index) const override;
+	UE_API uint16 GetGUIControlCount() const override;
+	UE_API FString GetGUIControlName(uint16 Index) const override;
+	UE_API uint16 GetRawControlCount() const override;
+	UE_API FString GetRawControlName(uint16 Index) const override;
+	UE_API uint16 GetJointCount() const override;
+	UE_API FString GetJointName(uint16 Index) const override;
+	UE_API uint16 GetJointIndexListCount() const override;
+	UE_API TArrayView<const uint16> GetJointIndicesForLOD(uint16 LOD) const override;
+	UE_API uint16 GetJointParentIndex(uint16 Index) const override;
+	UE_API uint16 GetBlendShapeChannelCount() const override;
+	UE_API FString GetBlendShapeChannelName(uint16 Index) const override;
+	UE_API uint16 GetBlendShapeChannelIndexListCount() const override;
+	UE_API TArrayView<const uint16> GetBlendShapeChannelIndicesForLOD(uint16 LOD) const override;
+	UE_API uint16 GetAnimatedMapCount() const override;
+	UE_API FString GetAnimatedMapName(uint16 Index) const override;
+	UE_API uint16 GetAnimatedMapIndexListCount() const override;
+	UE_API TArrayView<const uint16> GetAnimatedMapIndicesForLOD(uint16 LOD) const override;
+	UE_API uint16 GetMeshCount() const override;
+	UE_API FString GetMeshName(uint16 Index) const override;
+	UE_API uint16 GetMeshIndexListCount() const override;
+	UE_API TArrayView<const uint16> GetMeshIndicesForLOD(uint16 LOD) const override;
+	UE_API uint16 GetMeshBlendShapeChannelMappingCount() const override;
+	UE_API FMeshBlendShapeChannelMapping GetMeshBlendShapeChannelMapping(uint16 Index) const override;
+	UE_API TArrayView<const uint16> GetMeshBlendShapeChannelMappingIndicesForLOD(uint16 LOD) const override;
+	UE_API FVector GetNeutralJointTranslation(uint16 Index) const override;
+	UE_API FVector GetNeutralJointRotation(uint16 Index) const override;
 	// Behavior
-	TArrayView<const uint16> GetGUIToRawInputIndices() const override;
-	TArrayView<const uint16> GetGUIToRawOutputIndices() const override;
-	TArrayView<const float> GetGUIToRawFromValues() const override;
-	TArrayView<const float> GetGUIToRawToValues() const override;
-	TArrayView<const float> GetGUIToRawSlopeValues() const override;
-	TArrayView<const float> GetGUIToRawCutValues() const override;
-	uint16 GetPSDCount() const override;
-	TArrayView<const uint16> GetPSDRowIndices() const override;
-	TArrayView<const uint16> GetPSDColumnIndices() const override;
-	TArrayView<const float> GetPSDValues() const override;
-	uint16 GetJointRowCount() const override;
-	uint16 GetJointColumnCount() const override;
-	TArrayView<const uint16> GetJointVariableAttributeIndices(uint16 LOD) const override;
-	uint16 GetJointGroupCount() const override;
-	TArrayView<const uint16> GetJointGroupLODs(uint16 JointGroupIndex) const override;
-	TArrayView<const uint16> GetJointGroupInputIndices(uint16 JointGroupIndex) const override;
-	TArrayView<const uint16> GetJointGroupOutputIndices(uint16 JointGroupIndex) const override;
-	TArrayView<const float> GetJointGroupValues(uint16 JointGroupIndex) const override;
-	TArrayView<const uint16> GetJointGroupJointIndices(uint16 JointGroupIndex) const override;
-	TArrayView<const uint16> GetBlendShapeChannelLODs() const override;
-	TArrayView<const uint16> GetBlendShapeChannelOutputIndices() const override;
-	TArrayView<const uint16> GetBlendShapeChannelInputIndices() const override;
-	TArrayView<const uint16> GetAnimatedMapLODs() const override;
-	TArrayView<const uint16> GetAnimatedMapInputIndices() const override;
-	TArrayView<const uint16> GetAnimatedMapOutputIndices() const override;
-	TArrayView<const float> GetAnimatedMapFromValues() const override;
-	TArrayView<const float> GetAnimatedMapToValues() const override;
-	TArrayView<const float> GetAnimatedMapSlopeValues() const override;
-	TArrayView<const float> GetAnimatedMapCutValues() const override;
+	UE_API TArrayView<const uint16> GetGUIToRawInputIndices() const override;
+	UE_API TArrayView<const uint16> GetGUIToRawOutputIndices() const override;
+	UE_API TArrayView<const float> GetGUIToRawFromValues() const override;
+	UE_API TArrayView<const float> GetGUIToRawToValues() const override;
+	UE_API TArrayView<const float> GetGUIToRawSlopeValues() const override;
+	UE_API TArrayView<const float> GetGUIToRawCutValues() const override;
+	UE_API uint16 GetPSDCount() const override;
+	UE_API TArrayView<const uint16> GetPSDRowIndices() const override;
+	UE_API TArrayView<const uint16> GetPSDColumnIndices() const override;
+	UE_API TArrayView<const float> GetPSDValues() const override;
+	UE_API uint16 GetJointRowCount() const override;
+	UE_API uint16 GetJointColumnCount() const override;
+	UE_API TArrayView<const uint16> GetJointVariableAttributeIndices(uint16 LOD) const override;
+	UE_API uint16 GetJointGroupCount() const override;
+	UE_API TArrayView<const uint16> GetJointGroupLODs(uint16 JointGroupIndex) const override;
+	UE_API TArrayView<const uint16> GetJointGroupInputIndices(uint16 JointGroupIndex) const override;
+	UE_API TArrayView<const uint16> GetJointGroupOutputIndices(uint16 JointGroupIndex) const override;
+	UE_API TArrayView<const float> GetJointGroupValues(uint16 JointGroupIndex) const override;
+	UE_API TArrayView<const uint16> GetJointGroupJointIndices(uint16 JointGroupIndex) const override;
+	UE_API TArrayView<const uint16> GetBlendShapeChannelLODs() const override;
+	UE_API TArrayView<const uint16> GetBlendShapeChannelOutputIndices() const override;
+	UE_API TArrayView<const uint16> GetBlendShapeChannelInputIndices() const override;
+	UE_API TArrayView<const uint16> GetAnimatedMapLODs() const override;
+	UE_API TArrayView<const uint16> GetAnimatedMapInputIndices() const override;
+	UE_API TArrayView<const uint16> GetAnimatedMapOutputIndices() const override;
+	UE_API TArrayView<const float> GetAnimatedMapFromValues() const override;
+	UE_API TArrayView<const float> GetAnimatedMapToValues() const override;
+	UE_API TArrayView<const float> GetAnimatedMapSlopeValues() const override;
+	UE_API TArrayView<const float> GetAnimatedMapCutValues() const override;
 	// Geometry
-	uint32 GetVertexPositionCount(uint16 MeshIndex) const override;
-	FVector GetVertexPosition(uint16 MeshIndex, uint32 VertexIndex) const override;
-	TArrayView<const float> GetVertexPositionXs(uint16 MeshIndex) const override;
-	TArrayView<const float> GetVertexPositionYs(uint16 MeshIndex) const override;
-	TArrayView<const float> GetVertexPositionZs(uint16 MeshIndex) const override;
-	uint32 GetVertexTextureCoordinateCount(uint16 MeshIndex) const override;
-	FTextureCoordinate GetVertexTextureCoordinate(uint16 MeshIndex, uint32 TextureCoordinateIndex) const override;
-	uint32 GetVertexNormalCount(uint16 MeshIndex) const override;
-	FVector GetVertexNormal(uint16 MeshIndex, uint32 NormalIndex) const override;
-	uint32 GetFaceCount(uint16 MeshIndex) const override;
-	TArrayView<const uint32> GetFaceVertexLayoutIndices(uint16 MeshIndex, uint32 FaceIndex) const override;
-	uint32 GetVertexLayoutCount(uint16 MeshIndex) const override;
-	FVertexLayout GetVertexLayout(uint16 MeshIndex, uint32 LayoutIndex) const override;
-	uint16 GetMaximumInfluencePerVertex(uint16 MeshIndex) const override;
-	uint32 GetSkinWeightsCount(uint16 MeshIndex) const override;
-	TArrayView<const float> GetSkinWeightsValues(uint16 MeshIndex, uint32 VertexIndex) const override;
-	TArrayView<const uint16> GetSkinWeightsJointIndices(uint16 MeshIndex, uint32 VertexIndex) const override;
-	uint16 GetBlendShapeTargetCount(uint16 MeshIndex) const override;
-	uint16 GetBlendShapeChannelIndex(uint16 MeshIndex, uint16 BlendShapeTargetIndex) const override;
-	uint32 GetBlendShapeTargetDeltaCount(uint16 MeshIndex, uint16 BlendShapeTargetIndex) const override;
-	FVector GetBlendShapeTargetDelta(uint16 MeshIndex, uint16 BlendShapeTargetIndex, uint32 DeltaIndex) const override;
-	TArrayView<const float> GetBlendShapeTargetDeltaXs(uint16 MeshIndex, uint16 BlendShapeTargetIndex) const override;
-	TArrayView<const float> GetBlendShapeTargetDeltaYs(uint16 MeshIndex, uint16 BlendShapeTargetIndex) const override;
-	TArrayView<const float> GetBlendShapeTargetDeltaZs(uint16 MeshIndex, uint16 BlendShapeTargetIndex) const override;
-	TArrayView<const uint32> GetBlendShapeTargetVertexIndices(uint16 MeshIndex, uint16 BlendShapeTargetIndex) const override;
+	UE_API uint32 GetVertexPositionCount(uint16 MeshIndex) const override;
+	UE_API FVector GetVertexPosition(uint16 MeshIndex, uint32 VertexIndex) const override;
+	UE_API TArrayView<const float> GetVertexPositionXs(uint16 MeshIndex) const override;
+	UE_API TArrayView<const float> GetVertexPositionYs(uint16 MeshIndex) const override;
+	UE_API TArrayView<const float> GetVertexPositionZs(uint16 MeshIndex) const override;
+	UE_API uint32 GetVertexTextureCoordinateCount(uint16 MeshIndex) const override;
+	UE_API FTextureCoordinate GetVertexTextureCoordinate(uint16 MeshIndex, uint32 TextureCoordinateIndex) const override;
+	UE_API uint32 GetVertexNormalCount(uint16 MeshIndex) const override;
+	UE_API FVector GetVertexNormal(uint16 MeshIndex, uint32 NormalIndex) const override;
+	UE_API uint32 GetFaceCount(uint16 MeshIndex) const override;
+	UE_API TArrayView<const uint32> GetFaceVertexLayoutIndices(uint16 MeshIndex, uint32 FaceIndex) const override;
+	UE_API uint32 GetVertexLayoutCount(uint16 MeshIndex) const override;
+	UE_API FVertexLayout GetVertexLayout(uint16 MeshIndex, uint32 LayoutIndex) const override;
+	UE_API uint16 GetMaximumInfluencePerVertex(uint16 MeshIndex) const override;
+	UE_API uint32 GetSkinWeightsCount(uint16 MeshIndex) const override;
+	UE_API TArrayView<const float> GetSkinWeightsValues(uint16 MeshIndex, uint32 VertexIndex) const override;
+	UE_API TArrayView<const uint16> GetSkinWeightsJointIndices(uint16 MeshIndex, uint32 VertexIndex) const override;
+	UE_API uint16 GetBlendShapeTargetCount(uint16 MeshIndex) const override;
+	UE_API uint16 GetBlendShapeChannelIndex(uint16 MeshIndex, uint16 BlendShapeTargetIndex) const override;
+	UE_API uint32 GetBlendShapeTargetDeltaCount(uint16 MeshIndex, uint16 BlendShapeTargetIndex) const override;
+	UE_API FVector GetBlendShapeTargetDelta(uint16 MeshIndex, uint16 BlendShapeTargetIndex, uint32 DeltaIndex) const override;
+	UE_API TArrayView<const float> GetBlendShapeTargetDeltaXs(uint16 MeshIndex, uint16 BlendShapeTargetIndex) const override;
+	UE_API TArrayView<const float> GetBlendShapeTargetDeltaYs(uint16 MeshIndex, uint16 BlendShapeTargetIndex) const override;
+	UE_API TArrayView<const float> GetBlendShapeTargetDeltaZs(uint16 MeshIndex, uint16 BlendShapeTargetIndex) const override;
+	UE_API TArrayView<const uint32> GetBlendShapeTargetVertexIndices(uint16 MeshIndex, uint16 BlendShapeTargetIndex) const override;
 	// Machine Learned Behavior
-	uint16 GetMLControlCount() const override;
-	FString GetMLControlName(uint16 Index) const override;
-	uint16 GetNeuralNetworkCount() const override;
-	uint16 GetNeuralNetworkIndexListCount() const override;
-	TArrayView<const uint16> GetNeuralNetworkIndicesForLOD(uint16 LOD) const override;
-	uint16 GetMeshRegionCount(uint16 MeshIndex) const override;
-	FString GetMeshRegionName(uint16 MeshIndex, uint16 RegionIndex) const override;
-	TArrayView<const uint16> GetNeuralNetworkIndicesForMeshRegion(uint16 MeshIndex, uint16 RegionIndex) const override;
-	TArrayView<const uint16> GetNeuralNetworkInputIndices(uint16 NetIndex) const override;
-	TArrayView<const uint16> GetNeuralNetworkOutputIndices(uint16 NetIndex) const override;
-	uint16 GetNeuralNetworkLayerCount(uint16 NetIndex) const override;
-	EActivationFunction GetNeuralNetworkLayerActivationFunction(uint16 NetIndex, uint16 LayerIndex) const override;
-	TArrayView<const float> GetNeuralNetworkLayerActivationFunctionParameters(uint16 NetIndex, uint16 LayerIndex) const override;
-	TArrayView<const float> GetNeuralNetworkLayerBiases(uint16 NetIndex, uint16 LayerIndex) const override;
-	TArrayView<const float> GetNeuralNetworkLayerWeights(uint16 NetIndex, uint16 LayerIndex) const override;
+	UE_API uint16 GetMLControlCount() const override;
+	UE_API FString GetMLControlName(uint16 Index) const override;
+	UE_API uint16 GetNeuralNetworkCount() const override;
+	UE_API uint16 GetNeuralNetworkIndexListCount() const override;
+	UE_API TArrayView<const uint16> GetNeuralNetworkIndicesForLOD(uint16 LOD) const override;
+	UE_API uint16 GetMeshRegionCount(uint16 MeshIndex) const override;
+	UE_API FString GetMeshRegionName(uint16 MeshIndex, uint16 RegionIndex) const override;
+	UE_API TArrayView<const uint16> GetNeuralNetworkIndicesForMeshRegion(uint16 MeshIndex, uint16 RegionIndex) const override;
+	UE_API TArrayView<const uint16> GetNeuralNetworkInputIndices(uint16 NetIndex) const override;
+	UE_API TArrayView<const uint16> GetNeuralNetworkOutputIndices(uint16 NetIndex) const override;
+	UE_API uint16 GetNeuralNetworkLayerCount(uint16 NetIndex) const override;
+	UE_API EActivationFunction GetNeuralNetworkLayerActivationFunction(uint16 NetIndex, uint16 LayerIndex) const override;
+	UE_API TArrayView<const float> GetNeuralNetworkLayerActivationFunctionParameters(uint16 NetIndex, uint16 LayerIndex) const override;
+	UE_API TArrayView<const float> GetNeuralNetworkLayerBiases(uint16 NetIndex, uint16 LayerIndex) const override;
+	UE_API TArrayView<const float> GetNeuralNetworkLayerWeights(uint16 NetIndex, uint16 LayerIndex) const override;
 	// JointBehaviorMetadataReader
-	ETranslationRepresentation GetJointTranslationRepresentation(uint16 JointIndex) const override;
-	ERotationRepresentation GetJointRotationRepresentation(uint16 JointIndex) const override;
-	EScaleRepresentation GetJointScaleRepresentation(uint16 JointIndex) const override;
+	UE_API ETranslationRepresentation GetJointTranslationRepresentation(uint16 JointIndex) const override;
+	UE_API ERotationRepresentation GetJointRotationRepresentation(uint16 JointIndex) const override;
+	UE_API EScaleRepresentation GetJointScaleRepresentation(uint16 JointIndex) const override;
 	// RBFBehavior
-	uint16 GetRBFPoseCount() const override;
-	FString GetRBFPoseName(uint16 PoseIndex) const override;
-	TArrayView<const uint16> GetRBFPoseJointOutputIndices(uint16 PoseIndex) const override;
-	TArrayView<const uint16> GetRBFPoseBlendShapeChannelOutputIndices(uint16 PoseIndex) const override;
-	TArrayView<const uint16> GetRBFPoseAnimatedMapOutputIndices(uint16 PoseIndex) const override;
-	TArrayView<const float> GetRBFPoseJointOutputValues(uint16 PoseIndex) const override;
-	float GetRBFPoseScale(uint16 PoseIndex) const override;
-	uint16 GetRBFPoseControlCount() const override;
-	FString GetRBFPoseControlName(uint16 PoseControlIndex) const override;
-	TArrayView<const uint16> GetRBFPoseInputControlIndices(uint16 PoseIndex) const override;
-	TArrayView<const uint16> GetRBFPoseOutputControlIndices(uint16 PoseIndex) const override;
-	TArrayView<const float> GetRBFPoseOutputControlWeights(uint16 PoseIndex) const override;
-	uint16 GetRBFSolverCount() const override;
-	uint16 GetRBFSolverIndexListCount() const override;
-	TArrayView<const uint16> GetRBFSolverIndicesForLOD(uint16 LOD) const override;
-	FString GetRBFSolverName(uint16 SolverIndex) const override;
-	TArrayView<const uint16> GetRBFSolverRawControlIndices(uint16 SolverIndex) const override;
-	TArrayView<const uint16> GetRBFSolverPoseIndices(uint16 SolverIndex) const override;
-	TArrayView<const float> GetRBFSolverRawControlValues(uint16 SolverIndex) const override;
-	ERBFSolverType GetRBFSolverType(uint16 SolverIndex) const override;
-	float GetRBFSolverRadius(uint16 SolverIndex) const override;
-	EAutomaticRadius GetRBFSolverAutomaticRadius(uint16 SolverIndex) const override;
-	float GetRBFSolverWeightThreshold(uint16 SolverIndex) const override;
-	ERBFDistanceMethod GetRBFSolverDistanceMethod(uint16 SolverIndex) const override;
-	ERBFNormalizeMethod GetRBFSolverNormalizeMethod(uint16 SolverIndex) const override;
-	ERBFFunctionType GetRBFSolverFunctionType(uint16 SolverIndex) const override;
-	ETwistAxis GetRBFSolverTwistAxis(uint16 SolverIndex) const override;
+	UE_API uint16 GetRBFPoseCount() const override;
+	UE_API FString GetRBFPoseName(uint16 PoseIndex) const override;
+	UE_API TArrayView<const uint16> GetRBFPoseJointOutputIndices(uint16 PoseIndex) const override;
+	UE_API TArrayView<const uint16> GetRBFPoseBlendShapeChannelOutputIndices(uint16 PoseIndex) const override;
+	UE_API TArrayView<const uint16> GetRBFPoseAnimatedMapOutputIndices(uint16 PoseIndex) const override;
+	UE_API TArrayView<const float> GetRBFPoseJointOutputValues(uint16 PoseIndex) const override;
+	UE_API float GetRBFPoseScale(uint16 PoseIndex) const override;
+	UE_API uint16 GetRBFPoseControlCount() const override;
+	UE_API FString GetRBFPoseControlName(uint16 PoseControlIndex) const override;
+	UE_API TArrayView<const uint16> GetRBFPoseInputControlIndices(uint16 PoseIndex) const override;
+	UE_API TArrayView<const uint16> GetRBFPoseOutputControlIndices(uint16 PoseIndex) const override;
+	UE_API TArrayView<const float> GetRBFPoseOutputControlWeights(uint16 PoseIndex) const override;
+	UE_API uint16 GetRBFSolverCount() const override;
+	UE_API uint16 GetRBFSolverIndexListCount() const override;
+	UE_API TArrayView<const uint16> GetRBFSolverIndicesForLOD(uint16 LOD) const override;
+	UE_API FString GetRBFSolverName(uint16 SolverIndex) const override;
+	UE_API TArrayView<const uint16> GetRBFSolverRawControlIndices(uint16 SolverIndex) const override;
+	UE_API TArrayView<const uint16> GetRBFSolverPoseIndices(uint16 SolverIndex) const override;
+	UE_API TArrayView<const float> GetRBFSolverRawControlValues(uint16 SolverIndex) const override;
+	UE_API ERBFSolverType GetRBFSolverType(uint16 SolverIndex) const override;
+	UE_API float GetRBFSolverRadius(uint16 SolverIndex) const override;
+	UE_API EAutomaticRadius GetRBFSolverAutomaticRadius(uint16 SolverIndex) const override;
+	UE_API float GetRBFSolverWeightThreshold(uint16 SolverIndex) const override;
+	UE_API ERBFDistanceMethod GetRBFSolverDistanceMethod(uint16 SolverIndex) const override;
+	UE_API ERBFNormalizeMethod GetRBFSolverNormalizeMethod(uint16 SolverIndex) const override;
+	UE_API ERBFFunctionType GetRBFSolverFunctionType(uint16 SolverIndex) const override;
+	UE_API ETwistAxis GetRBFSolverTwistAxis(uint16 SolverIndex) const override;
 	// TwistSwingBehavior
-	uint16 GetTwistCount() const override;
-	ETwistAxis GetTwistSetupTwistAxis(uint16 TwistIndex) const override;
-	TArrayView<const uint16> GetTwistInputControlIndices(uint16 TwistIndex) const override;
-	TArrayView<const uint16> GetTwistOutputJointIndices(uint16 TwistIndex) const override;
-	TArrayView<const float> GetTwistBlendWeights(uint16 TwistIndex) const override;
-	uint16 GetSwingCount() const override;
-	ETwistAxis GetSwingSetupTwistAxis(uint16 SwingIndex) const override;
-	TArrayView<const uint16> GetSwingInputControlIndices(uint16 SwingIndex) const override;
-	TArrayView<const uint16> GetSwingOutputJointIndices(uint16 SwingIndex) const override;
-	TArrayView<const float> GetSwingBlendWeights(uint16 SwingIndex) const override;
+	UE_API uint16 GetTwistCount() const override;
+	UE_API ETwistAxis GetTwistSetupTwistAxis(uint16 TwistIndex) const override;
+	UE_API TArrayView<const uint16> GetTwistInputControlIndices(uint16 TwistIndex) const override;
+	UE_API TArrayView<const uint16> GetTwistOutputJointIndices(uint16 TwistIndex) const override;
+	UE_API TArrayView<const float> GetTwistBlendWeights(uint16 TwistIndex) const override;
+	UE_API uint16 GetSwingCount() const override;
+	UE_API ETwistAxis GetSwingSetupTwistAxis(uint16 SwingIndex) const override;
+	UE_API TArrayView<const uint16> GetSwingInputControlIndices(uint16 SwingIndex) const override;
+	UE_API TArrayView<const uint16> GetSwingOutputJointIndices(uint16 SwingIndex) const override;
+	UE_API TArrayView<const float> GetSwingBlendWeights(uint16 SwingIndex) const override;
 
-	void Unload(EDNADataLayer layer) override;
+	UE_API void Unload(EDNADataLayer layer) override;
 
-	dna::Reader* Unwrap() const override;
+	UE_API dna::Reader* Unwrap() const override;
 
 private:
 	TUniquePtr<IDNAReader> ReaderPtr;
 
 };
+
+#undef UE_API
diff --git a/Engine/Plugins/Animation/DNACalib/Source/DNACalibModule/Public/DNACalibModule.h b/Engine/Plugins/Animation/DNACalib/Source/DNACalibModule/Public/DNACalibModule.h
--- a/Engine/Plugins/Animation/DNACalib/Source/DNACalibModule/Public/DNACalibModule.h
+++ b/Engine/Plugins/Animation/DNACalib/Source/DNACalibModule/Public/DNACalibModule.h
@@ -5,11 +5,15 @@
 #include "CoreMinimal.h"
 #include "Modules/ModuleManager.h"
 
-class DNACALIBMODULE_API FDNACalibModule : public IModuleInterface
+#define UE_API DNACALIBMODULE_API
+
+class FDNACalibModule : public IModuleInterface
 {
 public:
 
 	/** IModuleInterface implementation */
-	virtual void StartupModule() override;
-	virtual void ShutdownModule() override;
+	UE_API virtual void StartupModule() override;
+	UE_API virtual void ShutdownModule() override;
 };
+
+#undef UE_API
diff --git a/Engine/Plugins/Animation/DeformerGraph/Source/OptimusEditor/Public/IOptimusEditorModule.h b/Engine/Plugins/Animation/DeformerGraph/Source/OptimusEditor/Public/IOptimusEditorModule.h
--- a/Engine/Plugins/Animation/DeformerGraph/Source/OptimusEditor/Public/IOptimusEditorModule.h
+++ b/Engine/Plugins/Animation/DeformerGraph/Source/OptimusEditor/Public/IOptimusEditorModule.h
@@ -14,7 +14,7 @@ class FOptimusEditorClipboard;
 class IOptimusEditor;
 class UOptimusDeformer;
 
-class OPTIMUSEDITOR_API IOptimusEditorModule
+class IOptimusEditorModule
 	: public IModuleInterface
 //	, public IHasMenuExtensibility
 //	, public IHasToolBarExtensibility,
diff --git a/Engine/Plugins/Animation/DeformerGraph/Source/OptimusEditor/Public/OptimusEditorClipboard.h b/Engine/Plugins/Animation/DeformerGraph/Source/OptimusEditor/Public/OptimusEditorClipboard.h
--- a/Engine/Plugins/Animation/DeformerGraph/Source/OptimusEditor/Public/OptimusEditorClipboard.h
+++ b/Engine/Plugins/Animation/DeformerGraph/Source/OptimusEditor/Public/OptimusEditorClipboard.h
@@ -5,6 +5,8 @@
 #include "Containers/Array.h"
 #include "HAL/Platform.h"
 
+#define UE_API OPTIMUSEDITOR_API
+
 
 class UClass;
 class UObject;
@@ -12,22 +14,24 @@ class UOptimusNodeGraph;
 class UOptimusNode;
 
 
-class OPTIMUSEDITOR_API FOptimusEditorClipboard
+class FOptimusEditorClipboard
 {
 public:
-	static void SetClipboardFromNodes(const TArray<UOptimusNode*>& InNodes);
+	static UE_API void SetClipboardFromNodes(const TArray<UOptimusNode*>& InNodes);
 
 	// Creates a self-contained 
-	static UOptimusNodeGraph* GetGraphFromClipboardContent(
+	static UE_API UOptimusNodeGraph* GetGraphFromClipboardContent(
 		UPackage* InTargetPackage
 		);
 
-	static bool HasValidClipboardContent();
+	static UE_API bool HasValidClipboardContent();
 
 private:
-	static bool CanCreateObjectsFromText(const TCHAR* InBuffer);
+	static UE_API bool CanCreateObjectsFromText(const TCHAR* InBuffer);
 
-	static bool ProcessObjectBuffer(UPackage* InPackage, UObject* InRootOuter, const TCHAR* InBuffer);
-	static bool CanCreateClass(const UClass* ObjectClass);
-	static bool ProcessPostCreateObject(UObject* InRootOuter, UObject* InNewObject);
+	static UE_API bool ProcessObjectBuffer(UPackage* InPackage, UObject* InRootOuter, const TCHAR* InBuffer);
+	static UE_API bool CanCreateClass(const UClass* ObjectClass);
+	static UE_API bool ProcessPostCreateObject(UObject* InRootOuter, UObject* InNewObject);
 };
+
+#undef UE_API
diff --git a/Engine/Plugins/Animation/GameplayInsights/Source/GameplayInsights/Public/IGameplayInsightsModule.h b/Engine/Plugins/Animation/GameplayInsights/Source/GameplayInsights/Public/IGameplayInsightsModule.h
--- a/Engine/Plugins/Animation/GameplayInsights/Source/GameplayInsights/Public/IGameplayInsightsModule.h
+++ b/Engine/Plugins/Animation/GameplayInsights/Source/GameplayInsights/Public/IGameplayInsightsModule.h
@@ -15,7 +15,7 @@ namespace TraceServices
 	class IAnalysisSession;
 }
 
-class GAMEPLAYINSIGHTS_API IGameplayInsightsModule
+class IGameplayInsightsModule
 	: public IModuleInterface
 {
 public:
diff --git a/Engine/Plugins/Animation/GameplayInsights/Source/GameplayInsights/Public/RewindDebuggerPlaceholderTrack.h b/Engine/Plugins/Animation/GameplayInsights/Source/GameplayInsights/Public/RewindDebuggerPlaceholderTrack.h
--- a/Engine/Plugins/Animation/GameplayInsights/Source/GameplayInsights/Public/RewindDebuggerPlaceholderTrack.h
+++ b/Engine/Plugins/Animation/GameplayInsights/Source/GameplayInsights/Public/RewindDebuggerPlaceholderTrack.h
@@ -5,7 +5,7 @@
 #include "CoreMinimal.h"
 #include "RewindDebuggerTrack.h"
 
-class GAMEPLAYINSIGHTS_API FRewindDebuggerPlaceholderTrack : public RewindDebugger::FRewindDebuggerTrack
+class FRewindDebuggerPlaceholderTrack : public RewindDebugger::FRewindDebuggerTrack
 {
 public:
 	FRewindDebuggerPlaceholderTrack(const FName& InObjectName, const FText& InDisplayName)
diff --git a/Engine/Plugins/Animation/GameplayInsights/Source/GameplayInsights/Public/SCurveTimelineView.h b/Engine/Plugins/Animation/GameplayInsights/Source/GameplayInsights/Public/SCurveTimelineView.h
--- a/Engine/Plugins/Animation/GameplayInsights/Source/GameplayInsights/Public/SCurveTimelineView.h
+++ b/Engine/Plugins/Animation/GameplayInsights/Source/GameplayInsights/Public/SCurveTimelineView.h
@@ -6,12 +6,14 @@
 #include "Widgets/SCompoundWidget.h"
 #include "Widgets/DeclarativeSyntaxSupport.h"
 
+#define UE_API GAMEPLAYINSIGHTS_API
+
 class FPaintArgs;
 class FSlateWindowElementList;
 class FTimeSliderController;
 class SToolTip;
 
-class GAMEPLAYINSIGHTS_API SCurveTimelineView : public SCompoundWidget
+class SCurveTimelineView : public SCompoundWidget
 {
 public:
 	struct FTimelineCurveData
@@ -60,7 +62,7 @@ public:
 	 * 
 	 * @param InArgs   A declaration from which to construct the widget
 	 */
-	void Construct( const FArguments& InArgs );
+	UE_API void Construct( const FArguments& InArgs );
 
 	// render the curve with a fixed min/max range, rather than using the min/max from data
 	void SetFixedRange(float Min, float Max)
@@ -76,12 +78,12 @@ public:
 	
 protected:
 	// SWidget interface
-	virtual FVector2D ComputeDesiredSize(float) const override;
-	virtual int32 OnPaint( const FPaintArgs& Args, const FGeometry& AllottedGeometry, const FSlateRect& MyCullingRect, FSlateWindowElementList& OutDrawElements, int32 LayerId, const FWidgetStyle& InWidgetStyle, bool bParentEnabled ) const override;
-	virtual FReply OnMouseMove(const FGeometry& MyGeometry, const FPointerEvent& MouseEvent) override;
+	UE_API virtual FVector2D ComputeDesiredSize(float) const override;
+	UE_API virtual int32 OnPaint( const FPaintArgs& Args, const FGeometry& AllottedGeometry, const FSlateRect& MyCullingRect, FSlateWindowElementList& OutDrawElements, int32 LayerId, const FWidgetStyle& InWidgetStyle, bool bParentEnabled ) const override;
+	UE_API virtual FReply OnMouseMove(const FGeometry& MyGeometry, const FPointerEvent& MouseEvent) override;
 	
-	int32 PaintCurve(const FGeometry& AllottedGeometry, const FSlateRect& MyCullingRect, FSlateWindowElementList& OutDrawElements, int32 LayerId, const FWidgetStyle& InWidgetStyle, bool bParentEnabled ) const;
-	void UpdateCurveToolTip(const FGeometry& MyGeometry, const FPointerEvent& MouseEvent);
+	UE_API int32 PaintCurve(const FGeometry& AllottedGeometry, const FSlateRect& MyCullingRect, FSlateWindowElementList& OutDrawElements, int32 LayerId, const FWidgetStyle& InWidgetStyle, bool bParentEnabled ) const;
+	UE_API void UpdateCurveToolTip(const FGeometry& MyGeometry, const FPointerEvent& MouseEvent);
 	
 	TAttribute<TRange<double>> ViewRange;
 	TAttribute<TSharedPtr<FTimelineCurveData>> CurveData;
@@ -102,3 +104,4 @@ protected:
 	float FixedRangeMax;
 };
 
+#undef UE_API
diff --git a/Engine/Plugins/Animation/GameplayInsights/Source/GameplayInsights/Public/SEventTimelineView.h b/Engine/Plugins/Animation/GameplayInsights/Source/GameplayInsights/Public/SEventTimelineView.h
--- a/Engine/Plugins/Animation/GameplayInsights/Source/GameplayInsights/Public/SEventTimelineView.h
+++ b/Engine/Plugins/Animation/GameplayInsights/Source/GameplayInsights/Public/SEventTimelineView.h
@@ -6,11 +6,13 @@
 #include "Widgets/SCompoundWidget.h"
 #include "Widgets/DeclarativeSyntaxSupport.h"
 
+#define UE_API GAMEPLAYINSIGHTS_API
+
 class FPaintArgs;
 class FSlateWindowElementList;
 class FTimeSliderController;
 
-class GAMEPLAYINSIGHTS_API SEventTimelineView : public SCompoundWidget
+class SEventTimelineView : public SCompoundWidget
 {
 public:
 	struct FTimelineEventData
@@ -58,14 +60,14 @@ public:
 	 * 
 	 * @param InArgs   A declaration from which to construct the widget
 	 */
-	void Construct( const FArguments& InArgs );
+	UE_API void Construct( const FArguments& InArgs );
 
 protected:
 	// SWidget interface
-	virtual FVector2D ComputeDesiredSize(float) const override;
-	virtual int32 OnPaint( const FPaintArgs& Args, const FGeometry& AllottedGeometry, const FSlateRect& MyCullingRect, FSlateWindowElementList& OutDrawElements, int32 LayerId, const FWidgetStyle& InWidgetStyle, bool bParentEnabled ) const override;
+	UE_API virtual FVector2D ComputeDesiredSize(float) const override;
+	UE_API virtual int32 OnPaint( const FPaintArgs& Args, const FGeometry& AllottedGeometry, const FSlateRect& MyCullingRect, FSlateWindowElementList& OutDrawElements, int32 LayerId, const FWidgetStyle& InWidgetStyle, bool bParentEnabled ) const override;
 	
-	int32 PaintEvents(const FGeometry& AllottedGeometry, const FSlateRect& MyCullingRect, FSlateWindowElementList& OutDrawElements, int32 LayerId, const FWidgetStyle& InWidgetStyle, bool bParentEnabled ) const;
+	UE_API int32 PaintEvents(const FGeometry& AllottedGeometry, const FSlateRect& MyCullingRect, FSlateWindowElementList& OutDrawElements, int32 LayerId, const FWidgetStyle& InWidgetStyle, bool bParentEnabled ) const;
 	
 	TAttribute<TRange<double>> ViewRange;
 	TAttribute<TSharedPtr<FTimelineEventData>> EventData;
@@ -75,3 +77,4 @@ protected:
 	const FSlateBrush* EventBorderBrush;
 };
 
+#undef UE_API
diff --git a/Engine/Plugins/Animation/GameplayInsights/Source/GameplayInsightsEditor/Public/InsightsSkeletalMeshComponent.h b/Engine/Plugins/Animation/GameplayInsights/Source/GameplayInsightsEditor/Public/InsightsSkeletalMeshComponent.h
--- a/Engine/Plugins/Animation/GameplayInsights/Source/GameplayInsightsEditor/Public/InsightsSkeletalMeshComponent.h
+++ b/Engine/Plugins/Animation/GameplayInsights/Source/GameplayInsightsEditor/Public/InsightsSkeletalMeshComponent.h
@@ -6,19 +6,23 @@
 #include "Components/SkeletalMeshComponent.h"
 #include "InsightsSkeletalMeshComponent.generated.h"
 
+#define UE_API GAMEPLAYINSIGHTSEDITOR_API
+
 class IAnimationProvider;
 struct FSkeletalMeshPoseMessage;
 struct FSkeletalMeshInfo;
 
-UCLASS(Hidden)
-class GAMEPLAYINSIGHTSEDITOR_API UInsightsSkeletalMeshComponent : public USkeletalMeshComponent
+UCLASS(MinimalAPI, Hidden)
+class UInsightsSkeletalMeshComponent : public USkeletalMeshComponent
 {
 	GENERATED_BODY()
 
 public:
 	// Set this component up from a provider & message
-	void SetPoseFromProvider(const IAnimationProvider& InProvider, const FSkeletalMeshPoseMessage& InMessage, const FSkeletalMeshInfo& SkeletalMeshInfo);
+	UE_API void SetPoseFromProvider(const IAnimationProvider& InProvider, const FSkeletalMeshPoseMessage& InMessage, const FSkeletalMeshInfo& SkeletalMeshInfo);
 
 	// USkeletalMeshComponent interface
-	virtual void InitAnim(bool bForceReInit) override;
+	UE_API virtual void InitAnim(bool bForceReInit) override;
 };
+
+#undef UE_API
diff --git a/Engine/Plugins/Animation/IKRig/Source/IKRigDeveloper/Public/AnimGraphNode_IKRig.h b/Engine/Plugins/Animation/IKRig/Source/IKRigDeveloper/Public/AnimGraphNode_IKRig.h
--- a/Engine/Plugins/Animation/IKRig/Source/IKRigDeveloper/Public/AnimGraphNode_IKRig.h
+++ b/Engine/Plugins/Animation/IKRig/Source/IKRigDeveloper/Public/AnimGraphNode_IKRig.h
@@ -9,6 +9,8 @@
 
 #include "AnimGraphNode_IKRig.generated.h"
 
+#define UE_API IKRIGDEVELOPER_API
+
 class FPrimitiveDrawInterface;
 class USkeletalMeshComponent;
 
@@ -83,8 +85,8 @@ private:
 };
 
 // Editor node for IKRig 
-UCLASS()
-class IKRIGDEVELOPER_API UAnimGraphNode_IKRig : public UAnimGraphNode_CustomProperty
+UCLASS(MinimalAPI)
+class UAnimGraphNode_IKRig : public UAnimGraphNode_CustomProperty
 {
 	GENERATED_BODY()
 
@@ -94,29 +96,29 @@ class IKRIGDEVELOPER_API UAnimGraphNode_IKRig : public UAnimGraphNode_CustomProp
 public:
 
 	UAnimGraphNode_IKRig() = default;
-	virtual ~UAnimGraphNode_IKRig();
+	UE_API virtual ~UAnimGraphNode_IKRig();
 	
 	// UEdGraphNode interface
-	virtual FText GetNodeTitle(ENodeTitleType::Type TitleType) const override;
-	virtual void PostEditChangeProperty(struct FPropertyChangedEvent& PropertyChangedEvent) override;
-	virtual void CreateCustomPins(TArray<UEdGraphPin*>* InOldPins) override;
-	virtual void CustomizeDetails(IDetailLayoutBuilder& DetailBuilder) override;
-	virtual void CustomizePinData(UEdGraphPin* Pin, FName SourcePropertyName, int32 ArrayIndex) const override;
+	UE_API virtual FText GetNodeTitle(ENodeTitleType::Type TitleType) const override;
+	UE_API virtual void PostEditChangeProperty(struct FPropertyChangedEvent& PropertyChangedEvent) override;
+	UE_API virtual void CreateCustomPins(TArray<UEdGraphPin*>* InOldPins) override;
+	UE_API virtual void CustomizeDetails(IDetailLayoutBuilder& DetailBuilder) override;
+	UE_API virtual void CustomizePinData(UEdGraphPin* Pin, FName SourcePropertyName, int32 ArrayIndex) const override;
 	// End of UEdGraphNode interface
 
 	// UAnimGraphNode_Base interface
-	virtual void CopyNodeDataToPreviewNode(FAnimNode_Base* AnimNode) override;
-	virtual void Draw(FPrimitiveDrawInterface* PDI, USkeletalMeshComponent* PreviewSkelMeshComp) const override;
-	virtual void ValidateAnimNodeDuringCompilation(USkeleton* ForSkeleton, FCompilerResultsLog& MessageLog) override;
+	UE_API virtual void CopyNodeDataToPreviewNode(FAnimNode_Base* AnimNode) override;
+	UE_API virtual void Draw(FPrimitiveDrawInterface* PDI, USkeletalMeshComponent* PreviewSkelMeshComp) const override;
+	UE_API virtual void ValidateAnimNodeDuringCompilation(USkeleton* ForSkeleton, FCompilerResultsLog& MessageLog) override;
 	// End of UAnimGraphNode_Base interface
 
 	// UK2Node interface
-	UObject* GetJumpTargetForDoubleClick() const;
+	UE_API UObject* GetJumpTargetForDoubleClick() const;
 	// END UK2Node
 
 	// Begin UObject Interface.
-	virtual void PostLoad() override;
-	virtual void PostEditUndo() override;
+	UE_API virtual void PostLoad() override;
+	UE_API virtual void PostEditUndo() override;
 	// End UObject Interface.
 
 	virtual bool NeedsToSpecifyValidTargetClass() const override { return false; }
@@ -129,25 +131,27 @@ protected:
 private:
 
 	// set pin's default value based on the FIKRigGoal default struct
-	static void SetPinDefaultValue(UEdGraphPin* InPin, const FName& InPropertyName);
+	static UE_API void SetPinDefaultValue(UEdGraphPin* InPin, const FName& InPropertyName);
 
 	// create custom pins from exposed goals as defined in the rig definition asset 
-	void CreateCustomPinsFromValidAsset();
+	UE_API void CreateCustomPinsFromValidAsset();
 
 	// recreate custom pins from old pins if the rig definition asset is not completely loaded
-	void CreateCustomPinsFromUnloadedAsset(TArray<UEdGraphPin*>* InOldPins);
+	UE_API void CreateCustomPinsFromUnloadedAsset(TArray<UEdGraphPin*>* InOldPins);
 
 	// Handle to the registered delegate
 	FDelegateHandle OnAssetPropertyChangedHandle;
 
 	// Global callback to anticipate on changes to the asset / goals
-	bool NeedsUpdate(UObject* ObjectBeingModified, FPropertyChangedEvent& PropertyChangedEvent) const;
-	void OnPropertyChanged(UObject* ObjectBeingModified, FPropertyChangedEvent& PropertyChangedEvent);
-	void BindPropertyChanges();
+	UE_API bool NeedsUpdate(UObject* ObjectBeingModified, FPropertyChangedEvent& PropertyChangedEvent) const;
+	UE_API void OnPropertyChanged(UObject* ObjectBeingModified, FPropertyChangedEvent& PropertyChangedEvent);
+	UE_API void BindPropertyChanges();
 
 	// update the goals' array within the anim node based on the asset
-	void UpdateGoalsFromAsset();
+	UE_API void UpdateGoalsFromAsset();
 
 	// setup goal based on it's asset definition 
-	static void SetupGoal(const UIKRigEffectorGoal* InAssetGoal, FIKRigGoal& OutGoal);
+	static UE_API void SetupGoal(const UIKRigEffectorGoal* InAssetGoal, FIKRigGoal& OutGoal);
 };
+
+#undef UE_API
diff --git a/Engine/Plugins/Animation/IKRig/Source/IKRigDeveloper/Public/AnimGraphNode_RetargetPoseFromMesh.h b/Engine/Plugins/Animation/IKRig/Source/IKRigDeveloper/Public/AnimGraphNode_RetargetPoseFromMesh.h
--- a/Engine/Plugins/Animation/IKRig/Source/IKRigDeveloper/Public/AnimGraphNode_RetargetPoseFromMesh.h
+++ b/Engine/Plugins/Animation/IKRig/Source/IKRigDeveloper/Public/AnimGraphNode_RetargetPoseFromMesh.h
@@ -8,13 +8,15 @@
 #include "AnimNodes/AnimNode_RetargetPoseFromMesh.h"
 #include "AnimGraphNode_RetargetPoseFromMesh.generated.h"
 
+#define UE_API IKRIGDEVELOPER_API
+
 class FPrimitiveDrawInterface;
 class USkeletalMeshComponent;
 
 // Facilitates runtime retargeting of an input pose using a user-specified IK Retargeter asset
 // Supports input poses from the animation graph, or from other skeletal mesh components
-UCLASS()
-class IKRIGDEVELOPER_API UAnimGraphNode_RetargetPoseFromMesh : public UAnimGraphNode_Base
+UCLASS(MinimalAPI)
+class UAnimGraphNode_RetargetPoseFromMesh : public UAnimGraphNode_Base
 {
 	GENERATED_BODY()
 
@@ -23,23 +25,25 @@ class IKRIGDEVELOPER_API UAnimGraphNode_RetargetPoseFromMesh : public UAnimGraph
 
 public:
 	// UEdGraphNode interface
-	virtual FText GetNodeTitle(ENodeTitleType::Type TitleType) const override;
-	virtual void PostEditChangeProperty(struct FPropertyChangedEvent& PropertyChangedEvent) override;
-	virtual void ValidateAnimNodeDuringCompilation(USkeleton* ForSkeleton, FCompilerResultsLog& MessageLog);
+	UE_API virtual FText GetNodeTitle(ENodeTitleType::Type TitleType) const override;
+	UE_API virtual void PostEditChangeProperty(struct FPropertyChangedEvent& PropertyChangedEvent) override;
+	UE_API virtual void ValidateAnimNodeDuringCompilation(USkeleton* ForSkeleton, FCompilerResultsLog& MessageLog);
 	// End of UEdGraphNode interface
 
 	// UAnimGraphNode_Base interface
-	virtual void PreloadRequiredAssets() override;
-	virtual void CopyNodeDataToPreviewNode(FAnimNode_Base* AnimNode) override;
-	virtual FEditorModeID GetEditorMode() const override;
-	virtual void Draw(FPrimitiveDrawInterface* PDI, USkeletalMeshComponent* PreviewSkelMeshComp) const override;
-	virtual void CustomizePinData(UEdGraphPin* Pin, FName SourcePropertyName, int32 ArrayIndex) const override;
+	UE_API virtual void PreloadRequiredAssets() override;
+	UE_API virtual void CopyNodeDataToPreviewNode(FAnimNode_Base* AnimNode) override;
+	UE_API virtual FEditorModeID GetEditorMode() const override;
+	UE_API virtual void Draw(FPrimitiveDrawInterface* PDI, USkeletalMeshComponent* PreviewSkelMeshComp) const override;
+	UE_API virtual void CustomizePinData(UEdGraphPin* Pin, FName SourcePropertyName, int32 ArrayIndex) const override;
 	virtual bool UsingCopyPoseFromMesh() const override { return true; };
 	// End of UAnimGraphNode_Base interface
 
 	// UK2Node interface
-	virtual UObject* GetJumpTargetForDoubleClick() const override;
+	UE_API virtual UObject* GetJumpTargetForDoubleClick() const override;
 	// End of UK2Node interface
 
-	static const FName AnimModeName;
+	static UE_API const FName AnimModeName;
 };
+
+#undef UE_API
diff --git a/Engine/Plugins/Animation/IKRig/Source/IKRigDeveloper/Public/IKRigDeveloper.h b/Engine/Plugins/Animation/IKRig/Source/IKRigDeveloper/Public/IKRigDeveloper.h
--- a/Engine/Plugins/Animation/IKRig/Source/IKRigDeveloper/Public/IKRigDeveloper.h
+++ b/Engine/Plugins/Animation/IKRig/Source/IKRigDeveloper/Public/IKRigDeveloper.h
@@ -5,9 +5,13 @@
 #include "CoreMinimal.h"
 #include "Modules/ModuleInterface.h"
 
-class IKRIGDEVELOPER_API FIKRigDeveloperModule : public IModuleInterface
+#define UE_API IKRIGDEVELOPER_API
+
+class FIKRigDeveloperModule : public IModuleInterface
 {
 public:
-	void StartupModule() override;
-	void ShutdownModule() override;
-};
\ No newline at end of file
+	UE_API void StartupModule() override;
+	UE_API void ShutdownModule() override;
+};
+
+#undef UE_API
diff --git a/Engine/Plugins/Animation/IKRig/Source/IKRigEditor/Public/IKRigEditor.h b/Engine/Plugins/Animation/IKRig/Source/IKRigEditor/Public/IKRigEditor.h
--- a/Engine/Plugins/Animation/IKRig/Source/IKRigEditor/Public/IKRigEditor.h
+++ b/Engine/Plugins/Animation/IKRig/Source/IKRigEditor/Public/IKRigEditor.h
@@ -5,11 +5,13 @@
 #include "CoreMinimal.h"
 #include "Modules/ModuleInterface.h"
 
-class IKRIGEDITOR_API FIKRigEditor : public IModuleInterface
+#define UE_API IKRIGEDITOR_API
+
+class FIKRigEditor : public IModuleInterface
 {
 public:
-	void StartupModule() override;
-	void ShutdownModule() override;
+	UE_API void StartupModule() override;
+	UE_API void ShutdownModule() override;
 
 private:
 	TSharedPtr<class FAssetTypeActions_AnimationAssetRetarget> RetargetAnimationAssetAction;
@@ -20,3 +22,5 @@ private:
 };
 
 DECLARE_LOG_CATEGORY_EXTERN(LogIKRigEditor, Warning, All);
+
+#undef UE_API
diff --git a/Engine/Plugins/Animation/IKRig/Source/IKRigEditor/Public/RetargetEditor/IKRetargetBatchOperation.h b/Engine/Plugins/Animation/IKRig/Source/IKRigEditor/Public/RetargetEditor/IKRetargetBatchOperation.h
--- a/Engine/Plugins/Animation/IKRig/Source/IKRigEditor/Public/RetargetEditor/IKRetargetBatchOperation.h
+++ b/Engine/Plugins/Animation/IKRig/Source/IKRigEditor/Public/RetargetEditor/IKRetargetBatchOperation.h
@@ -8,6 +8,8 @@
 
 #include "IKRetargetBatchOperation.generated.h"
 
+#define UE_API IKRIGEDITOR_API
+
 class UIKRetargeter;
 struct FScopedSlowTask;
 class USkeletalMesh;
@@ -44,7 +46,7 @@ enum class ERetargetRootLockMode : uint8
 };
 
 /// Data needed to run a batch "duplicate and retarget" operation on a set of animation assets
-struct IKRIGEDITOR_API FIKRetargetBatchOperationContext
+struct FIKRetargetBatchOperationContext
 {
 	
 public:
@@ -104,8 +106,8 @@ public:
 };
 
 // Encapsulate ability to batch duplicate and retarget a set of animation assets
-UCLASS(BlueprintType)
-class IKRIGEDITOR_API UIKRetargetBatchOperation : public UObject
+UCLASS(MinimalAPI, BlueprintType)
+class UIKRetargetBatchOperation : public UObject
 {
 	GENERATED_BODY()
 	
@@ -127,7 +129,7 @@ public:
 	 * 
 	 * Return: list of new animation files created.*/
 	UFUNCTION(BlueprintCallable, Category=IKBatchRetarget)
-	static TArray<FAssetData> DuplicateAndRetarget(
+	static UE_API TArray<FAssetData> DuplicateAndRetarget(
 		const TArray<FAssetData>& AssetsToRetarget,
 		USkeletalMesh* SourceMesh,
 		USkeletalMesh* TargetMesh,
@@ -139,39 +141,39 @@ public:
 		const bool bIncludeReferencedAssets=true);
 	
 	// Actually run the process to duplicate and retarget the assets for the given context
-	void RunRetarget(FIKRetargetBatchOperationContext& Context);
+	UE_API void RunRetarget(FIKRetargetBatchOperationContext& Context);
 
 private:
 
-	void Reset();
+	UE_API void Reset();
 
 	// Initialize set of referenced assets to retarget.
 	// @return	Number of assets that need retargeting. 
-	int32 GenerateAssetLists(const FIKRetargetBatchOperationContext& Context);
+	UE_API int32 GenerateAssetLists(const FIKRetargetBatchOperationContext& Context);
 
 	// Duplicate all the assets to retarget
-	void DuplicateRetargetAssets(const FIKRetargetBatchOperationContext& Context, FScopedSlowTask& Progress);
+	UE_API void DuplicateRetargetAssets(const FIKRetargetBatchOperationContext& Context, FScopedSlowTask& Progress);
 
 	// Retarget skeleton and animation on all the duplicates
-	void RetargetAssets(const FIKRetargetBatchOperationContext& Context, FScopedSlowTask& Progress);
+	UE_API void RetargetAssets(const FIKRetargetBatchOperationContext& Context, FScopedSlowTask& Progress);
 
 	// Convert animation on all the duplicates
-	void ConvertAnimation(const FIKRetargetBatchOperationContext& Context, FScopedSlowTask& Progress);
+	UE_API void ConvertAnimation(const FIKRetargetBatchOperationContext& Context, FScopedSlowTask& Progress);
 
 	// Copy/remap curves on all the duplicates
-	void RemapCurves(const FIKRetargetBatchOperationContext& Context, FScopedSlowTask& Progress);
+	UE_API void RemapCurves(const FIKRetargetBatchOperationContext& Context, FScopedSlowTask& Progress);
 
 	// Replace existing assets (optional)
-	void OverwriteExistingAssets(const FIKRetargetBatchOperationContext& Context, FScopedSlowTask& Progress);
+	UE_API void OverwriteExistingAssets(const FIKRetargetBatchOperationContext& Context, FScopedSlowTask& Progress);
 
 	// Output notifications of results
-	void NotifyUserOfResults(const FIKRetargetBatchOperationContext& Context, FScopedSlowTask& Progress) const;
+	UE_API void NotifyUserOfResults(const FIKRetargetBatchOperationContext& Context, FScopedSlowTask& Progress) const;
 
 	// Generate list of newly created assets to report to user
-	void GetNewAssets(TArray<UObject*>& NewAssets) const;
+	UE_API void GetNewAssets(TArray<UObject*>& NewAssets) const;
 
 	// If user cancelled half way, cleanup all the duplicated assets
-	void CleanupIfCancelled(const FScopedSlowTask& Progress) const;
+	UE_API void CleanupIfCancelled(const FScopedSlowTask& Progress) const;
 	
 	// Lists of assets to retarget. Populated from selection during init
 	TArray<UAnimationAsset*>	AnimationAssetsToRetarget;
@@ -183,3 +185,5 @@ private:
 
 	TMap<UAnimationAsset*, UAnimationAsset*>	RemappedAnimAssets;
 };
+
+#undef UE_API
