ControlRig: ensure that transform elements are correctly spawned in global space

when creating new elements in global space, make sure that:
- initial and current transforms are similar
- dirtiness is consistent with the initial transform just set

in a global space context, upper transforms and dirtinesses must be computed properly to avoid discrepancies when copying initial transforms and dirtinesses.

AddBone and AddSocket did, but AddNull did not, they now all use the same code path for consistency.

#rnx

diff --git a/Engine/Plugins/Animation/ControlRig/Source/ControlRig/Private/Rigs/RigHierarchyController.cpp b/Engine/Plugins/Animation/ControlRig/Source/ControlRig/Private/Rigs/RigHierarchyController.cpp
--- a/Engine/Plugins/Animation/ControlRig/Source/ControlRig/Private/Rigs/RigHierarchyController.cpp
+++ b/Engine/Plugins/Animation/ControlRig/Source/ControlRig/Private/Rigs/RigHierarchyController.cpp
@@ -67,7 +67,7 @@ void URigHierarchyController::Serialize(FArchive& Ar)
 
 URigHierarchy* URigHierarchyController::GetHierarchy() const
 {
-	return Cast<URigHierarchy>(GetOuter());;
+	return Cast<URigHierarchy>(GetOuter());
 }
 
 void URigHierarchyController::SetHierarchy(URigHierarchy* InHierarchy)
@@ -375,21 +375,10 @@ FRigElementKey URigHierarchyController::AddBone(FName InName, FRigElementKey InP
 		NewElement->Key.Type = ERigElementType::Bone;
 		NewElement->Key.Name = GetSafeNewName(InName, NewElement->Key.Type);
 		NewElement->BoneType = InBoneType;
-		AddElement(NewElement, Hierarchy->Get(Hierarchy->GetIndex(InParent)), true, InName);
-
-		if(bTransformInGlobal)
-		{
-			Hierarchy->SetTransform(NewElement, InTransform, ERigTransformType::InitialGlobal, true, false);
-			Hierarchy->SetTransform(NewElement, InTransform, ERigTransformType::CurrentGlobal, true, false);
-		}
-		else
-		{
-			Hierarchy->SetTransform(NewElement, InTransform, ERigTransformType::InitialLocal, true, false);
-			Hierarchy->SetTransform(NewElement, InTransform, ERigTransformType::CurrentLocal, true, false);
-		}
 
-		NewElement->GetTransform().Current = NewElement->GetTransform().Initial;
-		NewElement->GetDirtyState().Current = NewElement->GetDirtyState().Initial;
+		constexpr bool bMaintainGlobalTransform = true;
+		FRigBaseElement* Parent = Hierarchy->Get(Hierarchy->GetIndex(InParent));
+		AddTransformElementInternal(NewElement, Parent, InTransform, bTransformInGlobal, bMaintainGlobalTransform, InName);
 	}
 
 #if WITH_EDITOR
@@ -439,19 +428,10 @@ FRigElementKey URigHierarchyController::AddNull(FName InName, FRigElementKey InP
 		TGuardValue<bool> DisableCacheValidityChecks(Hierarchy->bEnableCacheValidityCheck, false);		
 		NewElement->Key.Type = ERigElementType::Null;
 		NewElement->Key.Name = GetSafeNewName(InName, NewElement->Key.Type);
-		AddElement(NewElement, Hierarchy->Get(Hierarchy->GetIndex(InParent)), false, InName);
 
-		if(bTransformInGlobal)
-		{
-			Hierarchy->SetTransform(NewElement, InTransform, ERigTransformType::InitialGlobal, true, false);
-		}
-		else
-		{
-			Hierarchy->SetTransform(NewElement, InTransform, ERigTransformType::InitialLocal, true, false);
-		}
-
-		NewElement->GetTransform().Current = NewElement->GetTransform().Initial;
-		NewElement->GetDirtyState().Current = NewElement->GetDirtyState().Initial;
+		constexpr bool bDoNotMaintainGlobalTransform = false;
+		FRigBaseElement* Parent = Hierarchy->Get(Hierarchy->GetIndex(InParent));
+		AddTransformElementInternal(NewElement, Parent, InTransform, bTransformInGlobal, bDoNotMaintainGlobalTransform, InName);
 	}
 
 #if WITH_EDITOR
@@ -802,21 +782,10 @@ FRigElementKey URigHierarchyController::AddSocket(FName InName, FRigElementKey I
 		TGuardValue<bool> DisableCacheValidityChecks(CurrentHierarchy->bEnableCacheValidityCheck, false);
 		NewElement->Key.Type = ERigElementType::Socket;
 		NewElement->Key.Name = GetSafeNewName(InName, NewElement->Key.Type);
-		AddElement(NewElement, CurrentHierarchy->Get(CurrentHierarchy->GetIndex(InParent)), true, InName);
-
-		if(bTransformInGlobal)
-		{
-			CurrentHierarchy->SetTransform(NewElement, InTransform, ERigTransformType::InitialGlobal, true, false);
-			CurrentHierarchy->SetTransform(NewElement, InTransform, ERigTransformType::CurrentGlobal, true, false);
-		}
-		else
-		{
-			CurrentHierarchy->SetTransform(NewElement, InTransform, ERigTransformType::InitialLocal, true, false);
-			CurrentHierarchy->SetTransform(NewElement, InTransform, ERigTransformType::CurrentLocal, true, false);
-		}
-
-		NewElement->GetTransform().Current = NewElement->GetTransform().Initial;
-		NewElement->GetDirtyState().Current = NewElement->GetDirtyState().Initial;
+		
+		constexpr bool bMaintainGlobalTransform = true;
+		FRigBaseElement* Parent = CurrentHierarchy->Get(CurrentHierarchy->GetIndex(InParent));
+		AddTransformElementInternal(NewElement, Parent, InTransform, bTransformInGlobal, bMaintainGlobalTransform, InName);
 
 		NewElement->SetColor(InColor, CurrentHierarchy);
 		NewElement->SetDescription(InDescription, CurrentHierarchy);
@@ -3123,6 +3092,31 @@ bool URigHierarchyController::RemoveElement(FRigBaseElement* InElement)
 	return NumElementsRemoved == 1;
 }
 
+void URigHierarchyController::AddTransformElementInternal(
+	FRigTransformElement* InNewTransformElement, FRigBaseElement* InParent,
+	const FTransform& InInitialTransform, const bool bTransformInGlobal,
+    const bool bMaintainGlobalTransform, const FName InName)
+{
+	URigHierarchy* Hierarchy = GetHierarchy();
+
+	AddElement(InNewTransformElement, InParent, bMaintainGlobalTransform, InName);
+
+	constexpr bool bAffectChildren = true, bNoUndo = false;
+	if (bTransformInGlobal)
+	{
+		Hierarchy->SetTransform(InNewTransformElement, InInitialTransform, ERigTransformType::InitialGlobal, bAffectChildren, bNoUndo);
+		Hierarchy->SetTransform(InNewTransformElement, InInitialTransform, ERigTransformType::CurrentGlobal, bAffectChildren, bNoUndo);
+	}
+	else
+	{
+		Hierarchy->SetTransform(InNewTransformElement, InInitialTransform, ERigTransformType::InitialLocal, bAffectChildren, bNoUndo);
+		Hierarchy->SetTransform(InNewTransformElement, InInitialTransform, ERigTransformType::CurrentLocal, bAffectChildren, bNoUndo);
+	}
+
+	InNewTransformElement->GetTransform().Current = InNewTransformElement->GetTransform().Initial;
+	InNewTransformElement->GetDirtyState().Current = InNewTransformElement->GetDirtyState().Initial;
+}
+
 FRigElementKey URigHierarchyController::RenameElement(FRigElementKey InElement, FName InName, bool bSetupUndo, bool bPrintPythonCommand, bool bClearSelection)
 {
 	if(!IsValid())
diff --git a/Engine/Plugins/Animation/ControlRig/Source/ControlRig/Public/Rigs/RigHierarchyController.h b/Engine/Plugins/Animation/ControlRig/Source/ControlRig/Public/Rigs/RigHierarchyController.h
--- a/Engine/Plugins/Animation/ControlRig/Source/ControlRig/Public/Rigs/RigHierarchyController.h
+++ b/Engine/Plugins/Animation/ControlRig/Source/ControlRig/Public/Rigs/RigHierarchyController.h
@@ -1106,6 +1106,17 @@ private:
 
 	UE_API void UpdateComponentsOnHierarchyKeyChange(const TArray<TPair<FRigHierarchyKey, FRigHierarchyKey>>& InKeyMap, bool bSetupUndoRedo);
 
+	/**
+	 * Adds a newly created transform element to the hierarchy
+	 * @param InNewElement The new transform element previously created.
+	 * @param InParent The (optional) parent of the new transform element (this can be null).
+	 * @param InInitialTransform The initial transform for the new element - either in local or global space, based on bTransformInGlobal.
+	 * @param bTransformInGlobal Set this to true if the Transform passed is expressed in global space, false for local space.
+ 	 * @param bMaintainGlobalTransform If set to true the child will stay in the same place spatially, otherwise it will maintain it's local transform (and potential move).
+	 * @param InName The original desired name.
+	 */
+	UE_API void AddTransformElementInternal(FRigTransformElement* InNewElement, FRigBaseElement* InParent, const FTransform& InInitialTransform, const bool bTransformInGlobal, const bool bMaintainGlobalTransform, const FName InName);
+	
 	/** 
 	 * If set to true all notifications coming from this hierarchy will be suspended
 	 */
