ControlRig editor does not execute post undo for the transient editor preview actor

diff --git a/Engine/Plugins/Animation/ControlRig/Source/ControlRigEditor/Private/Editor/ControlRigEditor.cpp b/Engine/Plugins/Animation/ControlRig/Source/ControlRigEditor/Private/Editor/ControlRigEditor.cpp
--- a/Engine/Plugins/Animation/ControlRig/Source/ControlRigEditor/Private/Editor/ControlRigEditor.cpp
+++ b/Engine/Plugins/Animation/ControlRig/Source/ControlRigEditor/Private/Editor/ControlRigEditor.cpp
@@ -2122,7 +2122,12 @@ void FControlRigBaseEditor::PostTransactionImpl(bool bSuccess, const FTransactio
 	bool bAffectsThisAsset = false;
 	for(const UObject* Object : TransactedObjects)
 	{
-		if(Object->GetPackage() == Package)
+		if (Object == RigBlueprint->GetObjectBeingDebugged())
+		{
+			bAffectsThisAsset = true;
+			break;
+		}
+		if (Object->GetPackage() == Package)
 		{
 			bAffectsThisAsset = true;
 			break;
