RigVM: Localized registry in cooked environments - first steps.

Removing global registry access mostly
Removing API use of bLockRegistry
Introduce a new RigVMRegistryContext providing access to the registry and to be passed between rigvm core APIs, so now we have the normal API, which sets up locks internally, and the NoLock API, which requires a registry context to be passed
Introducing notion of a localized registry  for the VM to serialize, initial unit tests + console variable with it disabled for now

#rnx

diff --git a/Engine/Plugins/Animation/ControlRig/Source/ControlRig/Private/Units/Animation/RigUnit_AnimAttribute.cpp b/Engine/Plugins/Animation/ControlRig/Source/ControlRig/Private/Units/Animation/RigUnit_AnimAttribute.cpp
--- a/Engine/Plugins/Animation/ControlRig/Source/ControlRig/Private/Units/Animation/RigUnit_AnimAttribute.cpp
+++ b/Engine/Plugins/Animation/ControlRig/Source/ControlRig/Private/Units/Animation/RigUnit_AnimAttribute.cpp
@@ -10,25 +10,26 @@
 
 #include UE_INLINE_GENERATED_CPP_BY_NAME(RigUnit_AnimAttribute)
 
-bool FRigDispatch_AnimAttributeBase::IsTypeSupported(const TRigVMTypeIndex& InTypeIndex)
+bool FRigDispatch_AnimAttributeBase::IsTypeSupported(const TRigVMTypeIndex& InTypeIndex, const FRigVMRegistryHandle& InRegistry) const
 {
-	const FRigVMRegistry_NoLock& Registry = FRigVMRegistry_NoLock::GetForRead();
-
-	static const TArray<TRigVMTypeIndex> SpecialTypes = {
-		RigVMTypeUtils::TypeIndex::Float,
-		RigVMTypeUtils::TypeIndex::Int32,
-		RigVMTypeUtils::TypeIndex::FString,
-		Registry.GetTypeIndex_NoLock<FTransform>(false),
-		Registry.GetTypeIndex_NoLock<FVector>(false),
-		Registry.GetTypeIndex_NoLock<FQuat>(false),
-	};
+	if(SpecialTypes.IsEmpty())
+	{
+		SpecialTypes = {
+			RigVMTypeUtils::TypeIndex::Float,
+			RigVMTypeUtils::TypeIndex::Int32,
+			RigVMTypeUtils::TypeIndex::FString,
+			InRegistry->GetTypeIndex_NoLock<FTransform>(false),
+			InRegistry->GetTypeIndex_NoLock<FVector>(false),
+			InRegistry->GetTypeIndex_NoLock<FQuat>(false),
+		};
+	}
 
 	if (SpecialTypes.Contains(InTypeIndex))
 	{
 		return true;
 	}
 
-	const FRigVMTemplateArgumentType& InType = Registry.GetType_NoLock(InTypeIndex);
+	const FRigVMTemplateArgumentType& InType = InRegistry->GetType_NoLock(InTypeIndex);
 
 	// cpp type object can become invalid because users can choose to delete
 	// user defined structs
@@ -262,11 +263,11 @@ void FRigDispatch_SetAnimAttribute::SetAnimAttributeDispatch(FRigVMExtendedExecu
 	}	
 }
 
-void FRigDispatch_AnimAttributeBase::RegisterDependencyTypes_NoLock(FRigVMRegistry_NoLock& InRegistry) const
+void FRigDispatch_AnimAttributeBase::RegisterDependencyTypes_NoLock(FRigVMRegistryHandle& InRegistry) const
 {
 	Super::RegisterDependencyTypes_NoLock(InRegistry);
 
-	InRegistry.FindOrAddType_NoLock(FRigVMTemplateArgumentType(), true);
+	InRegistry->FindOrAddType_NoLock(FRigVMTemplateArgumentType(), true);
 	
 	TArray<TWeakObjectPtr<const UScriptStruct>> AttributeTypes = UE::Anim::AttributeTypes::GetRegisteredTypes();
 	for (TWeakObjectPtr<const UScriptStruct> Type : AttributeTypes)
@@ -274,7 +275,7 @@ void FRigDispatch_AnimAttributeBase::RegisterDependencyTypes_NoLock(FRigVMRegist
 		UScriptStruct* TypePtr = const_cast<UScriptStruct*>(Type.Get());
 		if (TypePtr)
 		{
-			InRegistry.FindOrAddType_NoLock(FRigVMTemplateArgumentType(TypePtr));
+			InRegistry->FindOrAddType_NoLock(FRigVMTemplateArgumentType(TypePtr));
 		}
 	}
 }
@@ -308,18 +309,18 @@ FString FRigDispatch_AnimAttributeBase::GetNodeTitle(const FRigVMTemplateTypeMap
 }
 #endif
 
-const TArray<FRigVMTemplateArgumentInfo>& FRigDispatch_AnimAttributeBase::GetArgumentInfos() const
+const TArray<FRigVMTemplateArgumentInfo>& FRigDispatch_AnimAttributeBase::GetArgumentInfos(FRigVMRegistryHandle& InRegistry) const
 {
-	if (Infos.IsEmpty())
+	if (CachedArgumentInfos.IsEmpty())
 	{
-		NameArgIndex = Infos.Emplace(NameArgName, ERigVMPinDirection::Input, RigVMTypeUtils::TypeIndex::FName);
-		BoneNameArgIndex = Infos.Emplace(BoneNameArgName, ERigVMPinDirection::Input, RigVMTypeUtils::TypeIndex::FName);
+		NameArgIndex = CachedArgumentInfos.Emplace(NameArgName, ERigVMPinDirection::Input, RigVMTypeUtils::TypeIndex::FName);
+		BoneNameArgIndex = CachedArgumentInfos.Emplace(BoneNameArgName, ERigVMPinDirection::Input, RigVMTypeUtils::TypeIndex::FName);
 
-		CachedBoneNameArgIndex = Infos.Emplace(CachedBoneNameArgName, ERigVMPinDirection::Hidden, RigVMTypeUtils::TypeIndex::FName);
-		CachedBoneIndexArgIndex = Infos.Emplace(CachedBoneIndexArgName, ERigVMPinDirection::Hidden, RigVMTypeUtils::TypeIndex::Int32);
+		CachedBoneNameArgIndex = CachedArgumentInfos.Emplace(CachedBoneNameArgName, ERigVMPinDirection::Hidden, RigVMTypeUtils::TypeIndex::FName);
+		CachedBoneIndexArgIndex = CachedArgumentInfos.Emplace(CachedBoneIndexArgName, ERigVMPinDirection::Hidden, RigVMTypeUtils::TypeIndex::Int32);
 	}
 	
-	return Infos;
+	return CachedArgumentInfos;
 }
 
 #if WITH_EDITOR
@@ -359,27 +360,27 @@ FText FRigDispatch_AnimAttributeBase::GetArgumentTooltip(const FName& InArgument
 
 
 
-const TArray<FRigVMTemplateArgumentInfo>& FRigDispatch_GetAnimAttribute::GetArgumentInfos() const
+const TArray<FRigVMTemplateArgumentInfo>& FRigDispatch_GetAnimAttribute::GetArgumentInfos(FRigVMRegistryHandle& InRegistry) const
 {
-	if (ValueArgIndex == INDEX_NONE)
+	if (CachedArgumentInfos.IsEmpty())
 	{
-		Infos = Super::GetArgumentInfos(); 
+		CachedArgumentInfos = Super::GetArgumentInfos(InRegistry); 
 
 		// Will be open to any category, but will filter the type through our IsTypeSupported callback
 		// If we reduce this to multiple (more targeted) categories, and any of those categories have common types, bUseCategories will not be true
 		// and the template will not receive notifications of newly added types
 		const TArray<FRigVMTemplateArgument::ETypeCategory> Categories = {FRigVMTemplateArgument::ETypeCategory_SingleAnyValue};
-		DefaultArgIndex = Infos.Emplace(DefaultArgName, ERigVMPinDirection::Input, Categories, [](const TRigVMTypeIndex& Type) { return FRigDispatch_AnimAttributeBase::IsTypeSupported(Type); });
-		ValueArgIndex = Infos.Emplace(ValueArgName, ERigVMPinDirection::Output, Categories, [](const TRigVMTypeIndex& Type) { return FRigDispatch_AnimAttributeBase::IsTypeSupported(Type); });
+		DefaultArgIndex = CachedArgumentInfos.Emplace(DefaultArgName, ERigVMPinDirection::Input, Categories, [this](const TRigVMTypeIndex& Type, const FRigVMRegistryHandle& InRegistry) { return IsTypeSupported(Type, InRegistry); });
+		ValueArgIndex = CachedArgumentInfos.Emplace(ValueArgName, ERigVMPinDirection::Output, Categories, [this](const TRigVMTypeIndex& Type, const FRigVMRegistryHandle& InRegistry) { return IsTypeSupported(Type, InRegistry); });
 		
-		FoundArgIndex = Infos.Emplace(FoundArgName, ERigVMPinDirection::Output, RigVMTypeUtils::TypeIndex::Bool);
+		FoundArgIndex = CachedArgumentInfos.Emplace(FoundArgName, ERigVMPinDirection::Output, RigVMTypeUtils::TypeIndex::Bool);
 	}
 
-	return Infos;
+	return CachedArgumentInfos;
 }
 
 FRigVMTemplateTypeMap FRigDispatch_GetAnimAttribute::OnNewArgumentType(const FName& InArgumentName,
-	TRigVMTypeIndex InTypeIndex) const
+	TRigVMTypeIndex InTypeIndex, FRigVMRegistryHandle& InRegistry) const
 {
 	FRigVMTemplateTypeMap Types;
 
@@ -391,7 +392,7 @@ FRigVMTemplateTypeMap FRigDispatch_GetAnimAttribute::OnNewArgumentType(const FNa
 	
 	if (ensure(InArgumentName == DefaultArgName))
 	{
-		if (IsTypeSupported(InTypeIndex))
+		if (IsTypeSupported(InTypeIndex, InRegistry))
 		{
 			Types.Add(NameArgName, RigVMTypeUtils::TypeIndex::FName);
 			Types.Add(BoneNameArgName, RigVMTypeUtils::TypeIndex::FName);
@@ -405,9 +406,8 @@ FRigVMTemplateTypeMap FRigDispatch_GetAnimAttribute::OnNewArgumentType(const FNa
 	return Types;
 }
 
-FRigVMFunctionPtr FRigDispatch_GetAnimAttribute::GetDispatchFunctionImpl(const FRigVMTemplateTypeMap& InTypes) const
+FRigVMFunctionPtr FRigDispatch_GetAnimAttribute::GetDispatchFunctionImpl(const FRigVMTemplateTypeMap& InTypes, const FRigVMRegistryHandle& InRegistry) const
 {
-	const FRigVMRegistry_NoLock& Registry = FRigVMRegistry_NoLock::GetForRead();
 	const TRigVMTypeIndex& ValueTypeIndex = InTypes.FindChecked(TEXT("Value"));
 
 	if (ValueTypeIndex == RigVMTypeUtils::TypeIndex::Float)
@@ -422,20 +422,20 @@ FRigVMFunctionPtr FRigDispatch_GetAnimAttribute::GetDispatchFunctionImpl(const F
 	{
 		return &FRigDispatch_GetAnimAttribute::GetAnimAttributeDispatch<FString>;
 	}
-	if (ValueTypeIndex == Registry.GetTypeIndex_NoLock<FTransform>(false))
+	if (ValueTypeIndex == InRegistry->GetTypeIndex_NoLock<FTransform>(false))
 	{
 		return &FRigDispatch_GetAnimAttribute::GetAnimAttributeDispatch<FTransform>;
 	}
-	if (ValueTypeIndex == Registry.GetTypeIndex_NoLock<FQuat>(false))
+	if (ValueTypeIndex == InRegistry->GetTypeIndex_NoLock<FQuat>(false))
 	{
 		return &FRigDispatch_GetAnimAttribute::GetAnimAttributeDispatch<FQuat>;
 	}
-	if (ValueTypeIndex == Registry.GetTypeIndex_NoLock<FVector>(false))
+	if (ValueTypeIndex == InRegistry->GetTypeIndex_NoLock<FVector>(false))
 	{
 		return &FRigDispatch_GetAnimAttribute::GetAnimAttributeDispatch<FVector>;
 	}
 	
-	const FRigVMTemplateArgumentType& ValueType = Registry.GetType_NoLock(ValueTypeIndex);
+	const FRigVMTemplateArgumentType& ValueType = InRegistry->GetType_NoLock(ValueTypeIndex);
 	
 	if (UScriptStruct* ScriptStruct = Cast<UScriptStruct>(ValueType.CPPTypeObject))
 	{
@@ -445,22 +445,22 @@ FRigVMFunctionPtr FRigDispatch_GetAnimAttribute::GetDispatchFunctionImpl(const F
 	return nullptr;
 }
 
-const TArray<FRigVMTemplateArgumentInfo>& FRigDispatch_SetAnimAttribute::GetArgumentInfos() const
+const TArray<FRigVMTemplateArgumentInfo>& FRigDispatch_SetAnimAttribute::GetArgumentInfos(FRigVMRegistryHandle& InRegistry) const
 {
-	if (ValueArgIndex == INDEX_NONE)
+	if (CachedArgumentInfos.IsEmpty())
 	{
-		Infos = Super::GetArgumentInfos();
+		CachedArgumentInfos = Super::GetArgumentInfos(InRegistry);
 
 		// Will be open to any category, but will filter the type through our IsTypeSupported callback
 		// If we reduce this to multiple (more targeted) categories, and any of those categories have common types, bUseCategories will not be true
 		// and the template will not receive notifications of newly added types
 		static const TArray<FRigVMTemplateArgument::ETypeCategory> Categories = {FRigVMTemplateArgument::ETypeCategory_SingleAnyValue};
 		
-		ValueArgIndex = Infos.Emplace(ValueArgName, ERigVMPinDirection::Input, Categories, [](const TRigVMTypeIndex& Type) { return FRigDispatch_AnimAttributeBase::IsTypeSupported(Type); });		
-		SuccessArgIndex = Infos.Emplace(SuccessArgName, ERigVMPinDirection::Output, RigVMTypeUtils::TypeIndex::Bool);
+		ValueArgIndex = CachedArgumentInfos.Emplace(ValueArgName, ERigVMPinDirection::Input, Categories, [this](const TRigVMTypeIndex& Type, const FRigVMRegistryHandle& InRegistry) { return IsTypeSupported(Type, InRegistry); });		
+		SuccessArgIndex = CachedArgumentInfos.Emplace(SuccessArgName, ERigVMPinDirection::Output, RigVMTypeUtils::TypeIndex::Bool);
 	}
 	
-	return Infos;
+	return CachedArgumentInfos;
 }
 
 const TArray<FRigVMExecuteArgument>& FRigDispatch_SetAnimAttribute::GetExecuteArguments_Impl(const FRigVMDispatchContext& InContext) const
@@ -472,7 +472,8 @@ const TArray<FRigVMExecuteArgument>& FRigDispatch_SetAnimAttribute::GetExecuteAr
 }
 
 FRigVMTemplateTypeMap FRigDispatch_SetAnimAttribute::OnNewArgumentType(const FName& InArgumentName,
-                                                                       TRigVMTypeIndex InTypeIndex) const
+                                                                       TRigVMTypeIndex InTypeIndex,
+                                                                       FRigVMRegistryHandle& InRegistry) const
 {
 	FRigVMTemplateTypeMap Types;
 
@@ -480,7 +481,7 @@ FRigVMTemplateTypeMap FRigDispatch_SetAnimAttribute::OnNewArgumentType(const FNa
 	// likely not strictly needed for this node since it only has a single non-singleton argument
 	if (InArgumentName == ValueArgName)
 	{
-		if (IsTypeSupported(InTypeIndex))
+		if (IsTypeSupported(InTypeIndex, InRegistry))
 		{
 			Types.Add(NameArgName, RigVMTypeUtils::TypeIndex::FName);
 			Types.Add(BoneNameArgName, RigVMTypeUtils::TypeIndex::FName);
@@ -494,9 +495,8 @@ FRigVMTemplateTypeMap FRigDispatch_SetAnimAttribute::OnNewArgumentType(const FNa
 	return Types;	
 }
 
-FRigVMFunctionPtr FRigDispatch_SetAnimAttribute::GetDispatchFunctionImpl(const FRigVMTemplateTypeMap& InTypes) const
+FRigVMFunctionPtr FRigDispatch_SetAnimAttribute::GetDispatchFunctionImpl(const FRigVMTemplateTypeMap& InTypes, const FRigVMRegistryHandle& InRegistry) const
 {
-	const FRigVMRegistry_NoLock& Registry = FRigVMRegistry_NoLock::GetForRead();
 	const TRigVMTypeIndex& ValueTypeIndex = InTypes.FindChecked(TEXT("Value"));
 
 	if (ValueTypeIndex == RigVMTypeUtils::TypeIndex::Float)
@@ -511,22 +511,22 @@ FRigVMFunctionPtr FRigDispatch_SetAnimAttribute::GetDispatchFunctionImpl(const F
 	{
 		return &FRigDispatch_SetAnimAttribute::SetAnimAttributeDispatch<FString>;
 	}
-	if (ValueTypeIndex == Registry.GetTypeIndex_NoLock<FTransform>(false))
+	if (ValueTypeIndex == InRegistry->GetTypeIndex_NoLock<FTransform>(false))
 	{
 		return &FRigDispatch_SetAnimAttribute::SetAnimAttributeDispatch<FTransform>;
 	}
-	if (ValueTypeIndex == Registry.GetTypeIndex_NoLock<FQuat>(false))
+	if (ValueTypeIndex == InRegistry->GetTypeIndex_NoLock<FQuat>(false))
 	{
 		return &FRigDispatch_SetAnimAttribute::SetAnimAttributeDispatch<FQuat>;
 	}
-	if (ValueTypeIndex == Registry.GetTypeIndex_NoLock<FVector>(false))
+	if (ValueTypeIndex == InRegistry->GetTypeIndex_NoLock<FVector>(false))
 	{
 		return &FRigDispatch_SetAnimAttribute::SetAnimAttributeDispatch<FVector>;
 	}
 	
-	const FRigVMTemplateArgumentType& ValueType = Registry.GetType_NoLock(ValueTypeIndex);
+	const FRigVMTemplateArgumentType& ValueType = InRegistry->GetType_NoLock(ValueTypeIndex);
 	
-	if (UScriptStruct* ScriptStruct = Cast<UScriptStruct>(ValueType.CPPTypeObject))
+	if (Cast<UScriptStruct>(ValueType.CPPTypeObject))
 	{
 		return &FRigDispatch_SetAnimAttribute::SetAnimAttributeDispatch;
 	}
diff --git a/Engine/Plugins/Animation/ControlRig/Source/ControlRig/Private/Units/Core/RigUnit_UserData.cpp b/Engine/Plugins/Animation/ControlRig/Source/ControlRig/Private/Units/Core/RigUnit_UserData.cpp
--- a/Engine/Plugins/Animation/ControlRig/Source/ControlRig/Private/Units/Core/RigUnit_UserData.cpp
+++ b/Engine/Plugins/Animation/ControlRig/Source/ControlRig/Private/Units/Core/RigUnit_UserData.cpp
@@ -28,28 +28,27 @@ FString FRigDispatch_GetUserData::GetArgumentMetaData(const FName& InArgumentNam
 
 #endif
 
-const TArray<FRigVMTemplateArgumentInfo>& FRigDispatch_GetUserData::GetArgumentInfos() const
+const TArray<FRigVMTemplateArgumentInfo>& FRigDispatch_GetUserData::GetArgumentInfos(FRigVMRegistryHandle& InRegistry) const
 {
-	static TArray<FRigVMTemplateArgumentInfo> Infos;
-	if (Infos.IsEmpty())
+	if (CachedArgumentInfos.IsEmpty())
 	{
 		static const TArray<FRigVMTemplateArgument::ETypeCategory> ValueCategories = {
 			FRigVMTemplateArgument::ETypeCategory_SingleAnyValue,
 			FRigVMTemplateArgument::ETypeCategory_ArrayAnyValue
 		};
 		
-		Infos.Emplace(ArgNameSpaceName, ERigVMPinDirection::Input, RigVMTypeUtils::TypeIndex::FString);
-		Infos.Emplace(ArgPathName, ERigVMPinDirection::Input, RigVMTypeUtils::TypeIndex::FString);
-		Infos.Emplace(ArgDefaultName, ERigVMPinDirection::Input, ValueCategories);
-		Infos.Emplace(ArgResultName, ERigVMPinDirection::Output, ValueCategories);
-		Infos.Emplace(ArgFoundName, ERigVMPinDirection::Output, RigVMTypeUtils::TypeIndex::Bool);
+		CachedArgumentInfos.Emplace(ArgNameSpaceName, ERigVMPinDirection::Input, RigVMTypeUtils::TypeIndex::FString);
+		CachedArgumentInfos.Emplace(ArgPathName, ERigVMPinDirection::Input, RigVMTypeUtils::TypeIndex::FString);
+		CachedArgumentInfos.Emplace(ArgDefaultName, ERigVMPinDirection::Input, ValueCategories);
+		CachedArgumentInfos.Emplace(ArgResultName, ERigVMPinDirection::Output, ValueCategories);
+		CachedArgumentInfos.Emplace(ArgFoundName, ERigVMPinDirection::Output, RigVMTypeUtils::TypeIndex::Bool);
 	}
 	
-	return Infos;
+	return CachedArgumentInfos;
 }
 
 FRigVMTemplateTypeMap FRigDispatch_GetUserData::OnNewArgumentType(const FName& InArgumentName,
-	TRigVMTypeIndex InTypeIndex) const
+	TRigVMTypeIndex InTypeIndex, FRigVMRegistryHandle& InRegistry) const
 {
 	FRigVMTemplateTypeMap Types;
 	Types.Add(ArgNameSpaceName, RigVMTypeUtils::TypeIndex::FString);
@@ -60,12 +59,12 @@ FRigVMTemplateTypeMap FRigDispatch_GetUserData::OnNewArgumentType(const FName& I
 	return Types;
 }
 
-FRigVMFunctionPtr FRigDispatch_GetUserData::GetDispatchFunctionImpl(const FRigVMTemplateTypeMap& InTypes) const
+FRigVMFunctionPtr FRigDispatch_GetUserData::GetDispatchFunctionImpl(const FRigVMTemplateTypeMap& InTypes, const FRigVMRegistryHandle& InRegistry) const
 {
 	check(InTypes.FindChecked(ArgNameSpaceName) == RigVMTypeUtils::TypeIndex::FString);
 	check(InTypes.FindChecked(ArgPathName) == RigVMTypeUtils::TypeIndex::FString);
 	const TRigVMTypeIndex TypeIndex = InTypes.FindChecked(ArgDefaultName);
-	check(FRigVMRegistry_NoLock::GetForRead().CanMatchTypes_NoLock(TypeIndex, InTypes.FindChecked(ArgResultName), true));
+	check(InRegistry->CanMatchTypes_NoLock(TypeIndex, InTypes.FindChecked(ArgResultName), true));
 	check(InTypes.FindChecked(ArgFoundName) == RigVMTypeUtils::TypeIndex::Bool);
 	return &FRigDispatch_GetUserData::Execute;
 }
diff --git a/Engine/Plugins/Animation/ControlRig/Source/ControlRig/Private/Units/Hierarchy/RigUnit_Component.cpp b/Engine/Plugins/Animation/ControlRig/Source/ControlRig/Private/Units/Hierarchy/RigUnit_Component.cpp
--- a/Engine/Plugins/Animation/ControlRig/Source/ControlRig/Private/Units/Hierarchy/RigUnit_Component.cpp
+++ b/Engine/Plugins/Animation/ControlRig/Source/ControlRig/Private/Units/Hierarchy/RigUnit_Component.cpp
@@ -13,10 +13,9 @@
 // FRigDispatch_ComponentBase
 //////////////////////////////////////////////////////////////////////////////////////////////////////
 
-bool FRigDispatch_ComponentBase::IsTypeSupported(const TRigVMTypeIndex& InTypeIndex)
+bool FRigDispatch_ComponentBase::IsTypeSupported(const TRigVMTypeIndex& InTypeIndex, const FRigVMRegistryHandle& InRegistry)
 {
-	const FRigVMRegistry_NoLock& Registry = FRigVMRegistry_NoLock::GetForRead();
-	const FRigVMTemplateArgumentType& InType = Registry.GetType_NoLock(InTypeIndex);
+	const FRigVMTemplateArgumentType& InType = InRegistry->GetType_NoLock(InTypeIndex);
 
 	if (IsValid(InType.CPPTypeObject))
 	{
@@ -29,33 +28,30 @@ bool FRigDispatch_ComponentBase::IsTypeSupported(const TRigVMTypeIndex& InTypeIn
 	return false;
 }
 
-const TRigVMTypeIndex& FRigDispatch_ComponentBase::GetElementKeyType()
+const TRigVMTypeIndex& FRigDispatch_ComponentBase::GetElementKeyType(FRigVMRegistryHandle& InRegistry)
 {
-	static FRigVMRegistry_NoLock& Registry = FRigVMRegistry_NoLock::GetForWrite();
-	static TRigVMTypeIndex RigElementKeyType = Registry.FindOrAddType_NoLock(FRigVMTemplateArgumentType(FRigElementKey::StaticStruct()), false);
+	static TRigVMTypeIndex RigElementKeyType = InRegistry->FindOrAddType_NoLock(FRigVMTemplateArgumentType(FRigElementKey::StaticStruct()), false);
 	return RigElementKeyType;
 }
 
-const TRigVMTypeIndex& FRigDispatch_ComponentBase::GetComponentKeyType()
+const TRigVMTypeIndex& FRigDispatch_ComponentBase::GetComponentKeyType(FRigVMRegistryHandle& InRegistry)
 {
-	static FRigVMRegistry_NoLock& Registry = FRigVMRegistry_NoLock::GetForWrite();
-	static TRigVMTypeIndex RigComponentKeyType = Registry.FindOrAddType_NoLock(FRigVMTemplateArgumentType(FRigComponentKey::StaticStruct()), false);
+	static TRigVMTypeIndex RigComponentKeyType = InRegistry->FindOrAddType_NoLock(FRigVMTemplateArgumentType(FRigComponentKey::StaticStruct()), false);
 	return RigComponentKeyType;
 }
 
-void FRigDispatch_ComponentBase::RegisterDependencyTypes_NoLock(FRigVMRegistry_NoLock& InRegistry) const
+void FRigDispatch_ComponentBase::RegisterDependencyTypes_NoLock(FRigVMRegistryHandle& InRegistry) const
 {
 	Super::RegisterDependencyTypes_NoLock(InRegistry);
 
-	FRigVMRegistry_NoLock& Registry = FRigVMRegistry_NoLock::GetForWrite();
-	(void)Registry.FindOrAddType_NoLock(FRigVMTemplateArgumentType(FRigComponentKey::StaticStruct()), true);
+	(void)InRegistry->FindOrAddType_NoLock(FRigVMTemplateArgumentType(FRigComponentKey::StaticStruct()), true);
 	
 	for (TObjectIterator<UScriptStruct> ScriptStructIt; ScriptStructIt; ++ScriptStructIt)
 	{
 		UScriptStruct* ScriptStruct = *ScriptStructIt;
 		if(ScriptStruct != FRigBaseComponent::StaticStruct() && ScriptStruct->IsChildOf(FRigBaseComponent::StaticStruct()))
 		{
-			(void)Registry.FindOrAddType_NoLock(FRigVMTemplateArgumentType(ScriptStruct), true);
+			(void)InRegistry->FindOrAddType_NoLock(FRigVMTemplateArgumentType(ScriptStruct), true);
 		}
 	}
 }
@@ -118,18 +114,18 @@ FString FRigDispatch_ComponentBase::GetArgumentDefaultValue(const FName& InArgum
 // FRigDispatch_SpawnComponent
 //////////////////////////////////////////////////////////////////////////////////////////////////////
 
-const TArray<FRigVMTemplateArgumentInfo>& FRigDispatch_SpawnComponent::GetArgumentInfos() const
+const TArray<FRigVMTemplateArgumentInfo>& FRigDispatch_SpawnComponent::GetArgumentInfos(FRigVMRegistryHandle& InRegistry) const
 {
-	if (Infos.IsEmpty())
+	if (CachedArgumentInfos.IsEmpty())
 	{
-		ItemArgIndex = Infos.Emplace(ItemArgName, ERigVMPinDirection::Input, GetElementKeyType());
-		NameArgIndex = Infos.Emplace(NameArgName, ERigVMPinDirection::Input, RigVMTypeUtils::TypeIndex::FName);
-		ComponentArgIndex = Infos.Add(FRigVMTemplateArgumentInfo(ComponentArgName, ERigVMPinDirection::Input, {FRigVMTemplateArgument::ETypeCategory_SingleScriptStructValue}, &FRigDispatch_SpawnComponent::IsTypeSupported));
-		KeyArgIndex = Infos.Emplace(KeyArgName, ERigVMPinDirection::Output, GetComponentKeyType());
-		SuccessArgIndex = Infos.Emplace(SuccessArgName, ERigVMPinDirection::Output, RigVMTypeUtils::TypeIndex::Bool);
+		ItemArgIndex = CachedArgumentInfos.Emplace(ItemArgName, ERigVMPinDirection::Input, GetElementKeyType(InRegistry));
+		NameArgIndex = CachedArgumentInfos.Emplace(NameArgName, ERigVMPinDirection::Input, RigVMTypeUtils::TypeIndex::FName);
+		ComponentArgIndex = CachedArgumentInfos.Add(FRigVMTemplateArgumentInfo(ComponentArgName, ERigVMPinDirection::Input, {FRigVMTemplateArgument::ETypeCategory_SingleScriptStructValue}, &FRigDispatch_SpawnComponent::IsTypeSupported));
+		KeyArgIndex = CachedArgumentInfos.Emplace(KeyArgName, ERigVMPinDirection::Output, GetComponentKeyType(InRegistry));
+		SuccessArgIndex = CachedArgumentInfos.Emplace(SuccessArgName, ERigVMPinDirection::Output, RigVMTypeUtils::TypeIndex::Bool);
 	}
 	
-	return Infos;
+	return CachedArgumentInfos;
 }
 
 #if WITH_EDITOR
@@ -194,16 +190,16 @@ void FRigDispatch_SpawnComponent::Execute(FRigVMExtendedExecuteContext& InContex
 // FRigDispatch_GetComponentContent
 //////////////////////////////////////////////////////////////////////////////////////////////////////
 
-const TArray<FRigVMTemplateArgumentInfo>& FRigDispatch_GetComponentContent::GetArgumentInfos() const
+const TArray<FRigVMTemplateArgumentInfo>& FRigDispatch_GetComponentContent::GetArgumentInfos(FRigVMRegistryHandle& InRegistry) const
 {
-	if (Infos.IsEmpty())
+	if (CachedArgumentInfos.IsEmpty())
 	{
-		KeyArgIndex = Infos.Emplace(KeyArgName, ERigVMPinDirection::Input, GetComponentKeyType());
-		ComponentArgIndex = Infos.Add(FRigVMTemplateArgumentInfo(ComponentArgName, ERigVMPinDirection::Output, {FRigVMTemplateArgument::ETypeCategory_SingleScriptStructValue}, &FRigDispatch_GetComponentContent::IsTypeSupported));
-		SuccessArgIndex = Infos.Emplace(SuccessArgName, ERigVMPinDirection::Output, RigVMTypeUtils::TypeIndex::Bool);
+		KeyArgIndex = CachedArgumentInfos.Emplace(KeyArgName, ERigVMPinDirection::Input, GetComponentKeyType(InRegistry));
+		ComponentArgIndex = CachedArgumentInfos.Add(FRigVMTemplateArgumentInfo(ComponentArgName, ERigVMPinDirection::Output, {FRigVMTemplateArgument::ETypeCategory_SingleScriptStructValue}, &FRigDispatch_GetComponentContent::IsTypeSupported));
+		SuccessArgIndex = CachedArgumentInfos.Emplace(SuccessArgName, ERigVMPinDirection::Output, RigVMTypeUtils::TypeIndex::Bool);
 	}
 	
-	return Infos;
+	return CachedArgumentInfos;
 }
 
 void FRigDispatch_GetComponentContent::Execute(FRigVMExtendedExecuteContext& InContext, FRigVMMemoryHandleArray Handles, FRigVMPredicateBranchArray)
@@ -250,16 +246,16 @@ void FRigDispatch_GetComponentContent::Execute(FRigVMExtendedExecuteContext& InC
 // FRigDispatch_SetComponentContent
 //////////////////////////////////////////////////////////////////////////////////////////////////////
 
-const TArray<FRigVMTemplateArgumentInfo>& FRigDispatch_SetComponentContent::GetArgumentInfos() const
+const TArray<FRigVMTemplateArgumentInfo>& FRigDispatch_SetComponentContent::GetArgumentInfos(FRigVMRegistryHandle& InRegistry) const
 {
-	if (Infos.IsEmpty())
+	if (CachedArgumentInfos.IsEmpty())
 	{
-		KeyArgIndex = Infos.Emplace(KeyArgName, ERigVMPinDirection::Input, GetComponentKeyType());
-		ComponentArgIndex = Infos.Add(FRigVMTemplateArgumentInfo(ComponentArgName, ERigVMPinDirection::Input, {FRigVMTemplateArgument::ETypeCategory_SingleScriptStructValue}, &FRigDispatch_SetComponentContent::IsTypeSupported));
-		SuccessArgIndex = Infos.Emplace(SuccessArgName, ERigVMPinDirection::Output, RigVMTypeUtils::TypeIndex::Bool);
+		KeyArgIndex = CachedArgumentInfos.Emplace(KeyArgName, ERigVMPinDirection::Input, GetComponentKeyType(InRegistry));
+		ComponentArgIndex = CachedArgumentInfos.Add(FRigVMTemplateArgumentInfo(ComponentArgName, ERigVMPinDirection::Input, {FRigVMTemplateArgument::ETypeCategory_SingleScriptStructValue}, &FRigDispatch_SetComponentContent::IsTypeSupported));
+		SuccessArgIndex = CachedArgumentInfos.Emplace(SuccessArgName, ERigVMPinDirection::Output, RigVMTypeUtils::TypeIndex::Bool);
 	}
 	
-	return Infos;
+	return CachedArgumentInfos;
 }
 
 void FRigDispatch_SetComponentContent::Execute(FRigVMExtendedExecuteContext& InContext, FRigVMMemoryHandleArray Handles, FRigVMPredicateBranchArray)
diff --git a/Engine/Plugins/Animation/ControlRig/Source/ControlRig/Private/Units/Hierarchy/RigUnit_Metadata.cpp b/Engine/Plugins/Animation/ControlRig/Source/ControlRig/Private/Units/Hierarchy/RigUnit_Metadata.cpp
--- a/Engine/Plugins/Animation/ControlRig/Source/ControlRig/Private/Units/Hierarchy/RigUnit_Metadata.cpp
+++ b/Engine/Plugins/Animation/ControlRig/Source/ControlRig/Private/Units/Hierarchy/RigUnit_Metadata.cpp
@@ -56,17 +56,16 @@ FString FRigDispatch_MetadataBase::GetNodeTitle(const FRigVMTemplateTypeMap& InT
 
 #endif
 
-const TArray<FRigVMTemplateArgumentInfo>& FRigDispatch_MetadataBase::GetArgumentInfos() const
+const TArray<FRigVMTemplateArgumentInfo>& FRigDispatch_MetadataBase::GetArgumentInfos(FRigVMRegistryHandle& InRegistry) const
 {
-	if(Infos.IsEmpty())
+	if(CachedArgumentInfos.IsEmpty())
 	{
-		const FRigVMRegistry_NoLock& Registry = FRigVMRegistry_NoLock::GetForRead();
-		ItemArgIndex = Infos.Emplace(ItemArgName, ERigVMPinDirection::Input, Registry.GetTypeIndex_NoLock<FRigElementKey>());
-		NameArgIndex = Infos.Emplace(NameArgName, ERigVMPinDirection::Input, RigVMTypeUtils::TypeIndex::FName);
-		NameSpaceArgIndex = Infos.Emplace(NameSpaceArgName, ERigVMPinDirection::Input, Registry.GetTypeIndex_NoLock<ERigMetaDataNameSpace>());
-		CacheArgIndex = Infos.Emplace(CacheArgName, ERigVMPinDirection::Hidden, Registry.GetTypeIndex_NoLock<FCachedRigElement>());
+		ItemArgIndex = CachedArgumentInfos.Emplace(ItemArgName, ERigVMPinDirection::Input, InRegistry->GetTypeIndex_NoLock<FRigElementKey>());
+		NameArgIndex = CachedArgumentInfos.Emplace(NameArgName, ERigVMPinDirection::Input, RigVMTypeUtils::TypeIndex::FName);
+		NameSpaceArgIndex = CachedArgumentInfos.Emplace(NameSpaceArgName, ERigVMPinDirection::Input, InRegistry->GetTypeIndex_NoLock<ERigMetaDataNameSpace>());
+		CacheArgIndex = CachedArgumentInfos.Emplace(CacheArgName, ERigVMPinDirection::Hidden, InRegistry->GetTypeIndex_NoLock<FCachedRigElement>());
 	};
-	return Infos;
+	return CachedArgumentInfos;
 }
 
 #if WITH_EDITOR
@@ -129,50 +128,51 @@ FString FRigDispatch_MetadataBase::GetArgumentMetaData(const FName& InArgumentNa
 
 #endif
 
-const TArray<TRigVMTypeIndex>& FRigDispatch_MetadataBase::GetValueTypes() const
+const TArray<TRigVMTypeIndex>& FRigDispatch_MetadataBase::GetValueTypes(const FRigVMRegistryHandle& InRegistry) const
 {
-	static TArray<TRigVMTypeIndex> Types;
-	if(Types.IsEmpty())
+	if(ValueTypes.IsEmpty())
 	{
-		const FRigVMRegistry_NoLock& Registry = FRigVMRegistry_NoLock::GetForRead();
-		Types = {
+		ValueTypes = {
 			RigVMTypeUtils::TypeIndex::Bool,
 			RigVMTypeUtils::TypeIndex::Float,
 			RigVMTypeUtils::TypeIndex::Int32,
 			RigVMTypeUtils::TypeIndex::FName,
-			Registry.GetTypeIndex_NoLock<FVector>(false),
-			Registry.GetTypeIndex_NoLock<FRotator>(false),
-			Registry.GetTypeIndex_NoLock<FQuat>(false),
-			Registry.GetTypeIndex_NoLock<FTransform>(false),
-			Registry.GetTypeIndex_NoLock<FLinearColor>(false),
-			Registry.GetTypeIndex_NoLock<FRigElementKey>(false),
+			InRegistry->GetTypeIndex_NoLock<FVector>(false),
+			InRegistry->GetTypeIndex_NoLock<FRotator>(false),
+			InRegistry->GetTypeIndex_NoLock<FQuat>(false),
+			InRegistry->GetTypeIndex_NoLock<FTransform>(false),
+			InRegistry->GetTypeIndex_NoLock<FLinearColor>(false),
+			InRegistry->GetTypeIndex_NoLock<FRigElementKey>(false),
 			RigVMTypeUtils::TypeIndex::BoolArray,
 			RigVMTypeUtils::TypeIndex::FloatArray,
 			RigVMTypeUtils::TypeIndex::Int32Array,
 			RigVMTypeUtils::TypeIndex::FNameArray,
-			Registry.GetTypeIndex_NoLock<FVector>(true),
-			Registry.GetTypeIndex_NoLock<FRotator>(true),
-			Registry.GetTypeIndex_NoLock<FQuat>(true),
-			Registry.GetTypeIndex_NoLock<FTransform>(true),
-			Registry.GetTypeIndex_NoLock<FLinearColor>(true),
-			Registry.GetTypeIndex_NoLock<FRigElementKey>(true)
+			InRegistry->GetTypeIndex_NoLock<FVector>(true),
+			InRegistry->GetTypeIndex_NoLock<FRotator>(true),
+			InRegistry->GetTypeIndex_NoLock<FQuat>(true),
+			InRegistry->GetTypeIndex_NoLock<FTransform>(true),
+			InRegistry->GetTypeIndex_NoLock<FLinearColor>(true),
+			InRegistry->GetTypeIndex_NoLock<FRigElementKey>(true)
 		};
+
+		// remove invalid indices
+		ValueTypes.Remove(TRigVMTypeIndex());
 	}
-	return Types;
+	return ValueTypes;
 }
 
 ////////////////////////////////////////////////////////////////////////////////////////////////
 
-const TArray<FRigVMTemplateArgumentInfo>& FRigDispatch_GetMetadata::GetArgumentInfos() const
+const TArray<FRigVMTemplateArgumentInfo>& FRigDispatch_GetMetadata::GetArgumentInfos(FRigVMRegistryHandle& InRegistry) const
 {
 	if(ValueArgIndex == INDEX_NONE)
 	{
-		Infos = Super::GetArgumentInfos(); 
-		DefaultArgIndex = Infos.Emplace(DefaultArgName, ERigVMPinDirection::Input, GetValueTypes());
-		ValueArgIndex = Infos.Emplace(ValueArgName, ERigVMPinDirection::Output, GetValueTypes());
-		FoundArgIndex = Infos.Emplace(FoundArgName, ERigVMPinDirection::Output, RigVMTypeUtils::TypeIndex::Bool);
+		CachedArgumentInfos = Super::GetArgumentInfos(InRegistry); 
+		DefaultArgIndex = CachedArgumentInfos.Emplace(DefaultArgName, ERigVMPinDirection::Input, GetValueTypes(InRegistry));
+		ValueArgIndex = CachedArgumentInfos.Emplace(ValueArgName, ERigVMPinDirection::Output, GetValueTypes(InRegistry));
+		FoundArgIndex = CachedArgumentInfos.Emplace(FoundArgName, ERigVMPinDirection::Output, RigVMTypeUtils::TypeIndex::Bool);
 	};
-	return Infos;
+	return CachedArgumentInfos;
 }
 
 FRigBaseMetadata* FRigDispatch_GetMetadata::FindMetadata(const FRigVMExtendedExecuteContext& InContext,
@@ -192,9 +192,8 @@ FRigBaseMetadata* FRigDispatch_GetMetadata::FindMetadata(const FRigVMExtendedExe
 	return nullptr;
 }
 
-FRigVMFunctionPtr FRigDispatch_GetMetadata::GetDispatchFunctionImpl(const FRigVMTemplateTypeMap& InTypes) const
+FRigVMFunctionPtr FRigDispatch_GetMetadata::GetDispatchFunctionImpl(const FRigVMTemplateTypeMap& InTypes, const FRigVMRegistryHandle& InRegistry) const
 {
-	const FRigVMRegistry_NoLock& Registry = FRigVMRegistry_NoLock::GetForRead();
 	const TRigVMTypeIndex& ValueTypeIndex = InTypes.FindChecked(TEXT("Value"));
 	
 	if(ValueTypeIndex == RigVMTypeUtils::TypeIndex::Bool)
@@ -213,27 +212,27 @@ FRigVMFunctionPtr FRigDispatch_GetMetadata::GetDispatchFunctionImpl(const FRigVM
 	{
 		return &FRigDispatch_GetMetadata::GetMetadataDispatch<FName, FRigNameMetadata, ERigMetadataType::Name>;
 	}
-	if(ValueTypeIndex == Registry.GetTypeIndex_NoLock<FVector>(false))
+	if(ValueTypeIndex == InRegistry->GetTypeIndex_NoLock<FVector>(false))
 	{
 		return &FRigDispatch_GetMetadata::GetMetadataDispatch<FVector, FRigVectorMetadata, ERigMetadataType::Vector>;
 	}
-	if(ValueTypeIndex == Registry.GetTypeIndex_NoLock<FRotator>(false))
+	if(ValueTypeIndex == InRegistry->GetTypeIndex_NoLock<FRotator>(false))
 	{
 		return &FRigDispatch_GetMetadata::GetMetadataDispatch<FRotator, FRigRotatorMetadata, ERigMetadataType::Rotator>;
 	}
-	if(ValueTypeIndex == Registry.GetTypeIndex_NoLock<FQuat>(false))
+	if(ValueTypeIndex == InRegistry->GetTypeIndex_NoLock<FQuat>(false))
 	{
 		return &FRigDispatch_GetMetadata::GetMetadataDispatch<FQuat, FRigQuatMetadata, ERigMetadataType::Quat>;
 	}
-	if(ValueTypeIndex == Registry.GetTypeIndex_NoLock<FTransform>(false))
+	if(ValueTypeIndex == InRegistry->GetTypeIndex_NoLock<FTransform>(false))
 	{
 		return &FRigDispatch_GetMetadata::GetMetadataDispatch<FTransform, FRigTransformMetadata, ERigMetadataType::Transform>;
 	}
-	if(ValueTypeIndex == Registry.GetTypeIndex_NoLock<FLinearColor>(false))
+	if(ValueTypeIndex == InRegistry->GetTypeIndex_NoLock<FLinearColor>(false))
 	{
 		return &FRigDispatch_GetMetadata::GetMetadataDispatch<FLinearColor, FRigLinearColorMetadata, ERigMetadataType::LinearColor>;
 	}
-	if(ValueTypeIndex == Registry.GetTypeIndex_NoLock<FRigElementKey>(false))
+	if(ValueTypeIndex == InRegistry->GetTypeIndex_NoLock<FRigElementKey>(false))
 	{
 		return &FRigDispatch_GetMetadata::GetMetadataDispatch<FRigElementKey, FRigElementKeyMetadata, ERigMetadataType::RigElementKey>;
 	}
@@ -253,27 +252,27 @@ FRigVMFunctionPtr FRigDispatch_GetMetadata::GetDispatchFunctionImpl(const FRigVM
 	{
 		return &FRigDispatch_GetMetadata::GetMetadataDispatch<TArray<FName>, FRigNameArrayMetadata, ERigMetadataType::NameArray>;
 	}
-	if(ValueTypeIndex == Registry.GetTypeIndex_NoLock<FVector>(true))
+	if(ValueTypeIndex == InRegistry->GetTypeIndex_NoLock<FVector>(true))
 	{
 		return &FRigDispatch_GetMetadata::GetMetadataDispatch<TArray<FVector>, FRigVectorArrayMetadata, ERigMetadataType::VectorArray>;
 	}
-	if(ValueTypeIndex == Registry.GetTypeIndex_NoLock<FRotator>(true))
+	if(ValueTypeIndex == InRegistry->GetTypeIndex_NoLock<FRotator>(true))
 	{
 		return &FRigDispatch_GetMetadata::GetMetadataDispatch<TArray<FRotator>, FRigRotatorArrayMetadata, ERigMetadataType::RotatorArray>;
 	}
-	if(ValueTypeIndex == Registry.GetTypeIndex_NoLock<FQuat>(true))
+	if(ValueTypeIndex == InRegistry->GetTypeIndex_NoLock<FQuat>(true))
 	{
 		return &FRigDispatch_GetMetadata::GetMetadataDispatch<TArray<FQuat>, FRigQuatArrayMetadata, ERigMetadataType::QuatArray>;
 	}
-	if(ValueTypeIndex == Registry.GetTypeIndex_NoLock<FTransform>(true))
+	if(ValueTypeIndex == InRegistry->GetTypeIndex_NoLock<FTransform>(true))
 	{
 		return &FRigDispatch_GetMetadata::GetMetadataDispatch<TArray<FTransform>, FRigTransformArrayMetadata, ERigMetadataType::TransformArray>;
 	}
-	if(ValueTypeIndex == Registry.GetTypeIndex_NoLock<FLinearColor>(true))
+	if(ValueTypeIndex == InRegistry->GetTypeIndex_NoLock<FLinearColor>(true))
 	{
 		return &FRigDispatch_GetMetadata::GetMetadataDispatch<TArray<FLinearColor>, FRigLinearColorArrayMetadata, ERigMetadataType::LinearColorArray>;
 	}
-	if(ValueTypeIndex == Registry.GetTypeIndex_NoLock<FRigElementKey>(true))
+	if(ValueTypeIndex == InRegistry->GetTypeIndex_NoLock<FRigElementKey>(true))
 	{
 		return &FRigDispatch_GetMetadata::GetMetadataDispatch<TArray<FRigElementKey>, FRigElementKeyArrayMetadata, ERigMetadataType::RigElementKeyArray>;
 	}
@@ -283,15 +282,15 @@ FRigVMFunctionPtr FRigDispatch_GetMetadata::GetDispatchFunctionImpl(const FRigVM
 
 ////////////////////////////////////////////////////////////////////////////////////////////////
 
-const TArray<FRigVMTemplateArgumentInfo>& FRigDispatch_SetMetadata::GetArgumentInfos() const
+const TArray<FRigVMTemplateArgumentInfo>& FRigDispatch_SetMetadata::GetArgumentInfos(FRigVMRegistryHandle& InRegistry) const
 {
 	if(ValueArgIndex == INDEX_NONE)
 	{
-		Infos = Super::GetArgumentInfos(); 
-		ValueArgIndex = Infos.Emplace(ValueArgName, ERigVMPinDirection::Input, GetValueTypes());
-		SuccessArgIndex = Infos.Emplace(SuccessArgName, ERigVMPinDirection::Output, RigVMTypeUtils::TypeIndex::Bool);
+		CachedArgumentInfos = Super::GetArgumentInfos(InRegistry); 
+		ValueArgIndex = CachedArgumentInfos.Emplace(ValueArgName, ERigVMPinDirection::Input, GetValueTypes(InRegistry));
+		SuccessArgIndex = CachedArgumentInfos.Emplace(SuccessArgName, ERigVMPinDirection::Output, RigVMTypeUtils::TypeIndex::Bool);
 	};
-	return Infos;
+	return CachedArgumentInfos;
 }
 
 const TArray<FRigVMExecuteArgument>& FRigDispatch_SetMetadata::GetExecuteArguments_Impl(const FRigVMDispatchContext& InContext) const
@@ -318,9 +317,8 @@ FRigBaseMetadata* FRigDispatch_SetMetadata::FindOrAddMetadata(const FControlRigE
 	return nullptr;
 }
 
-FRigVMFunctionPtr FRigDispatch_SetMetadata::GetDispatchFunctionImpl(const FRigVMTemplateTypeMap& InTypes) const
+FRigVMFunctionPtr FRigDispatch_SetMetadata::GetDispatchFunctionImpl(const FRigVMTemplateTypeMap& InTypes, const FRigVMRegistryHandle& InRegistry) const
 {
-	const FRigVMRegistry_NoLock& Registry = FRigVMRegistry_NoLock::GetForRead();
 	const TRigVMTypeIndex& ValueTypeIndex = InTypes.FindChecked(TEXT("Value"));
 	
 	if(ValueTypeIndex == RigVMTypeUtils::TypeIndex::Bool)
@@ -339,27 +337,27 @@ FRigVMFunctionPtr FRigDispatch_SetMetadata::GetDispatchFunctionImpl(const FRigVM
 	{
 		return &FRigDispatch_SetMetadata::SetMetadataDispatch<FName, FRigNameMetadata, ERigMetadataType::Name>;
 	}
-	if(ValueTypeIndex == Registry.GetTypeIndex_NoLock<FVector>(false))
+	if(ValueTypeIndex == InRegistry->GetTypeIndex_NoLock<FVector>(false))
 	{
 		return &FRigDispatch_SetMetadata::SetMetadataDispatch<FVector, FRigVectorMetadata, ERigMetadataType::Vector>;
 	}
-	if(ValueTypeIndex == Registry.GetTypeIndex_NoLock<FRotator>(false))
+	if(ValueTypeIndex == InRegistry->GetTypeIndex_NoLock<FRotator>(false))
 	{
 		return &FRigDispatch_SetMetadata::SetMetadataDispatch<FRotator, FRigRotatorMetadata, ERigMetadataType::Rotator>;
 	}
-	if(ValueTypeIndex == Registry.GetTypeIndex_NoLock<FQuat>(false))
+	if(ValueTypeIndex == InRegistry->GetTypeIndex_NoLock<FQuat>(false))
 	{
 		return &FRigDispatch_SetMetadata::SetMetadataDispatch<FQuat, FRigQuatMetadata, ERigMetadataType::Quat>;
 	}
-	if(ValueTypeIndex == Registry.GetTypeIndex_NoLock<FTransform>(false))
+	if(ValueTypeIndex == InRegistry->GetTypeIndex_NoLock<FTransform>(false))
 	{
 		return &FRigDispatch_SetMetadata::SetMetadataDispatch<FTransform, FRigTransformMetadata, ERigMetadataType::Transform>;
 	}
-	if(ValueTypeIndex == Registry.GetTypeIndex_NoLock<FLinearColor>(false))
+	if(ValueTypeIndex == InRegistry->GetTypeIndex_NoLock<FLinearColor>(false))
 	{
 		return &FRigDispatch_SetMetadata::SetMetadataDispatch<FLinearColor, FRigLinearColorMetadata, ERigMetadataType::LinearColor>;
 	}
-	if(ValueTypeIndex == Registry.GetTypeIndex_NoLock<FRigElementKey>(false))
+	if(ValueTypeIndex == InRegistry->GetTypeIndex_NoLock<FRigElementKey>(false))
 	{
 		return &FRigDispatch_SetMetadata::SetMetadataDispatch<FRigElementKey, FRigElementKeyMetadata, ERigMetadataType::RigElementKey>;
 	}
@@ -379,27 +377,27 @@ FRigVMFunctionPtr FRigDispatch_SetMetadata::GetDispatchFunctionImpl(const FRigVM
 	{
 		return &FRigDispatch_SetMetadata::SetMetadataDispatch<TArray<FName>, FRigNameArrayMetadata, ERigMetadataType::NameArray>;
 	}
-	if(ValueTypeIndex == Registry.GetTypeIndex_NoLock<FVector>(true))
+	if(ValueTypeIndex == InRegistry->GetTypeIndex_NoLock<FVector>(true))
 	{
 		return &FRigDispatch_SetMetadata::SetMetadataDispatch<TArray<FVector>, FRigVectorArrayMetadata, ERigMetadataType::VectorArray>;
 	}
-	if(ValueTypeIndex == Registry.GetTypeIndex_NoLock<FRotator>(true))
+	if(ValueTypeIndex == InRegistry->GetTypeIndex_NoLock<FRotator>(true))
 	{
 		return &FRigDispatch_SetMetadata::SetMetadataDispatch<TArray<FRotator>, FRigRotatorArrayMetadata, ERigMetadataType::RotatorArray>;
 	}
-	if(ValueTypeIndex == Registry.GetTypeIndex_NoLock<FQuat>(true))
+	if(ValueTypeIndex == InRegistry->GetTypeIndex_NoLock<FQuat>(true))
 	{
 		return &FRigDispatch_SetMetadata::SetMetadataDispatch<TArray<FQuat>, FRigQuatArrayMetadata, ERigMetadataType::QuatArray>;
 	}
-	if(ValueTypeIndex == Registry.GetTypeIndex_NoLock<FTransform>(true))
+	if(ValueTypeIndex == InRegistry->GetTypeIndex_NoLock<FTransform>(true))
 	{
 		return &FRigDispatch_SetMetadata::SetMetadataDispatch<TArray<FTransform>, FRigTransformArrayMetadata, ERigMetadataType::TransformArray>;
 	}
-	if(ValueTypeIndex == Registry.GetTypeIndex_NoLock<FLinearColor>(true))
+	if(ValueTypeIndex == InRegistry->GetTypeIndex_NoLock<FLinearColor>(true))
 	{
 		return &FRigDispatch_SetMetadata::SetMetadataDispatch<TArray<FLinearColor>, FRigLinearColorArrayMetadata, ERigMetadataType::LinearColorArray>;
 	}
-	if(ValueTypeIndex == Registry.GetTypeIndex_NoLock<FRigElementKey>(true))
+	if(ValueTypeIndex == InRegistry->GetTypeIndex_NoLock<FRigElementKey>(true))
 	{
 		return &FRigDispatch_SetMetadata::SetMetadataDispatch<TArray<FRigElementKey>, FRigElementKeyArrayMetadata, ERigMetadataType::RigElementKeyArray>;
 	}
@@ -825,17 +823,17 @@ FRigUnit_FilterItemsByMetadataTags_Execute()
 	}
 }
 
-const TArray<FRigVMTemplateArgumentInfo>& FRigDispatch_GetModuleMetadata::GetArgumentInfos() const
+const TArray<FRigVMTemplateArgumentInfo>& FRigDispatch_GetModuleMetadata::GetArgumentInfos(FRigVMRegistryHandle& InRegistry) const
 {
 	if(ValueArgIndex == INDEX_NONE)
 	{
-		NameArgIndex = Infos.Emplace(NameArgName, ERigVMPinDirection::Input, RigVMTypeUtils::TypeIndex::FName);
-		NameSpaceArgIndex = Infos.Emplace(NameSpaceArgName, ERigVMPinDirection::Input, FRigVMRegistry_NoLock::GetForRead().GetTypeIndex_NoLock<ERigMetaDataNameSpace>());
-		DefaultArgIndex = Infos.Emplace(DefaultArgName, ERigVMPinDirection::Input, GetValueTypes());
-		ValueArgIndex = Infos.Emplace(ValueArgName, ERigVMPinDirection::Output, GetValueTypes());
-		FoundArgIndex = Infos.Emplace(FoundArgName, ERigVMPinDirection::Output, RigVMTypeUtils::TypeIndex::Bool);
+		NameArgIndex = CachedArgumentInfos.Emplace(NameArgName, ERigVMPinDirection::Input, RigVMTypeUtils::TypeIndex::FName);
+		NameSpaceArgIndex = CachedArgumentInfos.Emplace(NameSpaceArgName, ERigVMPinDirection::Input, InRegistry->GetTypeIndex_NoLock<ERigMetaDataNameSpace>());
+		DefaultArgIndex = CachedArgumentInfos.Emplace(DefaultArgName, ERigVMPinDirection::Input, GetValueTypes(InRegistry));
+		ValueArgIndex = CachedArgumentInfos.Emplace(ValueArgName, ERigVMPinDirection::Output, GetValueTypes(InRegistry));
+		FoundArgIndex = CachedArgumentInfos.Emplace(FoundArgName, ERigVMPinDirection::Output, RigVMTypeUtils::TypeIndex::Bool);
 	};
-	return Infos;
+	return CachedArgumentInfos;
 }
 
 FRigBaseMetadata* FRigDispatch_GetModuleMetadata::FindMetadata(const FRigVMExtendedExecuteContext& InContext, const FName& InName, ERigMetadataType InType, ERigMetaDataNameSpace InNameSpace)
@@ -865,9 +863,8 @@ FRigBaseMetadata* FRigDispatch_GetModuleMetadata::FindMetadata(const FRigVMExten
 	return nullptr;
 }
 
-FRigVMFunctionPtr FRigDispatch_GetModuleMetadata::GetDispatchFunctionImpl(const FRigVMTemplateTypeMap& InTypes) const
+FRigVMFunctionPtr FRigDispatch_GetModuleMetadata::GetDispatchFunctionImpl(const FRigVMTemplateTypeMap& InTypes, const FRigVMRegistryHandle& InRegistry) const
 {
-	const FRigVMRegistry_NoLock& Registry = FRigVMRegistry_NoLock::GetForRead();
 	const TRigVMTypeIndex& ValueTypeIndex = InTypes.FindChecked(TEXT("Value"));
 	
 	if(ValueTypeIndex == RigVMTypeUtils::TypeIndex::Bool)
@@ -886,27 +883,27 @@ FRigVMFunctionPtr FRigDispatch_GetModuleMetadata::GetDispatchFunctionImpl(const
 	{
 		return &FRigDispatch_GetModuleMetadata::GetModuleMetadataDispatch<FName, FRigNameMetadata, ERigMetadataType::Name>;
 	}
-	if(ValueTypeIndex == Registry.GetTypeIndex_NoLock<FVector>(false))
+	if(ValueTypeIndex == InRegistry->GetTypeIndex_NoLock<FVector>(false))
 	{
 		return &FRigDispatch_GetModuleMetadata::GetModuleMetadataDispatch<FVector, FRigVectorMetadata, ERigMetadataType::Vector>;
 	}
-	if(ValueTypeIndex == Registry.GetTypeIndex_NoLock<FRotator>(false))
+	if(ValueTypeIndex == InRegistry->GetTypeIndex_NoLock<FRotator>(false))
 	{
 		return &FRigDispatch_GetModuleMetadata::GetModuleMetadataDispatch<FRotator, FRigRotatorMetadata, ERigMetadataType::Rotator>;
 	}
-	if(ValueTypeIndex == Registry.GetTypeIndex_NoLock<FQuat>(false))
+	if(ValueTypeIndex == InRegistry->GetTypeIndex_NoLock<FQuat>(false))
 	{
 		return &FRigDispatch_GetModuleMetadata::GetModuleMetadataDispatch<FQuat, FRigQuatMetadata, ERigMetadataType::Quat>;
 	}
-	if(ValueTypeIndex == Registry.GetTypeIndex_NoLock<FTransform>(false))
+	if(ValueTypeIndex == InRegistry->GetTypeIndex_NoLock<FTransform>(false))
 	{
 		return &FRigDispatch_GetModuleMetadata::GetModuleMetadataDispatch<FTransform, FRigTransformMetadata, ERigMetadataType::Transform>;
 	}
-	if(ValueTypeIndex == Registry.GetTypeIndex_NoLock<FLinearColor>(false))
+	if(ValueTypeIndex == InRegistry->GetTypeIndex_NoLock<FLinearColor>(false))
 	{
 		return &FRigDispatch_GetModuleMetadata::GetModuleMetadataDispatch<FLinearColor, FRigLinearColorMetadata, ERigMetadataType::LinearColor>;
 	}
-	if(ValueTypeIndex == Registry.GetTypeIndex_NoLock<FRigElementKey>(false))
+	if(ValueTypeIndex == InRegistry->GetTypeIndex_NoLock<FRigElementKey>(false))
 	{
 		return &FRigDispatch_GetModuleMetadata::GetModuleMetadataDispatch<FRigElementKey, FRigElementKeyMetadata, ERigMetadataType::RigElementKey>;
 	}
@@ -926,27 +923,27 @@ FRigVMFunctionPtr FRigDispatch_GetModuleMetadata::GetDispatchFunctionImpl(const
 	{
 		return &FRigDispatch_GetModuleMetadata::GetModuleMetadataDispatch<TArray<FName>, FRigNameArrayMetadata, ERigMetadataType::NameArray>;
 	}
-	if(ValueTypeIndex == Registry.GetTypeIndex_NoLock<FVector>(true))
+	if(ValueTypeIndex == InRegistry->GetTypeIndex_NoLock<FVector>(true))
 	{
 		return &FRigDispatch_GetModuleMetadata::GetModuleMetadataDispatch<TArray<FVector>, FRigVectorArrayMetadata, ERigMetadataType::VectorArray>;
 	}
-	if(ValueTypeIndex == Registry.GetTypeIndex_NoLock<FRotator>(true))
+	if(ValueTypeIndex == InRegistry->GetTypeIndex_NoLock<FRotator>(true))
 	{
 		return &FRigDispatch_GetModuleMetadata::GetModuleMetadataDispatch<TArray<FRotator>, FRigRotatorArrayMetadata, ERigMetadataType::RotatorArray>;
 	}
-	if(ValueTypeIndex == Registry.GetTypeIndex_NoLock<FQuat>(true))
+	if(ValueTypeIndex == InRegistry->GetTypeIndex_NoLock<FQuat>(true))
 	{
 		return &FRigDispatch_GetModuleMetadata::GetModuleMetadataDispatch<TArray<FQuat>, FRigQuatArrayMetadata, ERigMetadataType::QuatArray>;
 	}
-	if(ValueTypeIndex == Registry.GetTypeIndex_NoLock<FTransform>(true))
+	if(ValueTypeIndex == InRegistry->GetTypeIndex_NoLock<FTransform>(true))
 	{
 		return &FRigDispatch_GetModuleMetadata::GetModuleMetadataDispatch<TArray<FTransform>, FRigTransformArrayMetadata, ERigMetadataType::TransformArray>;
 	}
-	if(ValueTypeIndex == Registry.GetTypeIndex_NoLock<FLinearColor>(true))
+	if(ValueTypeIndex == InRegistry->GetTypeIndex_NoLock<FLinearColor>(true))
 	{
 		return &FRigDispatch_GetModuleMetadata::GetModuleMetadataDispatch<TArray<FLinearColor>, FRigLinearColorArrayMetadata, ERigMetadataType::LinearColorArray>;
 	}
-	if(ValueTypeIndex == Registry.GetTypeIndex_NoLock<FRigElementKey>(true))
+	if(ValueTypeIndex == InRegistry->GetTypeIndex_NoLock<FRigElementKey>(true))
 	{
 		return &FRigDispatch_GetModuleMetadata::GetModuleMetadataDispatch<TArray<FRigElementKey>, FRigElementKeyArrayMetadata, ERigMetadataType::RigElementKeyArray>;
 	}
@@ -954,16 +951,16 @@ FRigVMFunctionPtr FRigDispatch_GetModuleMetadata::GetDispatchFunctionImpl(const
 	return nullptr;
 }
 
-const TArray<FRigVMTemplateArgumentInfo>& FRigDispatch_SetModuleMetadata::GetArgumentInfos() const
+const TArray<FRigVMTemplateArgumentInfo>& FRigDispatch_SetModuleMetadata::GetArgumentInfos(FRigVMRegistryHandle& InRegistry) const
 {
 	if(ValueArgIndex == INDEX_NONE)
 	{
-		NameArgIndex = Infos.Emplace(NameArgName, ERigVMPinDirection::Input, RigVMTypeUtils::TypeIndex::FName);
-		NameSpaceArgIndex = Infos.Emplace(NameSpaceArgName, ERigVMPinDirection::Input, FRigVMRegistry_NoLock::GetForRead().GetTypeIndex_NoLock<ERigMetaDataNameSpace>());
-		ValueArgIndex = Infos.Emplace(ValueArgName, ERigVMPinDirection::Input, GetValueTypes());
-		SuccessArgIndex = Infos.Emplace(SuccessArgName, ERigVMPinDirection::Output, RigVMTypeUtils::TypeIndex::Bool);
+		NameArgIndex = CachedArgumentInfos.Emplace(NameArgName, ERigVMPinDirection::Input, RigVMTypeUtils::TypeIndex::FName);
+		NameSpaceArgIndex = CachedArgumentInfos.Emplace(NameSpaceArgName, ERigVMPinDirection::Input, InRegistry->GetTypeIndex_NoLock<ERigMetaDataNameSpace>());
+		ValueArgIndex = CachedArgumentInfos.Emplace(ValueArgName, ERigVMPinDirection::Input, GetValueTypes(InRegistry));
+		SuccessArgIndex = CachedArgumentInfos.Emplace(SuccessArgName, ERigVMPinDirection::Output, RigVMTypeUtils::TypeIndex::Bool);
 	};
-	return Infos;
+	return CachedArgumentInfos;
 }
 
 FRigBaseMetadata* FRigDispatch_SetModuleMetadata::FindOrAddMetadata(const FControlRigExecuteContext& InContext, const FName& InName, ERigMetadataType InType, ERigMetaDataNameSpace InNameSpace)
@@ -993,9 +990,8 @@ FRigBaseMetadata* FRigDispatch_SetModuleMetadata::FindOrAddMetadata(const FContr
 	return nullptr;
 }
 
-FRigVMFunctionPtr FRigDispatch_SetModuleMetadata::GetDispatchFunctionImpl(const FRigVMTemplateTypeMap& InTypes) const
+FRigVMFunctionPtr FRigDispatch_SetModuleMetadata::GetDispatchFunctionImpl(const FRigVMTemplateTypeMap& InTypes, const FRigVMRegistryHandle& InRegistry) const
 {
-	const FRigVMRegistry_NoLock& Registry = FRigVMRegistry_NoLock::GetForRead();
 	const TRigVMTypeIndex& ValueTypeIndex = InTypes.FindChecked(TEXT("Value"));
 	
 	if(ValueTypeIndex == RigVMTypeUtils::TypeIndex::Bool)
@@ -1014,27 +1010,27 @@ FRigVMFunctionPtr FRigDispatch_SetModuleMetadata::GetDispatchFunctionImpl(const
 	{
 		return &FRigDispatch_SetModuleMetadata::SetModuleMetadataDispatch<FName, FRigNameMetadata, ERigMetadataType::Name>;
 	}
-	if(ValueTypeIndex == Registry.GetTypeIndex_NoLock<FVector>(false))
+	if(ValueTypeIndex == InRegistry->GetTypeIndex_NoLock<FVector>(false))
 	{
 		return &FRigDispatch_SetModuleMetadata::SetModuleMetadataDispatch<FVector, FRigVectorMetadata, ERigMetadataType::Vector>;
 	}
-	if(ValueTypeIndex == Registry.GetTypeIndex_NoLock<FRotator>(false))
+	if(ValueTypeIndex == InRegistry->GetTypeIndex_NoLock<FRotator>(false))
 	{
 		return &FRigDispatch_SetModuleMetadata::SetModuleMetadataDispatch<FRotator, FRigRotatorMetadata, ERigMetadataType::Rotator>;
 	}
-	if(ValueTypeIndex == Registry.GetTypeIndex_NoLock<FQuat>(false))
+	if(ValueTypeIndex == InRegistry->GetTypeIndex_NoLock<FQuat>(false))
 	{
 		return &FRigDispatch_SetModuleMetadata::SetModuleMetadataDispatch<FQuat, FRigQuatMetadata, ERigMetadataType::Quat>;
 	}
-	if(ValueTypeIndex == Registry.GetTypeIndex_NoLock<FTransform>(false))
+	if(ValueTypeIndex == InRegistry->GetTypeIndex_NoLock<FTransform>(false))
 	{
 		return &FRigDispatch_SetModuleMetadata::SetModuleMetadataDispatch<FTransform, FRigTransformMetadata, ERigMetadataType::Transform>;
 	}
-	if(ValueTypeIndex == Registry.GetTypeIndex_NoLock<FLinearColor>(false))
+	if(ValueTypeIndex == InRegistry->GetTypeIndex_NoLock<FLinearColor>(false))
 	{
 		return &FRigDispatch_SetModuleMetadata::SetModuleMetadataDispatch<FLinearColor, FRigLinearColorMetadata, ERigMetadataType::LinearColor>;
 	}
-	if(ValueTypeIndex == Registry.GetTypeIndex_NoLock<FRigElementKey>(false))
+	if(ValueTypeIndex == InRegistry->GetTypeIndex_NoLock<FRigElementKey>(false))
 	{
 		return &FRigDispatch_SetModuleMetadata::SetModuleMetadataDispatch<FRigElementKey, FRigElementKeyMetadata, ERigMetadataType::RigElementKey>;
 	}
@@ -1054,27 +1050,27 @@ FRigVMFunctionPtr FRigDispatch_SetModuleMetadata::GetDispatchFunctionImpl(const
 	{
 		return &FRigDispatch_SetModuleMetadata::SetModuleMetadataDispatch<TArray<FName>, FRigNameArrayMetadata, ERigMetadataType::NameArray>;
 	}
-	if(ValueTypeIndex == Registry.GetTypeIndex_NoLock<FVector>(true))
+	if(ValueTypeIndex == InRegistry->GetTypeIndex_NoLock<FVector>(true))
 	{
 		return &FRigDispatch_SetModuleMetadata::SetModuleMetadataDispatch<TArray<FVector>, FRigVectorArrayMetadata, ERigMetadataType::VectorArray>;
 	}
-	if(ValueTypeIndex == Registry.GetTypeIndex_NoLock<FRotator>(true))
+	if(ValueTypeIndex == InRegistry->GetTypeIndex_NoLock<FRotator>(true))
 	{
 		return &FRigDispatch_SetModuleMetadata::SetModuleMetadataDispatch<TArray<FRotator>, FRigRotatorArrayMetadata, ERigMetadataType::RotatorArray>;
 	}
-	if(ValueTypeIndex == Registry.GetTypeIndex_NoLock<FQuat>(true))
+	if(ValueTypeIndex == InRegistry->GetTypeIndex_NoLock<FQuat>(true))
 	{
 		return &FRigDispatch_SetModuleMetadata::SetModuleMetadataDispatch<TArray<FQuat>, FRigQuatArrayMetadata, ERigMetadataType::QuatArray>;
 	}
-	if(ValueTypeIndex == Registry.GetTypeIndex_NoLock<FTransform>(true))
+	if(ValueTypeIndex == InRegistry->GetTypeIndex_NoLock<FTransform>(true))
 	{
 		return &FRigDispatch_SetModuleMetadata::SetModuleMetadataDispatch<TArray<FTransform>, FRigTransformArrayMetadata, ERigMetadataType::TransformArray>;
 	}
-	if(ValueTypeIndex == Registry.GetTypeIndex_NoLock<FLinearColor>(true))
+	if(ValueTypeIndex == InRegistry->GetTypeIndex_NoLock<FLinearColor>(true))
 	{
 		return &FRigDispatch_SetModuleMetadata::SetModuleMetadataDispatch<TArray<FLinearColor>, FRigLinearColorArrayMetadata, ERigMetadataType::LinearColorArray>;
 	}
-	if(ValueTypeIndex == Registry.GetTypeIndex_NoLock<FRigElementKey>(true))
+	if(ValueTypeIndex == InRegistry->GetTypeIndex_NoLock<FRigElementKey>(true))
 	{
 		return &FRigDispatch_SetModuleMetadata::SetModuleMetadataDispatch<TArray<FRigElementKey>, FRigElementKeyArrayMetadata, ERigMetadataType::RigElementKeyArray>;
 	}
diff --git a/Engine/Plugins/Animation/ControlRig/Source/ControlRig/Public/Units/Animation/RigUnit_AnimAttribute.h b/Engine/Plugins/Animation/ControlRig/Source/ControlRig/Public/Units/Animation/RigUnit_AnimAttribute.h
--- a/Engine/Plugins/Animation/ControlRig/Source/ControlRig/Public/Units/Animation/RigUnit_AnimAttribute.h
+++ b/Engine/Plugins/Animation/ControlRig/Source/ControlRig/Public/Units/Animation/RigUnit_AnimAttribute.h
@@ -141,13 +141,13 @@ struct FRigDispatch_AnimAttributeBase : public FRigDispatchFactory
 {
 	GENERATED_BODY()
 
-	UE_API virtual void RegisterDependencyTypes_NoLock(FRigVMRegistry_NoLock& InRegistry) const override;
+	UE_API virtual void RegisterDependencyTypes_NoLock(FRigVMRegistryHandle& InRegistry) const override;
 
 #if WITH_EDITOR
 	UE_API virtual FString GetNodeTitle(const FRigVMTemplateTypeMap& InTypes) const override;;
 #endif
 	
-	UE_API virtual const TArray<FRigVMTemplateArgumentInfo>& GetArgumentInfos() const override;
+	UE_API virtual const TArray<FRigVMTemplateArgumentInfo>& GetArgumentInfos(FRigVMRegistryHandle& InRegistry) const override;
 	virtual bool IsSet() const { return false; }
 
 #if WITH_EDITOR
@@ -155,11 +155,9 @@ struct FRigDispatch_AnimAttributeBase : public FRigDispatchFactory
 #endif
 
 protected:
-	static UE_API bool IsTypeSupported(const TRigVMTypeIndex& InTypeIndex);
+	UE_API bool IsTypeSupported(const TRigVMTypeIndex& InTypeIndex, const FRigVMRegistryHandle& InRegistry) const;
 	static UE_API const TArray<FRigVMTemplateArgument::ETypeCategory>& GetValueTypeCategory();
 	
-	mutable TArray<FRigVMTemplateArgumentInfo> Infos;
-
 	// input
 	mutable int32 NameArgIndex = INDEX_NONE;
 	mutable int32 BoneNameArgIndex = INDEX_NONE;
@@ -175,6 +173,8 @@ protected:
 	mutable int32 CachedBoneNameArgIndex = INDEX_NONE;
 	mutable int32 CachedBoneIndexArgIndex = INDEX_NONE;
 
+	mutable TArray<TRigVMTypeIndex> SpecialTypes;
+
 	static inline const FLazyName NameArgName = FLazyName(TEXT("Name"));
 	static inline const FLazyName BoneNameArgName = FLazyName(TEXT("BoneName"));
 	static inline const FLazyName CachedBoneNameArgName = FLazyName(TEXT("CachedBoneName"));
@@ -199,12 +199,12 @@ struct FRigDispatch_GetAnimAttribute: public FRigDispatch_AnimAttributeBase
 		FactoryScriptStruct = StaticStruct();
 	}
 
-	UE_API virtual const TArray<FRigVMTemplateArgumentInfo>& GetArgumentInfos() const override;
-	UE_API virtual FRigVMTemplateTypeMap OnNewArgumentType(const FName& InArgumentName, TRigVMTypeIndex InTypeIndex) const override;
+	UE_API virtual const TArray<FRigVMTemplateArgumentInfo>& GetArgumentInfos(FRigVMRegistryHandle& InRegistry) const override;
+	UE_API virtual FRigVMTemplateTypeMap OnNewArgumentType(const FName& InArgumentName, TRigVMTypeIndex InTypeIndex, FRigVMRegistryHandle& InRegistry) const override;
 
 	
 protected:
-	UE_API virtual FRigVMFunctionPtr GetDispatchFunctionImpl(const FRigVMTemplateTypeMap& InTypes) const override;
+	UE_API virtual FRigVMFunctionPtr GetDispatchFunctionImpl(const FRigVMTemplateTypeMap& InTypes, const FRigVMRegistryHandle& InRegistry) const override;
 
 #if WITH_EDITOR
 	template<typename ValueType>
@@ -279,12 +279,12 @@ struct FRigDispatch_SetAnimAttribute: public FRigDispatch_AnimAttributeBase
 	}
 
 	virtual bool IsSet() const override { return true; }
-	UE_API virtual const TArray<FRigVMTemplateArgumentInfo>& GetArgumentInfos() const override;
+	UE_API virtual const TArray<FRigVMTemplateArgumentInfo>& GetArgumentInfos(FRigVMRegistryHandle& InRegistry) const override;
 	UE_API virtual const TArray<FRigVMExecuteArgument>& GetExecuteArguments_Impl(const FRigVMDispatchContext& InContext) const override;
-	UE_API virtual FRigVMTemplateTypeMap OnNewArgumentType(const FName& InArgumentName, TRigVMTypeIndex InTypeIndex) const override;
+	UE_API virtual FRigVMTemplateTypeMap OnNewArgumentType(const FName& InArgumentName, TRigVMTypeIndex InTypeIndex, FRigVMRegistryHandle& InRegistry) const override;
 	
 protected:
-	UE_API virtual FRigVMFunctionPtr GetDispatchFunctionImpl(const FRigVMTemplateTypeMap& InTypes) const override;
+	UE_API virtual FRigVMFunctionPtr GetDispatchFunctionImpl(const FRigVMTemplateTypeMap& InTypes, const FRigVMRegistryHandle& InRegistry) const override;
 
 #if WITH_EDITOR
 	template<typename ValueType>
diff --git a/Engine/Plugins/Animation/ControlRig/Source/ControlRig/Public/Units/Core/RigUnit_UserData.h b/Engine/Plugins/Animation/ControlRig/Source/ControlRig/Public/Units/Core/RigUnit_UserData.h
--- a/Engine/Plugins/Animation/ControlRig/Source/ControlRig/Public/Units/Core/RigUnit_UserData.h
+++ b/Engine/Plugins/Animation/ControlRig/Source/ControlRig/Public/Units/Core/RigUnit_UserData.h
@@ -25,12 +25,12 @@ public:
 #if WITH_EDITOR
 	virtual FString GetArgumentMetaData(const FName& InArgumentName, const FName& InMetaDataKey) const override;
 #endif
-	virtual const TArray<FRigVMTemplateArgumentInfo>& GetArgumentInfos() const override;
-	virtual FRigVMTemplateTypeMap OnNewArgumentType(const FName& InArgumentName, TRigVMTypeIndex InTypeIndex) const override;
+	virtual const TArray<FRigVMTemplateArgumentInfo>& GetArgumentInfos(FRigVMRegistryHandle& InRegistry) const override;
+	virtual FRigVMTemplateTypeMap OnNewArgumentType(const FName& InArgumentName, TRigVMTypeIndex InTypeIndex, FRigVMRegistryHandle& InRegistry) const override;
 
 protected:
 
-	virtual FRigVMFunctionPtr GetDispatchFunctionImpl(const FRigVMTemplateTypeMap& InTypes) const override;
+	virtual FRigVMFunctionPtr GetDispatchFunctionImpl(const FRigVMTemplateTypeMap& InTypes, const FRigVMRegistryHandle& InRegistry) const override;
 	static void Execute(FRigVMExtendedExecuteContext& InContext, FRigVMMemoryHandleArray Handles, FRigVMPredicateBranchArray RigVMBranches);
 
 	static inline constexpr TCHAR ArgNameSpaceName[] = TEXT("NameSpace");
diff --git a/Engine/Plugins/Animation/ControlRig/Source/ControlRig/Public/Units/Hierarchy/RigUnit_Component.h b/Engine/Plugins/Animation/ControlRig/Source/ControlRig/Public/Units/Hierarchy/RigUnit_Component.h
--- a/Engine/Plugins/Animation/ControlRig/Source/ControlRig/Public/Units/Hierarchy/RigUnit_Component.h
+++ b/Engine/Plugins/Animation/ControlRig/Source/ControlRig/Public/Units/Hierarchy/RigUnit_Component.h
@@ -12,7 +12,7 @@ struct FRigDispatch_ComponentBase : public FRigDispatchFactory
 {
 	GENERATED_BODY()
 
-	UE_API virtual void RegisterDependencyTypes_NoLock(FRigVMRegistry_NoLock& InRegistry) const override;
+	UE_API virtual void RegisterDependencyTypes_NoLock(FRigVMRegistryHandle& InRegistry) const override;
 	virtual UScriptStruct* GetExecuteContextStruct() const override { return FControlRigExecuteContext::StaticStruct(); }
 	UE_API virtual const TArray<FRigVMExecuteArgument>& GetExecuteArguments_Impl(const FRigVMDispatchContext& InContext) const override;
 
@@ -24,11 +24,10 @@ struct FRigDispatch_ComponentBase : public FRigDispatchFactory
 protected:
 
 	virtual bool IsMutable() const { return false; }
-	static UE_API bool IsTypeSupported(const TRigVMTypeIndex& InTypeIndex);
-	static UE_API const TRigVMTypeIndex& GetElementKeyType();
-	static UE_API const TRigVMTypeIndex& GetComponentKeyType();
+	static UE_API bool IsTypeSupported(const TRigVMTypeIndex& InTypeIndex, const FRigVMRegistryHandle& InRegistry);
+	static UE_API const TRigVMTypeIndex& GetElementKeyType(FRigVMRegistryHandle& InRegistry);
+	static UE_API const TRigVMTypeIndex& GetComponentKeyType(FRigVMRegistryHandle& InRegistry);
 
-	mutable TArray<FRigVMTemplateArgumentInfo> Infos;
 	mutable TArray<FRigVMExecuteArgument> ExecuteInfos;
 
 	mutable int32 NameArgIndex = INDEX_NONE;
@@ -61,7 +60,7 @@ struct FRigDispatch_SpawnComponent : public FRigDispatch_ComponentBase
 		FactoryScriptStruct = FRigDispatch_SpawnComponent::StaticStruct();
 	}
 
-	UE_API virtual const TArray<FRigVMTemplateArgumentInfo>& GetArgumentInfos() const override;
+	UE_API virtual const TArray<FRigVMTemplateArgumentInfo>& GetArgumentInfos(FRigVMRegistryHandle& InRegistry) const override;
 #if WITH_EDITOR
 	UE_API virtual FString GetKeywords() const override;
 #endif
@@ -70,7 +69,7 @@ protected:
 
 	virtual bool IsMutable() const override { return true; }
 
-	virtual FRigVMFunctionPtr GetDispatchFunctionImpl(const FRigVMTemplateTypeMap& InTypes) const override
+	virtual FRigVMFunctionPtr GetDispatchFunctionImpl(const FRigVMTemplateTypeMap& InTypes, const FRigVMRegistryHandle& InRegistry) const override
 	{
 		return &Execute;
 	}
@@ -91,11 +90,11 @@ struct FRigDispatch_GetComponentContent : public FRigDispatch_ComponentBase
 		FactoryScriptStruct = FRigDispatch_GetComponentContent::StaticStruct();
 	}
 
-	UE_API virtual const TArray<FRigVMTemplateArgumentInfo>& GetArgumentInfos() const override;
+	UE_API virtual const TArray<FRigVMTemplateArgumentInfo>& GetArgumentInfos(FRigVMRegistryHandle& InRegistry) const override;
 
 protected:
 
-	virtual FRigVMFunctionPtr GetDispatchFunctionImpl(const FRigVMTemplateTypeMap& InTypes) const override
+	virtual FRigVMFunctionPtr GetDispatchFunctionImpl(const FRigVMTemplateTypeMap& InTypes, const FRigVMRegistryHandle& InRegistry) const override
 	{
 		return &Execute;
 	}
@@ -116,13 +115,13 @@ struct FRigDispatch_SetComponentContent : public FRigDispatch_ComponentBase
 		FactoryScriptStruct = FRigDispatch_SetComponentContent::StaticStruct();
 	}
 
-	UE_API virtual const TArray<FRigVMTemplateArgumentInfo>& GetArgumentInfos() const override;
+	UE_API virtual const TArray<FRigVMTemplateArgumentInfo>& GetArgumentInfos(FRigVMRegistryHandle& InRegistry) const override;
 
 protected:
 
 	virtual bool IsMutable() const override { return true; }
 
-	virtual FRigVMFunctionPtr GetDispatchFunctionImpl(const FRigVMTemplateTypeMap& InTypes) const override
+	virtual FRigVMFunctionPtr GetDispatchFunctionImpl(const FRigVMTemplateTypeMap& InTypes, const FRigVMRegistryHandle& InRegistry) const override
 	{
 		return &Execute;
 	}
diff --git a/Engine/Plugins/Animation/ControlRig/Source/ControlRig/Public/Units/Hierarchy/RigUnit_Metadata.h b/Engine/Plugins/Animation/ControlRig/Source/ControlRig/Public/Units/Hierarchy/RigUnit_Metadata.h
--- a/Engine/Plugins/Animation/ControlRig/Source/ControlRig/Public/Units/Hierarchy/RigUnit_Metadata.h
+++ b/Engine/Plugins/Animation/ControlRig/Source/ControlRig/Public/Units/Hierarchy/RigUnit_Metadata.h
@@ -45,7 +45,7 @@ struct FRigDispatch_MetadataBase : public FRigDispatchFactory
 	UE_API virtual FString GetNodeTitle(const FRigVMTemplateTypeMap& InTypes) const override;;
 	virtual FString GetNodeTitlePrefix() const { return FString(); }
 #endif
-	UE_API virtual const TArray<FRigVMTemplateArgumentInfo>& GetArgumentInfos() const override;
+	UE_API virtual const TArray<FRigVMTemplateArgumentInfo>& GetArgumentInfos(FRigVMRegistryHandle& InRegistry) const override;
 	virtual bool IsSetMetadata() const { return false; }
 
 #if WITH_EDITOR
@@ -56,10 +56,8 @@ struct FRigDispatch_MetadataBase : public FRigDispatchFactory
 
 protected:
 
-	UE_API const TArray<TRigVMTypeIndex>& GetValueTypes() const;
+	UE_API const TArray<TRigVMTypeIndex>& GetValueTypes(const FRigVMRegistryHandle& InRegistry) const;
 
-	mutable TArray<FRigVMTemplateArgumentInfo> Infos;
-	
 	mutable int32 ItemArgIndex = INDEX_NONE;
 	mutable int32 NameArgIndex = INDEX_NONE;
 	mutable int32 NameSpaceArgIndex = INDEX_NONE;
@@ -68,6 +66,7 @@ protected:
 	mutable int32 ValueArgIndex = INDEX_NONE;
 	mutable int32 FoundArgIndex = INDEX_NONE;
 	mutable int32 SuccessArgIndex = INDEX_NONE;
+	mutable TArray<TRigVMTypeIndex> ValueTypes;
 
 	static inline const FLazyName ItemArgName = FLazyName(TEXT("Item"));
 	static inline const FLazyName NameArgName = FLazyName(TEXT("Name"));
@@ -189,12 +188,12 @@ struct FRigDispatch_GetMetadata : public FRigDispatch_MetadataBase
 {
 	GENERATED_BODY()
 
-	UE_API virtual const TArray<FRigVMTemplateArgumentInfo>& GetArgumentInfos() const override;
+	UE_API virtual const TArray<FRigVMTemplateArgumentInfo>& GetArgumentInfos(FRigVMRegistryHandle& InRegistry) const override;
 
 protected:
 
 	static UE_API FRigBaseMetadata* FindMetadata(const FRigVMExtendedExecuteContext& InContext, const FRigElementKey& InKey, const FName& InName, ERigMetadataType InType, ERigMetaDataNameSpace InNameSpace, FCachedRigElement& Cache);
-	UE_API virtual FRigVMFunctionPtr GetDispatchFunctionImpl(const FRigVMTemplateTypeMap& InTypes) const override;
+	UE_API virtual FRigVMFunctionPtr GetDispatchFunctionImpl(const FRigVMTemplateTypeMap& InTypes, const FRigVMRegistryHandle& InRegistry) const override;
 
 #if WITH_EDITOR
 	template<typename ValueType>
@@ -253,14 +252,14 @@ struct FRigDispatch_SetMetadata : public FRigDispatch_MetadataBase
 {
 	GENERATED_BODY()
 
-	UE_API virtual const TArray<FRigVMTemplateArgumentInfo>& GetArgumentInfos() const override;
+	UE_API virtual const TArray<FRigVMTemplateArgumentInfo>& GetArgumentInfos(FRigVMRegistryHandle& InRegistry) const override;
 	UE_API virtual const TArray<FRigVMExecuteArgument>& GetExecuteArguments_Impl(const FRigVMDispatchContext& InContext) const override;
 	virtual bool IsSetMetadata() const override { return true; }
 
 protected:
 
 	static UE_API FRigBaseMetadata* FindOrAddMetadata(const FControlRigExecuteContext& InContext, const FRigElementKey& InKey, const FName& InName, ERigMetadataType InType, ERigMetaDataNameSpace InNameSpace, FCachedRigElement& Cache);
-	UE_API virtual FRigVMFunctionPtr GetDispatchFunctionImpl(const FRigVMTemplateTypeMap& InTypes) const override;
+	UE_API virtual FRigVMFunctionPtr GetDispatchFunctionImpl(const FRigVMTemplateTypeMap& InTypes, const FRigVMRegistryHandle& InRegistry) const override;
 
 #if WITH_EDITOR
 	template<typename ValueType>
@@ -875,7 +874,7 @@ struct FRigDispatch_GetModuleMetadata : public FRigDispatch_GetMetadata
 {
 	GENERATED_BODY()
 
-	UE_API virtual const TArray<FRigVMTemplateArgumentInfo>& GetArgumentInfos() const override;
+	UE_API virtual const TArray<FRigVMTemplateArgumentInfo>& GetArgumentInfos(FRigVMRegistryHandle& InRegistry) const override;
 #if WITH_EDITOR
 	virtual FString GetNodeTitlePrefix() const override { return TEXT("Module "); }
 #endif
@@ -883,7 +882,7 @@ struct FRigDispatch_GetModuleMetadata : public FRigDispatch_GetMetadata
 protected:
 
 	static UE_API FRigBaseMetadata* FindMetadata(const FRigVMExtendedExecuteContext& InContext, const FName& InName, ERigMetadataType InType, ERigMetaDataNameSpace InNameSpace);
-	UE_API virtual FRigVMFunctionPtr GetDispatchFunctionImpl(const FRigVMTemplateTypeMap& InTypes) const override;
+	UE_API virtual FRigVMFunctionPtr GetDispatchFunctionImpl(const FRigVMTemplateTypeMap& InTypes, const FRigVMRegistryHandle& InRegistry) const override;
 
 #if WITH_EDITOR
 	template<typename ValueType>
@@ -938,7 +937,7 @@ struct FRigDispatch_SetModuleMetadata : public FRigDispatch_SetMetadata
 {
 	GENERATED_BODY()
 
-	UE_API virtual const TArray<FRigVMTemplateArgumentInfo>& GetArgumentInfos() const override;
+	UE_API virtual const TArray<FRigVMTemplateArgumentInfo>& GetArgumentInfos(FRigVMRegistryHandle& InRegistry) const override;
 #if WITH_EDITOR
 	virtual FString GetNodeTitlePrefix() const override { return TEXT("Module "); }
 #endif
@@ -946,7 +945,7 @@ struct FRigDispatch_SetModuleMetadata : public FRigDispatch_SetMetadata
 protected:
 
 	static UE_API FRigBaseMetadata* FindOrAddMetadata(const FControlRigExecuteContext& InContext, const FName& InName, ERigMetadataType InType, ERigMetaDataNameSpace InNameSpace);
-	UE_API virtual FRigVMFunctionPtr GetDispatchFunctionImpl(const FRigVMTemplateTypeMap& InTypes) const override;
+	UE_API virtual FRigVMFunctionPtr GetDispatchFunctionImpl(const FRigVMTemplateTypeMap& InTypes, const FRigVMRegistryHandle& InRegistry) const override;
 
 #if WITH_EDITOR
 	template<typename ValueType>
diff --git a/Engine/Plugins/Animation/ControlRig/Source/ControlRig/Public/Units/RigDispatchFactory.h b/Engine/Plugins/Animation/ControlRig/Source/ControlRig/Public/Units/RigDispatchFactory.h
--- a/Engine/Plugins/Animation/ControlRig/Source/ControlRig/Public/Units/RigDispatchFactory.h
+++ b/Engine/Plugins/Animation/ControlRig/Source/ControlRig/Public/Units/RigDispatchFactory.h
@@ -20,11 +20,11 @@ struct FRigDispatchFactory : public FRigVMDispatchFactory
 		return FControlRigExecuteContext::StaticStruct();
 	}
 
-	virtual void RegisterDependencyTypes_NoLock(FRigVMRegistry_NoLock& InRegistry) const override
+	virtual void RegisterDependencyTypes_NoLock(FRigVMRegistryHandle& InRegistry) const override
 	{
-		InRegistry.FindOrAddType_NoLock(FControlRigExecuteContext::StaticStruct());
-		InRegistry.FindOrAddType_NoLock(FRigElementKey::StaticStruct());
-    	InRegistry.FindOrAddType_NoLock(FCachedRigElement::StaticStruct());
+		InRegistry->FindOrAddType_NoLock(FControlRigExecuteContext::StaticStruct());
+		InRegistry->FindOrAddType_NoLock(FRigElementKey::StaticStruct());
+    	InRegistry->FindOrAddType_NoLock(FCachedRigElement::StaticStruct());
 	}
 
 #if WITH_EDITOR
diff --git a/Engine/Plugins/Experimental/AnimNext/Source/AnimNext/Internal/Param/RigVMDispatch_GetLayerParameter.h b/Engine/Plugins/Experimental/AnimNext/Source/AnimNext/Internal/Param/RigVMDispatch_GetLayerParameter.h
--- a/Engine/Plugins/Experimental/AnimNext/Source/AnimNext/Internal/Param/RigVMDispatch_GetLayerParameter.h
+++ b/Engine/Plugins/Experimental/AnimNext/Source/AnimNext/Internal/Param/RigVMDispatch_GetLayerParameter.h
@@ -27,15 +27,15 @@ private:
 	friend struct UE::AnimNext::UncookedOnly::FUtils;
 
 	virtual UScriptStruct* GetExecuteContextStruct() const override { return FAnimNextExecuteContext::StaticStruct(); }
-	UE_API virtual FName GetArgumentNameForOperandIndex(int32 InOperandIndex, int32 InTotalOperands) const override;
-	UE_API virtual const TArray<FRigVMTemplateArgumentInfo>& GetArgumentInfos() const override;
+	UE_API virtual FName GetArgumentNameForOperandIndex(int32 InOperandIndex, int32 InTotalOperands, FRigVMRegistryHandle& InRegistry) const override;
+	UE_API virtual const TArray<FRigVMTemplateArgumentInfo>& GetArgumentInfos(FRigVMRegistryHandle& InRegistry) const override;
 #if WITH_EDITOR
 	UE_API virtual FString GetArgumentMetaData(const FName& InArgumentName, const FName& InMetaDataKey) const override;
 #endif
-	UE_API virtual FRigVMTemplateTypeMap OnNewArgumentType(const FName& InArgumentName, TRigVMTypeIndex InTypeIndex) const override;
+	UE_API virtual FRigVMTemplateTypeMap OnNewArgumentType(const FName& InArgumentName, TRigVMTypeIndex InTypeIndex, FRigVMRegistryHandle& InRegistry) const override;
 	virtual bool IsSingleton() const override { return true; }
 
-	virtual FRigVMFunctionPtr GetDispatchFunctionImpl(const FRigVMTemplateTypeMap& InTypes) const override
+	virtual FRigVMFunctionPtr GetDispatchFunctionImpl(const FRigVMTemplateTypeMap& InTypes, const FRigVMRegistryHandle& InRegistry) const override
 	{
 		return &FRigVMDispatch_GetLayerParameter::Execute;
 	}
diff --git a/Engine/Plugins/Experimental/AnimNext/Source/AnimNext/Internal/Param/RigVMDispatch_GetParameter.h b/Engine/Plugins/Experimental/AnimNext/Source/AnimNext/Internal/Param/RigVMDispatch_GetParameter.h
--- a/Engine/Plugins/Experimental/AnimNext/Source/AnimNext/Internal/Param/RigVMDispatch_GetParameter.h
+++ b/Engine/Plugins/Experimental/AnimNext/Source/AnimNext/Internal/Param/RigVMDispatch_GetParameter.h
@@ -30,15 +30,15 @@ private:
 	friend struct UE::AnimNext::UncookedOnly::FUtils;
 
 	virtual UScriptStruct* GetExecuteContextStruct() const { return FAnimNextExecuteContext::StaticStruct(); }
-	UE_API virtual FName GetArgumentNameForOperandIndex(int32 InOperandIndex, int32 InTotalOperands) const override;
-	UE_API virtual const TArray<FRigVMTemplateArgumentInfo>& GetArgumentInfos() const override;
+	UE_API virtual FName GetArgumentNameForOperandIndex(int32 InOperandIndex, int32 InTotalOperands, FRigVMRegistryHandle& InRegistry) const override;
+	UE_API virtual const TArray<FRigVMTemplateArgumentInfo>& GetArgumentInfos(FRigVMRegistryHandle& InRegistry) const override;
 #if WITH_EDITOR
 	UE_API virtual FString GetArgumentMetaData(const FName& InArgumentName, const FName& InMetaDataKey) const override;
 #endif
-	UE_API virtual FRigVMTemplateTypeMap OnNewArgumentType(const FName& InArgumentName, TRigVMTypeIndex InTypeIndex) const override;
+	UE_API virtual FRigVMTemplateTypeMap OnNewArgumentType(const FName& InArgumentName, TRigVMTypeIndex InTypeIndex, FRigVMRegistryHandle& InRegistry) const override;
 	virtual bool IsSingleton() const override { return true; }
 
-	virtual FRigVMFunctionPtr GetDispatchFunctionImpl(const FRigVMTemplateTypeMap& InTypes) const override
+	virtual FRigVMFunctionPtr GetDispatchFunctionImpl(const FRigVMTemplateTypeMap& InTypes, const FRigVMRegistryHandle& InRegistry) const override
 	{
 		return &FRigVMDispatch_GetParameter::Execute;
 	}
diff --git a/Engine/Plugins/Experimental/AnimNext/Source/AnimNext/Internal/Variables/RigVMDispatch_CallHoistedAccessorFunction.h b/Engine/Plugins/Experimental/AnimNext/Source/AnimNext/Internal/Variables/RigVMDispatch_CallHoistedAccessorFunction.h
--- a/Engine/Plugins/Experimental/AnimNext/Source/AnimNext/Internal/Variables/RigVMDispatch_CallHoistedAccessorFunction.h
+++ b/Engine/Plugins/Experimental/AnimNext/Source/AnimNext/Internal/Variables/RigVMDispatch_CallHoistedAccessorFunction.h
@@ -22,10 +22,10 @@ struct FRigVMDispatch_CallHoistedAccessorFunctionBase : public FRigVMDispatchFac
 
 private:
 	virtual UScriptStruct* GetExecuteContextStruct() const override { return FAnimNextExecuteContext::StaticStruct(); }
-	UE_API virtual FName GetArgumentNameForOperandIndex(int32 InOperandIndex, int32 InTotalOperands) const override;
-	UE_API virtual void RegisterDependencyTypes_NoLock(FRigVMRegistry_NoLock& InRegistry) const override;
-	UE_API virtual const TArray<FRigVMTemplateArgumentInfo>& GetArgumentInfos() const override;
-	UE_API virtual FRigVMTemplateTypeMap OnNewArgumentType(const FName& InArgumentName, TRigVMTypeIndex InTypeIndex) const override;
+	UE_API virtual FName GetArgumentNameForOperandIndex(int32 InOperandIndex, int32 InTotalOperands, FRigVMRegistryHandle& InRegistry) const override;
+	UE_API virtual void RegisterDependencyTypes_NoLock(FRigVMRegistryHandle& InRegistry) const override;
+	UE_API virtual const TArray<FRigVMTemplateArgumentInfo>& GetArgumentInfos(FRigVMRegistryHandle& InRegistry) const override;
+	UE_API virtual FRigVMTemplateTypeMap OnNewArgumentType(const FName& InArgumentName, TRigVMTypeIndex InTypeIndex, FRigVMRegistryHandle& InRegistry) const override;
 	virtual bool IsSingleton() const override { return true; }
 };
 
@@ -38,7 +38,7 @@ struct FRigVMDispatch_CallHoistedAccessorFunctionNative : public FRigVMDispatch_
 	UE_API FRigVMDispatch_CallHoistedAccessorFunctionNative();
 
 private:
-	virtual FRigVMFunctionPtr GetDispatchFunctionImpl(const FRigVMTemplateTypeMap& InTypes) const override
+	virtual FRigVMFunctionPtr GetDispatchFunctionImpl(const FRigVMTemplateTypeMap& InTypes, const FRigVMRegistryHandle& InRegistry) const override
 	{
 		return &FRigVMDispatch_CallHoistedAccessorFunctionNative::Execute;
 	}
@@ -54,7 +54,7 @@ struct FRigVMDispatch_CallHoistedAccessorFunctionScript : public FRigVMDispatch_
 	UE_API FRigVMDispatch_CallHoistedAccessorFunctionScript();
 
 private:
-	virtual FRigVMFunctionPtr GetDispatchFunctionImpl(const FRigVMTemplateTypeMap& InTypes) const override
+	virtual FRigVMFunctionPtr GetDispatchFunctionImpl(const FRigVMTemplateTypeMap& InTypes, const FRigVMRegistryHandle& InRegistry) const override
 	{
 		return &FRigVMDispatch_CallHoistedAccessorFunctionScript::Execute;
 	}
diff --git a/Engine/Plugins/Experimental/AnimNext/Source/AnimNext/Internal/Variables/RigVMDispatch_CallObjectAccessorFunction.h b/Engine/Plugins/Experimental/AnimNext/Source/AnimNext/Internal/Variables/RigVMDispatch_CallObjectAccessorFunction.h
--- a/Engine/Plugins/Experimental/AnimNext/Source/AnimNext/Internal/Variables/RigVMDispatch_CallObjectAccessorFunction.h
+++ b/Engine/Plugins/Experimental/AnimNext/Source/AnimNext/Internal/Variables/RigVMDispatch_CallObjectAccessorFunction.h
@@ -22,10 +22,10 @@ struct FRigVMDispatch_CallObjectAccessorFunctionBase : public FRigVMDispatchFact
 
 private:
 	virtual UScriptStruct* GetExecuteContextStruct() const override { return FAnimNextExecuteContext::StaticStruct(); }
-	UE_API virtual FName GetArgumentNameForOperandIndex(int32 InOperandIndex, int32 InTotalOperands) const override;
-	UE_API virtual void RegisterDependencyTypes_NoLock(FRigVMRegistry_NoLock& InRegistry) const override;
-	UE_API virtual const TArray<FRigVMTemplateArgumentInfo>& GetArgumentInfos() const override;
-	UE_API virtual FRigVMTemplateTypeMap OnNewArgumentType(const FName& InArgumentName, TRigVMTypeIndex InTypeIndex) const override;
+	UE_API virtual FName GetArgumentNameForOperandIndex(int32 InOperandIndex, int32 InTotalOperands, FRigVMRegistryHandle& InRegistry) const override;
+	UE_API virtual void RegisterDependencyTypes_NoLock(FRigVMRegistryHandle& InRegistry) const override;
+	UE_API virtual const TArray<FRigVMTemplateArgumentInfo>& GetArgumentInfos(FRigVMRegistryHandle& InRegistry) const override;
+	UE_API virtual FRigVMTemplateTypeMap OnNewArgumentType(const FName& InArgumentName, TRigVMTypeIndex InTypeIndex, FRigVMRegistryHandle& InRegistry) const override;
 	virtual bool IsSingleton() const override { return true; }
 };
 
@@ -38,7 +38,7 @@ struct FRigVMDispatch_CallObjectAccessorFunctionNative : public FRigVMDispatch_C
 	UE_API FRigVMDispatch_CallObjectAccessorFunctionNative();
 
 private:
-	virtual FRigVMFunctionPtr GetDispatchFunctionImpl(const FRigVMTemplateTypeMap& InTypes) const override
+	virtual FRigVMFunctionPtr GetDispatchFunctionImpl(const FRigVMTemplateTypeMap& InTypes, const FRigVMRegistryHandle& InRegistry) const override
 	{
 		return &FRigVMDispatch_CallObjectAccessorFunctionNative::Execute;
 	}
@@ -54,7 +54,7 @@ struct FRigVMDispatch_CallObjectAccessorFunctionScript : public FRigVMDispatch_C
 	UE_API FRigVMDispatch_CallObjectAccessorFunctionScript();
 
 private:
-	virtual FRigVMFunctionPtr GetDispatchFunctionImpl(const FRigVMTemplateTypeMap& InTypes) const override
+	virtual FRigVMFunctionPtr GetDispatchFunctionImpl(const FRigVMTemplateTypeMap& InTypes, const FRigVMRegistryHandle& InRegistry) const override
 	{
 		return &FRigVMDispatch_CallObjectAccessorFunctionScript::Execute;
 	}
diff --git a/Engine/Plugins/Experimental/AnimNext/Source/AnimNext/Internal/Variables/RigVMDispatch_GetObjectProperty.h b/Engine/Plugins/Experimental/AnimNext/Source/AnimNext/Internal/Variables/RigVMDispatch_GetObjectProperty.h
--- a/Engine/Plugins/Experimental/AnimNext/Source/AnimNext/Internal/Variables/RigVMDispatch_GetObjectProperty.h
+++ b/Engine/Plugins/Experimental/AnimNext/Source/AnimNext/Internal/Variables/RigVMDispatch_GetObjectProperty.h
@@ -22,13 +22,13 @@ struct FRigVMDispatch_GetObjectProperty : public FRigVMDispatchFactory
 	
 private:
 	virtual UScriptStruct* GetExecuteContextStruct() const override { return FAnimNextExecuteContext::StaticStruct(); }
-	UE_API virtual FName GetArgumentNameForOperandIndex(int32 InOperandIndex, int32 InTotalOperands) const override;
-	UE_API virtual void RegisterDependencyTypes_NoLock(FRigVMRegistry_NoLock& InRegistry) const override;
-	UE_API virtual const TArray<FRigVMTemplateArgumentInfo>& GetArgumentInfos() const override;
-	UE_API virtual FRigVMTemplateTypeMap OnNewArgumentType(const FName& InArgumentName, TRigVMTypeIndex InTypeIndex) const override;
+	UE_API virtual FName GetArgumentNameForOperandIndex(int32 InOperandIndex, int32 InTotalOperands, FRigVMRegistryHandle& InRegistry) const override;
+	UE_API virtual void RegisterDependencyTypes_NoLock(FRigVMRegistryHandle& InRegistry) const override;
+	UE_API virtual const TArray<FRigVMTemplateArgumentInfo>& GetArgumentInfos(FRigVMRegistryHandle& InRegistry) const override;
+	UE_API virtual FRigVMTemplateTypeMap OnNewArgumentType(const FName& InArgumentName, TRigVMTypeIndex InTypeIndex, FRigVMRegistryHandle& InRegistry) const override;
 	virtual bool IsSingleton() const override { return true; }
 
-	virtual FRigVMFunctionPtr GetDispatchFunctionImpl(const FRigVMTemplateTypeMap& InTypes) const override
+	virtual FRigVMFunctionPtr GetDispatchFunctionImpl(const FRigVMTemplateTypeMap& InTypes, const FRigVMRegistryHandle& InRegistry) const override
 	{
 		return &FRigVMDispatch_GetObjectProperty::Execute;
 	}
diff --git a/Engine/Plugins/Experimental/AnimNext/Source/AnimNext/Private/Param/RigVMDispatch_GetLayerParameter.cpp b/Engine/Plugins/Experimental/AnimNext/Source/AnimNext/Private/Param/RigVMDispatch_GetLayerParameter.cpp
--- a/Engine/Plugins/Experimental/AnimNext/Source/AnimNext/Private/Param/RigVMDispatch_GetLayerParameter.cpp
+++ b/Engine/Plugins/Experimental/AnimNext/Source/AnimNext/Private/Param/RigVMDispatch_GetLayerParameter.cpp
@@ -13,7 +13,7 @@ FRigVMDispatch_GetLayerParameter::FRigVMDispatch_GetLayerParameter()
 	FactoryScriptStruct = StaticStruct();
 }
 
-FName FRigVMDispatch_GetLayerParameter::GetArgumentNameForOperandIndex(int32 InOperandIndex, int32 InTotalOperands) const
+FName FRigVMDispatch_GetLayerParameter::GetArgumentNameForOperandIndex(int32 InOperandIndex, int32 InTotalOperands, FRigVMRegistryHandle& InRegistry) const
 {
 	static const FName ArgumentNames[] =
 	{
@@ -43,10 +43,9 @@ FString FRigVMDispatch_GetLayerParameter::GetArgumentMetaData(const FName& InArg
 }
 #endif
 
-const TArray<FRigVMTemplateArgumentInfo>& FRigVMDispatch_GetLayerParameter::GetArgumentInfos() const
+const TArray<FRigVMTemplateArgumentInfo>& FRigVMDispatch_GetLayerParameter::GetArgumentInfos(FRigVMRegistryHandle& InRegistry) const
 {
-	static TArray<FRigVMTemplateArgumentInfo> Infos;
-	if(Infos.IsEmpty())
+	if(CachedArgumentInfos.IsEmpty())
 	{
 		static const TArray<FRigVMTemplateArgument::ETypeCategory> ValueCategories =
 		{
@@ -54,16 +53,16 @@ const TArray<FRigVMTemplateArgumentInfo>& FRigVMDispatch_GetLayerParameter::GetA
 			FRigVMTemplateArgument::ETypeCategory_ArrayAnyValue
 		};
 		
-		Infos.Emplace(ParameterName, ERigVMPinDirection::Input, RigVMTypeUtils::TypeIndex::FName);
-		Infos.Emplace(ValueName, ERigVMPinDirection::Output, ValueCategories);
-		Infos.Emplace(ParameterIdName, ERigVMPinDirection::Hidden, RigVMTypeUtils::TypeIndex::UInt32);
-		Infos.Emplace(TypeHandleName, ERigVMPinDirection::Hidden, RigVMTypeUtils::TypeIndex::UInt32);
+		CachedArgumentInfos.Emplace(ParameterName, ERigVMPinDirection::Input, RigVMTypeUtils::TypeIndex::FName);
+		CachedArgumentInfos.Emplace(ValueName, ERigVMPinDirection::Output, ValueCategories);
+		CachedArgumentInfos.Emplace(ParameterIdName, ERigVMPinDirection::Hidden, RigVMTypeUtils::TypeIndex::UInt32);
+		CachedArgumentInfos.Emplace(TypeHandleName, ERigVMPinDirection::Hidden, RigVMTypeUtils::TypeIndex::UInt32);
 	}
 
-	return Infos;
+	return CachedArgumentInfos;
 }
 
-FRigVMTemplateTypeMap FRigVMDispatch_GetLayerParameter::OnNewArgumentType(const FName& InArgumentName, TRigVMTypeIndex InTypeIndex) const
+FRigVMTemplateTypeMap FRigVMDispatch_GetLayerParameter::OnNewArgumentType(const FName& InArgumentName, TRigVMTypeIndex InTypeIndex, FRigVMRegistryHandle& InRegistry) const
 {
 	FRigVMTemplateTypeMap Types;
 	Types.Add(ParameterName, RigVMTypeUtils::TypeIndex::FName);
diff --git a/Engine/Plugins/Experimental/AnimNext/Source/AnimNext/Private/Param/RigVMDispatch_GetParameter.cpp b/Engine/Plugins/Experimental/AnimNext/Source/AnimNext/Private/Param/RigVMDispatch_GetParameter.cpp
--- a/Engine/Plugins/Experimental/AnimNext/Source/AnimNext/Private/Param/RigVMDispatch_GetParameter.cpp
+++ b/Engine/Plugins/Experimental/AnimNext/Source/AnimNext/Private/Param/RigVMDispatch_GetParameter.cpp
@@ -13,7 +13,7 @@ FRigVMDispatch_GetParameter::FRigVMDispatch_GetParameter()
 	FactoryScriptStruct = StaticStruct();
 }
 
-FName FRigVMDispatch_GetParameter::GetArgumentNameForOperandIndex(int32 InOperandIndex, int32 InTotalOperands) const
+FName FRigVMDispatch_GetParameter::GetArgumentNameForOperandIndex(int32 InOperandIndex, int32 InTotalOperands, FRigVMRegistryHandle& InRegistry) const
 {
 	static const FName ArgumentNames[] =
 	{
@@ -43,10 +43,9 @@ FString FRigVMDispatch_GetParameter::GetArgumentMetaData(const FName& InArgument
 }
 #endif
 
-const TArray<FRigVMTemplateArgumentInfo>& FRigVMDispatch_GetParameter::GetArgumentInfos() const
+const TArray<FRigVMTemplateArgumentInfo>& FRigVMDispatch_GetParameter::GetArgumentInfos(FRigVMRegistryHandle& InRegistry) const
 {
-	static TArray<FRigVMTemplateArgumentInfo> Infos;
-	if(Infos.IsEmpty())
+	if(CachedArgumentInfos.IsEmpty())
 	{
 		static const TArray<FRigVMTemplateArgument::ETypeCategory> ValueCategories =
 		{
@@ -54,16 +53,16 @@ const TArray<FRigVMTemplateArgumentInfo>& FRigVMDispatch_GetParameter::GetArgume
 			FRigVMTemplateArgument::ETypeCategory_ArrayAnyValue
 		};
 
-		Infos.Emplace(ParameterName, ERigVMPinDirection::Input, RigVMTypeUtils::TypeIndex::FName);
-		Infos.Emplace(ValueName, ERigVMPinDirection::Output, ValueCategories);
-		Infos.Emplace(ParameterIdName, ERigVMPinDirection::Hidden, RigVMTypeUtils::TypeIndex::UInt32);
-		Infos.Emplace(TypeHandleName, ERigVMPinDirection::Hidden, RigVMTypeUtils::TypeIndex::UInt32);
+		CachedArgumentInfos.Emplace(ParameterName, ERigVMPinDirection::Input, RigVMTypeUtils::TypeIndex::FName);
+		CachedArgumentInfos.Emplace(ValueName, ERigVMPinDirection::Output, ValueCategories);
+		CachedArgumentInfos.Emplace(ParameterIdName, ERigVMPinDirection::Hidden, RigVMTypeUtils::TypeIndex::UInt32);
+		CachedArgumentInfos.Emplace(TypeHandleName, ERigVMPinDirection::Hidden, RigVMTypeUtils::TypeIndex::UInt32);
 	}
 
-	return Infos;
+	return CachedArgumentInfos;
 }
 
-FRigVMTemplateTypeMap FRigVMDispatch_GetParameter::OnNewArgumentType(const FName& InArgumentName, TRigVMTypeIndex InTypeIndex) const
+FRigVMTemplateTypeMap FRigVMDispatch_GetParameter::OnNewArgumentType(const FName& InArgumentName, TRigVMTypeIndex InTypeIndex, FRigVMRegistryHandle& InRegistry) const
 {
 	FRigVMTemplateTypeMap Types;
 	Types.Add(ParameterName, InTypeIndex);
diff --git a/Engine/Plugins/Experimental/AnimNext/Source/AnimNext/Private/Param/RigVMDispatch_GetScopedParameter.cpp b/Engine/Plugins/Experimental/AnimNext/Source/AnimNext/Private/Param/RigVMDispatch_GetScopedParameter.cpp
--- a/Engine/Plugins/Experimental/AnimNext/Source/AnimNext/Private/Param/RigVMDispatch_GetScopedParameter.cpp
+++ b/Engine/Plugins/Experimental/AnimNext/Source/AnimNext/Private/Param/RigVMDispatch_GetScopedParameter.cpp
@@ -13,7 +13,7 @@ FRigVMDispatch_GetScopedParameter::FRigVMDispatch_GetScopedParameter()
 	FactoryScriptStruct = StaticStruct();
 }
 
-FName FRigVMDispatch_GetScopedParameter::GetArgumentNameForOperandIndex(int32 InOperandIndex, int32 InTotalOperands) const
+FName FRigVMDispatch_GetScopedParameter::GetArgumentNameForOperandIndex(int32 InOperandIndex, int32 InTotalOperands, FRigVMRegistryHandle& InRegistry) const
 {
 	static const FName ArgumentNames[] =
 	{
@@ -42,10 +42,9 @@ FString FRigVMDispatch_GetScopedParameter::GetArgumentMetaData(const FName& InAr
 }
 #endif
 
-const TArray<FRigVMTemplateArgumentInfo>& FRigVMDispatch_GetScopedParameter::GetArgumentInfos() const
+const TArray<FRigVMTemplateArgumentInfo>& FRigVMDispatch_GetScopedParameter::GetArgumentInfos(FRigVMRegistryHandle& InRegistry) const
 {
-	static TArray<FRigVMTemplateArgumentInfo> Infos;
-	if(Infos.IsEmpty())
+	if(CachedArgumentInfos.IsEmpty())
 	{
 		static const TArray<FRigVMTemplateArgument::ETypeCategory> ValueCategories =
 		{
@@ -53,20 +52,18 @@ const TArray<FRigVMTemplateArgumentInfo>& FRigVMDispatch_GetScopedParameter::Get
 			FRigVMTemplateArgument::ETypeCategory_ArrayAnyValue
 		};
 
-		const FRigVMRegistry_NoLock& Registry = FRigVMRegistry_NoLock::GetForRead();
-		Infos.Emplace(ParameterName, ERigVMPinDirection::Input, RigVMTypeUtils::TypeIndex::FString);
-		Infos.Emplace(ValueName, ERigVMPinDirection::Output, ValueCategories);
-		Infos.Emplace(ParameterIdName, ERigVMPinDirection::Hidden, RigVMTypeUtils::TypeIndex::UInt32);
-		Infos.Emplace(TypeHandleName, ERigVMPinDirection::Hidden, RigVMTypeUtils::TypeIndex::UInt32);
+		CachedArgumentInfos.Emplace(ParameterName, ERigVMPinDirection::Input, RigVMTypeUtils::TypeIndex::FString);
+		CachedArgumentInfos.Emplace(ValueName, ERigVMPinDirection::Output, ValueCategories);
+		CachedArgumentInfos.Emplace(ParameterIdName, ERigVMPinDirection::Hidden, RigVMTypeUtils::TypeIndex::UInt32);
+		CachedArgumentInfos.Emplace(TypeHandleName, ERigVMPinDirection::Hidden, RigVMTypeUtils::TypeIndex::UInt32);
 	}
 
-	return Infos;
+	return CachedArgumentInfos;
 }
  
-FRigVMTemplateTypeMap FRigVMDispatch_GetScopedParameter::OnNewArgumentType(const FName& InArgumentName, TRigVMTypeIndex InTypeIndex) const
+FRigVMTemplateTypeMap FRigVMDispatch_GetScopedParameter::OnNewArgumentType(const FName& InArgumentName, TRigVMTypeIndex InTypeIndex, FRigVMRegistryHandle& InRegistry) const
 {
 	FRigVMTemplateTypeMap Types;
-	const FRigVMRegistry_NoLock& Registry = FRigVMRegistry_NoLock::GetForRead();
 	Types.Add(ParameterName, RigVMTypeUtils::TypeIndex::FString);
 	Types.Add(ValueName, InTypeIndex);
 	Types.Add(ParameterIdName, RigVMTypeUtils::TypeIndex::UInt32);
diff --git a/Engine/Plugins/Experimental/AnimNext/Source/AnimNext/Private/Param/RigVMDispatch_SetLayerParameter.cpp b/Engine/Plugins/Experimental/AnimNext/Source/AnimNext/Private/Param/RigVMDispatch_SetLayerParameter.cpp
--- a/Engine/Plugins/Experimental/AnimNext/Source/AnimNext/Private/Param/RigVMDispatch_SetLayerParameter.cpp
+++ b/Engine/Plugins/Experimental/AnimNext/Source/AnimNext/Private/Param/RigVMDispatch_SetLayerParameter.cpp
@@ -15,7 +15,7 @@ FRigVMDispatch_SetLayerParameter::FRigVMDispatch_SetLayerParameter()
 	FactoryScriptStruct = StaticStruct();
 }
 
-FName FRigVMDispatch_SetLayerParameter::GetArgumentNameForOperandIndex(int32 InOperandIndex, int32 InTotalOperands) const
+FName FRigVMDispatch_SetLayerParameter::GetArgumentNameForOperandIndex(int32 InOperandIndex, int32 InTotalOperands, FRigVMRegistryHandle& InRegistry) const
 {
 	static const FName ArgumentNames[] = 
 	{
@@ -45,10 +45,9 @@ FString FRigVMDispatch_SetLayerParameter::GetArgumentMetaData(const FName& InArg
 }
 #endif
 
-const TArray<FRigVMTemplateArgumentInfo>& FRigVMDispatch_SetLayerParameter::GetArgumentInfos() const
+const TArray<FRigVMTemplateArgumentInfo>& FRigVMDispatch_SetLayerParameter::GetArgumentInfos(FRigVMRegistryHandle& InRegistry) const
 {
-	static TArray<FRigVMTemplateArgumentInfo> Infos;
-	if(Infos.IsEmpty())
+	if(CachedArgumentInfos.IsEmpty())
 	{
 		static const TArray<FRigVMTemplateArgument::ETypeCategory> ValueCategories =
 		{
@@ -56,13 +55,13 @@ const TArray<FRigVMTemplateArgumentInfo>& FRigVMDispatch_SetLayerParameter::GetA
 			FRigVMTemplateArgument::ETypeCategory_ArrayAnyValue
 		};
 	
-		Infos.Emplace(ParameterName, ERigVMPinDirection::Input, RigVMTypeUtils::TypeIndex::FName);
-		Infos.Emplace(ValueName, ERigVMPinDirection::Input, ValueCategories);
-		Infos.Emplace(ParameterIdName, ERigVMPinDirection::Hidden, RigVMTypeUtils::TypeIndex::UInt32);
-		Infos.Emplace(TypeHandleName, ERigVMPinDirection::Hidden, RigVMTypeUtils::TypeIndex::UInt32);
+		CachedArgumentInfos.Emplace(ParameterName, ERigVMPinDirection::Input, RigVMTypeUtils::TypeIndex::FName);
+		CachedArgumentInfos.Emplace(ValueName, ERigVMPinDirection::Input, ValueCategories);
+		CachedArgumentInfos.Emplace(ParameterIdName, ERigVMPinDirection::Hidden, RigVMTypeUtils::TypeIndex::UInt32);
+		CachedArgumentInfos.Emplace(TypeHandleName, ERigVMPinDirection::Hidden, RigVMTypeUtils::TypeIndex::UInt32);
 	}
 
-	return Infos;
+	return CachedArgumentInfos;
 }
 
 TArray<FRigVMExecuteArgument>& FRigVMDispatch_SetLayerParameter::GetExecuteArguments_Impl(const FRigVMDispatchContext& InContext) const
@@ -71,7 +70,7 @@ TArray<FRigVMExecuteArgument>& FRigVMDispatch_SetLayerParameter::GetExecuteArgum
 	return Arguments;
 }
 
-FRigVMTemplateTypeMap FRigVMDispatch_SetLayerParameter::OnNewArgumentType(const FName& InArgumentName, TRigVMTypeIndex InTypeIndex) const
+FRigVMTemplateTypeMap FRigVMDispatch_SetLayerParameter::OnNewArgumentType(const FName& InArgumentName, TRigVMTypeIndex InTypeIndex, FRigVMRegistryHandle& InRegistry) const
 {
 	FRigVMTemplateTypeMap Types;
 	Types.Add(ParameterName, RigVMTypeUtils::TypeIndex::FName);
diff --git a/Engine/Plugins/Experimental/AnimNext/Source/AnimNext/Private/Variables/RigVMDispatch_CallHoistedAccessorFunction.cpp b/Engine/Plugins/Experimental/AnimNext/Source/AnimNext/Private/Variables/RigVMDispatch_CallHoistedAccessorFunction.cpp
--- a/Engine/Plugins/Experimental/AnimNext/Source/AnimNext/Private/Variables/RigVMDispatch_CallHoistedAccessorFunction.cpp
+++ b/Engine/Plugins/Experimental/AnimNext/Source/AnimNext/Private/Variables/RigVMDispatch_CallHoistedAccessorFunction.cpp
@@ -9,7 +9,7 @@ FRigVMDispatch_CallHoistedAccessorFunctionBase::FRigVMDispatch_CallHoistedAccess
 	FactoryScriptStruct = StaticStruct();
 }
 
-FName FRigVMDispatch_CallHoistedAccessorFunctionBase::GetArgumentNameForOperandIndex(int32 InOperandIndex, int32 InTotalOperands) const
+FName FRigVMDispatch_CallHoistedAccessorFunctionBase::GetArgumentNameForOperandIndex(int32 InOperandIndex, int32 InTotalOperands, FRigVMRegistryHandle& InRegistry) const
 {
 	static const FName ArgumentNames[] =
 	{
@@ -21,17 +21,17 @@ FName FRigVMDispatch_CallHoistedAccessorFunctionBase::GetArgumentNameForOperandI
 	return ArgumentNames[InOperandIndex];
 }
 
-void FRigVMDispatch_CallHoistedAccessorFunctionBase::RegisterDependencyTypes_NoLock(FRigVMRegistry_NoLock& InRegistry) const
+void FRigVMDispatch_CallHoistedAccessorFunctionBase::RegisterDependencyTypes_NoLock(FRigVMRegistryHandle& InRegistry) const
 {
 	static UScriptStruct* const AllowedStructTypes[] =
 	{
 		FAnimNextSoftFunctionPtr::StaticStruct(),
 	};
 
-	InRegistry.RegisterStructTypes_NoLock(AllowedStructTypes);
+	InRegistry->RegisterStructTypes_NoLock(AllowedStructTypes);
 	for(UScriptStruct* const ScriptStruct : AllowedStructTypes)
 	{
-		InRegistry.FindOrAddType_NoLock(FRigVMTemplateArgumentType(ScriptStruct));
+		InRegistry->FindOrAddType_NoLock(FRigVMTemplateArgumentType(ScriptStruct));
 	}
 
 	static TPair<UClass*, FRigVMRegistry::ERegisterObjectOperation> const AllowedObjectTypes[] =
@@ -39,18 +39,17 @@ void FRigVMDispatch_CallHoistedAccessorFunctionBase::RegisterDependencyTypes_NoL
 		{ UObject::StaticClass(), FRigVMRegistry::ERegisterObjectOperation::Class },
 	};
 
-	InRegistry.RegisterObjectTypes_NoLock(AllowedObjectTypes);
+	InRegistry->RegisterObjectTypes_NoLock(AllowedObjectTypes);
 
 	for(const TPair<UClass*, FRigVMRegistry::ERegisterObjectOperation>& ObjectType : AllowedObjectTypes)
 	{
-		InRegistry.FindOrAddType_NoLock(FRigVMTemplateArgumentType(ObjectType.Key));
+		InRegistry->FindOrAddType_NoLock(FRigVMTemplateArgumentType(ObjectType.Key));
 	}
 }
 
-const TArray<FRigVMTemplateArgumentInfo>& FRigVMDispatch_CallHoistedAccessorFunctionBase::GetArgumentInfos() const
+const TArray<FRigVMTemplateArgumentInfo>& FRigVMDispatch_CallHoistedAccessorFunctionBase::GetArgumentInfos(FRigVMRegistryHandle& InRegistry) const
 {
-	static TArray<FRigVMTemplateArgumentInfo> Infos;
-	if(Infos.IsEmpty())
+	if(CachedArgumentInfos.IsEmpty())
 	{
 		static const TArray<FRigVMTemplateArgument::ETypeCategory> ValueCategories =
 		{
@@ -58,22 +57,19 @@ const TArray<FRigVMTemplateArgumentInfo>& FRigVMDispatch_CallHoistedAccessorFunc
 			FRigVMTemplateArgument::ETypeCategory_ArrayAnyValue
 		};
 
-		const FRigVMRegistry_NoLock& Registry = FRigVMRegistry_NoLock::GetForRead();
-		Infos.Emplace(ObjectName, ERigVMPinDirection::Input, Registry.GetTypeIndex_NoLock<UObject>());
-		Infos.Emplace(FunctionName, ERigVMPinDirection::Input, Registry.GetTypeIndex_NoLock<FAnimNextSoftFunctionPtr>());
-		Infos.Emplace(ValueName, ERigVMPinDirection::Output, ValueCategories);
+		CachedArgumentInfos.Emplace(ObjectName, ERigVMPinDirection::Input, InRegistry->GetTypeIndex_NoLock<UObject>());
+		CachedArgumentInfos.Emplace(FunctionName, ERigVMPinDirection::Input, InRegistry->GetTypeIndex_NoLock<FAnimNextSoftFunctionPtr>());
+		CachedArgumentInfos.Emplace(ValueName, ERigVMPinDirection::Output, ValueCategories);
 	}
 
-	return Infos;
+	return CachedArgumentInfos;
 }
 
-FRigVMTemplateTypeMap FRigVMDispatch_CallHoistedAccessorFunctionBase::OnNewArgumentType(const FName& InArgumentName, TRigVMTypeIndex InTypeIndex) const
+FRigVMTemplateTypeMap FRigVMDispatch_CallHoistedAccessorFunctionBase::OnNewArgumentType(const FName& InArgumentName, TRigVMTypeIndex InTypeIndex, FRigVMRegistryHandle& InRegistry) const
 {
-	const FRigVMRegistry_NoLock& Registry = FRigVMRegistry_NoLock::GetForRead();
-	
 	FRigVMTemplateTypeMap Types;
-	Types.Add(ObjectName, Registry.GetTypeIndex_NoLock<UObject>());
-	Types.Add(FunctionName, Registry.GetTypeIndex_NoLock<FAnimNextSoftFunctionPtr>());
+	Types.Add(ObjectName, InRegistry->GetTypeIndex_NoLock<UObject>());
+	Types.Add(FunctionName, InRegistry->GetTypeIndex_NoLock<FAnimNextSoftFunctionPtr>());
 	Types.Add(ValueName, InTypeIndex);
 	return Types;
 }
