Control Rig: Spawn index base parenting of modules, module reordering

#rnx

diff --git a/Engine/Plugins/Animation/ControlRig/Source/ControlRig/Private/ModularRig.cpp b/Engine/Plugins/Animation/ControlRig/Source/ControlRig/Private/ModularRig.cpp
--- a/Engine/Plugins/Animation/ControlRig/Source/ControlRig/Private/ModularRig.cpp
+++ b/Engine/Plugins/Animation/ControlRig/Source/ControlRig/Private/ModularRig.cpp
@@ -3,6 +3,7 @@
 #include "ModularRig.h"
 #include "Units/Execution/RigUnit_BeginExecution.h"
 #include "Units/Execution/RigUnit_InteractionExecution.h"
+#include "Units/Execution/RigUnit_PrepareForExecution.h"
 #include "ControlRigObjectBinding.h"
 #include "Rigs/RigHierarchyController.h"
 #include "ControlRigComponent.h"
@@ -535,6 +536,9 @@ void UModularRig::ExecuteQueue()
 #if WITH_EDITOR
 	TMap<FRigModuleInstance*, FFirstEntryEventGuard> FirstModuleEvent;
 #endif
+
+	int32 PostConstructionSpawnIndex = INDEX_NONE;
+	TArray<FRigModuleInstance*> ModulesWithPendingSpawnIndex;
 	
 	while(ExecutionQueue.IsValidIndex(ExecutionQueueFront))
 	{
@@ -543,6 +547,18 @@ void UModularRig::ExecuteQueue()
 		{
 			if (UControlRig* ModuleRig = ModuleInstance->GetRig())
 			{
+				// Let the module know of the spawn index the hierarchy is currently going through
+				if(ExecutionElement.EventName == FRigUnit_PrepareForExecution::EventName)
+				{
+					ModuleInstance->ConstructionSpawnStartIndex = PostConstructionSpawnIndex = Hierarchy->GetNextSpawnIndex();
+					ModulesWithPendingSpawnIndex.Add(ModuleInstance);
+				}
+				else if(ExecutionElement.EventName == FRigUnit_PostPrepareForExecution::EventName)
+				{
+					ModuleInstance->PostConstructionSpawnStartIndex = Hierarchy->GetNextSpawnIndex();
+					ModulesWithPendingSpawnIndex.Remove(ModuleInstance);
+				}
+
 				if (!ModuleRig->SupportsEvent(ExecutionElement.EventName))
 				{
 					ExecutionQueueFront++;
@@ -656,6 +672,11 @@ void UModularRig::ExecuteQueue()
 		
 		ExecutionQueueFront++;
 	}
+
+	for(FRigModuleInstance* Module : ModulesWithPendingSpawnIndex)
+	{
+		Module->PostConstructionSpawnStartIndex = Hierarchy->GetNextSpawnIndex();
+	}
 }
 
 void UModularRig::ResetExecutionQueue()
diff --git a/Engine/Plugins/Animation/ControlRig/Source/ControlRig/Private/ModularRigController.cpp b/Engine/Plugins/Animation/ControlRig/Source/ControlRig/Private/ModularRigController.cpp
--- a/Engine/Plugins/Animation/ControlRig/Source/ControlRig/Private/ModularRigController.cpp
+++ b/Engine/Plugins/Animation/ControlRig/Source/ControlRig/Private/ModularRigController.cpp
@@ -339,34 +339,38 @@ bool UModularRigController::ConnectConnectorToElements(const FRigElementKey& InC
 
 	// First disconnect before connecting to anything else. This might disconnect other secondary/optional connectors.
 	TMap<FRigElementKey, TArray<FRigElementKey>> PreviousConnections;
-	if (CurrentTarget.IsValid())
 	{
-		const TGuardValue<bool> DisableAutomaticReparenting(bAutomaticReparenting, false);
-		DisconnectConnector_Internal(InConnectorKey, false, &PreviousConnections, bSetupUndo);
-	}
+		TGuardValue<bool> NotificationsGuard(bSuspendNotifications, true);
+		
+		if (CurrentTarget.IsValid())
+		{
+			const TGuardValue<bool> DisableAutomaticReparenting(bAutomaticReparenting, false);
+			DisconnectConnector_Internal(InConnectorKey, false, &PreviousConnections, bSetupUndo);
+		}
 
-	Model->Connections.AddConnection(InConnectorKey, TargetKeys);
+		Model->Connections.AddConnection(InConnectorKey, TargetKeys);
 
-	// restore previous connections if possible
-	for(const TPair<FRigElementKey, TArray<FRigElementKey>>& PreviousConnection : PreviousConnections)
-	{
-		if(!Model->Connections.HasConnection(PreviousConnection.Key))
+		// restore previous connections if possible
+		for(const TPair<FRigElementKey, TArray<FRigElementKey>>& PreviousConnection : PreviousConnections)
 		{
-			FText ErrorMessageForPreviousConnection;
-			if(CanConnectConnectorToElements(PreviousConnection.Key, PreviousConnection.Value, ErrorMessageForPreviousConnection))
+			if(!Model->Connections.HasConnection(PreviousConnection.Key))
 			{
-				(void)ConnectConnectorToElements(PreviousConnection.Key, PreviousConnection.Value, bSetupUndo, false, false);
+				FText ErrorMessageForPreviousConnection;
+				if(CanConnectConnectorToElements(PreviousConnection.Key, PreviousConnection.Value, ErrorMessageForPreviousConnection))
+				{
+					(void)ConnectConnectorToElements(PreviousConnection.Key, PreviousConnection.Value, bSetupUndo, false, false);
+				}
 			}
 		}
 	}
-	
-	Notify(EModularRigNotification::ConnectionChanged, Module);
 
 #if WITH_EDITOR
 	if (UControlRig* RigCDO = Module->Class->GetDefaultObject<UControlRig>())
 	{
 		if (ModularRig)
 		{
+			TGuardValue<bool> NotificationsGuard(bSuspendNotifications, true);
+
 			if (URigHierarchy* Hierarchy = ModularRig->GetHierarchy())
 			{
 				bool bResolvedPrimaryConnector = false;
@@ -413,10 +417,12 @@ bool UModularRigController::ConnectConnectorToElements(const FRigElementKey& InC
 		
 		if (TArray<FRigElementKey>* Previous = PreviousConnections.Find(InConnectorKey))
 		{
-			ConnectConnectorToElements(InConnectorKey, *Previous, bSetupUndo, bAutomaticReparenting, false);
+			return ConnectConnectorToElements(InConnectorKey, *Previous, bSetupUndo, bAutomaticReparenting, false);
 		}
 	}
 
+	Notify(EModularRigNotification::ConnectionChanged, Module);
+
 #if WITH_EDITOR
 	TransactionPtr.Reset();
 #endif
@@ -615,6 +621,23 @@ TArray<FRigElementKey> UModularRigController::DisconnectCyclicConnectors(bool bS
 		const FName ConnectorModuleName = Hierarchy->GetModuleFName(Connection.Connector);
 		if(const FRigModuleReference* ConnectorModule = Model->FindModule(ConnectorModuleName))
 		{
+			const FRigConnectorElement* ConnectorElement = Hierarchy->Find<FRigConnectorElement>(Connection.Connector, false);
+			if(ConnectorElement == nullptr)
+			{
+				continue;
+			}
+
+			const FRigModuleInstance* ConnectorModuleInstance = ModularRig->FindModule(ConnectorModuleName);
+			if(ConnectorModuleInstance == nullptr)
+			{
+				continue;
+			}
+
+			const int32& SpawnStartIndex =
+				ConnectorElement->IsPostConstructionConnector() ?
+				ConnectorModuleInstance->PostConstructionSpawnStartIndex :
+				ConnectorModuleInstance->ConstructionSpawnStartIndex;
+
 			for(const FRigElementKey& Target : Connection.Targets)
 			{
 				const FName TargetModuleName = Hierarchy->GetModuleFName(Target);
@@ -631,7 +654,18 @@ TArray<FRigElementKey> UModularRigController::DisconnectCyclicConnectors(bool bS
 					continue;
 				}
 
-				if(!Model->IsModuleParentedTo(ConnectorModule, TargetModule))
+				const int32 TargetSpawnIndex = Hierarchy->GetSpawnIndex(Target);
+				if(TargetSpawnIndex != INDEX_NONE)
+				{
+					if(TargetSpawnIndex >= SpawnStartIndex)
+					{
+						ConnectorsToDisconnect.Add(Connection.Connector);
+						break;
+					}
+				}
+
+				// alternatively if we don't have a valid spawn index relay back to module parenting rules
+				else if(!Model->IsModuleParentedTo(ConnectorModule, TargetModule))
 				{
 					ConnectorsToDisconnect.Add(Connection.Connector);
 					break;
@@ -1566,6 +1600,105 @@ bool UModularRigController::ReparentModule(const FName& InModuleName, const FNam
 	return true;
 }
 
+bool UModularRigController::ReorderModule(const FName& InModuleName, int32 InModuleIndex, bool bSetupUndo)
+{
+	FRigModuleReference* Module = FindModule(InModuleName);
+	if (!Module)
+	{
+		UE_LOG(LogControlRig, Error, TEXT("Could not find module %s"), *InModuleName.ToString());
+		return false;
+	}
+
+	TArray<FRigModuleReference*> ChildModules;
+	
+	const FRigModuleReference* ParentModule = Module->GetParentModule();
+	if(ParentModule)
+	{
+		ChildModules = ParentModule->CachedChildren;
+	}
+	else
+	{
+		ChildModules = Model->RootModules;
+	}
+
+	if(ChildModules.Num() == 1)
+	{
+		UE_LOG(LogControlRig, Error, TEXT("Cannot reorder sole module %s"), *InModuleName.ToString());
+		return false;
+	}
+
+	if(InModuleIndex < 0 || InModuleIndex >= ChildModules.Num())
+	{
+		UE_LOG(LogControlRig, Error, TEXT("Cannot reorder module %s, index is out of range (0 to %d)"), *InModuleName.ToString(), ChildModules.Num() - 1);
+		return false;
+	}
+
+	int32 CurrentIndex = ChildModules.Find(Module);
+	if(CurrentIndex == InModuleIndex)
+	{
+		UE_LOG(LogControlRig, Error, TEXT("Cannot reorder module %s, module is already at desired index"), *InModuleName.ToString());
+		return false;
+	}
+
+#if WITH_EDITOR
+	TSharedPtr<FScopedTransaction> TransactionPtr;
+	if (bSetupUndo)
+	{
+		TransactionPtr = MakeShared<FScopedTransaction>(NSLOCTEXT("ModularRigController", "ReorderModuleTransaction", "Reorder Module"), !GIsTransacting);
+		if(UBlueprint* Blueprint = Cast<UBlueprint>(GetOuter()))
+		{
+			Blueprint->Modify();
+		}
+	}
+#endif
+
+	// create a flat list of module name
+	TArray<FName> AllModuleNames;
+	for(const FRigModuleReference& ModuleReference : Model->Modules)
+	{
+		AllModuleNames.Add(ModuleReference.Name);
+	}
+
+	// now create a list of child names
+	TArray<FName> ChildModuleNames;
+	for(const FRigModuleReference* ChildModule : ChildModules)
+	{
+		ChildModuleNames.Add(ChildModule->Name);
+	}
+
+	// move the module up
+	while(CurrentIndex > InModuleIndex)
+	{
+		ChildModuleNames.Swap(CurrentIndex, CurrentIndex-1);
+		AllModuleNames.Swap(AllModuleNames.Find(ChildModuleNames[CurrentIndex]), AllModuleNames.Find(ChildModuleNames[CurrentIndex-1]));
+		CurrentIndex--;
+	}
+
+	// move the module down
+	while(CurrentIndex < InModuleIndex)
+	{
+		ChildModuleNames.Swap(CurrentIndex, CurrentIndex+1);
+		AllModuleNames.Swap(AllModuleNames.Find(ChildModuleNames[CurrentIndex]), AllModuleNames.Find(ChildModuleNames[CurrentIndex+1]));
+		CurrentIndex++;
+	}
+
+	// Sort the modules by their new name order
+	Algo::SortBy(Model->Modules, [AllModuleNames](const FRigModuleReference& Module) -> int32
+	{
+		return AllModuleNames.IndexOfByKey(Module.Name);
+	});
+
+	Model->UpdateCachedChildren();
+
+	Notify(EModularRigNotification::ModuleReordered, Module);
+
+#if WITH_EDITOR
+	TransactionPtr.Reset();
+#endif
+	
+	return true;
+}
+
 FName UModularRigController::MirrorModule(const FName& InModuleName, const FRigVMMirrorSettings& InSettings, bool bSetupUndo)
 {
 	FRigModuleReference* OriginalModule = FindModule(InModuleName);
diff --git a/Engine/Plugins/Animation/ControlRig/Source/ControlRig/Private/ModularRigRuleManager.cpp b/Engine/Plugins/Animation/ControlRig/Source/ControlRig/Private/ModularRigRuleManager.cpp
--- a/Engine/Plugins/Animation/ControlRig/Source/ControlRig/Private/ModularRigRuleManager.cpp
+++ b/Engine/Plugins/Animation/ControlRig/Source/ControlRig/Private/ModularRigRuleManager.cpp
@@ -228,6 +228,21 @@ void UModularRigRuleManager::FilterInvalidModules(FWorkData& InOutWorkData)
 		}
 		else
 		{
+			const int32& SpawnStartIndex =
+				InOutWorkData.Connector->IsPostConstructionConnector() ?
+					InOutWorkData.Module->PostConstructionSpawnStartIndex :
+					InOutWorkData.Module->ConstructionSpawnStartIndex;
+
+			const int32 TargetSpawnIndex = InOutWorkData.Hierarchy->GetSpawnIndex(Result.GetKey());
+
+			if(SpawnStartIndex != INDEX_NONE && TargetSpawnIndex >= SpawnStartIndex)
+			{
+				static const FText CannotConnectBelowNameSpaceMessage = LOCTEXT("CannotConnectLaterSpawnElement", "Cannot connect to an element spawn later than the connector's module.");
+				Result.SetInvalidTarget(CannotConnectBelowNameSpaceMessage);
+				return;
+			}
+
+			// if we don't have a spawn index at all - let's rely on module hierarchy only.
 			if(InOutWorkData.Module && InOutWorkData.Module->HasChildModule(MatchModuleName, true))
 			{
 				static const FText CannotConnectBelowNameSpaceMessage = LOCTEXT("CannotConnectBelowNameSpace", "Cannot connect to element below the connector's namespace.");
diff --git a/Engine/Plugins/Animation/ControlRig/Source/ControlRig/Private/Rigs/RigHierarchy.cpp b/Engine/Plugins/Animation/ControlRig/Source/ControlRig/Private/Rigs/RigHierarchy.cpp
--- a/Engine/Plugins/Animation/ControlRig/Source/ControlRig/Private/Rigs/RigHierarchy.cpp
+++ b/Engine/Plugins/Animation/ControlRig/Source/ControlRig/Private/Rigs/RigHierarchy.cpp
@@ -1749,9 +1749,9 @@ bool URigHierarchy::IsProcedural(const FRigHierarchyKey& InKey) const
 	return false;
 }
 
-int32 URigHierarchy::GetIndex(const FRigElementKey& InKey) const
+int32 URigHierarchy::GetIndex(const FRigElementKey& InKey, bool bFollowRedirector) const
 {
-	if(ElementKeyRedirector)
+	if(ElementKeyRedirector && bFollowRedirector)
 	{
 		if(FRigElementKeyRedirector::FCachedKeyArray* CachedRigElements = ElementKeyRedirector->Find(InKey))
 		{
@@ -1799,6 +1799,22 @@ int32 URigHierarchy::GetIndex(const FRigElementKey& InKey) const
 	return INDEX_NONE;
 }
 
+int32 URigHierarchy::GetSpawnIndex(const FRigHierarchyKey& InKey) const
+{
+	if(InKey.IsComponent())
+	{
+		if(const FRigBaseComponent* Component = FindComponent(InKey.GetComponent()))
+		{
+			return Component->GetSpawnIndex();
+		}
+	}
+	else if(const FRigBaseElement* Element = Find(InKey.GetElement()))
+	{
+		return Element->GetSpawnIndex();
+	}
+	return INDEX_NONE;
+}
+
 const FRigBaseComponent* URigHierarchy::FindComponent(const FRigComponentKey& InKey) const
 {
 	const int32 ComponentIndex = GetComponentIndex(InKey);
@@ -2072,6 +2088,11 @@ bool URigHierarchy::CanAddComponent(FRigElementKey InElementKey, const FRigBaseC
 	return true;
 }
 
+int32 URigHierarchy::GetNextSpawnIndex() const
+{
+	return Num() + NumComponents();
+}
+
 TArray<FRigSocketState> URigHierarchy::GetSocketStates() const
 {
 	const TArray<FRigElementKey> Keys = GetSocketKeys(true);
@@ -6849,6 +6870,7 @@ FRigBaseComponent* URigHierarchy::MakeComponent(const UScriptStruct* InComponent
 	check(InElement->GetOwner() == this);
 	
 	int32 ComponentIndex = INDEX_NONE;
+	const int32 SpawnIndex = GetNextSpawnIndex();
 	for(int32 ExistingIndex = 0; ExistingIndex < ElementComponents.Num(); ExistingIndex++)
 	{
 		if(!ElementComponents[ExistingIndex].IsValid())
@@ -6872,6 +6894,7 @@ FRigBaseComponent* URigHierarchy::MakeComponent(const UScriptStruct* InComponent
 	Component->Key = FRigComponentKey(InElement->GetKey(), UniqueName.GetFName());
 	Component->Element = InElement;
 	Component->IndexInElement = InElement->ComponentIndices.Add(Component->IndexInHierarchy);
+	Component->SpawnIndex = SpawnIndex;
 
 	ComponentIndexLookup.Add(Component->GetKey(), Component->GetIndexInHierarchy());
 	
diff --git a/Engine/Plugins/Animation/ControlRig/Source/ControlRig/Private/Rigs/RigHierarchyController.cpp b/Engine/Plugins/Animation/ControlRig/Source/ControlRig/Private/Rigs/RigHierarchyController.cpp
--- a/Engine/Plugins/Animation/ControlRig/Source/ControlRig/Private/Rigs/RigHierarchyController.cpp
+++ b/Engine/Plugins/Animation/ControlRig/Source/ControlRig/Private/Rigs/RigHierarchyController.cpp
@@ -2856,6 +2856,7 @@ int32 URigHierarchyController::AddElement(FRigBaseElement* InElementToAdd, FRigB
 
 	InElementToAdd->CachedNameString.Reset();
 	InElementToAdd->SubIndex = Hierarchy->Num(InElementToAdd->Key.Type);
+	InElementToAdd->SpawnIndex = Hierarchy->GetNextSpawnIndex();
 	InElementToAdd->Index = Hierarchy->Elements.Add(InElementToAdd);
 	Hierarchy->ElementsPerType[URigHierarchy::RigElementTypeToFlatIndex(InElementToAdd->GetKey().Type)].Add(InElementToAdd);
 	Hierarchy->ElementIndexLookup.Add(InElementToAdd->Key, InElementToAdd->Index);
diff --git a/Engine/Plugins/Animation/ControlRig/Source/ControlRig/Private/Rigs/RigHierarchyElements.cpp b/Engine/Plugins/Animation/ControlRig/Source/ControlRig/Private/Rigs/RigHierarchyElements.cpp
--- a/Engine/Plugins/Animation/ControlRig/Source/ControlRig/Private/Rigs/RigHierarchyElements.cpp
+++ b/Engine/Plugins/Animation/ControlRig/Source/ControlRig/Private/Rigs/RigHierarchyElements.cpp
@@ -253,6 +253,7 @@ void FRigBaseElement::InitializeFrom(const FRigBaseElement* InOther)
 	Index = InOther->Index;
 	SubIndex = InOther->SubIndex;
 	CreatedAtInstructionIndex = InOther->CreatedAtInstructionIndex;
+	SpawnIndex = InOther->SpawnIndex;
 	bSelected = false;
 }
 
@@ -2184,6 +2185,7 @@ FRigConnectorSettings::FRigConnectorSettings()
 	: Type(EConnectorType::Primary)
 	, bOptional(false)
 	, bIsArray(false)
+	, bPostConstruction(false)
 {
 }
 
@@ -2202,6 +2204,7 @@ void FRigConnectorSettings::Save(FArchive& Ar)
 	Ar << Type;
 	Ar << bOptional;
 	Ar << bIsArray;
+	Ar << bPostConstruction;
 
 	int32 NumRules = Rules.Num();
 	Ar << NumRules;
@@ -2233,6 +2236,15 @@ void FRigConnectorSettings::Load(FArchive& Ar)
 		bIsArray = false;
 	}
 
+	if (Ar.CustomVer(FControlRigObjectVersion::GUID) >= FControlRigObjectVersion::RigHierarchyPostConstructionConnectors)
+	{
+		Ar << bPostConstruction;
+	}
+	else
+	{
+		bPostConstruction = false;
+	}
+
 	int32 NumRules = 0;
 	Ar << NumRules;
 	Rules.SetNumZeroed(NumRules);
@@ -2260,6 +2272,10 @@ bool FRigConnectorSettings::operator==(const FRigConnectorSettings& InOther) con
 	{
 		return false;
 	}
+	if(bPostConstruction != InOther.bPostConstruction)
+	{
+		return false;
+	}
 	if(Rules.Num() != InOther.Rules.Num())
 	{
 		return false;
@@ -2289,6 +2305,7 @@ uint32 GetTypeHash(const FRigConnectorSettings& Settings)
 	uint32 Hash = HashCombine(GetTypeHash(Settings.Type), Settings.GetRulesHash());
 	Hash = HashCombine(Hash, GetTypeHash(Settings.bOptional));
 	Hash = HashCombine(Hash, GetTypeHash(Settings.bIsArray));
+	Hash = HashCombine(Hash, GetTypeHash(Settings.bPostConstruction));
 	return Hash;
 }
 
diff --git a/Engine/Plugins/Animation/ControlRig/Source/ControlRig/Private/Units/Execution/RigUnit_RigModules.cpp b/Engine/Plugins/Animation/ControlRig/Source/ControlRig/Private/Units/Execution/RigUnit_RigModules.cpp
--- a/Engine/Plugins/Animation/ControlRig/Source/ControlRig/Private/Units/Execution/RigUnit_RigModules.cpp
+++ b/Engine/Plugins/Animation/ControlRig/Source/ControlRig/Private/Units/Execution/RigUnit_RigModules.cpp
@@ -1,6 +1,7 @@
 // Copyright Epic Games, Inc. All Rights Reserved.
 
 #include "Units/Execution/RigUnit_RigModules.h"
+#include "Units/Execution/RigUnit_PrepareForExecution.h"
 #include "Units/RigUnitContext.h"
 #include "ControlRig.h"
 
@@ -26,6 +27,22 @@ FRigUnit_ResolveArrayConnector_Execute()
 
 	if (const URigHierarchy* Hierarchy = ExecuteContext.Hierarchy)
 	{
+		// make sure to only avoid post connectors during construction
+		if(const FRigConnectorElement* ConnectorElement = Hierarchy->Find<FRigConnectorElement>(Connector, false))
+		{
+			if(ConnectorElement->IsPostConstructionConnector())
+			{
+				if(ExecuteContext.GetEventName() == FRigUnit_PrepareForExecution::EventName)
+				{
+					static constexpr TCHAR ErrorMessageFormat[] = TEXT("Connector '%s' is only valid during the %s event.");
+					const FString ErrorMessage = FString::Printf(ErrorMessageFormat, *Connector.Name.ToString(), *FRigUnit_PostPrepareForExecution::EventName.ToString());
+					UE_CONTROLRIG_RIGUNIT_REPORT_ERROR(TEXT("%s"), *ErrorMessage);
+					bIsConnected = false;
+					return;
+				}
+			}
+		}
+		
 		Result = Hierarchy->GetResolvedTargets(Connector);
 		if(SkipSocket && !Result.IsEmpty())
 		{
diff --git a/Engine/Plugins/Animation/ControlRig/Source/ControlRig/Public/ControlRigObjectVersion.h b/Engine/Plugins/Animation/ControlRig/Source/ControlRig/Public/ControlRigObjectVersion.h
--- a/Engine/Plugins/Animation/ControlRig/Source/ControlRig/Public/ControlRigObjectVersion.h
+++ b/Engine/Plugins/Animation/ControlRig/Source/ControlRig/Public/ControlRigObjectVersion.h
@@ -127,6 +127,9 @@ struct FControlRigObjectVersion
 		// Previous name and parent maps serialized as FRigHierarchyKey
 		RigHierarchyPreviousNameAndParentMapUsingHierarchyKey,
 
+		// New setting for connectors to optionally specify their use only during post construction
+		RigHierarchyPostConstructionConnectors,
+
 		// -----<new versions can be added above this line>-------------------------------------------------
 		VersionPlusOne,
 		LatestVersion = VersionPlusOne - 1
diff --git a/Engine/Plugins/Animation/ControlRig/Source/ControlRig/Public/ModularRig.h b/Engine/Plugins/Animation/ControlRig/Source/ControlRig/Public/ModularRig.h
--- a/Engine/Plugins/Animation/ControlRig/Source/ControlRig/Public/ModularRig.h
+++ b/Engine/Plugins/Animation/ControlRig/Source/ControlRig/Public/ModularRig.h
@@ -57,6 +57,8 @@ public:
 	, RigPtr(nullptr)
 	, ParentModuleName(NAME_None)
 	, CachedParentModule(nullptr)
+	, ConstructionSpawnStartIndex(INDEX_NONE)
+	, PostConstructionSpawnStartIndex(INDEX_NONE)
 	{
 	}
 
@@ -82,6 +84,8 @@ public:
 	FRigModuleInstance* CachedParentModule;
 	TArray<FRigModuleInstance*> CachedChildren;
 	mutable FCachedRigElement PrimaryConnector;
+	int32 ConstructionSpawnStartIndex;
+	int32 PostConstructionSpawnStartIndex;
 
 	UE_API UControlRig* GetRig() const;
 	UE_API void SetRig(UControlRig* InRig);
diff --git a/Engine/Plugins/Animation/ControlRig/Source/ControlRig/Public/ModularRigController.h b/Engine/Plugins/Animation/ControlRig/Source/ControlRig/Public/ModularRigController.h
--- a/Engine/Plugins/Animation/ControlRig/Source/ControlRig/Public/ModularRigController.h
+++ b/Engine/Plugins/Animation/ControlRig/Source/ControlRig/Public/ModularRigController.h
@@ -85,6 +85,9 @@ class UModularRigController : public UObject
 	UFUNCTION(BlueprintCallable, Category = "ControlRig | Modules")
 	UE_API bool ReparentModule(const FName& InModuleName, const FName& InNewParentModuleName, bool bSetupUndo = true);
 
+	UFUNCTION(BlueprintCallable, Category = "ControlRig | Modules")
+	UE_API bool ReorderModule(const FName& InModuleName, int32 InModuleIndex, bool bSetupUndo = true);
+
 	UFUNCTION(BlueprintCallable, Category = "ControlRig | Modules")
 	UE_API FName MirrorModule(const FName& InModuleName, const FRigVMMirrorSettings& InSettings, bool bSetupUndo = true);
 
diff --git a/Engine/Plugins/Animation/ControlRig/Source/ControlRig/Public/ModularRigModel.h b/Engine/Plugins/Animation/ControlRig/Source/ControlRig/Public/ModularRigModel.h
--- a/Engine/Plugins/Animation/ControlRig/Source/ControlRig/Public/ModularRigModel.h
+++ b/Engine/Plugins/Animation/ControlRig/Source/ControlRig/Public/ModularRigModel.h
@@ -26,6 +26,8 @@ enum class EModularRigNotification : uint8
 
 	ModuleReparented,
 
+	ModuleReordered,
+
 	ConnectionChanged,
 
 	ModuleConfigValueChanged,
diff --git a/Engine/Plugins/Animation/ControlRig/Source/ControlRig/Public/Rigs/RigHierarchy.h b/Engine/Plugins/Animation/ControlRig/Source/ControlRig/Public/Rigs/RigHierarchy.h
--- a/Engine/Plugins/Animation/ControlRig/Source/ControlRig/Public/Rigs/RigHierarchy.h
+++ b/Engine/Plugins/Animation/ControlRig/Source/ControlRig/Public/Rigs/RigHierarchy.h
@@ -463,9 +463,18 @@ public:
 	/**
 	 * Returns the index of an element given its key
 	 * @param InKey The key of the element to retrieve the index for
+	 * @param bFollowRedirector false if we want to retrieve a connector itself
 	 * @return The index of the element or INDEX_NONE
 	 */
-	CONTROLRIG_API int32 GetIndex(const FRigElementKey& InKey) const;
+	CONTROLRIG_API int32 GetIndex(const FRigElementKey& InKey, bool bFollowRedirector = true) const;
+
+	/**
+	 * Returns the spawn index of an element / component given its key.
+	 * The spawn index indicates the spawning order of elements.
+	 * @param InKey The key of the element to retrieve the spawn index for
+	 * @return The spawn index of the element or INDEX_NONE
+	 */
+	CONTROLRIG_API int32 GetSpawnIndex(const FRigHierarchyKey& InKey) const;
 
 	/**
 	 * Returns the key and index pair of an element given its key
@@ -701,12 +710,13 @@ public:
 	 * This templated method also casts to the chosen
 	 * element type but does not guarantee a valid result.
 	 * @param InKey The key of the element to retrieve.
+	 * @param bFollowRedirector false if we want to retrieve a connector itself
 	 * @return The retrieved element or nullptr.
 	 */
 	template<typename T>
-	const T* Find(const FRigElementKey& InKey) const
+	const T* Find(const FRigElementKey& InKey, bool bFollowRedirector = true) const
 	{
-		return Get<T>(GetIndex(InKey));
+		return Get<T>(GetIndex(InKey, bFollowRedirector));
 	}
 
 	/**
@@ -878,6 +888,8 @@ public:
 	CONTROLRIG_API bool CanAddComponent(FRigElementKey InElementKey, const UScriptStruct* InComponentStruct, FString* OutFailureReason = nullptr) const;
 	CONTROLRIG_API bool CanAddComponent(FRigElementKey InElementKey, const FRigBaseComponent* InComponent, FString* OutFailureReason = nullptr) const;
 
+	CONTROLRIG_API int32 GetNextSpawnIndex() const;
+
 private:
 	/**
 	* Returns bone element for a given key, for scripting purpose only, for cpp usage, use Find<FRigBoneElement>()
diff --git a/Engine/Plugins/Animation/ControlRig/Source/ControlRig/Public/Rigs/RigHierarchyComponents.h b/Engine/Plugins/Animation/ControlRig/Source/ControlRig/Public/Rigs/RigHierarchyComponents.h
--- a/Engine/Plugins/Animation/ControlRig/Source/ControlRig/Public/Rigs/RigHierarchyComponents.h
+++ b/Engine/Plugins/Animation/ControlRig/Source/ControlRig/Public/Rigs/RigHierarchyComponents.h
@@ -82,6 +82,7 @@ public:
 		IndexInElement = InOther.IndexInElement;
 		IndexInHierarchy = InOther.IndexInHierarchy;
 		CreatedAtInstructionIndex = InOther.CreatedAtInstructionIndex;
+		SpawnIndex = InOther.SpawnIndex;
 		bSelected = InOther.bSelected;
 		return *this;
 	}
@@ -126,6 +127,7 @@ protected:
 	int32 IndexInHierarchy = INDEX_NONE;
 	int32 IndexInElement = INDEX_NONE;
 	int32 CreatedAtInstructionIndex = INDEX_NONE;
+	int32 SpawnIndex = INDEX_NONE;
 	bool bSelected = false;
 
 	mutable FString CachedNameString;
@@ -156,6 +158,7 @@ public:
 	int32 GetIndexInHierarchy() const { return IndexInHierarchy; }
 	bool IsSelected() const { return bSelected; }
 	int32 GetCreatedAtInstructionIndex() const { return CreatedAtInstructionIndex; }
+	int32 GetSpawnIndex() const { return SpawnIndex; }
 	bool IsProcedural() const { return CreatedAtInstructionIndex != INDEX_NONE; }
 	const FRigElementKey& GetElementKey() const { return Key.ElementKey; }
 	const FRigBaseElement* GetElement() const { return Element; }
diff --git a/Engine/Plugins/Animation/ControlRig/Source/ControlRig/Public/Rigs/RigHierarchyElements.h b/Engine/Plugins/Animation/ControlRig/Source/ControlRig/Public/Rigs/RigHierarchyElements.h
--- a/Engine/Plugins/Animation/ControlRig/Source/ControlRig/Public/Rigs/RigHierarchyElements.h
+++ b/Engine/Plugins/Animation/ControlRig/Source/ControlRig/Public/Rigs/RigHierarchyElements.h
@@ -878,6 +878,9 @@ protected:
 	UPROPERTY(BlueprintReadOnly, Transient, Category = RigElement, meta = (AllowPrivateAccess = "true"))
 	int32 CreatedAtInstructionIndex = INDEX_NONE;
 
+	UPROPERTY(BlueprintReadOnly, Transient, Category = RigElement, meta = (AllowPrivateAccess = "true"))
+	int32 SpawnIndex = INDEX_NONE;
+
 	UPROPERTY(BlueprintReadOnly, Transient, Category = RigElement, meta = (AllowPrivateAccess = "true"))
 	bool bSelected = false;
 
@@ -920,6 +923,7 @@ public:
 	int32 GetSubIndex() const { return SubIndex; }
 	bool IsSelected() const { return bSelected; }
 	int32 GetCreatedAtInstructionIndex() const { return CreatedAtInstructionIndex; }
+	int32 GetSpawnIndex() const { return SpawnIndex; }
 	bool IsProcedural() const { return CreatedAtInstructionIndex != INDEX_NONE; }
 	const URigHierarchy* GetOwner() const { return Owner; }
 	URigHierarchy* GetOwner() { return Owner; }
@@ -1959,6 +1963,9 @@ struct FRigConnectorSettings
 	UPROPERTY(BlueprintReadWrite, EditAnywhere, Category = Connector)
 	bool bIsArray;
 
+	UPROPERTY(BlueprintReadWrite, EditAnywhere, Category = Connector)
+	bool bPostConstruction;
+
 	UPROPERTY(BlueprintReadWrite, EditAnywhere, Category = Connector)
 	TArray<FRigConnectionRuleStash> Rules;
 
@@ -2033,6 +2040,7 @@ struct FRigConnectorElement final : public FRigBaseElement
 	bool IsSecondary() const { return Settings.Type == EConnectorType::Secondary; }
 	bool IsOptional() const { return IsSecondary() && Settings.bOptional; }
 	bool IsArrayConnector() const { return IsSecondary() && Settings.bIsArray; }
+	bool IsPostConstructionConnector() const { return IsSecondary() && Settings.bPostConstruction; }
 
 private:
 	explicit FRigConnectorElement(URigHierarchy* InOwner)
diff --git a/Engine/Plugins/Animation/ControlRig/Source/ControlRigDeveloper/Private/ControlRigBlueprint.cpp b/Engine/Plugins/Animation/ControlRig/Source/ControlRigDeveloper/Private/ControlRigBlueprint.cpp
--- a/Engine/Plugins/Animation/ControlRig/Source/ControlRigDeveloper/Private/ControlRigBlueprint.cpp
+++ b/Engine/Plugins/Animation/ControlRig/Source/ControlRigDeveloper/Private/ControlRigBlueprint.cpp
@@ -3094,6 +3094,7 @@ void UControlRigBlueprint::HandleRigModulesModified(EModularRigNotification InNo
 		}
 		case EModularRigNotification::ModuleRenamed:
 		case EModularRigNotification::ModuleReparented:
+		case EModularRigNotification::ModuleReordered:
 		{
 			if (InModule)
 			{
diff --git a/Engine/Plugins/Animation/ControlRig/Source/ControlRigEditor/Private/ControlRigElementDetails.cpp b/Engine/Plugins/Animation/ControlRig/Source/ControlRigEditor/Private/ControlRigElementDetails.cpp
--- a/Engine/Plugins/Animation/ControlRig/Source/ControlRigEditor/Private/ControlRigElementDetails.cpp
+++ b/Engine/Plugins/Animation/ControlRig/Source/ControlRigEditor/Private/ControlRigElementDetails.cpp
@@ -1112,6 +1112,21 @@ void FRigBaseElementDetails::CustomizeDetails(IDetailLayoutBuilder& DetailBuilde
 		];
 	}
 
+	GeneralCategory.AddCustomRow(FText::FromString(TEXT("Spawn Index")), true /* advanced */)
+	.IsEnabled(!bIsProcedural)
+	.NameContent()
+	[
+		SNew(STextBlock)
+		.Text(FText::FromString(TEXT("Spawn Index")))
+		.Font(IDetailLayoutBuilder::GetDetailFont())
+	]
+	.ValueContent()
+	[
+		SNew(STextBlock)
+		.Text(this, &FRigBaseElementDetails::GetSpawnIndex)
+		.Font(IDetailLayoutBuilder::GetDetailFont())
+	];
+
 	// if we are not a bone, control or null
 	if(!IsAnyElementOfType(ERigElementType::Bone) &&
 		!IsAnyElementOfType(ERigElementType::Control) &&
@@ -1241,6 +1256,15 @@ bool FRigBaseElementDetails::OnVerifyNameChanged(const FText& InText, FText& Out
 	return true;
 }
 
+FText FRigBaseElementDetails::GetSpawnIndex() const
+{
+	if(PerElementInfos.Num() > 1)
+	{
+		return ControlRigDetailsMultipleValues;
+	}
+	return FText::FromString(FString::FromInt(PerElementInfos[0].GetElement<FRigBaseElement>()->GetSpawnIndex()));
+}
+
 void FRigBaseElementDetails::OnStructContentsChanged(FProperty* InProperty, const TSharedRef<IPropertyUtilities> PropertyUtilities)
 {
 	const FPropertyChangedEvent ChangeEvent(InProperty, EPropertyChangeType::ValueSet);
@@ -6302,6 +6326,20 @@ void FRigConnectorElementDetails::CustomizeSettings(IDetailLayoutBuilder& Detail
 		IsArrayPropertyRow.ToolTip(LOCTEXT("PrimaryConnectorsDontAllowArrayToolTip", "Primary Connectors don't support arrays. Add a secondary connector for that."));
 	}
 
+	IsPostConstructionHandle = SettingsHandle->GetChildHandle(GET_MEMBER_NAME_CHECKED(FRigConnectorSettings, bPostConstruction));
+	IDetailPropertyRow& IsPostConstructionPropertyRow = SettingsCategory
+		.AddProperty(IsPostConstructionHandle)
+		.IsEnabled(IsArrayEnabled());
+
+	if(!IsArrayEnabled())
+	{
+		IsPostConstructionPropertyRow.ToolTip(LOCTEXT("PrimaryConnectorsDontAllowPostConstructionToolTip", "Primary Connectors cannot be used during post construction."));
+	}
+	else
+	{
+		IsPostConstructionPropertyRow.ToolTip(LOCTEXT("PostConstructionConnectorToolTip", "This connector will be available only during the Post Construction Event.\nYou can connect any element spawned during Construction from any module to it."));
+	}
+	
 	if(!bHideRules)
 	{
 		SettingsCategory
diff --git a/Engine/Plugins/Animation/ControlRig/Source/ControlRigEditor/Private/ControlRigElementDetails.h b/Engine/Plugins/Animation/ControlRig/Source/ControlRigEditor/Private/ControlRigElementDetails.h
--- a/Engine/Plugins/Animation/ControlRig/Source/ControlRigEditor/Private/ControlRigElementDetails.h
+++ b/Engine/Plugins/Animation/ControlRig/Source/ControlRigEditor/Private/ControlRigElementDetails.h
@@ -240,6 +240,7 @@ public:
 	FText GetName() const;
 	void SetName(const FText& InNewText, ETextCommit::Type InCommitType);
 	bool OnVerifyNameChanged(const FText& InText, FText& OutErrorMessage);
+	FText GetSpawnIndex() const;
 
 	void OnStructContentsChanged(FProperty* InProperty, const TSharedRef<IPropertyUtilities> PropertyUtilities);
 	bool IsConstructionModeEnabled() const;
@@ -897,6 +898,7 @@ private:
 
 	TSharedPtr<IPropertyHandle> ConnectorTypeHandle;
 	TSharedPtr<IPropertyHandle> IsArrayHandle;
+	TSharedPtr<IPropertyHandle> IsPostConstructionHandle;
 };
 
 class FRigSocketElementDetails : public FRigNullElementDetails
diff --git a/Engine/Plugins/Animation/ControlRig/Source/ControlRigEditor/Private/ControlRigModuleDetails.cpp b/Engine/Plugins/Animation/ControlRig/Source/ControlRigEditor/Private/ControlRigModuleDetails.cpp
--- a/Engine/Plugins/Animation/ControlRig/Source/ControlRigEditor/Private/ControlRigModuleDetails.cpp
+++ b/Engine/Plugins/Animation/ControlRig/Source/ControlRigEditor/Private/ControlRigModuleDetails.cpp
@@ -164,6 +164,8 @@ void FRigModuleInstanceDetails::CustomizeDetails(IDetailLayoutBuilder& DetailBui
 		return;
 	}
 
+	TWeakPtr<FRigModuleInstanceDetails> WeakThisPtr = SharedThis(this);
+
 	TArray<FName> OriginalCategoryNames;
 	DetailBuilder.GetCategoryNames(OriginalCategoryNames);
 
@@ -206,7 +208,35 @@ void FRigModuleInstanceDetails::CustomizeDetails(IDetailLayoutBuilder& DetailBui
 			.Text(this, &FRigModuleInstanceDetails::GetRigClassPath)
 		];
 
-		TWeakPtr<FRigModuleInstanceDetails> WeakThisPtr = SharedThis(this);
+		GeneralCategory.AddCustomRow(FText::FromString(TEXT("Construction Start Spawn Index")), true /* advanced */)
+		.NameContent()
+		[
+			SNew(STextBlock)
+			.Text(FText::FromString(TEXT("Construction Start Index")))
+			.Font(IDetailLayoutBuilder::GetDetailFont())
+			.ToolTipText(LOCTEXT("ConstructionStartSpawnIndexToolTip", "The construction spawn index up to which the module can read from,\nconnectors can connect to. This represents any elements that spawned\nduring construction prior to this module."))
+		]
+		.ValueContent()
+		[
+			SNew(STextBlock)
+			.Text(this, &FRigModuleInstanceDetails::GetConstructionStartSpawnIndex)
+			.Font(IDetailLayoutBuilder::GetDetailFont())
+		];
+
+		GeneralCategory.AddCustomRow(FText::FromString(TEXT("Post-Construction Start Spawn Index")), true /* advanced */)
+		.NameContent()
+		[
+			SNew(STextBlock)
+			.Text(FText::FromString(TEXT("Post-Construction Start Index")))
+			.Font(IDetailLayoutBuilder::GetDetailFont())
+			.ToolTipText(LOCTEXT("PostConstructionStartSpawnIndexToolTip", "The post construction spawn index up to which the module can read from,\nconnectors can connect to. This represents any elements that spawned\nduring construction of any module - as well as during\npost-construction prior to this module."))
+		]
+		.ValueContent()
+		[
+			SNew(STextBlock)
+			.Text(this, &FRigModuleInstanceDetails::GetPostConstructionStartSpawnIndex)
+			.Font(IDetailLayoutBuilder::GetDetailFont())
+		];
 
 		GeneralCategory.AddCustomRow(FText::FromString(TEXT("Variant Tags")))
 		.NameContent()
@@ -346,9 +376,13 @@ void FRigModuleInstanceDetails::CustomizeDetails(IDetailLayoutBuilder& DetailBui
 				bool bIsArrayConnector = IsArrayConnector[ConnectorIndex].Get(false);
 
 				FRigTreeDelegates RigTreeDelegates;
-				RigTreeDelegates.OnGetHierarchy = FOnGetRigTreeHierarchy::CreateLambda([this]()
+				RigTreeDelegates.OnGetHierarchy = FOnGetRigTreeHierarchy::CreateLambda([WeakThisPtr]() -> const URigHierarchy*
 				{
-					return PerModuleInfos[0].GetModularRig()->GetHierarchy();
+					if (TSharedPtr<FRigModuleInstanceDetails> StrongThis = WeakThisPtr.Pin())
+					{
+						return StrongThis->PerModuleInfos[0].GetModularRig()->GetHierarchy();
+					}
+					return nullptr;
 				});
 
 				static const FSlateBrush* PrimaryBrush = FControlRigEditorStyle::Get().GetBrush("ControlRig.ConnectorPrimary");
@@ -1071,6 +1105,34 @@ bool FRigModuleInstanceDetails::OnConnectorTargetChanged(TArray<FRigElementKey>
 	return bResult;
 }
 
+FText FRigModuleInstanceDetails::GetConstructionStartSpawnIndex() const
+{
+	if(PerModuleInfos.Num() > 1)
+	{
+		return ControlRigModuleDetailsMultipleValues;
+	}
+	if(PerModuleInfos[0].GetModule()->ConstructionSpawnStartIndex == INDEX_NONE)
+	{
+		static const FText ConstructionWasNeverRun = LOCTEXT("ConstructionNeverRun", "Construction was never run.");
+		return ConstructionWasNeverRun;
+	}
+	return FText::FromString(FString::FromInt(PerModuleInfos[0].GetModule()->ConstructionSpawnStartIndex));
+}
+
+FText FRigModuleInstanceDetails::GetPostConstructionStartSpawnIndex() const
+{
+	if(PerModuleInfos.Num() > 1)
+	{
+		return ControlRigModuleDetailsMultipleValues;
+	}
+	if(PerModuleInfos[0].GetModule()->PostConstructionSpawnStartIndex == INDEX_NONE)
+	{
+		static const FText PostConstructionWasNeverRun = LOCTEXT("PostConstructionNeverRun", "Post-Construction was never run.");
+		return PostConstructionWasNeverRun;
+	}
+	return FText::FromString(FString::FromInt(PerModuleInfos[0].GetModule()->PostConstructionSpawnStartIndex));
+}
+
 const FRigModuleInstanceDetails::FPerModuleInfo& FRigModuleInstanceDetails::FindModule(const FName& InModuleName) const
 {
 	const FPerModuleInfo* Info = FindModuleByPredicate([InModuleName](const FPerModuleInfo& Info)
diff --git a/Engine/Plugins/Animation/ControlRig/Source/ControlRigEditor/Private/ControlRigModuleDetails.h b/Engine/Plugins/Animation/ControlRig/Source/ControlRigEditor/Private/ControlRigModuleDetails.h
--- a/Engine/Plugins/Animation/ControlRig/Source/ControlRigEditor/Private/ControlRigModuleDetails.h
+++ b/Engine/Plugins/Animation/ControlRig/Source/ControlRigEditor/Private/ControlRigModuleDetails.h
@@ -50,7 +50,10 @@ public:
 
 	void OnConfigValueChanged(const FPropertyChangedEvent& InPropertyChangedEvent);
 	bool OnConnectorTargetChanged(TArray<FRigElementKey> InTargets, FRigModuleConnector InConnector);
-	
+
+	FText GetConstructionStartSpawnIndex() const;
+	FText GetPostConstructionStartSpawnIndex() const;
+
 	struct FPerModuleInfo
 	{
 		FPerModuleInfo()
diff --git a/Engine/Plugins/Animation/ControlRig/Source/ControlRigEditor/Private/Editor/SModularRigModel.cpp b/Engine/Plugins/Animation/ControlRig/Source/ControlRigEditor/Private/Editor/SModularRigModel.cpp
--- a/Engine/Plugins/Animation/ControlRig/Source/ControlRigEditor/Private/Editor/SModularRigModel.cpp
+++ b/Engine/Plugins/Animation/ControlRig/Source/ControlRigEditor/Private/Editor/SModularRigModel.cpp
@@ -891,7 +891,7 @@ void SModularRigModel::HandleDeleteModules(const TArray<FName>& InModuleNames)
 	}
 }
 
-void SModularRigModel::HandleReparentModules(const TArray<FName>& InModuleNames, const FName& InParentModuleName)
+void SModularRigModel::HandleReparentModules(const TArray<FName>& InModuleNames, const FName& InParentModuleName, int32 NewModuleIndex)
 {
 	if (ControlRigBlueprint.IsValid())
 	{
@@ -903,6 +903,10 @@ void SModularRigModel::HandleReparentModules(const TArray<FName>& InModuleNames,
 		for (const FName& ModuleName : InModuleNames)
 		{
 			Controller->ReparentModule(ModuleName, InParentModuleName);
+			if(NewModuleIndex != INDEX_NONE)
+			{
+				Controller->ReorderModule(ModuleName, NewModuleIndex);
+			}
 		}
 	}
 }
@@ -1395,6 +1399,7 @@ void SModularRigModel::OnModularRigModified(EModularRigNotification InNotif, con
 		case EModularRigNotification::ModuleRenamed:
 		case EModularRigNotification::ModuleRemoved:
 		case EModularRigNotification::ModuleReparented:
+		case EModularRigNotification::ModuleReordered:
 		case EModularRigNotification::ConnectionChanged:
 		case EModularRigNotification::ModuleConfigValueChanged:
 		case EModularRigNotification::ModuleShortNameChanged:
@@ -1563,10 +1568,32 @@ TOptional<EItemDropZone> SModularRigModel::OnCanAcceptDrop(const FDragDropEvent&
 {
 	const TOptional<EItemDropZone> InvalidDropZone;
 	TOptional<EItemDropZone> ReturnDropZone = DropZone;
+	int32 NewModuleIndex = INDEX_NONE;
 
-	if(DropZone == EItemDropZone::BelowItem && !TargetItem.IsValid())
+	if(DropZone == EItemDropZone::BelowItem || DropZone == EItemDropZone::AboveItem)
 	{
 		DropZone = EItemDropZone::OntoItem;
+		
+		if(TargetItem.IsValid())
+		{
+			TSharedPtr<FModularRigTreeElement> ChildTargetItem;
+			Swap(TargetItem, ChildTargetItem);
+			
+			if(ControlRigBlueprint.IsValid())
+			{
+				FModularRigModel& Model = ControlRigBlueprint.Get()->ModularRigModel;
+				if(const FRigModuleReference* ChildModule = Model.FindModule(ChildTargetItem->ModuleName))
+				{
+					if(const FRigModuleReference* ParentModule = ChildModule->GetParentModule())
+					{
+						if(const TSharedPtr<FModularRigTreeElement> ParentItem = TreeView->FindElement(ParentModule->Name.ToString()))
+						{
+							TargetItem = ParentItem;
+						}
+					}
+				}
+			}
+		}
 	}
 	
 	if(DropZone != EItemDropZone::OntoItem)
@@ -1617,6 +1644,50 @@ TOptional<EItemDropZone> SModularRigModel::OnCanAcceptDrop(const FDragDropEvent&
 
 FReply SModularRigModel::OnAcceptDrop(const FDragDropEvent& DragDropEvent, EItemDropZone DropZone, TSharedPtr<FModularRigTreeElement> TargetItem)
 {
+	int32 NewModuleIndex = INDEX_NONE;
+	
+	if(DropZone == EItemDropZone::BelowItem || DropZone == EItemDropZone::AboveItem)
+	{
+		if(TargetItem.IsValid())
+		{
+			TSharedPtr<FModularRigTreeElement> ChildTargetItem;
+			Swap(TargetItem, ChildTargetItem);
+			
+			if(ControlRigBlueprint.IsValid())
+			{
+				FModularRigModel& Model = ControlRigBlueprint.Get()->ModularRigModel;
+				if(const FRigModuleReference* ChildModule = Model.FindModule(ChildTargetItem->ModuleName))
+				{
+					TArray<FRigModuleReference*> Children;
+					if(const FRigModuleReference* ParentModule = ChildModule->GetParentModule())
+					{
+						Children = ParentModule->CachedChildren;
+						
+						if(const TSharedPtr<FModularRigTreeElement> ParentItem = TreeView->FindElement(ParentModule->Name.ToString()))
+						{
+							TargetItem = ParentItem;
+						}
+					}
+					else
+					{
+						Children = Model.RootModules;
+					}
+
+					for(int32 ChildIndex = 0; ChildIndex < Children.Num(); ChildIndex++)
+					{
+						if(Children[ChildIndex] == ChildModule)
+						{
+							NewModuleIndex = ChildIndex;
+							break;
+						}
+					}
+				}
+			}
+		}
+
+		DropZone = EItemDropZone::OntoItem;
+	}
+
 	FName ParentModuleName = NAME_None;
 	if (TargetItem.IsValid())
 	{
@@ -1659,7 +1730,7 @@ FReply SModularRigModel::OnAcceptDrop(const FDragDropEvent& DragDropEvent, EItem
 	else if(ModuleDragDropOperation)
 	{
 		const TArray<FName> ModuleNames = ModuleDragDropOperation->GetModules();
-		HandleReparentModules(ModuleNames, ParentModuleName);
+		HandleReparentModules(ModuleNames, ParentModuleName, NewModuleIndex);
 	}
 	
 	return FReply::Unhandled();
diff --git a/Engine/Plugins/Animation/ControlRig/Source/ControlRigEditor/Private/Editor/SModularRigModel.h b/Engine/Plugins/Animation/ControlRig/Source/ControlRigEditor/Private/Editor/SModularRigModel.h
--- a/Engine/Plugins/Animation/ControlRig/Source/ControlRigEditor/Private/Editor/SModularRigModel.h
+++ b/Engine/Plugins/Animation/ControlRig/Source/ControlRigEditor/Private/Editor/SModularRigModel.h
@@ -76,7 +76,7 @@ private:
 	void HandleDeleteModules(const TArray<FName>& InModuleNames);
 
 	/** Reparent items */
-	void HandleReparentModules(const TArray<FName>& InModuleNames, const FName& InParentModuleName);
+	void HandleReparentModules(const TArray<FName>& InModuleNames, const FName& InParentModuleName, int32 NewModuleIndex);
 
 	/** Mirror items */
 	void HandleMirrorModules();
