Control Rig: Fix static analysis issue with ModularRigRuleManager

#rnx

diff --git a/Engine/Plugins/Animation/ControlRig/Source/ControlRig/Private/ModularRigRuleManager.cpp b/Engine/Plugins/Animation/ControlRig/Source/ControlRig/Private/ModularRigRuleManager.cpp
--- a/Engine/Plugins/Animation/ControlRig/Source/ControlRig/Private/ModularRigRuleManager.cpp
+++ b/Engine/Plugins/Animation/ControlRig/Source/ControlRig/Private/ModularRigRuleManager.cpp
@@ -228,25 +228,28 @@ void UModularRigRuleManager::FilterInvalidModules(FWorkData& InOutWorkData)
 		}
 		else
 		{
-			const int32& SpawnStartIndex =
-				InOutWorkData.Connector->IsPostConstructionConnector() ?
-					InOutWorkData.Module->PostConstructionSpawnStartIndex :
-					InOutWorkData.Module->ConstructionSpawnStartIndex;
-
-			const int32 TargetSpawnIndex = InOutWorkData.Hierarchy->GetSpawnIndex(Result.GetKey());
-
-			if(SpawnStartIndex != INDEX_NONE && TargetSpawnIndex >= SpawnStartIndex)
-			{
-				static const FText CannotConnectBelowNameSpaceMessage = LOCTEXT("CannotConnectLaterSpawnElement", "Cannot connect to an element spawn later than the connector's module.");
-				Result.SetInvalidTarget(CannotConnectBelowNameSpaceMessage);
-				return;
-			}
-
-			// if we don't have a spawn index at all - let's rely on module hierarchy only.
-			if(InOutWorkData.Module && InOutWorkData.Module->HasChildModule(MatchModuleName, true))
+			if(InOutWorkData.Module)
 			{
-				static const FText CannotConnectBelowNameSpaceMessage = LOCTEXT("CannotConnectBelowNameSpace", "Cannot connect to element below the connector's namespace.");
-				Result.SetInvalidTarget(CannotConnectBelowNameSpaceMessage);
+				const int32 SpawnStartIndex =
+					InOutWorkData.Connector->IsPostConstructionConnector() ?
+						InOutWorkData.Module->PostConstructionSpawnStartIndex :
+						InOutWorkData.Module->ConstructionSpawnStartIndex;
+
+				const int32 TargetSpawnIndex = InOutWorkData.Hierarchy->GetSpawnIndex(Result.GetKey());
+
+				if(SpawnStartIndex != INDEX_NONE && TargetSpawnIndex >= SpawnStartIndex)
+				{
+					static const FText CannotConnectBelowNameSpaceMessage = LOCTEXT("CannotConnectLaterSpawnElement", "Cannot connect to an element spawn later than the connector's module.");
+					Result.SetInvalidTarget(CannotConnectBelowNameSpaceMessage);
+					return;
+				}
+
+				// if we don't have a spawn index at all - let's rely on module hierarchy only.
+				if(SpawnStartIndex == INDEX_NONE && InOutWorkData.Module->HasChildModule(MatchModuleName, true))
+				{
+					static const FText CannotConnectBelowNameSpaceMessage = LOCTEXT("CannotConnectBelowNameSpace", "Cannot connect to element below the connector's namespace.");
+					Result.SetInvalidTarget(CannotConnectBelowNameSpaceMessage);
+				}
 			}
 		}
 	});
