Handle loading components that no longer have a parent, due to remove of the top-level component feature.

diff --git a/Engine/Plugins/Animation/ControlRig/Source/ControlRig/Private/Rigs/RigHierarchy.cpp b/Engine/Plugins/Animation/ControlRig/Source/ControlRig/Private/Rigs/RigHierarchy.cpp
--- a/Engine/Plugins/Animation/ControlRig/Source/ControlRig/Private/Rigs/RigHierarchy.cpp
+++ b/Engine/Plugins/Animation/ControlRig/Source/ControlRig/Private/Rigs/RigHierarchy.cpp
@@ -734,12 +734,24 @@ void URigHierarchy::Load(FArchive& Ar)
 				Component->IndexInHierarchy = ComponentIndex;
 
 				FRigBaseElement* Element = Find(Component->GetElementKey());
-				check(Element);
-					
-				Component->Element = Element;
-				Component->IndexInElement = Element->ComponentIndices.Add(Component->IndexInHierarchy);
+				if (ensure(Element))
+				{
+					Component->Element = Element;
+					Component->IndexInElement = Element->ComponentIndices.Add(Component->IndexInHierarchy);
 
-				ComponentIndexLookup.Add(Component->GetKey(), ComponentIndex);
+					ComponentIndexLookup.Add(Component->GetKey(), ComponentIndex);
+				}
+				else
+				{
+					UE_LOG(
+						LogControlRig,
+						Warning,
+						TEXT("%s: Can't find an owning element for component '%s.%s' - this component will not be loaded."),
+						*GetPathName(),
+						*Component->GetElementKey().Name.ToString(),
+						*Component->GetKey().Name.ToString()
+					);
+				}
 			}
 		}
 
@@ -1143,9 +1155,23 @@ void URigHierarchy::CopyHierarchy(URigHierarchy* InHierarchy)
 			TargetComponent->IndexInHierarchy = TargetComponentIndex;
 
 			TargetComponent->Element = Find(TargetComponent->GetElementKey());
-			TargetComponent->IndexInElement = TargetComponent->Element->ComponentIndices.Add(TargetComponent->IndexInHierarchy);
-	
-			ComponentIndexLookup.Add(TargetComponent->GetKey(), TargetComponentIndex);
+			if (ensure(TargetComponent->Element))
+			{
+				TargetComponent->IndexInElement = TargetComponent->Element->ComponentIndices.Add(TargetComponent->IndexInHierarchy);
+				ComponentIndexLookup.Add(TargetComponent->GetKey(), TargetComponentIndex);
+			}
+			else
+			{
+				TargetComponent->IndexInElement = -1;
+				UE_LOG(
+					LogControlRig,
+					Warning,
+					TEXT("%s: Can't find an owning element for component '%s.%s'."),
+					*GetPathName(),
+					*TargetComponent->GetElementKey().Name.ToString(),
+					*TargetComponent->GetKey().Name.ToString()
+				);
+			}
 		}
 
 		// copy the topology version from the hierarchy.
