Anim In Engine: Fix Flickering with constraints, added new API helper function for this, CL also contains some other helpers that will be needed.

diff --git a/Engine/Plugins/Animation/ControlRig/Source/ControlRigEditor/Private/Tools/EvaluateSequencerTools.cpp b/Engine/Plugins/Animation/ControlRig/Source/ControlRigEditor/Private/Tools/EvaluateSequencerTools.cpp
--- a/Engine/Plugins/Animation/ControlRig/Source/ControlRigEditor/Private/Tools/EvaluateSequencerTools.cpp
+++ b/Engine/Plugins/Animation/ControlRig/Source/ControlRigEditor/Private/Tools/EvaluateSequencerTools.cpp
@@ -29,6 +29,7 @@
 #include "Constraints/TransformConstraintChannelInterface.h"
 #include "Components/SkeletalMeshComponent.h"
 #include "Transform/TransformConstraintUtil.h"
+#include "EditMode/ControlRigEditMode.h"
 
 class UControlRig;
 
@@ -109,8 +110,7 @@ namespace AIE
 			while (Parent)
 			{
 				TArray<USkeletalMeshComponent*> MeshComps;
-				Parent->GetComponents(MeshComps, true);
-
+				Parent->GetComponents(MeshComps, true);	
 				for (USkeletalMeshComponent* MeshComp : MeshComps)
 				{
 					MeshComp->TickAnimation(0.0f, false);
@@ -189,11 +189,43 @@ namespace AIE
 		}
 		OutTransform = ControlRig->GetControlGlobalTransform(ControlName) * ParentTransform;
 		return true;
+
 	}
 
+	bool FEvalHelpers::EvaluateSequencer(UWorld* World, ISequencer* Sequencer, const TSet<const UMovieSceneTrack*>& RelevantTracks)
+	{
+		if (World == nullptr || Sequencer == nullptr || Sequencer->GetFocusedMovieSceneSequence() == nullptr
+			|| Sequencer->GetFocusedMovieSceneSequence()->GetMovieScene() == nullptr)
+		{
+			return false;
+		}
+		const FQualifiedFrameTime QFrameTime = Sequencer->GetGlobalTime();
+		FMovieSceneContext Context = FMovieSceneContext(FMovieSceneEvaluationRange(QFrameTime.Time, QFrameTime.Rate), Sequencer->GetPlaybackStatus()).SetHasJumped(true);
+		//have no way to just evaluate specific tracks in Sequencer so do all of them.
+		Sequencer->GetEvaluationTemplate().EvaluateSynchronousBlocking(Context);
+		const FConstraintsManagerController& Controller = FConstraintsManagerController::Get(World);
+		Controller.EvaluateAllConstraints();
+
+		for (const UMovieSceneTrack* Track : RelevantTracks)
+		{
+			if (const UMovieSceneControlRigParameterTrack* CRTrack = Cast<UMovieSceneControlRigParameterTrack>(Track))
+			{
+				if (UControlRig* ControlRig = CRTrack->GetControlRig())
+				{
+					ControlRig->Evaluate_AnyThread();
+					if (ControlRig->GetObjectBinding())
+					{
+						ControlRig->EvaluateSkeletalMeshComponent(0.0);
+					}
+				}
+			}
+		}
+		return true;
+	}
+
+
 	bool FEvalHelpers::CalculateWorldTransforms(UWorld* World, ISequencer* Sequencer, const UE::AIE::FFrameTimeByIndex& FrameTimeByIndex,
-		const TArray<int32>& Indices, TArray<UE::AIE::FActorAndWorldTransforms>& Actors, TMap<UControlRig*, UE::AIE::FControlRigAndWorldTransforms>& ControlRigs,
-		const TSet<const UMovieSceneTrack*>& DependentTracks)
+		const TArray<int32>& Indices, TArray<UE::AIE::FActorAndWorldTransforms>& Actors, TMap<UControlRig*, UE::AIE::FControlRigAndWorldTransforms>& ControlRigs)
 	{
 		if (World == nullptr || Sequencer == nullptr || Sequencer->GetFocusedMovieSceneSequence() == nullptr
 			|| Sequencer->GetFocusedMovieSceneSequence()->GetMovieScene() == nullptr)
@@ -208,9 +240,8 @@ namespace AIE
 		FFrameRate TickResolution = MovieScene->GetTickResolution();
 		FFrameRate DisplayRate = MovieScene->GetDisplayRate();
 
-
 		TSet<UControlRig*> ControlsToRevalaute;
-		
+		FControlRigEditMode::FTurnOffPosePoseUpdate  TurnOff; //stop flashing
 		for (int32 Index : Indices)
 		{
 			const FFrameNumber FrameNumber = FrameTimeByIndex.CalculateFrame(Index);
@@ -257,25 +288,6 @@ namespace AIE
 			}
 		}
 		
-		Sequencer->ForceEvaluate();
-		//we now need to make sure the control rigs are back up to date.
-		//we do this by going through the control rig tracks that the actors/rigs are dependent up
-
-		for (const UMovieSceneTrack* Track : DependentTracks)
-		{
-			if (const UMovieSceneControlRigParameterTrack* CRTrack = Cast<UMovieSceneControlRigParameterTrack>(Track))
-			{
-				if (UControlRig* ControlRig = CRTrack->GetControlRig())
-				{
-					ControlRig->Evaluate_AnyThread();
-					if (ControlRig->GetObjectBinding())
-					{
-						ControlRig->EvaluateSkeletalMeshComponent(0.0);
-					}
-				}
-			}
-		}
-		
 		return true;
 	}
 
@@ -429,6 +441,7 @@ namespace AIE
 		const FConstraintsManagerController& Controller = FConstraintsManagerController::Get(Actor->GetWorld());
 		TArrayView<FMovieSceneDoubleChannel*> Channels = TransformSection->GetChannelProxy().GetChannels<FMovieSceneDoubleChannel>();
 		const bool bhasContraint = TransformSection->GetConstraintsChannels().Num() > 0;
+		FControlRigEditMode::FTurnOffPosePoseUpdate  TurnOff; //stop flashing
 
 		for (int32 Index = 0; Index < Frames.Num(); ++Index)
 		{
diff --git a/Engine/Plugins/Animation/ControlRig/Source/ControlRigEditor/Public/Tools/EvaluateSequencerTools.h b/Engine/Plugins/Animation/ControlRig/Source/ControlRigEditor/Public/Tools/EvaluateSequencerTools.h
--- a/Engine/Plugins/Animation/ControlRig/Source/ControlRigEditor/Public/Tools/EvaluateSequencerTools.h
+++ b/Engine/Plugins/Animation/ControlRig/Source/ControlRigEditor/Public/Tools/EvaluateSequencerTools.h
@@ -56,7 +56,7 @@ struct  FFrameTimeByIndex
 		FrameStep(InFrameStep)
 	{
 		//make sure end frame falls on frame step 
-		NumFrames = (InEndFrame.Value - StartFrame.Value) / (FrameStep.Value);
+		NumFrames = (InEndFrame.Value - StartFrame.Value) / (FrameStep.Value) + 1;
 		EndFrame.Value = StartFrame.Value + (NumFrames * FrameStep.Value);
 	};
 
@@ -81,6 +81,19 @@ struct  FFrameTimeByIndex
 		}
 		return INDEX_NONE;
 	}
+	TArray<int32> GetFullIndexArray() const
+	{
+		TArray<int32> IndexArray;
+		if (NumFrames > 0)
+		{
+			IndexArray.SetNum(NumFrames);
+			for (int32 Index = 0; Index < NumFrames; ++Index)
+			{
+				IndexArray[Index] = Index;
+			}
+		}
+		return IndexArray;
+	}
 };
 
 //Array of transforms, that maybe be sparse with only some transforms set
@@ -114,6 +127,19 @@ struct FActorAndWorldTransforms
 	FActorForWorldTransforms Actor;
 	TSharedPtr<FArrayOfTransforms> WorldTransforms;
 	TSharedPtr<FArrayOfTransforms> ParentTransforms;
+
+	void MakeCopy(FActorAndWorldTransforms& OutCopy) const
+	{
+		OutCopy.Actor = Actor;
+		OutCopy.WorldTransforms = MakeShared<FArrayOfTransforms>();
+		OutCopy.ParentTransforms = MakeShared<FArrayOfTransforms>();
+		if (WorldTransforms.IsValid())
+		{
+			OutCopy.SetNumOfTransforms(WorldTransforms->Num());
+		}
+	}
+
+
 	void SetNumOfTransforms(int32 Num)
 	{
 		if (WorldTransforms)
@@ -154,9 +180,25 @@ struct FControlRigAndWorldTransforms
 		}
 	}
 
+	void MakeCopy(FControlRigAndWorldTransforms& OutCopy) const
+	{
+		OutCopy.ControlRig = ControlRig;
+		OutCopy.ParentTransforms = MakeShared<FArrayOfTransforms>();
+		for (const TPair<FName, TSharedPtr<FArrayOfTransforms>>& Pair : ControlAndWorldTransforms)
+		{
+			TSharedPtr<FArrayOfTransforms> SharedPtr = MakeShared<FArrayOfTransforms>();
+			OutCopy.AddWorldTransform(Pair.Key, SharedPtr);
+		}
+		if (ParentTransforms.IsValid())
+		{
+			OutCopy.SetNumOfTransforms(ParentTransforms->Num());
+		}
+	}
+
 	TWeakObjectPtr<UControlRig> ControlRig;
 	TSharedPtr<FArrayOfTransforms > ParentTransforms;
 	TMap<FName, TSharedPtr<FArrayOfTransforms>> ControlAndWorldTransforms; //one for each control in the FControlRigAndWorldTransform
+
 };
 
 struct FSetTransformHelpers
@@ -178,7 +220,7 @@ struct FSetTransformHelpers
 struct FSequencerTransformDependencies
 {
 	//Get the list of tracks in the moviescene with this binding that may effect the final transform of something in the binding.
-	//Currently it will return a set of trnasform, control rig, attachment, path and  skeletal anim tracks.
+	//Currently it will return a set of transform, control rig, attachment, path and  skeletal anim tracks.
 	static UE_API TArray<UMovieSceneTrack*> GetTransformAffectingTracks(UMovieScene* InMovieScene, const FMovieSceneBinding& Binding);
 
 	//Main function to calculate the dependencies for an actor and a set of tracks under that actor
@@ -209,8 +251,10 @@ struct FEvalHelpers
 {
 	//The main evaluation function to use a sparse set of times and caclculate their world transforms for a set of sequencer actors/control rigs.
 	static UE_API bool CalculateWorldTransforms(UWorld* World, ISequencer* Sequencer, const FFrameTimeByIndex& FrameTimeByIndex,
-		const TArray<int32>& Indices, TArray<FActorAndWorldTransforms>& Actors, TMap<UControlRig*, FControlRigAndWorldTransforms>& ControlRigs,
-		const TSet<const UMovieSceneTrack*>& DependentTracks);
+		const TArray<int32>& Indices, TArray<FActorAndWorldTransforms>& Actors, TMap<UControlRig*, FControlRigAndWorldTransforms>& ControlRigs);
+
+	//Evaluate sequencer for the relevant tracks, making sure depencencies such as constraints and control rigs also tick
+	static UE_API bool EvaluateSequencer(UWorld* World, ISequencer* Sequencer, const TSet<const UMovieSceneTrack*>& RelevantTracks);
 };
 
 
diff --git a/Engine/Plugins/Editor/SequencerAnimTools/Source/SequencerAnimTools/Private/Sequencer/SequencerTrailHierarchy.cpp b/Engine/Plugins/Editor/SequencerAnimTools/Source/SequencerAnimTools/Private/Sequencer/SequencerTrailHierarchy.cpp
--- a/Engine/Plugins/Editor/SequencerAnimTools/Source/SequencerAnimTools/Private/Sequencer/SequencerTrailHierarchy.cpp
+++ b/Engine/Plugins/Editor/SequencerAnimTools/Source/SequencerAnimTools/Private/Sequencer/SequencerTrailHierarchy.cpp
@@ -439,9 +439,9 @@ void FSequencerTrailHierarchy::EvaluateSequencerAndSetTransforms()
 				if (IndicesToCalcluate.Num() > 0)
 				{
 					UE::AIE::FEvalHelpers::CalculateWorldTransforms(World, SequencerPtr.Get(), CurrentFramesInfo.CurrentFrameTimes,
-						IndicesToCalcluate, EvaluatingActors, EvaluatingControlRigs, DependentTracks);
+						IndicesToCalcluate, EvaluatingActors, EvaluatingControlRigs);
+					UE::AIE::FEvalHelpers::EvaluateSequencer(World, SequencerPtr.Get(), DependentTracks);
 				}
-
 			}
 			
 			const TRange<FFrameNumber> Range = TickEvalRange;
