Control Rig: Fix cook not preloading rig modules when a modular rig is used in an ABP
#rnx

diff --git a/Engine/Plugins/Animation/ControlRig/Source/ControlRigDeveloper/Private/AnimGraphNode_ControlRig.cpp b/Engine/Plugins/Animation/ControlRig/Source/ControlRigDeveloper/Private/AnimGraphNode_ControlRig.cpp
--- a/Engine/Plugins/Animation/ControlRig/Source/ControlRigDeveloper/Private/AnimGraphNode_ControlRig.cpp
+++ b/Engine/Plugins/Animation/ControlRig/Source/ControlRigDeveloper/Private/AnimGraphNode_ControlRig.cpp
@@ -13,6 +13,7 @@
 #include "RigVMBlueprintGeneratedClass.h"
 #include "ControlRigBlueprint.h"
 #include "ControlRigObjectBinding.h"
+#include "ModularRig.h"
 #include "Misc/DefaultValueHelper.h"
 
 #include UE_INLINE_GENERATED_CPP_BY_NAME(AnimGraphNode_ControlRig)
@@ -473,6 +474,24 @@ void UAnimGraphNode_ControlRig::PostReconstructNode()
 	Super::PostReconstructNode();
 }
 
+void UAnimGraphNode_ControlRig::PreloadRequiredAssets()
+{
+	Super::PreloadRequiredAssets();
+
+	if (UObject* CDO = Node.ControlRigClass.GetDefaultObject())
+	{
+		PreloadObject(CDO);
+		if (UModularRig* ModularRig = Cast<UModularRig>(CDO))
+		{
+			ModularRig->ModularRigModel.ForEachModule([this](const FRigModuleReference* Module) -> bool
+			{
+				PreloadObject(Module->Class.Get());
+				return true;
+			});
+		}
+	}
+}
+
 void UAnimGraphNode_ControlRig::CustomizePinData(UEdGraphPin* Pin, FName SourcePropertyName, int32 ArrayIndex) const
 {
 	Super::CustomizePinData(Pin, SourcePropertyName, ArrayIndex);
diff --git a/Engine/Plugins/Animation/ControlRig/Source/ControlRigDeveloper/Public/AnimGraphNode_ControlRig.h b/Engine/Plugins/Animation/ControlRig/Source/ControlRigDeveloper/Public/AnimGraphNode_ControlRig.h
--- a/Engine/Plugins/Animation/ControlRig/Source/ControlRigDeveloper/Public/AnimGraphNode_ControlRig.h
+++ b/Engine/Plugins/Animation/ControlRig/Source/ControlRigDeveloper/Public/AnimGraphNode_ControlRig.h
@@ -21,6 +21,7 @@ class UAnimGraphNode_ControlRig : public UAnimGraphNode_CustomProperty
 	virtual void CustomizeDetails(IDetailLayoutBuilder& DetailBuilder) override;
 	virtual void PostEditChangeProperty(FPropertyChangedEvent& PropertyChangedEvent) override;
 	virtual void PostReconstructNode() override;
+	virtual void PreloadRequiredAssets() override;
 
 private:
 	// UEdGraphNode interface
