Control Rig: Speculative fix for linux server crash. Note ECS has this check already.

diff --git a/Engine/Plugins/Animation/ControlRig/Source/ControlRig/Private/Sequencer/MovieSceneControlRigParameterTemplate.cpp b/Engine/Plugins/Animation/ControlRig/Source/ControlRig/Private/Sequencer/MovieSceneControlRigParameterTemplate.cpp
--- a/Engine/Plugins/Animation/ControlRig/Source/ControlRig/Private/Sequencer/MovieSceneControlRigParameterTemplate.cpp
+++ b/Engine/Plugins/Animation/ControlRig/Source/ControlRig/Private/Sequencer/MovieSceneControlRigParameterTemplate.cpp
@@ -1978,6 +1978,10 @@ void FMovieSceneControlRigParameterTemplate::Evaluate(const FMovieSceneEvaluatio
 			{
 				UWorld* GameWorld =  BoundObjects[0].Get()->GetWorld();
 				ControlRig = Section->GetControlRig(GameWorld);
+				if (ControlRig == nullptr)
+				{
+					return;
+				}
 			}
 		}
 
