Control Rig: Attempt fix of crash when a control rig editor is opening while the editor is closing
#rnx

diff --git a/Engine/Plugins/Animation/ControlRig/Source/ControlRigEditor/Private/ControlRigBlueprintActions.cpp b/Engine/Plugins/Animation/ControlRig/Source/ControlRigEditor/Private/ControlRigBlueprintActions.cpp
--- a/Engine/Plugins/Animation/ControlRig/Source/ControlRigEditor/Private/ControlRigBlueprintActions.cpp
+++ b/Engine/Plugins/Animation/ControlRig/Source/ControlRigEditor/Private/ControlRigBlueprintActions.cpp
@@ -51,6 +51,11 @@ UFactory* FControlRigBlueprintActions::GetFactoryForBlueprintType(UBlueprint* In
 
 void FControlRigBlueprintActions::OpenAssetEditor( const TArray<UObject*>& InObjects, TSharedPtr<IToolkitHost> EditWithinLevelEditor )
 {
+	if (IsEngineExitRequested())
+	{
+		return;
+	}
+	
 	EToolkitMode::Type Mode = EditWithinLevelEditor.IsValid() ? EToolkitMode::WorldCentric : EToolkitMode::Standalone;
 
 	for (auto ObjIt = InObjects.CreateConstIterator(); ObjIt; ++ObjIt)
