Control Rig: Fix crash when undoing after opening editor
#rnx

diff --git a/Engine/Plugins/Animation/ControlRig/Source/ControlRigEditor/Private/Editor/ControlRigEditor.cpp b/Engine/Plugins/Animation/ControlRig/Source/ControlRigEditor/Private/Editor/ControlRigEditor.cpp
--- a/Engine/Plugins/Animation/ControlRig/Source/ControlRigEditor/Private/Editor/ControlRigEditor.cpp
+++ b/Engine/Plugins/Animation/ControlRig/Source/ControlRigEditor/Private/Editor/ControlRigEditor.cpp
@@ -3365,6 +3365,9 @@ void FControlRigBaseEditor::HandlePreviewMeshChanged(USkeletalMesh* InOldSkeleta
 	{
 		if (UControlRigBlueprint* ControlRigBP = GetControlRigBlueprint())
 		{
+			const bool bTransact = !GIsTransacting && IsEditorInitialized();
+			const FScopedTransaction Transaction(NSLOCTEXT("ControlRigEditor", "SetPreviewMesh", "Set Preview Mesh"), bTransact);
+			
 			ControlRigBP->SetPreviewMesh(InNewSkeletalMesh);
 			URigHierarchy* BPHierarchy = ControlRigBP->GetHierarchy();
 
diff --git a/Engine/Plugins/Animation/ControlRig/Source/ControlRigEditor/Private/Editor/ControlRigEditor.h b/Engine/Plugins/Animation/ControlRig/Source/ControlRigEditor/Private/Editor/ControlRigEditor.h
--- a/Engine/Plugins/Animation/ControlRig/Source/ControlRigEditor/Private/Editor/ControlRigEditor.h
+++ b/Engine/Plugins/Animation/ControlRig/Source/ControlRigEditor/Private/Editor/ControlRigEditor.h
@@ -123,7 +123,7 @@ protected:
 	virtual bool IsControlRigLegacyEditor() const = 0;
 	virtual URigVMBlueprint* GetRigVMBlueprint() const = 0;
 	virtual URigVMHost* GetRigVMHost() const = 0;
-	
+	virtual bool IsEditorInitialized() const = 0;
 	
 	virtual TSharedRef<FUICommandList> GetToolkitCommands() = 0;
 	virtual FPreviewScene* GetPreviewScene() = 0;
diff --git a/Engine/Plugins/Animation/ControlRig/Source/ControlRigEditor/Private/Editor/ControlRigLegacyEditor.h b/Engine/Plugins/Animation/ControlRig/Source/ControlRigEditor/Private/Editor/ControlRigLegacyEditor.h
--- a/Engine/Plugins/Animation/ControlRig/Source/ControlRigEditor/Private/Editor/ControlRigLegacyEditor.h
+++ b/Engine/Plugins/Animation/ControlRig/Source/ControlRigEditor/Private/Editor/ControlRigLegacyEditor.h
@@ -32,6 +32,7 @@ public:
 	virtual bool IsControlRigLegacyEditor() const { return true; };
 	virtual URigVMBlueprint* GetRigVMBlueprint() const { return IControlRigLegacyEditor::GetRigVMBlueprint(); }
 	virtual URigVMHost* GetRigVMHost() const { return FRigVMEditorBase::GetRigVMHost(); }
+	virtual bool IsEditorInitialized() const override { return bRigVMEditorInitialized; }
 	virtual TSharedRef<FUICommandList> GetToolkitCommands() { return IControlRigLegacyEditor::GetToolkitCommands(); }
 	virtual FPreviewScene* GetPreviewScene() { return IControlRigLegacyEditor::GetPreviewScene(); }
 	virtual bool IsDetailsPanelRefreshSuspended() const { return IControlRigLegacyEditor::IsDetailsPanelRefreshSuspended(); }
diff --git a/Engine/Plugins/Animation/ControlRig/Source/ControlRigEditor/Private/Editor/ControlRigNewEditor.h b/Engine/Plugins/Animation/ControlRig/Source/ControlRigEditor/Private/Editor/ControlRigNewEditor.h
--- a/Engine/Plugins/Animation/ControlRig/Source/ControlRigEditor/Private/Editor/ControlRigNewEditor.h
+++ b/Engine/Plugins/Animation/ControlRig/Source/ControlRigEditor/Private/Editor/ControlRigNewEditor.h
@@ -32,6 +32,7 @@ public:
 	virtual bool IsControlRigNewEditor() const { return true; };
 	virtual URigVMBlueprint* GetRigVMBlueprint() const { return IControlRigNewEditor::GetRigVMBlueprint(); }
 	virtual URigVMHost* GetRigVMHost() const { return FRigVMEditorBase::GetRigVMHost(); }
+	virtual bool IsEditorInitialized() const override { return bRigVMEditorInitialized; }
 	virtual TSharedRef<FUICommandList> GetToolkitCommands() { return IControlRigNewEditor::GetToolkitCommands(); }
 	virtual FPreviewScene* GetPreviewScene() { return &PreviewScene; }
 	virtual bool IsDetailsPanelRefreshSuspended() const { return IControlRigNewEditor::IsDetailsPanelRefreshSuspended(); }
diff --git a/Engine/Plugins/Animation/ControlRig/Source/ControlRigEditor/Public/IControlRigEditor.h b/Engine/Plugins/Animation/ControlRig/Source/ControlRigEditor/Public/IControlRigEditor.h
--- a/Engine/Plugins/Animation/ControlRig/Source/ControlRigEditor/Public/IControlRigEditor.h
+++ b/Engine/Plugins/Animation/ControlRig/Source/ControlRigEditor/Public/IControlRigEditor.h
@@ -16,6 +16,7 @@ class IControlRigEditor
 {
 public:
 	virtual URigVMHost* GetRigVMHost() const = 0;
+	virtual bool IsEditorInitialized() const = 0;
 };
 
 #if WITH_RIGVMLEGACYEDITOR
diff --git a/Engine/Plugins/Runtime/RigVM/Source/RigVM/Private/RigVMHost.cpp b/Engine/Plugins/Runtime/RigVM/Source/RigVM/Private/RigVMHost.cpp
--- a/Engine/Plugins/Runtime/RigVM/Source/RigVM/Private/RigVMHost.cpp
+++ b/Engine/Plugins/Runtime/RigVM/Source/RigVM/Private/RigVMHost.cpp
@@ -1389,7 +1389,10 @@ void URigVMHost::InitializeFromCDO()
 		// we initialize all other instances of Control Rig from the CDO here
 		URigVMHost* CDO = GetClass()->GetDefaultObject<URigVMHost>();
 
-		ensure(VM == CDO->VM);
+		if (VM == nullptr || VM != CDO->VM)
+		{
+			VM = CDO->VM;
+		}
 
 		PostInitInstanceIfRequired();
 
