Anim in Engine: Fix an issue in Anim Details where one could no longer navigate through property values directly

diff --git a/Engine/Plugins/Animation/ControlRig/Source/ControlRigEditor/Private/AnimDetails/AnimDetailsNavigableWidgetRegistrar.cpp b/Engine/Plugins/Animation/ControlRig/Source/ControlRigEditor/Private/AnimDetails/AnimDetailsNavigableWidgetRegistrar.cpp
new file mode 100644
--- /dev/null
+++ b/Engine/Plugins/Animation/ControlRig/Source/ControlRigEditor/Private/AnimDetails/AnimDetailsNavigableWidgetRegistrar.cpp
@@ -0,0 +1,32 @@
+// Copyright Epic Games, Inc. All Rights Reserved.
+
+#include "AnimDetailsNavigableWidgetRegistrar.h"
+
+#include "AnimDetailsProxyManager.h"
+
+namespace UE::ControlRigEditor
+{
+	FAnimDetailsNavigableWidgetRegistrar::~FAnimDetailsNavigableWidgetRegistrar()
+	{
+		if (!IsEngineExitRequested() &&
+			WeakProxyManager.IsValid() &&
+			WeakWidget.IsValid())
+		{
+			WeakProxyManager->GetNavigableWidgetRegistry().Unregister(WeakWidget.Pin().ToSharedRef());
+		}
+	}
+
+	void FAnimDetailsNavigableWidgetRegistrar::Register(UAnimDetailsProxyManager& ProxyManager, const TSharedRef<SWidget> Widget)
+	{
+		if (WeakProxyManager.IsValid() &&
+			WeakWidget.IsValid())
+		{
+			WeakProxyManager->GetNavigableWidgetRegistry().Unregister(WeakWidget.Pin().ToSharedRef());
+		}
+
+		ProxyManager.GetNavigableWidgetRegistry().Register(Widget);
+
+		WeakProxyManager = &ProxyManager;
+		WeakWidget = Widget;
+	}
+}
diff --git a/Engine/Plugins/Animation/ControlRig/Source/ControlRigEditor/Private/AnimDetails/AnimDetailsNavigableWidgetRegistrar.h b/Engine/Plugins/Animation/ControlRig/Source/ControlRigEditor/Private/AnimDetails/AnimDetailsNavigableWidgetRegistrar.h
new file mode 100644
--- /dev/null
+++ b/Engine/Plugins/Animation/ControlRig/Source/ControlRigEditor/Private/AnimDetails/AnimDetailsNavigableWidgetRegistrar.h
@@ -0,0 +1,25 @@
+// Copyright Epic Games, Inc. All Rights Reserved.
+
+#pragma once
+
+#include "Templates/SharedPointer.h"
+#include "UObject/WeakObjectPtr.h"
+
+class SWidget;
+class UAnimDetailsProxyManager;
+
+namespace UE::ControlRigEditor
+{
+	/** Responsible to handle registration of navigable widgets with the navigable widget registry */
+	struct FAnimDetailsNavigableWidgetRegistrar
+	{
+		~FAnimDetailsNavigableWidgetRegistrar();
+
+		/** Registers a widget with the navigable widget registry */
+		void Register(UAnimDetailsProxyManager& ProxyManager, const TSharedRef<SWidget> Widget);
+
+	private:
+		TWeakObjectPtr<UAnimDetailsProxyManager> WeakProxyManager;
+		TWeakPtr<SWidget> WeakWidget;
+	};
+}
diff --git a/Engine/Plugins/Animation/ControlRig/Source/ControlRigEditor/Private/AnimDetails/AnimDetailsNavigableWidgetRegistry.cpp b/Engine/Plugins/Animation/ControlRig/Source/ControlRigEditor/Private/AnimDetails/AnimDetailsNavigableWidgetRegistry.cpp
new file mode 100644
--- /dev/null
+++ b/Engine/Plugins/Animation/ControlRig/Source/ControlRigEditor/Private/AnimDetails/AnimDetailsNavigableWidgetRegistry.cpp
@@ -0,0 +1,78 @@
+// Copyright Epic Games, Inc. All Rights Reserved.
+
+#include "AnimDetailsNavigableWidgetRegistry.h"
+
+#include "Framework/Application/SlateApplication.h"
+
+namespace UE::ControlRigEditor
+{
+	const TSharedPtr<SWidget> FAnimDetailsNavigableWidgetRegistry::GetNext(const TSharedRef<SWidget>& Current) const
+	{
+		WeakWidgets.Remove(nullptr);
+
+		const int32 CurrentIndex = WeakWidgets.IndexOfByKey(Current);
+		if (CurrentIndex == INDEX_NONE)
+		{
+			return nullptr;
+		}
+
+		for (int32 NextIndex = CurrentIndex + 1; NextIndex < WeakWidgets.Num(); NextIndex++)
+		{
+			if (!WeakWidgets[NextIndex].IsValid())
+			{
+				continue;
+			}
+
+			const TSharedRef<SWidget> NextWidget = WeakWidgets[NextIndex].Pin().ToSharedRef();
+			if (NextWidget->IsEnabled()) 
+			{
+				if (IsWidgetVisible(NextWidget))
+				{
+					return NextWidget;
+				}
+			}
+		}
+
+		return nullptr;
+	}
+
+	bool FAnimDetailsNavigableWidgetRegistry::IsWidgetVisible(const TSharedRef<SWidget>& Widget) const
+	{
+		if (Widget->GetVisibility() == EVisibility::Collapsed ||
+			Widget->GetVisibility() == EVisibility::Hidden)
+		{
+			return false;
+		}
+
+		// In some cases details views just minimize their children, but don't flag them hidden or collapsed
+		if (Widget->GetCachedGeometry().GetAbsoluteSize().IsNearlyZero())
+		{
+			return false;
+		}
+
+		// In some cases details views just remove their widgets from the window, but don't flag them hidden or collapsed
+		const TSharedPtr<SWidget> ActiveTopLevelWindow = FSlateApplication::Get().GetActiveTopLevelWindow();
+		TSharedPtr<SWidget> VisibilityTestWidget = Widget->GetParentWidget();
+		while (VisibilityTestWidget.IsValid())
+		{
+			if (VisibilityTestWidget == ActiveTopLevelWindow)
+			{
+				return true;
+			}
+
+			VisibilityTestWidget = VisibilityTestWidget->GetParentWidget();
+		}
+
+		return false;
+	}
+
+	void FAnimDetailsNavigableWidgetRegistry::Register(const TSharedRef<SWidget>& InWidget)
+	{
+		WeakWidgets.Add(InWidget);
+	}
+
+	void FAnimDetailsNavigableWidgetRegistry::Unregister(const TSharedRef<SWidget>& InWidget)
+	{
+		WeakWidgets.Remove(InWidget);
+	}
+}
diff --git a/Engine/Plugins/Animation/ControlRig/Source/ControlRigEditor/Private/AnimDetails/AnimDetailsNavigableWidgetRegistry.h b/Engine/Plugins/Animation/ControlRig/Source/ControlRigEditor/Private/AnimDetails/AnimDetailsNavigableWidgetRegistry.h
new file mode 100644
--- /dev/null
+++ b/Engine/Plugins/Animation/ControlRig/Source/ControlRigEditor/Private/AnimDetails/AnimDetailsNavigableWidgetRegistry.h
@@ -0,0 +1,33 @@
+// Copyright Epic Games, Inc. All Rights Reserved.
+
+#pragma once
+
+#include "Templates/SharedPointer.h"
+
+class SWidget;
+
+namespace UE::ControlRigEditor
+{
+	class FAnimDetailsNavigableWidgetRegistry
+	{
+		// Let the registrar register and unregister widgets
+		friend struct FAnimDetailsNavigableWidgetRegistrar;
+
+	public:
+		/** Returns the next widget given the current, or nullptr if there is no next widget. Useful for navigation */
+		const TSharedPtr<SWidget> GetNext(const TSharedRef<SWidget>& Current) const;
+
+	private:
+		/** Does a deep test if a widget is truly visible. */
+		bool IsWidgetVisible(const TSharedRef<SWidget>& Widget) const;
+
+		/** Registers a widget */
+		void Register(const TSharedRef<SWidget>& InWidget);
+
+		/** Unregisters a widget */
+		void Unregister(const TSharedRef<SWidget>& InWidget);
+
+		/** The current widgets */
+		mutable TArray<TWeakPtr<SWidget>> WeakWidgets;
+	};
+}
diff --git a/Engine/Plugins/Animation/ControlRig/Source/ControlRigEditor/Private/AnimDetails/AnimDetailsProxyManager.h b/Engine/Plugins/Animation/ControlRig/Source/ControlRigEditor/Private/AnimDetails/AnimDetailsProxyManager.h
--- a/Engine/Plugins/Animation/ControlRig/Source/ControlRigEditor/Private/AnimDetails/AnimDetailsProxyManager.h
+++ b/Engine/Plugins/Animation/ControlRig/Source/ControlRigEditor/Private/AnimDetails/AnimDetailsProxyManager.h
@@ -3,6 +3,7 @@
 #pragma once
 
 #include "AnimDetailsFilter.h"
+#include "AnimDetailsNavigableWidgetRegistry.h"
 #include "EditorUndoClient.h"
 #include "Engine/TimerHandle.h"
 #include "TimerManager.h"
@@ -65,6 +66,9 @@ public:
 	/** Returns the member name of the proxies member */
 	static const FName GetProxiesMemberNameChecked() { return GET_MEMBER_NAME_CHECKED(UAnimDetailsProxyManager, Proxies); }
 
+	/** Returns the navigable widget registry that belongs to this proxy manager. Useful for navigation */
+	UE::ControlRigEditor::FAnimDetailsNavigableWidgetRegistry& GetNavigableWidgetRegistry() { return NavigableWidgetRegistry; }
+
 private:
 	//~ Begin FSelfRegisteringEditorUndoClient interface
 	virtual void PostUndo(bool bSuccess) override;
@@ -140,6 +144,9 @@ private:
 	/** The filter used for anim details */
 	UE::ControlRigEditor::FAnimDetailsFilter AnimDetailsFilter;
 
+	/** Registry for value widgets, useful for navigation */
+	UE::ControlRigEditor::FAnimDetailsNavigableWidgetRegistry NavigableWidgetRegistry;
+
 	/** The sequencer object this manager currently handles */
 	TWeakPtr<ISequencer> WeakSequencer;
 
diff --git a/Engine/Plugins/Animation/ControlRig/Source/ControlRigEditor/Private/AnimDetails/Widgets/SAnimDetailsValueBoolean.cpp b/Engine/Plugins/Animation/ControlRig/Source/ControlRigEditor/Private/AnimDetails/Widgets/SAnimDetailsValueBoolean.cpp
--- a/Engine/Plugins/Animation/ControlRig/Source/ControlRigEditor/Private/AnimDetails/Widgets/SAnimDetailsValueBoolean.cpp
+++ b/Engine/Plugins/Animation/ControlRig/Source/ControlRigEditor/Private/AnimDetails/Widgets/SAnimDetailsValueBoolean.cpp
@@ -30,21 +30,53 @@ namespace UE::ControlRigEditor
 		WeakProxyManager = ProxyManager;
 		WeakPropertyHandle = InPropertyHandle;
 
+		CheckBox = SNew(SCheckBox)
+			.Type(ESlateCheckBoxType::CheckBox)
+			.IsChecked(this, &SAnimDetailsValueBoolean::GetCheckState)
+			.OnCheckStateChanged(this, &SAnimDetailsValueBoolean::OnCheckStateChanged);
+
 		ChildSlot
 		[
 			SAssignNew(SelectionBorder, SAnimDetailsPropertySelectionBorder, InPropertyHandle)
 			.RequiresModifierKeys(true)
 			[
-				SNew(SCheckBox)
-				.Type(ESlateCheckBoxType::CheckBox)
-				.IsChecked(this, &SAnimDetailsValueBoolean::GetCheckState)
-				.OnCheckStateChanged(this, &SAnimDetailsValueBoolean::OnCheckStateChanged)
+				CheckBox.ToSharedRef()
 			]
 		];
 
+		NavigableWidgetRegistrar.Register(*ProxyManager, CheckBox.ToSharedRef());
+
 		FAnimDetailsMultiEditUtil::Get().Join(ProxyManager, InPropertyHandle);
 	}
 
+	FNavigationReply SAnimDetailsValueBoolean::OnNavigation(const FGeometry& MyGeometry, const FNavigationEvent& InNavigationEvent)
+	{
+		if (InNavigationEvent.GetNavigationType() != EUINavigation::Next)
+		{
+			return SCompoundWidget::OnNavigation(MyGeometry, InNavigationEvent);
+		}
+		else
+		{
+			return FNavigationReply::Custom(FNavigationDelegate::CreateLambda([WeakThis = AsWeak(), this](EUINavigation UINaviagation)
+				{
+					if (!WeakThis.IsValid() ||
+						!WeakProxyManager.IsValid() ||
+						!CheckBox.IsValid())
+					{
+						return SNullWidget::NullWidget;
+					}
+
+					const TSharedPtr<SWidget> NextWidget = WeakProxyManager->GetNavigableWidgetRegistry().GetNext(CheckBox.ToSharedRef());
+					if (!NextWidget.IsValid())
+					{
+						return SNullWidget::NullWidget;
+					}
+
+					return NextWidget.ToSharedRef();
+				}));
+		}
+	}
+
 	ECheckBoxState SAnimDetailsValueBoolean::GetCheckState() const
 	{
 		bool bChecked{};
diff --git a/Engine/Plugins/Animation/ControlRig/Source/ControlRigEditor/Private/AnimDetails/Widgets/SAnimDetailsValueBoolean.h b/Engine/Plugins/Animation/ControlRig/Source/ControlRigEditor/Private/AnimDetails/Widgets/SAnimDetailsValueBoolean.h
--- a/Engine/Plugins/Animation/ControlRig/Source/ControlRigEditor/Private/AnimDetails/Widgets/SAnimDetailsValueBoolean.h
+++ b/Engine/Plugins/Animation/ControlRig/Source/ControlRigEditor/Private/AnimDetails/Widgets/SAnimDetailsValueBoolean.h
@@ -2,9 +2,11 @@
 
 #pragma once
 
+#include "AnimDetails/AnimDetailsNavigableWidgetRegistrar.h"
 #include "Widgets/SCompoundWidget.h"
 
 class IPropertyHandle;
+class SCheckBox;
 class UAnimDetailsProxyManager;
 
 namespace UE::ControlRigEditor
@@ -35,15 +37,25 @@ namespace UE::ControlRigEditor
 		void Construct(const FArguments& InArgs, const TSharedRef<IPropertyHandle>& PropertyHandle);
 
 	private:
+		//~ Begin SWidget Interface
+		virtual FNavigationReply OnNavigation(const FGeometry& MyGeometry, const FNavigationEvent& InNavigationEvent) override;
+		//~ End SWidget Interface
+
 		/** Returns the current check box state */
 		ECheckBoxState GetCheckState() const;
 
 		/** Called when the check box state changed */
 		void OnCheckStateChanged(ECheckBoxState CheckBoxState);
 
+		/** The check box used to display the boolean state */
+		TSharedPtr<SCheckBox> CheckBox;
+
 		/** Border that handles the selction of the displayed property */
 		TSharedPtr<SAnimDetailsPropertySelectionBorder> SelectionBorder;
 
+		/** Handles registration of the inner widget that can be navigated to */
+		FAnimDetailsNavigableWidgetRegistrar NavigableWidgetRegistrar;
+
 		/** The proxy manager that holds the displayed property */
 		TWeakObjectPtr<UAnimDetailsProxyManager> WeakProxyManager;
 
diff --git a/Engine/Plugins/Animation/ControlRig/Source/ControlRigEditor/Private/AnimDetails/Widgets/SAnimDetailsValueNumeric.cpp b/Engine/Plugins/Animation/ControlRig/Source/ControlRigEditor/Private/AnimDetails/Widgets/SAnimDetailsValueNumeric.cpp
--- a/Engine/Plugins/Animation/ControlRig/Source/ControlRigEditor/Private/AnimDetails/Widgets/SAnimDetailsValueNumeric.cpp
+++ b/Engine/Plugins/Animation/ControlRig/Source/ControlRigEditor/Private/AnimDetails/Widgets/SAnimDetailsValueNumeric.cpp
@@ -208,12 +208,45 @@ namespace UE::ControlRigEditor
 		}
 	}
 
+	template <typename NumericType>
+	FNavigationReply SAnimDetailsValueNumeric<NumericType>::OnNavigation(const FGeometry& MyGeometry, const FNavigationEvent& InNavigationEvent)
+	{
+		if (InNavigationEvent.GetNavigationType() != EUINavigation::Next)
+		{
+			return SCompoundWidget::OnNavigation(MyGeometry, InNavigationEvent);
+		}
+		else
+		{
+			return FNavigationReply::Custom(FNavigationDelegate::CreateLambda([WeakThis = AsWeak(), this](EUINavigation UINaviagation)
+				{
+					if (!WeakThis.IsValid() || 
+						!WeakProxyManager.IsValid() || 
+						!SpinBox.IsValid())
+					{
+						return SNullWidget::NullWidget;
+					}
+
+					const TSharedPtr<SWidget> NextWidget = WeakProxyManager->GetNavigableWidgetRegistry().GetNext(SpinBox.ToSharedRef());
+					if (!NextWidget.IsValid())
+					{
+						return SNullWidget::NullWidget;
+					}
+
+					return NextWidget.ToSharedRef();
+				}));
+		}
+	}
+
+	template FNavigationReply SAnimDetailsValueNumeric<double>::OnNavigation(const FGeometry& MyGeometry, const FNavigationEvent& InNavigationEvent);
+	template FNavigationReply SAnimDetailsValueNumeric<int64>::OnNavigation(const FGeometry& MyGeometry, const FNavigationEvent& InNavigationEvent);
+
 	template <typename NumericType>
 	void SAnimDetailsValueNumeric<NumericType>::ForceRefresh()
 	{
 		RefreshTimerHandle.Invalidate();
 
-		if (!WeakPropertyHandle.IsValid())
+		if (!WeakProxyManager.IsValid() ||
+			!WeakPropertyHandle.IsValid())
 		{
 			ChildSlot
 			[
@@ -222,13 +255,14 @@ namespace UE::ControlRigEditor
 
 			return;
 		}
+		UAnimDetailsProxyManager* ProxyManager = WeakProxyManager.Get();
 
 		const TSharedRef<IPropertyHandle> PropertyHandle = WeakPropertyHandle.Pin().ToSharedRef();
 		const PropertyUtils::FFloatingPointMetaData<NumericType>& PropertyMetaData = PropertyUtils::FFloatingPointMetaData<NumericType>(PropertyHandle);
 
 		const uint8 NumFractionalDigits = GetDefault<UAnimDetailsSettings>()->NumFractionalDigits;
 
-		const TSharedRef<SWidget> SpinBox = SNew(SSpinBox<NumericType>)
+		SpinBox = SNew(SSpinBox<NumericType>)
 			.Font(FAppStyle::GetFontStyle("PropertyWindow.NormalFont"))
 			.Value(this, &SAnimDetailsValueNumeric<NumericType>::GetValue)
 			.OnGetDisplayValue(this, &SAnimDetailsValueNumeric<NumericType>::OnGetDisplayValue)
@@ -251,6 +285,8 @@ namespace UE::ControlRigEditor
 			.LinearDeltaSensitivity(1)
 			.TypeInterface(TypeInterface);
 
+		NavigableWidgetRegistrar.Register(*ProxyManager, SpinBox.ToSharedRef());
+
 		const TSharedRef<SHorizontalBox> ContentBox = SNew(SHorizontalBox);
 
 		const bool bHasLabel = Label.Widget != SNullWidget::NullWidget;
@@ -264,7 +300,7 @@ namespace UE::ControlRigEditor
 					.HAlign(HAlign_Fill)
 					.VAlign(VAlign_Center)
 					[
-						SpinBox
+						SpinBox.ToSharedRef()
 					]
 
 					+ SOverlay::Slot()
@@ -289,14 +325,14 @@ namespace UE::ControlRigEditor
 
 			ContentBox->AddSlot()
 				[
-					SpinBox
+					SpinBox.ToSharedRef()
 				];
 		}
 		else
 		{
 			ContentBox->AddSlot()
 				[
-					SpinBox
+					SpinBox.ToSharedRef()
 				];
 		}
 
diff --git a/Engine/Plugins/Animation/ControlRig/Source/ControlRigEditor/Private/AnimDetails/Widgets/SAnimDetailsValueNumeric.h b/Engine/Plugins/Animation/ControlRig/Source/ControlRigEditor/Private/AnimDetails/Widgets/SAnimDetailsValueNumeric.h
--- a/Engine/Plugins/Animation/ControlRig/Source/ControlRigEditor/Private/AnimDetails/Widgets/SAnimDetailsValueNumeric.h
+++ b/Engine/Plugins/Animation/ControlRig/Source/ControlRigEditor/Private/AnimDetails/Widgets/SAnimDetailsValueNumeric.h
@@ -2,6 +2,7 @@
 
 #pragma once
 
+#include "AnimDetails/AnimDetailsNavigableWidgetRegistrar.h"
 #include "Engine/TimerHandle.h"
 #include "Widgets/Layout/SBorder.h"
 #include "Widgets/SCompoundWidget.h"
@@ -92,6 +93,10 @@ namespace UE::ControlRigEditor
 		}
 
 	private:
+		//~ Begin SWidget Interface
+		virtual FNavigationReply OnNavigation(const FGeometry& MyGeometry, const FNavigationEvent& InNavigationEvent) override;
+		//~ End SWidget Interface
+
 		/** Refreshes the widget on the next tick */
 		void RequestRefresh();
 
@@ -143,6 +148,12 @@ namespace UE::ControlRigEditor
 		/** Padding around the label content */
 		FMargin LabelPadding;
 
+		/** The spinbox that is used to display the numeric value*/
+		TSharedPtr<SWidget> SpinBox;
+
+		/** Handles registration of the inner widget that can be navigated to */
+		FAnimDetailsNavigableWidgetRegistrar NavigableWidgetRegistrar;
+
 		/** Type inteface for the spinbox */
 		TSharedPtr<TAnimDetailsNumericTypeInterface<NumericType>> TypeInterface;
 	};
