Fixed AnimNode Control Rig variable mappings Initialization, as it was not generating mapping functions correctly on the initial call to OnInitializeAnimInstance

diff --git a/Engine/Plugins/Animation/ControlRig/Source/ControlRig/Private/AnimNode_ControlRig.cpp b/Engine/Plugins/Animation/ControlRig/Source/ControlRig/Private/AnimNode_ControlRig.cpp
--- a/Engine/Plugins/Animation/ControlRig/Source/ControlRig/Private/AnimNode_ControlRig.cpp
+++ b/Engine/Plugins/Animation/ControlRig/Source/ControlRig/Private/AnimNode_ControlRig.cpp
@@ -422,6 +422,7 @@ bool FAnimNode_ControlRig::UpdateControlRigIfNeeded(const UAnimInstance* InAnimI
 				if(NewControlRig)
 				{
 					Swap(NewControlRig, ControlRig);
+					SetTargetInstance(ControlRig);
 				}
 			}
 			else
@@ -443,6 +444,7 @@ bool FAnimNode_ControlRig::UpdateControlRigIfNeeded(const UAnimInstance* InAnimI
 			ControlRig->GetObjectBinding()->BindToObject(InAnimInstance->GetOwningComponent());
 			ControlRig->Initialize(true);
 			ControlRig->RequestInit();
+			SetTargetInstance(ControlRig);
 		}
 
 		ControlRigHierarchyMappings.ResetRefPoseSetterHash();
