Control Rig: Need to update Rotation Order table with modular rig updated names

diff --git a/Engine/Plugins/Animation/ControlRig/Source/ControlRig/Private/Sequencer/MovieSceneControlRigParameterTrack.cpp b/Engine/Plugins/Animation/ControlRig/Source/ControlRig/Private/Sequencer/MovieSceneControlRigParameterTrack.cpp
--- a/Engine/Plugins/Animation/ControlRig/Source/ControlRig/Private/Sequencer/MovieSceneControlRigParameterTrack.cpp
+++ b/Engine/Plugins/Animation/ControlRig/Source/ControlRig/Private/Sequencer/MovieSceneControlRigParameterTrack.cpp
@@ -104,6 +104,32 @@ bool UMovieSceneControlRigParameterTrack::SupportsType(TSubclassOf<UMovieSceneSe
 	return SectionClass == UMovieSceneControlRigParameterSection::StaticClass();
 }
 
+//we have possible name changes with modular rigs, need to update the rotation order map with this new name
+void UMovieSceneControlRigParameterTrack::CheckForNameChanges()
+{
+	if (ControlRig && ControlRig->IsModularRig())
+	{
+		TArray<FName> RotationOrderNames;
+		ControlsRotationOrder.GenerateKeyArray(RotationOrderNames);
+		for (const FName& OrderControlName : RotationOrderNames)
+		{
+			if (FControlRotationOrder* RotationOrder = ControlsRotationOrder.Find(OrderControlName))
+			{
+				if (const FRigControlElement* Control = ControlRig->FindControl(OrderControlName)) //this will find old names also
+				{
+					const FName& ControlName = Control->GetFName();
+					if (ControlName != OrderControlName)
+					{
+						FControlRotationOrder NewRotationOrder = *RotationOrder;
+						ControlsRotationOrder.Remove(OrderControlName);
+						ControlsRotationOrder.Add(ControlName, NewRotationOrder);
+					}
+				}
+			}
+		}
+	}
+}
+
 UMovieSceneSection* UMovieSceneControlRigParameterTrack::CreateNewSection()
 {
 	UMovieSceneControlRigParameterSection* NewSection = NewObject<UMovieSceneControlRigParameterSection>(this, NAME_None, RF_Transactional);
@@ -532,6 +558,8 @@ void UMovieSceneControlRigParameterTrack::ReconstructControlRig()
 		ControlRig->ConditionalPostLoad();
 		ApplyControlRigSettingsOverrides();
 		ControlRig->Initialize();
+		CheckForNameChanges();
+
 		for (UMovieSceneSection* Section : Sections)
 		{
 			if (Section)
@@ -647,6 +675,7 @@ void UMovieSceneControlRigParameterTrack::HandleOnPostConstructed_GameThread()
 {
 	if (IsValid(ControlRig))
 	{
+		CheckForNameChanges();
 		TArray<FRigControlElement*> SortedControls;
 		ControlRig->GetControlsInOrder(SortedControls);
 #if WITH_EDITOR
@@ -809,6 +838,9 @@ void UMovieSceneControlRigParameterTrack::ReplaceControlRig(UControlRig* NewCont
 			ControlRig->Rename(nullptr, this);
 		}
 	}
+
+	CheckForNameChanges();
+
 	for (UMovieSceneSection* Section : Sections)
 	{
 		if (UMovieSceneControlRigParameterSection* CRSection = Cast<UMovieSceneControlRigParameterSection>(Section))
diff --git a/Engine/Plugins/Animation/ControlRig/Source/ControlRig/Public/Sequencer/MovieSceneControlRigParameterTrack.h b/Engine/Plugins/Animation/ControlRig/Source/ControlRig/Public/Sequencer/MovieSceneControlRigParameterTrack.h
--- a/Engine/Plugins/Animation/ControlRig/Source/ControlRig/Public/Sequencer/MovieSceneControlRigParameterTrack.h
+++ b/Engine/Plugins/Animation/ControlRig/Source/ControlRig/Public/Sequencer/MovieSceneControlRigParameterTrack.h
@@ -206,6 +206,9 @@ private:
 	IMovieSceneConstrainedSection::FConstraintChannelAddedEvent OnConstraintChannelAdded;
 	void ReconstructControlRig();
 
+	//look for name changes in particularly with modular rigs so we can update the rotation order map
+	void CheckForNameChanges();
+
 private:
 
 
