Control Rig: Need to update Constraint Channels with updated Modular rig names, or constraints come in broken

diff --git a/Engine/Plugins/Animation/ControlRig/Source/ControlRig/Private/Sequencer/MovieSceneControlRigParameterSection.cpp b/Engine/Plugins/Animation/ControlRig/Source/ControlRig/Private/Sequencer/MovieSceneControlRigParameterSection.cpp
--- a/Engine/Plugins/Animation/ControlRig/Source/ControlRig/Private/Sequencer/MovieSceneControlRigParameterSection.cpp
+++ b/Engine/Plugins/Animation/ControlRig/Source/ControlRig/Private/Sequencer/MovieSceneControlRigParameterSection.cpp
@@ -3057,6 +3057,44 @@ void UMovieSceneControlRigParameterSection::RecreateWithThisControlRig(UControlR
 				}
 			}
 		});
+
+	}
+
+	//also need update Modular Rig FNames on the transformable handles on the constraint channels for those that are constraints. Update both child and parent
+	auto UpdateHandleName = [](UTransformableControlHandle* Handle)
+	{
+		if (Handle)
+		{
+			if (UControlRig* ControlRig = Handle->ControlRig.Get())
+			{
+				if (ControlRig->IsModularRig())
+				{
+					if (const FRigControlElement* Control = ControlRig->FindControl(Handle->ControlName))
+					{
+						const FName& ControlName = Control->GetFName();
+						if (ControlName != Handle->ControlName)
+						{
+							Handle->ControlName = ControlName;
+						}
+					}
+				}
+			}
+		}
+	};
+
+	for (FConstraintAndActiveChannel& ConstraintChannel : ConstraintsChannels)
+	{
+		if (UTickableTransformConstraint* TransformConstraint = Cast<UTickableTransformConstraint>(ConstraintChannel.GetConstraint()))
+		{
+			if (UTransformableControlHandle* Handle = Cast<UTransformableControlHandle>(TransformConstraint->ChildTRSHandle))
+			{
+				UpdateHandleName(Handle);
+			}
+			if (UTransformableControlHandle* Handle = Cast<UTransformableControlHandle>(TransformConstraint->ParentTRSHandle))
+			{
+				UpdateHandleName(Handle);
+			}
+		}
 	}
 
 	for (FRigControlElement* ControlElement : SortedControls)
