Control Rig: Search&Replace to support switching element key type

#rnx

diff --git a/Engine/Plugins/Animation/ControlRig/Source/ControlRigEditor/Private/ControlRigEditorModule.cpp b/Engine/Plugins/Animation/ControlRig/Source/ControlRigEditor/Private/ControlRigEditorModule.cpp
--- a/Engine/Plugins/Animation/ControlRig/Source/ControlRigEditor/Private/ControlRigEditorModule.cpp
+++ b/Engine/Plugins/Animation/ControlRig/Source/ControlRigEditor/Private/ControlRigEditorModule.cpp
@@ -607,8 +607,38 @@ void FControlRigEditorModule::GetNodeContextMenuActions(IRigVMClientHost* RigVMC
 						if(NewNameStr != OldNameStr)
 						{
 							Key.Name = *NewNameStr;
-							if(ControlRigBlueprint->Hierarchy->GetIndex(Key) != INDEX_NONE)
+
+							bool bIsValid = ControlRigBlueprint->Hierarchy->GetIndex(Key) != INDEX_NONE;
+
+							if (!bIsValid)
 							{
+								for (ERigElementType ElementType = ERigElementType::First; ElementType <= ERigElementType::Last; ElementType = static_cast<ERigElementType>(static_cast<int32>(ElementType) + 1))
+								{
+									Key.Type = ElementType;
+									if (ControlRigBlueprint->Hierarchy->GetIndex(Key) != INDEX_NONE)
+									{
+										bIsValid = true;
+										break;
+									}
+								}
+							}
+
+							if(bIsValid)
+							{
+								if (Key.Type != Pair.Value.Type)
+								{
+									if (const URigVMPin* ParentPin = Pin->GetParentPin())
+									{
+										if (ParentPin->GetCPPTypeObject() == FRigElementKey::StaticStruct())
+										{
+											if (const URigVMPin* TypePin = ParentPin->FindSubPin(GET_MEMBER_NAME_STRING_CHECKED(FRigElementKey, Type)))
+											{
+												const FString TypeAsString = StaticEnum<ERigElementType>()->GetNameStringByValue(static_cast<int64>(Key.Type));
+												Controller->SetPinDefaultValue(TypePin->GetPinPath(), TypeAsString, false, true, false, true);
+											}
+										}
+									}
+								}
 								Controller->SetPinDefaultValue(Pin->GetPinPath(), NewNameStr, false, true, false, true);
 								ReplacedNames++;
 							}
