Control Rig: Second attempt to fix cooking issues for modular rigs.
#rnx

diff --git a/Engine/Plugins/Animation/ControlRig/Source/ControlRigDeveloper/Private/AnimGraphNode_ControlRig.cpp b/Engine/Plugins/Animation/ControlRig/Source/ControlRigDeveloper/Private/AnimGraphNode_ControlRig.cpp
--- a/Engine/Plugins/Animation/ControlRig/Source/ControlRigDeveloper/Private/AnimGraphNode_ControlRig.cpp
+++ b/Engine/Plugins/Animation/ControlRig/Source/ControlRigDeveloper/Private/AnimGraphNode_ControlRig.cpp
@@ -485,7 +485,7 @@ void UAnimGraphNode_ControlRig::PreloadRequiredAssets()
 		{
 			ModularRig->ModularRigModel.ForEachModule([this](const FRigModuleReference* Module) -> bool
 			{
-				PreloadObject(Module->Class.Get());
+				(void) Module->Class.LoadSynchronous();
 				return true;
 			});
 		}
