Fixed an editor undo redo crash at Control Rig variable mappings, when performing a Redo of a previowsly Undoed SetPreviewMesh transaction in an AnimBP that contains a Control Rig.

diff --git a/Engine/Plugins/Animation/ControlRig/Source/ControlRig/Private/Tools/ControlRigVariableMappings.cpp b/Engine/Plugins/Animation/ControlRig/Source/ControlRig/Private/Tools/ControlRigVariableMappings.cpp
--- a/Engine/Plugins/Animation/ControlRig/Source/ControlRig/Private/Tools/ControlRigVariableMappings.cpp
+++ b/Engine/Plugins/Animation/ControlRig/Source/ControlRig/Private/Tools/ControlRigVariableMappings.cpp
@@ -176,9 +176,9 @@ void FControlRigVariableMappings::PropagateInputProperties(const UObject* InSour
 
 	if (AnimNodeLocals::bUseFunctions && !UpdateFunctions.IsEmpty())
 	{
-		Algo::ForEach(UpdateFunctions, [InSourceInstance](const PropertyUpdateFunction& InFunc)
+		Algo::ForEach(UpdateFunctions, [TargetHierarchy, InSourceInstance](const PropertyUpdateFunction& InFunc)
 			{
-				InFunc(InSourceInstance);
+				InFunc(TargetHierarchy, InSourceInstance);
 			});
 	}
 	else
@@ -505,9 +505,15 @@ void FControlRigVariableMappings::AddControlFunction(FRigControlElement* InContr
 		{
 			if (ensure(CastField<FBoolProperty>(InSourceProperty)))
 			{
-				return AddUpdateFunction([InSourceProperty, InTargetHierarchy, InControlElement](const UObject* InSourceInstance)
+				return AddUpdateFunction([InSourceProperty, ElementKey = InControlElement->GetKey(), ElementIndex = InControlElement->GetIndex()](URigHierarchy* InTargetHierarchy, const UObject* InSourceInstance)
 					{
-						SetControlValue<bool>(InControlElement, InSourceProperty, InTargetHierarchy, InSourceInstance);
+						if (FRigBaseElement* BaseElement = InTargetHierarchy->Get(ElementIndex))
+						{
+							if (BaseElement->GetKey() == ElementKey)
+							{
+								SetControlValue<bool>(CastChecked<FRigControlElement>(BaseElement), InSourceProperty, InTargetHierarchy, InSourceInstance);
+							}
+						}
 					});
 			}
 			break;
@@ -517,9 +523,15 @@ void FControlRigVariableMappings::AddControlFunction(FRigControlElement* InContr
 		{
 			if (ensure(CastField<FFloatProperty>(InSourceProperty)))
 			{
-				return AddUpdateFunction([InSourceProperty, InTargetHierarchy, InControlElement](const UObject* InSourceInstance)
+				return AddUpdateFunction([InSourceProperty, ElementKey = InControlElement->GetKey(), ElementIndex = InControlElement->GetIndex()](URigHierarchy* InTargetHierarchy, const UObject* InSourceInstance)
 					{
-						SetControlValue<float>(InControlElement, InSourceProperty, InTargetHierarchy, InSourceInstance);
+						if (FRigBaseElement* BaseElement = InTargetHierarchy->Get(ElementIndex))
+						{
+							if (BaseElement->GetKey() == ElementKey)
+							{
+								SetControlValue<float>(CastChecked<FRigControlElement>(BaseElement), InSourceProperty, InTargetHierarchy, InSourceInstance);
+							}
+						}
 					});
 			}
 			break;
@@ -528,9 +540,15 @@ void FControlRigVariableMappings::AddControlFunction(FRigControlElement* InContr
 		{
 			if (ensure(CastField<FIntProperty>(InSourceProperty)))
 			{
-				return AddUpdateFunction([InSourceProperty, InTargetHierarchy, InControlElement](const UObject* InSourceInstance)
+				return AddUpdateFunction([InSourceProperty, ElementKey = InControlElement->GetKey(), ElementIndex = InControlElement->GetIndex()](URigHierarchy* InTargetHierarchy, const UObject* InSourceInstance)
 					{
-						SetControlValue<int32>(InControlElement, InSourceProperty, InTargetHierarchy, InSourceInstance);
+						if (FRigBaseElement* BaseElement = InTargetHierarchy->Get(ElementIndex))
+						{
+							if (BaseElement->GetKey() == ElementKey)
+							{
+								SetControlValue<int32>(CastChecked<FRigControlElement>(BaseElement), InSourceProperty, InTargetHierarchy, InSourceInstance);
+							}
+						}
 					});
 			}
 			break;
@@ -542,9 +560,15 @@ void FControlRigVariableMappings::AddControlFunction(FRigControlElement* InContr
 			{
 				if (ensure(StructProperty->Struct == TBaseStructure<FVector2D>::Get()))
 				{
-					return AddUpdateFunction([InSourceProperty, InTargetHierarchy, InControlElement](const UObject* InSourceInstance)
+					return AddUpdateFunction([InSourceProperty, ElementKey = InControlElement->GetKey(), ElementIndex = InControlElement->GetIndex()](URigHierarchy* InTargetHierarchy, const UObject* InSourceInstance)
 						{
-							SetControlValue<FVector2D>(InControlElement, InSourceProperty, InTargetHierarchy, InSourceInstance);
+							if (FRigBaseElement* BaseElement = InTargetHierarchy->Get(ElementIndex))
+							{
+								if (BaseElement->GetKey() == ElementKey)
+								{
+									SetControlValue<FVector2D>(CastChecked<FRigControlElement>(BaseElement), InSourceProperty, InTargetHierarchy, InSourceInstance);
+								}
+							}
 						});
 				}
 			}
@@ -558,9 +582,15 @@ void FControlRigVariableMappings::AddControlFunction(FRigControlElement* InContr
 			{
 				if (ensure(StructProperty->Struct == TBaseStructure<FVector>::Get()))
 				{
-					return AddUpdateFunction([InSourceProperty, InTargetHierarchy, InControlElement](const UObject* InSourceInstance)
+					return AddUpdateFunction([InSourceProperty, ElementKey = InControlElement->GetKey(), ElementIndex = InControlElement->GetIndex()](URigHierarchy* InTargetHierarchy, const UObject* InSourceInstance)
 						{
-							SetControlValue<FVector>(InControlElement, InSourceProperty, InTargetHierarchy, InSourceInstance);
+							if (FRigBaseElement* BaseElement = InTargetHierarchy->Get(ElementIndex))
+							{
+								if (BaseElement->GetKey() == ElementKey)
+								{
+									SetControlValue<FVector>(CastChecked<FRigControlElement>(BaseElement), InSourceProperty, InTargetHierarchy, InSourceInstance);
+								}
+							}
 						});
 				}
 			}
@@ -573,9 +603,15 @@ void FControlRigVariableMappings::AddControlFunction(FRigControlElement* InContr
 			{
 				if (ensure(StructProperty->Struct == TBaseStructure<FRotator>::Get()))
 				{
-					return AddUpdateFunction([InSourceProperty, InTargetHierarchy, InControlElement](const UObject* InSourceInstance)
+					return AddUpdateFunction([InSourceProperty, ElementKey = InControlElement->GetKey(), ElementIndex = InControlElement->GetIndex()](URigHierarchy* InTargetHierarchy, const UObject* InSourceInstance)
 						{
-							SetControlValue<FRotator>(InControlElement, InSourceProperty, InTargetHierarchy, InSourceInstance);
+							if (FRigBaseElement* BaseElement = InTargetHierarchy->Get(ElementIndex))
+							{
+								if (BaseElement->GetKey() == ElementKey)
+								{
+									SetControlValue<FRotator>(CastChecked<FRigControlElement>(BaseElement), InSourceProperty, InTargetHierarchy, InSourceInstance);
+								}
+							}
 						});
 				}
 			}
@@ -588,9 +624,15 @@ void FControlRigVariableMappings::AddControlFunction(FRigControlElement* InContr
 			{
 				if (ensure(StructProperty->Struct == TBaseStructure<FTransform>::Get()))
 				{
-					return AddUpdateFunction([InSourceProperty, InTargetHierarchy, InControlElement](const UObject* InSourceInstance)
+					return AddUpdateFunction([InSourceProperty, ElementKey = InControlElement->GetKey(), ElementIndex = InControlElement->GetIndex()](URigHierarchy* InTargetHierarchy, const UObject* InSourceInstance)
 						{
-							SetControlValue<FTransform>(InControlElement, InSourceProperty, InTargetHierarchy, InSourceInstance);
+							if (FRigBaseElement* BaseElement = InTargetHierarchy->Get(ElementIndex))
+							{
+								if (BaseElement->GetKey() == ElementKey)
+								{
+									SetControlValue<FTransform>(CastChecked<FRigControlElement>(BaseElement), InSourceProperty, InTargetHierarchy, InSourceInstance);
+								}
+							}
 						});
 				}
 			}
@@ -603,9 +645,15 @@ void FControlRigVariableMappings::AddControlFunction(FRigControlElement* InContr
 			{
 				if (ensure(StructProperty->Struct == TBaseStructure<FTransform>::Get()))
 				{
-					return AddUpdateFunction([InSourceProperty, InTargetHierarchy, InControlElement](const UObject* InSourceInstance)
+					return AddUpdateFunction([InSourceProperty, ElementKey = InControlElement->GetKey(), ElementIndex = InControlElement->GetIndex()](URigHierarchy* InTargetHierarchy, const UObject* InSourceInstance)
 						{
-							SetControlValue<FTransform, FTransformNoScale>(InControlElement, InSourceProperty, InTargetHierarchy, InSourceInstance);
+							if (FRigBaseElement* BaseElement = InTargetHierarchy->Get(ElementIndex))
+							{
+								if (BaseElement->GetKey() == ElementKey)
+								{
+									SetControlValue<FTransform, FTransformNoScale>(CastChecked<FRigControlElement>(BaseElement), InSourceProperty, InTargetHierarchy, InSourceInstance);
+								}
+							}
 						});
 				}
 			}
@@ -618,11 +666,17 @@ void FControlRigVariableMappings::AddControlFunction(FRigControlElement* InContr
 			{
 				if (ensure(StructProperty->Struct == TBaseStructure<FTransform>::Get()))
 				{
-					return AddUpdateFunction([InSourceProperty, InTargetHierarchy, InControlElement](const UObject* InSourceInstance)
+					return AddUpdateFunction([InSourceProperty, ElementKey = InControlElement->GetKey(), ElementIndex = InControlElement->GetIndex()](URigHierarchy* InTargetHierarchy, const UObject* InSourceInstance)
 						{
-							const FTransform* ValuePtr = InSourceProperty->ContainerPtrToValuePtr<FTransform>(InSourceInstance);
-							const FRigControlValue Value = FRigControlValue::Make<FEulerTransform>(FEulerTransform(*ValuePtr));
-							InTargetHierarchy->SetControlValue(InControlElement, Value, ERigControlValueType::Current);
+							if (FRigBaseElement* BaseElement = InTargetHierarchy->Get(ElementIndex))
+							{
+								if (BaseElement->GetKey() == ElementKey)
+								{
+									const FTransform* ValuePtr = InSourceProperty->ContainerPtrToValuePtr<FTransform>(InSourceInstance);
+									const FRigControlValue Value = FRigControlValue::Make<FEulerTransform>(FEulerTransform(*ValuePtr));
+									InTargetHierarchy->SetControlValue(CastChecked<FRigControlElement>(BaseElement), Value, ERigControlValueType::Current);
+								}
+							}
 						});
 				}
 			}
@@ -653,7 +707,7 @@ void FControlRigVariableMappings::AddVariableFunction(const FRigVMExternalVariab
 	
 	if (CastField<FBoolProperty>(InSourceProperty) && InVariable.TypeName == RigVMTypeUtils::BoolTypeName)
 	{
-		return AddUpdateFunction([InSourceProperty, this, VariableIndex](const UObject* InSourceInstance)
+		return AddUpdateFunction([InSourceProperty, this, VariableIndex](URigHierarchy* InTargetHierarchy, const UObject* InSourceInstance)
 		{
 			SetVariableValue<bool>(Variables[VariableIndex], InSourceProperty, InSourceInstance);
 		});
@@ -662,14 +716,14 @@ void FControlRigVariableMappings::AddVariableFunction(const FRigVMExternalVariab
 	{
 		if(InVariable.TypeName == RigVMTypeUtils::FloatTypeName)
 		{
-			return AddUpdateFunction([InSourceProperty, this, VariableIndex](const UObject* InSourceInstance)
+			return AddUpdateFunction([InSourceProperty, this, VariableIndex](URigHierarchy* InTargetHierarchy, const UObject* InSourceInstance)
 			{
 				SetVariableValue<float>(Variables[VariableIndex], InSourceProperty, InSourceInstance);
 			});
 		}
 		else
 		{
-			return AddUpdateFunction([InSourceProperty, this, VariableIndex](const UObject* InSourceInstance)
+			return AddUpdateFunction([InSourceProperty, this, VariableIndex](URigHierarchy* InTargetHierarchy, const UObject* InSourceInstance)
 			{
 				SetVariableValue<float, double>(Variables[VariableIndex], InSourceProperty, InSourceInstance);
 			});
@@ -679,14 +733,14 @@ void FControlRigVariableMappings::AddVariableFunction(const FRigVMExternalVariab
 	{
 		if(InVariable.TypeName == RigVMTypeUtils::FloatTypeName)
 		{
-			return AddUpdateFunction([InSourceProperty, this, VariableIndex](const UObject* InSourceInstance)
+			return AddUpdateFunction([InSourceProperty, this, VariableIndex](URigHierarchy* InTargetHierarchy, const UObject* InSourceInstance)
 			{
 				SetVariableValue<double, float>(Variables[VariableIndex], InSourceProperty, InSourceInstance);
 			});
 		}
 		else
 		{
-			return AddUpdateFunction([InSourceProperty, this, VariableIndex](const UObject* InSourceInstance)
+			return AddUpdateFunction([InSourceProperty, this, VariableIndex](URigHierarchy* InTargetHierarchy, const UObject* InSourceInstance)
 			{
 				SetVariableValue<double>(Variables[VariableIndex], InSourceProperty, InSourceInstance);
 			});
@@ -694,28 +748,28 @@ void FControlRigVariableMappings::AddVariableFunction(const FRigVMExternalVariab
 	}
 	else if (CastField<FIntProperty>(InSourceProperty) && InVariable.TypeName == RigVMTypeUtils::Int32TypeName)
 	{
-		return AddUpdateFunction([InSourceProperty, this, VariableIndex](const UObject* InSourceInstance)
+		return AddUpdateFunction([InSourceProperty, this, VariableIndex](URigHierarchy* InTargetHierarchy, const UObject* InSourceInstance)
 		{
 			SetVariableValue<int32>(Variables[VariableIndex], InSourceProperty, InSourceInstance);
 		});
 	}
 	else if (CastField<FNameProperty>(InSourceProperty) && InVariable.TypeName == RigVMTypeUtils::FNameTypeName)
 	{
-		return AddUpdateFunction([InSourceProperty, this, VariableIndex](const UObject* InSourceInstance)
+		return AddUpdateFunction([InSourceProperty, this, VariableIndex](URigHierarchy* InTargetHierarchy, const UObject* InSourceInstance)
 		{
 			SetVariableValue<FName>(Variables[VariableIndex], InSourceProperty, InSourceInstance);
 		});
 	}
 	else if (CastField<FNameProperty>(InSourceProperty) && InVariable.TypeName == RigVMTypeUtils::FStringTypeName)
 	{
-		return AddUpdateFunction([InSourceProperty, this, VariableIndex](const UObject* InSourceInstance)
+		return AddUpdateFunction([InSourceProperty, this, VariableIndex](URigHierarchy* InTargetHierarchy, const UObject* InSourceInstance)
 		{
 			SetVariableValue<FString>(Variables[VariableIndex], InSourceProperty, InSourceInstance);
 		});
 	}
 	else if (FStructProperty* StructProperty = CastField<FStructProperty>(InSourceProperty))
 	{
-		return AddUpdateFunction([StructProperty, this, VariableIndex](const UObject* InSourceInstance)
+		return AddUpdateFunction([StructProperty, this, VariableIndex](URigHierarchy* InTargetHierarchy, const UObject* InSourceInstance)
 		{
 			const uint8* SrcPtr = StructProperty->ContainerPtrToValuePtr<uint8>(InSourceInstance);
 			StructProperty->Struct->CopyScriptStruct(Variables[VariableIndex].Memory, SrcPtr, 1);
@@ -725,7 +779,7 @@ void FControlRigVariableMappings::AddVariableFunction(const FRigVMExternalVariab
 	{
 		if(ensure(ArrayProperty->SameType(InVariable.Property)))
 		{
-			return AddUpdateFunction([ArrayProperty, this, VariableIndex](const UObject* InSourceInstance)
+			return AddUpdateFunction([ArrayProperty, this, VariableIndex](URigHierarchy* InTargetHierarchy, const UObject* InSourceInstance)
 			{
 				const uint8* SrcPtr = ArrayProperty->ContainerPtrToValuePtr<uint8>(InSourceInstance);
 				ArrayProperty->CopyCompleteValue(Variables[VariableIndex].Memory, SrcPtr);
@@ -736,7 +790,7 @@ void FControlRigVariableMappings::AddVariableFunction(const FRigVMExternalVariab
 	{
 		if(ensure(ObjectProperty->SameType(InVariable.Property)))
 		{
-			return AddUpdateFunction([ObjectProperty, this, VariableIndex](const UObject* InSourceInstance)
+			return AddUpdateFunction([ObjectProperty, this, VariableIndex](URigHierarchy* InTargetHierarchy, const UObject* InSourceInstance)
 			{
 				const uint8* SrcPtr = ObjectProperty->ContainerPtrToValuePtr<uint8>(InSourceInstance);
 				ObjectProperty->CopyCompleteValue(Variables[VariableIndex].Memory, SrcPtr);
@@ -747,7 +801,7 @@ void FControlRigVariableMappings::AddVariableFunction(const FRigVMExternalVariab
 	{
 		if(ensure(EnumProperty->SameType(InVariable.Property)))
 		{
-			return AddUpdateFunction([EnumProperty, this, VariableIndex](const UObject* InSourceInstance)
+			return AddUpdateFunction([EnumProperty, this, VariableIndex](URigHierarchy* InTargetHierarchy, const UObject* InSourceInstance)
 			{
 				const uint8* SrcPtr = EnumProperty->ContainerPtrToValuePtr<uint8>(InSourceInstance);
 				EnumProperty->CopyCompleteValue(Variables[VariableIndex].Memory, SrcPtr);
@@ -758,7 +812,7 @@ void FControlRigVariableMappings::AddVariableFunction(const FRigVMExternalVariab
 	{
 		if (ensure(ByteProperty->SameType(InVariable.Property)))
 		{
-			return AddUpdateFunction([ByteProperty, this, VariableIndex](const UObject* InSourceInstance)
+			return AddUpdateFunction([ByteProperty, this, VariableIndex](URigHierarchy* InTargetHierarchy, const UObject* InSourceInstance)
 				{
 					const uint8* SrcPtr = ByteProperty->ContainerPtrToValuePtr<uint8>(InSourceInstance);
 					ByteProperty->CopyCompleteValue(Variables[VariableIndex].Memory, SrcPtr);
diff --git a/Engine/Plugins/Animation/ControlRig/Source/ControlRig/Public/Tools/ControlRigVariableMappings.h b/Engine/Plugins/Animation/ControlRig/Source/ControlRig/Public/Tools/ControlRigVariableMappings.h
--- a/Engine/Plugins/Animation/ControlRig/Source/ControlRig/Public/Tools/ControlRigVariableMappings.h
+++ b/Engine/Plugins/Animation/ControlRig/Source/ControlRig/Public/Tools/ControlRigVariableMappings.h
@@ -133,7 +133,7 @@ struct FControlRigVariableMappings
 	}
 
 private:
-	using PropertyUpdateFunction = TFunction<void(const UObject*)>;
+	using PropertyUpdateFunction = TFunction<void(URigHierarchy*, const UObject*)>;
 	using CustomPropertyUpdateFunction = TFunction<void()>;
 
 	/*
