Fixed a crash in Control Rig Pose adapter when removing bones from the skeleton while having a Blueprint with a Control rig opened in the asset editor

diff --git a/Engine/Plugins/Animation/ControlRig/Source/ControlRig/Private/Tools/ControlRigPoseAdapter.cpp b/Engine/Plugins/Animation/ControlRig/Source/ControlRig/Private/Tools/ControlRigPoseAdapter.cpp
--- a/Engine/Plugins/Animation/ControlRig/Source/ControlRig/Private/Tools/ControlRigPoseAdapter.cpp
+++ b/Engine/Plugins/Animation/ControlRig/Source/ControlRig/Private/Tools/ControlRigPoseAdapter.cpp
@@ -202,7 +202,7 @@ void FControlRigPoseAdapter::UpdateInputOutputMappingIfRequired(UControlRig* InC
 			{
 				// get bone name, and find reverse mapping
 				const FSkeletonPoseBoneIndex BoneIndex = InRequiredBones.GetSkeletonPoseIndexFromCompactPoseIndex(FCompactPoseBoneIndex(Index));
-				if (BoneIndex.IsValid())
+				if (BoneIndex.IsValid() && RefSkeleton->IsValidIndex(BoneIndex.GetInt()))
 				{
 					FName TargetNodeName = RefSkeleton->GetBoneName(BoneIndex.GetInt());
 					FName* SourceName = TargetToSourceMappingTable.Find(TargetNodeName);
@@ -233,7 +233,7 @@ void FControlRigPoseAdapter::UpdateInputOutputMappingIfRequired(UControlRig* InC
 			for (uint16 Index = 0; Index < NumBonesInPose; ++Index)
 			{
 				const FSkeletonPoseBoneIndex BoneIndex = InRequiredBones.GetSkeletonPoseIndexFromCompactPoseIndex(FCompactPoseBoneIndex(Index));
-				if (BoneIndex.IsValid())
+				if (BoneIndex.IsValid() && RefSkeleton->IsValidIndex(BoneIndex.GetInt()))
 				{
 					const FName& BoneName = RefSkeleton->GetBoneName(BoneIndex.GetInt());
 					if (NodeNames.Contains(BoneName))
