Constraints: thread-safe notification

#rnx

diff --git a/Engine/Plugins/Animation/ControlRig/Source/ControlRig/Private/Constraints/ControlRigTransformableHandle.cpp b/Engine/Plugins/Animation/ControlRig/Source/ControlRig/Private/Constraints/ControlRigTransformableHandle.cpp
--- a/Engine/Plugins/Animation/ControlRig/Source/ControlRig/Private/Constraints/ControlRigTransformableHandle.cpp
+++ b/Engine/Plugins/Animation/ControlRig/Source/ControlRig/Private/Constraints/ControlRigTransformableHandle.cpp
@@ -29,6 +29,9 @@ namespace UE::Private
 		return InControlRig ? NotifyingRigs.Contains(InControlRig) : false;
 	}
 
+	// only allow notification from one thread at a time
+	FCriticalSection NotificationLock;
+	
 	struct FControlPoseChangedNotifier
 	{
 		FControlPoseChangedNotifier(const URigHierarchy* InHierarchy, const FRigControlElement* InControl, const USceneComponent* InComponent)
@@ -642,7 +645,10 @@ void UTransformableControlHandle::OnControlModified(
 			EHandleEvent::GlobalTransformUpdated : EHandleEvent::LocalTransformUpdated;
 
 		if (InControl->GetFName() == ControlName)
-		{	// if that handle is wrapping InControl
+		{
+			FScopeLock ScopeLock(&UE::Private::NotificationLock);
+			
+			// if that handle is wrapping InControl
 			if (InContext.bConstraintUpdate)
 			{
 				GetEvaluationBinding().bPendingFlush = true;
@@ -667,6 +673,8 @@ void UTransformableControlHandle::OnControlModified(
 			// the control being modified is not the one wrapped by this handle 
 			if (const FRigControlElement* Control = ControlRig->FindControl(ControlName))
 			{
+				FScopeLock ScopeLock(&UE::Private::NotificationLock);
+			
 				if (InContext.bConstraintUpdate)
 				{
 					GetEvaluationBinding().bPendingFlush = true;
