Anim In Engine: Fix cases where Anim Details Proxy Manager may bind to the same delegates multiple times

#rnx

diff --git a/Engine/Plugins/Animation/ControlRig/Source/ControlRigEditor/Private/AnimDetails/AnimDetailsProxyManager.cpp b/Engine/Plugins/Animation/ControlRig/Source/ControlRigEditor/Private/AnimDetails/AnimDetailsProxyManager.cpp
--- a/Engine/Plugins/Animation/ControlRig/Source/ControlRigEditor/Private/AnimDetails/AnimDetailsProxyManager.cpp
+++ b/Engine/Plugins/Animation/ControlRig/Source/ControlRigEditor/Private/AnimDetails/AnimDetailsProxyManager.cpp
@@ -115,8 +115,8 @@ bool UAnimDetailsProxyManager::UpdateSequencer()
 	FControlRigEditMode* EditMode = static_cast<FControlRigEditMode*>(GLevelEditorModeTools().GetActiveMode(FControlRigEditMode::ModeName));
 	if (EditMode &&
 		!EditMode->OnControlRigAddedOrRemoved().IsBoundToObject(this) &&
-		EditMode->OnControlRigShapeActorsRecreated().IsBoundToObject(this) &&
-		EditMode->OnControlRigSelected().IsBoundToObject(this))
+		!EditMode->OnControlRigShapeActorsRecreated().IsBoundToObject(this) &&
+		!EditMode->OnControlRigSelected().IsBoundToObject(this))
 	{
 		EditMode->OnControlRigAddedOrRemoved().AddUObject(this, &UAnimDetailsProxyManager::OnControlRigControlAdded);
 		EditMode->OnControlRigShapeActorsRecreated().AddUObject(this, &UAnimDetailsProxyManager::OnControlRigShapeActorsRecreated);
