Control Rig: Fix for bogus Horde Android issue, see https://epicgames.slack.com/archives/C042RSGQVPD/p1749846390290029

diff --git a/Engine/Plugins/Animation/ControlRig/Source/ControlRig/Private/Sequencer/MovieSceneControlRigParameterSection.cpp b/Engine/Plugins/Animation/ControlRig/Source/ControlRig/Private/Sequencer/MovieSceneControlRigParameterSection.cpp
--- a/Engine/Plugins/Animation/ControlRig/Source/ControlRig/Private/Sequencer/MovieSceneControlRigParameterSection.cpp
+++ b/Engine/Plugins/Animation/ControlRig/Source/ControlRig/Private/Sequencer/MovieSceneControlRigParameterSection.cpp
@@ -3065,11 +3065,11 @@ void UMovieSceneControlRigParameterSection::RecreateWithThisControlRig(UControlR
 	{
 		if (Handle)
 		{
-			if (UControlRig* ControlRig = Handle->ControlRig.Get())
+			if (UControlRig* HandleControlRig = Handle->ControlRig.Get())
 			{
-				if (ControlRig->IsModularRig())
+				if (HandleControlRig->IsModularRig())
 				{
-					if (const FRigControlElement* Control = ControlRig->FindControl(Handle->ControlName))
+					if (const FRigControlElement* Control = HandleControlRig->FindControl(Handle->ControlName))
 					{
 						const FName& ControlName = Control->GetFName();
 						if (ControlName != Handle->ControlName)
