Anim In Engine: Hosted channels no longer enabled if parent is not enabled, this breaks selection. Also fix for anim layer selection with static meshes

diff --git a/Engine/Plugins/Animation/ControlRig/Source/ControlRig/Private/Sequencer/MovieSceneControlRigParameterSection.cpp b/Engine/Plugins/Animation/ControlRig/Source/ControlRig/Private/Sequencer/MovieSceneControlRigParameterSection.cpp
--- a/Engine/Plugins/Animation/ControlRig/Source/ControlRig/Private/Sequencer/MovieSceneControlRigParameterSection.cpp
+++ b/Engine/Plugins/Animation/ControlRig/Source/ControlRig/Private/Sequencer/MovieSceneControlRigParameterSection.cpp
@@ -2064,11 +2064,13 @@ EMovieSceneChannelProxyType UMovieSceneControlRigParameterSection::CacheChannelP
 				FName ParentControlName = NAME_None;
 				FText Group;
 
+				bool bParentIsEnabled = true;
 				if (Hierarchy->ShouldBeGrouped(ControlElement))
 				{
 					if (const FRigControlElement* ParentControlElement = Cast<FRigControlElement>(Hierarchy->GetFirstParent(ControlElement)))
 					{
 						ParentControlName = ParentControlElement->GetFName();
+						bParentIsEnabled = GetControlNameMask(ParentControlName);
 						if (FText* Found = ControlToGroupName.Find(ParentControlElement))
 						{
 							Group = *Found;
@@ -2080,7 +2082,7 @@ EMovieSceneChannelProxyType UMovieSceneControlRigParameterSection::CacheChannelP
 					}
 				}
 
-				bool bEnabled = GetControlNameMask(ControlElement->GetFName());
+				const bool bEnabled = GetControlNameMask(ControlElement->GetFName()) && bParentIsEnabled;
 
 #if WITH_EDITOR
 				switch (ControlElement->Settings.ControlType)
diff --git a/Engine/Plugins/Animation/ControlRig/Source/ControlRigEditor/Private/Sequencer/AnimLayers/AnimLayers.cpp b/Engine/Plugins/Animation/ControlRig/Source/ControlRigEditor/Private/Sequencer/AnimLayers/AnimLayers.cpp
--- a/Engine/Plugins/Animation/ControlRig/Source/ControlRigEditor/Private/Sequencer/AnimLayers/AnimLayers.cpp
+++ b/Engine/Plugins/Animation/ControlRig/Source/ControlRigEditor/Private/Sequencer/AnimLayers/AnimLayers.cpp
@@ -698,68 +698,81 @@ void UAnimLayer::GetSections(TArray<UMovieSceneSection*>& OutSections) const
 
 void UAnimLayer::SetSelectedInList(bool bInValue) 
 {
-	bIsSelectedInList = bInValue;
-	if (TSharedPtr<ISequencer> SequencerPtr = UAnimLayers::GetSequencerFromAsset())
+	if (bIsSelectedInList != bInValue)
 	{
-		using namespace UE::Sequencer;
-		FOutlinerSelection& SelectedOutlinerItems = SequencerPtr->GetViewModel()->GetSelection()->Outliner;
-
-		FTrackRowModelStorageExtension* TrackRowModelStorage = SequencerPtr->GetViewModel()->GetRootModel()->CastDynamic<FTrackRowModelStorageExtension>();
-		check(TrackRowModelStorage);
-		for (const TPair<TWeakObjectPtr<UObject>, FAnimLayerItem>& Pair : AnimLayerItems)
+		bIsSelectedInList = bInValue;
+		if (TSharedPtr<ISequencer> SequencerPtr = UAnimLayers::GetSequencerFromAsset())
 		{
-			if (Pair.Key != nullptr)
+			using namespace UE::Sequencer;
+			FOutlinerSelection& SelectedOutlinerItems = SequencerPtr->GetViewModel()->GetSelection()->Outliner;
+
+			FTrackRowModelStorageExtension* TrackRowModelStorage = SequencerPtr->GetViewModel()->GetRootModel()->CastDynamic<FTrackRowModelStorageExtension>();
+			check(TrackRowModelStorage);
+			for (const TPair<TWeakObjectPtr<UObject>, FAnimLayerItem>& Pair : AnimLayerItems)
 			{
-				for (const FAnimLayerSectionItem& SectionItem : Pair.Value.SectionItems)
+				if (Pair.Key != nullptr)
 				{
-					if (SectionItem.Section.IsValid())
+					for (const FAnimLayerSectionItem& SectionItem : Pair.Value.SectionItems)
 					{
-						const UControlRig* ControlRig = Cast<UControlRig>(Pair.Key.Get());
-						if (ControlRig)
+						if (SectionItem.Section.IsValid())
 						{
-							if (UMovieSceneControlRigParameterSection* CRSection = Cast<UMovieSceneControlRigParameterSection>(SectionItem.Section.Get()))
+							const UControlRig* ControlRig = Cast<UControlRig>(Pair.Key.Get());
+							if (ControlRig)
 							{
-								for (const TPair<FName, FAnimLayerPropertyAndChannels>& ControlName : SectionItem.AnimLayerSet.Names)
+								if (UMovieSceneControlRigParameterSection* CRSection = Cast<UMovieSceneControlRigParameterSection>(SectionItem.Section.Get()))
 								{
-									if (bIsSelectedInList && ControlRig->IsControlSelected(ControlName.Key) == false)
+									for (const TPair<FName, FAnimLayerPropertyAndChannels>& ControlName : SectionItem.AnimLayerSet.Names)
 									{
-										continue; //don't select it if not selected
-									}
-									FChannelMapInfo* pChannelIndex = CRSection->ControlChannelMap.Find(ControlName.Key);
-									if (pChannelIndex != nullptr)
-									{
-										if (pChannelIndex->ParentControlIndex == INDEX_NONE)
+										if (bIsSelectedInList && ControlRig->IsControlSelected(ControlName.Key) == false)
 										{
-											int32 CategoryIndex = CRSection->GetActiveCategoryIndex(ControlName.Key);
-											if (CategoryIndex != INDEX_NONE)
+											continue; //don't select it if not selected
+										}
+										FChannelMapInfo* pChannelIndex = CRSection->ControlChannelMap.Find(ControlName.Key);
+										if (pChannelIndex != nullptr)
+										{
+											if (pChannelIndex->ParentControlIndex == INDEX_NONE)
 											{
-												SequencerPtr->SelectByNthCategoryNode(CRSection, CategoryIndex, bIsSelectedInList);
+												int32 CategoryIndex = CRSection->GetActiveCategoryIndex(ControlName.Key);
+												if (CategoryIndex != INDEX_NONE)
+												{
+													SequencerPtr->SelectByNthCategoryNode(CRSection, CategoryIndex, bIsSelectedInList);
+												}
 											}
 										}
 									}
 								}
 							}
-						}
-						if(bIsSelectedInList == false || ControlRig == nullptr) //if not a control rig we select the whole trackrow, always make sure to deselect it
-						{
-							if (UMovieSceneTrack* Track = SectionItem.Section->GetTypedOuter<UMovieSceneTrack>())
+							if (bIsSelectedInList == false || ControlRig == nullptr) //if not a control rig we select the whole trackrow, always make sure to deselect it
 							{
-								const int32 RowIndex = SectionItem.Section->GetRowIndex();
-								TSharedPtr<FTrackRowModel> TrackRowModel = TrackRowModelStorage->FindModelForTrackRow(Track, RowIndex);
-								if (bIsSelectedInList)
-								{
-									SelectedOutlinerItems.Select(TrackRowModel);
-								}
-								else
+								if (UMovieSceneTrack* Track = SectionItem.Section->GetTypedOuter<UMovieSceneTrack>())
 								{
-									SelectedOutlinerItems.Deselect(TrackRowModel);
+									const int32 RowIndex = SectionItem.Section->GetRowIndex();
+									if (TSharedPtr<FTrackRowModel> TrackRowModel = TrackRowModelStorage->FindModelForTrackRow(Track, RowIndex))
+									{
+										if (bIsSelectedInList)
+										{
+											if (ControlRig == nullptr && RowIndex != 0) //first deselect the 0th track row model which is the parent, yes this will deselect the actor but that is okay for now
+											{
+												if (TSharedPtr<FTrackRowModel> ParentTrackRwtModel = TrackRowModelStorage->FindModelForTrackRow(Track, 0))
+												{
+													SelectedOutlinerItems.Deselect(ParentTrackRwtModel);
+												}
+											}
+											SelectedOutlinerItems.Select(TrackRowModel);
+
+										}
+										else
+										{
+											SelectedOutlinerItems.Deselect(TrackRowModel);
+										}
+									}
 								}
 							}
 						}
 					}
 				}
 			}
-		}	
+		}
 	}
 }
 
