Anim In Engine: Anim details, do not show object if keys are selected only

diff --git a/Engine/Plugins/Animation/ControlRig/Source/ControlRigEditor/Private/AnimDetails/AnimDetailsProxyManager.cpp b/Engine/Plugins/Animation/ControlRig/Source/ControlRigEditor/Private/AnimDetails/AnimDetailsProxyManager.cpp
--- a/Engine/Plugins/Animation/ControlRig/Source/ControlRigEditor/Private/AnimDetails/AnimDetailsProxyManager.cpp
+++ b/Engine/Plugins/Animation/ControlRig/Source/ControlRigEditor/Private/AnimDetails/AnimDetailsProxyManager.cpp
@@ -22,6 +22,10 @@
 #include "LevelEditorViewport.h"
 #include "MovieScene.h"
 #include "MovieSceneSequence.h"
+#include "MVVM/Extensions/IObjectBindingExtension.h"
+#include "MVVM/Selection/Selection.h"
+#include "MVVM/ViewModels/SequencerEditorViewModel.h"
+#include "MVVM/Selection/SequencerOutlinerSelection.h"
 #include "Rigs/RigHierarchyDefines.h"
 #include "Tracks/MovieScene3DTransformTrack.h"
 #include "Tracks/MovieSceneBoolTrack.h"
@@ -258,8 +262,26 @@ void UAnimDetailsProxyManager::ForceUpdateProxies()
 
 	// Update sequencer proxies
 	{
+		//just care abour outliner selection
 		TArray<FGuid> ObjectGuids;
-		Sequencer->GetSelectedObjects(ObjectGuids);
+		TSet<FGuid> ObjectGuidsSet;
+		using namespace UE::Sequencer;
+		const TSharedPtr<FSequencerEditorViewModel> SequencerViewModel = Sequencer->GetViewModel();
+		if (TSharedPtr<FSequencerSelection> SequencerSelection = SequencerViewModel ? SequencerViewModel->GetSelection() : nullptr)
+		{
+			if (SequencerSelection->Outliner.Num() > 0)
+			{
+				for (FViewModelPtr Model : SequencerSelection->Outliner)
+				{
+					TSharedPtr<IObjectBindingExtension> ObjectBinding = Model->FindAncestorOfType<IObjectBindingExtension>(true);
+					if (ObjectBinding)
+					{
+						ObjectGuidsSet.Add(ObjectBinding->GetObjectGuid());
+					}
+				}
+			}
+		}
+		ObjectGuids = ObjectGuidsSet.Array();
 
 		UMovieSceneSequence* MovieSceneSequence = WeakSequencer.IsValid() ? WeakSequencer.Pin()->GetRootMovieSceneSequence() : nullptr;
 		UMovieScene* MovieScene = MovieSceneSequence ? MovieSceneSequence->GetMovieScene() : nullptr;
