Control Rig: Avoid crash when compiling rig while in PIE
#rnx

diff --git a/Engine/Plugins/Animation/ControlRig/Source/ControlRig/Public/Rigs/RigHierarchy.h b/Engine/Plugins/Animation/ControlRig/Source/ControlRig/Public/Rigs/RigHierarchy.h
--- a/Engine/Plugins/Animation/ControlRig/Source/ControlRig/Public/Rigs/RigHierarchy.h
+++ b/Engine/Plugins/Animation/ControlRig/Source/ControlRig/Public/Rigs/RigHierarchy.h
@@ -5438,6 +5438,7 @@ private:
 	friend class FControlRigHierarchyShrinkElementStorage;
 	friend class FControlRigHierarchyRelinkElementStorage;
 	friend class FRigHierarchyPoseAdapter;
+	friend class UControlRigGraph;
 };
 
 struct FRigHierarchyInteractionBracket
diff --git a/Engine/Plugins/Animation/ControlRig/Source/ControlRigDeveloper/Public/Graph/ControlRigGraph.h b/Engine/Plugins/Animation/ControlRig/Source/ControlRigDeveloper/Public/Graph/ControlRigGraph.h
--- a/Engine/Plugins/Animation/ControlRig/Source/ControlRigDeveloper/Public/Graph/ControlRigGraph.h
+++ b/Engine/Plugins/Animation/ControlRig/Source/ControlRigDeveloper/Public/Graph/ControlRigGraph.h
@@ -88,6 +88,7 @@ private:
 	void CacheNameListForHierarchy(UControlRig* InControlRig, URigHierarchy* InHierarchy, TArray<TSharedPtr<FRigVMStringWithTag>>& OutNameList, bool bFilter = true)
 	{
         TArray<FRigVMStringWithTag> Names;
+		UE::TScopeLock Lock(InHierarchy->ElementsLock);
 		for (auto Element : *InHierarchy)
 		{
 			if(Element->IsA<T>())
