Anim Details: no need to update the constraints for each proxy as it will be forced only once right after creation

#rnx

diff --git a/Engine/Plugins/Animation/ControlRig/Source/ControlRigEditor/Private/AnimDetails/AnimDetailsProxyManager.cpp b/Engine/Plugins/Animation/ControlRig/Source/ControlRigEditor/Private/AnimDetails/AnimDetailsProxyManager.cpp
--- a/Engine/Plugins/Animation/ControlRig/Source/ControlRigEditor/Private/AnimDetails/AnimDetailsProxyManager.cpp
+++ b/Engine/Plugins/Animation/ControlRig/Source/ControlRigEditor/Private/AnimDetails/AnimDetailsProxyManager.cpp
@@ -481,14 +481,7 @@ UAnimDetailsProxyBase* UAnimDetailsProxyManager::GetOrCreateSequencerProxy(UObje
 			ControlType == ERigControlType::Transform ||
 			ControlType == ERigControlType::TransformNoScale ||
 			ControlType == ERigControlType::EulerTransform);
-
-		// Reevaluate constraints
 		Proxy->Modify();
-		UWorld* World = GCurrentLevelEditingViewportClient ? GCurrentLevelEditingViewportClient->GetWorld() : nullptr;
-		const FConstraintsManagerController& Controller = FConstraintsManagerController::Get(World);
-		Controller.EvaluateAllConstraints();
-		
-		Proxy->UpdateProxyValues();
 
 		Proxies.Add(Proxy);
 	}
