Animation Layers: Adjustment Blending. Modify local curves or world space transforms based upon percentage changes over the additive layer.

diff --git a/Engine/Plugins/Animation/ControlRig/Source/ControlRigEditor/Private/Sequencer/AnimLayers/AnimLayers.cpp b/Engine/Plugins/Animation/ControlRig/Source/ControlRigEditor/Private/Sequencer/AnimLayers/AnimLayers.cpp
--- a/Engine/Plugins/Animation/ControlRig/Source/ControlRigEditor/Private/Sequencer/AnimLayers/AnimLayers.cpp
+++ b/Engine/Plugins/Animation/ControlRig/Source/ControlRigEditor/Private/Sequencer/AnimLayers/AnimLayers.cpp
@@ -41,6 +41,10 @@
 #include "BakingAnimationKeySettings.h"
 #include "Algo/Accumulate.h"
 #include "MovieSceneToolHelpers.h"
+#include "Tools/EvaluateSequencerTools.h"
+#include "LevelEditorViewport.h"
+#include "EulerTransform.h"
+
 #define LOCTEXT_NAMESPACE "AnimLayers"
 
 
@@ -1848,6 +1852,17 @@ static void AddNamesToMask(FAnimLayerSectionItem* Owner, UMovieSceneControlRigPa
 	}
 }
 
+void FAnimLayerItem::SetSectionsActive(bool bIsActive)
+{
+	for (FAnimLayerSectionItem& SectionItem : SectionItems)
+	{
+		if (SectionItem.Section.IsValid())
+		{
+			SectionItem.Section->SetIsActive(bIsActive);
+		}
+	}
+}
+
 FAnimLayerSectionItem* FAnimLayerItem::FindMatchingSectionItem(UMovieSceneSection* InMovieSceneSection)
 {
 	if (InMovieSceneSection)
@@ -3224,6 +3239,11 @@ void UAnimLayers::SetUpControlRigSection(UMovieSceneControlRigParameterSection*
 	}
 }
 
+namespace UE
+{
+namespace AIE
+{
+
 struct FKeyInterval
 {
 	FFrameNumber StartFrame;
@@ -3252,6 +3272,19 @@ void GetPairs(ChannelType* Channel, TArray<FKeyInterval>& OutKeyIntervals)
 	}
 }
 
+template<typename ChannelType>
+void GetFrames(ChannelType* Channel, TSet<FFrameNumber>& OutKeyTimes)
+{
+	using ChannelValueType = typename ChannelType::ChannelValueType;
+	TMovieSceneChannelData<ChannelValueType> ChannelData = Channel->GetData();
+
+	TArrayView<FFrameNumber> Times = ChannelData.GetTimes();
+	for (int32 Index = 0; Index < Times.Num(); ++Index)
+	{
+		OutKeyTimes.Add(Times[Index]);
+	}
+}
+
 template<typename ChannelType>
 void  EvaluateCurveOverRange(ChannelType* Channel, const FFrameNumber& StartTime, const FFrameNumber& EndTime,
 	const FFrameNumber& Interval, TArray<TPair<FFrameNumber, double>>& OutKeys)
@@ -3340,8 +3373,424 @@ void AdjustmentBlend(UMovieSceneSection* Section, TArrayView<ChannelType*>& Base
 	}
 }
 
+UMovieSceneSection* GetSectionFromAnimLayerItem(const FAnimLayerItem& AnimLayerItem, UControlRig* OptionalControlRig)
+{
+	for (const FAnimLayerSectionItem& SectionItem : AnimLayerItem.SectionItems)
+	{
+		if (SectionItem.Section.IsValid())
+		{
+			if (OptionalControlRig)
+			{
+				if (UMovieSceneControlRigParameterSection* CRSection = Cast<UMovieSceneControlRigParameterSection>(SectionItem.Section.Get()))
+				{
+					if (OptionalControlRig == CRSection->GetControlRig())
+					{
+						return CRSection;
+					}
+				}
+			}
+			else
+			{
+				return SectionItem.Section.Get();
+			}
+		}
+	}
+	return nullptr;
+}
+
+void GetEvaluatingItemsFromLayer(const FAnimLayerItem& AnimLayerItem, int32 NumOfTransforms, TArray<UE::AIE::FActorAndWorldTransforms>& OutEvaluatingActors, TMap<UControlRig*, UE::AIE::FControlRigAndWorldTransforms>& OutEvaluatingControlRigs)
+{
+	for (const FAnimLayerSectionItem& SectionItem : AnimLayerItem.SectionItems)
+	{
+		if (SectionItem.Section.IsValid())
+		{
+			if (UMovieSceneControlRigParameterSection* CRSection = Cast<UMovieSceneControlRigParameterSection>(SectionItem.Section.Get()))
+			{
+				if (UControlRig* ControlRig = CRSection->GetControlRig())
+				{
+					FControlRigAndWorldTransforms ControlRigAndWorldTransforms;
+					ControlRigAndWorldTransforms.ControlRig = ControlRig;
+					ControlRigAndWorldTransforms.ParentTransforms =  MakeShared<FArrayOfTransforms>();
+					ControlRigAndWorldTransforms.ParentTransforms->SetNum(NumOfTransforms);
+
+					for (const TPair<FName, FAnimLayerPropertyAndChannels>& ControlName : SectionItem.AnimLayerSet.Names)
+					{
+						TSharedPtr<FArrayOfTransforms> WorldTransforms = MakeShared<FArrayOfTransforms>();
+						WorldTransforms->SetNum(NumOfTransforms);
+						ControlRigAndWorldTransforms.ControlAndWorldTransforms.Add(ControlName.Key, WorldTransforms);
+
+					}
+					OutEvaluatingControlRigs.Add(ControlRig, ControlRigAndWorldTransforms);		
+				}
+			}
+			else if (UMovieScene3DTransformSection* LayerSection = Cast<UMovieScene3DTransformSection>(SectionItem.Section.Get()))
+			{
+				//tbd if we do actors or not yet currently will use curve based adjustment blending that only works for world space curves
+			}
+		}
+	}
+}
+
+
+void CalculateTransformDiff(TSharedPtr<UE::AIE::FArrayOfTransforms>& Base, TSharedPtr<UE::AIE::FArrayOfTransforms> Layer, TSharedPtr<UE::AIE::FArrayOfTransforms>& OutDiff)
+{
+	if (Base->Num() > 0 && (Base->Num() == Layer->Num()))
+	{
+		OutDiff->SetNum(Base->Num());
+		for (int32 Index = 0; Index < Base->Num(); ++Index)
+		{
+
+			OutDiff->Transforms[Index] = Layer->Transforms[Index].GetRelativeTransform(Base->Transforms[Index]);
+		}
+
+	}
+}
+
+struct FTransformAsArray
+{
+	double Values[9];
+	FTransformAsArray()
+	{
+		for (int32 Index = 0; Index < 3; ++Index)
+		{
+			const int32 RealIndex = Index * 3;
+			const int32 RealIndex1 = RealIndex + 1;
+			const int32 RealIndex2 = RealIndex + 2;
+			Values[RealIndex] = 0;
+			Values[RealIndex1] = 0;
+			Values[RealIndex2] = 0;
+		}
+	}
+	FTransformAsArray(const FTransformAsArray&) = default;
+	FTransformAsArray(const FTransform& Transform)
+	{
+		const FVector& Location = Transform.GetLocation();
+		const FVector& EulerAngles = Transform.GetRotation().Euler();
+		const FVector& Scale = Transform.GetScale3D();
+		Values[0] = Location.X;
+		Values[1] = Location.Y;
+		Values[2] = Location.Z;
+		Values[3] = EulerAngles.X;
+		Values[4] = EulerAngles.Y;
+		Values[5] = EulerAngles.Z;
+		Values[6] = Scale.X;
+		Values[7] = Scale.Y;
+		Values[8] = Scale.Z;
+
+	}
+
+	FEulerTransform AsEulerTransform() const
+	{
+		FVector Location(Values[0], Values[1], Values[2]);
+		FRotator Rotator(Values[4], Values[5], Values[3]);
+		FVector Scale(Values[6], Values[7], Values[8]);
+
+		return FEulerTransform(Location, Rotator, Scale);
+	}
+
+	FTransform AsTransform() const
+	{
+		FVector Location(Values[0], Values[1], Values[2]);
+		FRotator Rotator(Values[4], Values[5], Values[3]);
+		FVector Scale(Values[6], Values[7], Values[8]);
+
+		return FTransform(Rotator, Location, Scale);
+	}
+
+	double& operator[](int32 Index)
+	{
+		return Values[Index];
+	}
+
+	const double& operator[](int32 Index)const
+	{
+		return Values[Index];
+	}
+
+	
+	FTransformAsArray operator+=(const FTransformAsArray& RHS)
+	{
+		for (int32 Index = 0; Index < 3; ++Index)
+		{
+			const int32 RealIndex = Index * 3;
+			const int32 RealIndex1 = RealIndex + 1;
+			const int32 RealIndex2 = RealIndex + 2;
+
+			Values[RealIndex]  +=  RHS[RealIndex];
+			Values[RealIndex1] += RHS[RealIndex1];
+			Values[RealIndex2] += RHS[RealIndex2];
+		}
+		return *this;
+	}
+};
+
+FTransformAsArray operator-(const FTransformAsArray& LHS, const FTransformAsArray& RHS) {
+	FTransformAsArray ReturnVal;
+	for (int32 Index = 0; Index < 3; ++Index)
+	{
+		const int32 RealIndex = Index * 3;
+		const int32 RealIndex1 = RealIndex + 1;
+		const int32 RealIndex2 = RealIndex + 2;
+
+		ReturnVal[RealIndex] = LHS[RealIndex] - RHS[RealIndex];
+		ReturnVal[RealIndex1] = LHS[RealIndex1] - RHS[RealIndex1];
+		ReturnVal[RealIndex2] = LHS[RealIndex2] - RHS[RealIndex2];
+	}
+	return ReturnVal;
+}
+
+void GetWorldTransformPercentageOfChange(int32 FrameStartIndex, int32 FrameEndIndex, const TPair<FName, TSharedPtr<UE::AIE::FArrayOfTransforms>>& Pair, TArray<FTransformAsArray>& ValueDifferences, TArray<TArray<double>>& PercentageDifferences)
+{
+	const int32 NumKeys = FrameEndIndex - FrameStartIndex + 1;
+	FTransformAsArray CurrentArray;
+	FTransformAsArray NextArray(FTransformAsArray((*Pair.Value)[FrameStartIndex]));
+	FTransformAsArray TotalDifferences,TotalChange;
+	for (int32 Index = FrameStartIndex; Index < FrameEndIndex; ++Index)
+	{
+		CurrentArray = NextArray;
+		FTransform& NextTransform = (*Pair.Value)[Index + 1];
+		NextArray = FTransformAsArray(NextTransform);
+		ValueDifferences[Index] = NextArray - CurrentArray;
+		TotalChange += ValueDifferences[Index];
+	}
+	PercentageDifferences.SetNum(9); //one per channel
+	for (int32 ChannelIndex = 0; ChannelIndex < 9; ++ChannelIndex)
+	{
+		TArray<double>& PercentageDiffs = PercentageDifferences[ChannelIndex];
+		if (FMath::IsNearlyZero(TotalChange[ChannelIndex]) == false)
+		{
+			const double TotalChangePercentage = (100.0 / TotalChange[ChannelIndex]);
+			PercentageDiffs.SetNum(ValueDifferences.Num());
+			for (int32 Index = FrameStartIndex; Index <= FrameEndIndex; ++Index)
+			{
+				PercentageDiffs[Index] = TotalChangePercentage * ValueDifferences[Index][ChannelIndex];
+			}
+		}
+		else
+		{
+			PercentageDiffs.SetNum(0);
+		}
+	}
+
+}
+//get evaluating rigs and actors for the layer we will adjust
+//make a copy for the base
+//mute adjustment , calculate base curves
+//unmet adjustment, calculate modified curves
+
+bool ControlTypeValidForAdjustment(const FRigControlElement* ControlElement)
+{
+	switch (ControlElement->Settings.ControlType)
+	{
+		case ERigControlType::Position:
+		case ERigControlType::Rotator:
+		case ERigControlType::Transform:
+		case ERigControlType::TransformNoScale:
+		case ERigControlType::EulerTransform:
+			return true;
+	}
+	return false;
+}
+
+void AdjustmentBlendWorldSpace(ISequencer* Sequencer,  FAnimLayerItem& BaseAnimLayerItem,  FAnimLayerItem& AnimLayerItem)
+{
+	if (!Sequencer)
+	{
+		return;
+	}
+
+	TRange<FFrameNumber> Range = Sequencer->GetFocusedMovieSceneSequence()->GetMovieScene()->GetPlaybackRange();
+	TOptional<TRange<FFrameNumber>> OptionalRange = Sequencer->GetSubSequenceRange();
+	if (OptionalRange.IsSet())
+	{
+		Range = TRange<FFrameNumber>(OptionalRange.GetValue().GetLowerBoundValue(), OptionalRange.GetValue().GetUpperBoundValue());
+	}
+
+	FFrameNumber FramesPerTick = FFrameRate::TransformTime(FFrameNumber(1), Sequencer->GetFocusedDisplayRate(),
+		Sequencer->GetFocusedTickResolution()).RoundToFrame();
+
+	FFrameTimeByIndex CurrentFrameTimes;
+	CurrentFrameTimes = FFrameTimeByIndex(Range.GetLowerBoundValue(), Range.GetUpperBoundValue(), FramesPerTick);
+
+	UWorld* World = GCurrentLevelEditingViewportClient ? GCurrentLevelEditingViewportClient->GetWorld() : nullptr;
+
+	TArray<FActorAndWorldTransforms> EvaluatingActors, BaseEvaluatingActors;
+	TMap<UControlRig*, FControlRigAndWorldTransforms> EvaluatingControlRigs, BaseEvaluatingControlRigs;
+
+	GetEvaluatingItemsFromLayer(AnimLayerItem, CurrentFrameTimes.NumFrames, EvaluatingActors, EvaluatingControlRigs);
+
+	//make copy and get times
+	TSet<FFrameNumber> Frames;
+	for (const FActorAndWorldTransforms& ActorAndWorld : EvaluatingActors)
+	{
+		FActorAndWorldTransforms Copy;
+		ActorAndWorld.MakeCopy(Copy);
+		BaseEvaluatingActors.Add(ActorAndWorld);
+		if (UMovieSceneSection* Section = GetSectionFromAnimLayerItem(AnimLayerItem, nullptr))
+		{
+			TArrayView<FMovieSceneDoubleChannel*> LayerDoubleChannels = Section->GetChannelProxy().GetChannels<FMovieSceneDoubleChannel>();
+			for (FMovieSceneDoubleChannel* Channel : LayerDoubleChannels)
+			{
+				GetFrames(Channel, Frames);
+			}
+		}
+	}
+	for (const TPair< UControlRig*, FControlRigAndWorldTransforms>& ControlRigAndWorld : EvaluatingControlRigs)
+	{
+		FControlRigAndWorldTransforms Copy;
+		ControlRigAndWorld.Value.MakeCopy(Copy);
+		BaseEvaluatingControlRigs.Add(ControlRigAndWorld.Key, Copy);
+
+		if (UMovieSceneSection* Section = GetSectionFromAnimLayerItem(AnimLayerItem, ControlRigAndWorld.Key))
+		{
+			TArrayView<FMovieSceneFloatChannel*> LayerFloatChannels = Section->GetChannelProxy().GetChannels<FMovieSceneFloatChannel>();
+			for (FMovieSceneFloatChannel* Channel : LayerFloatChannels)
+			{
+				GetFrames(Channel, Frames);
+			}
+		}
+	}
+	//get base world space for every frame
+	AnimLayerItem.SetSectionsActive(false);
+	FEvalHelpers::CalculateWorldTransforms(World, Sequencer, CurrentFrameTimes,
+		CurrentFrameTimes.GetFullIndexArray(), BaseEvaluatingActors, BaseEvaluatingControlRigs);
+	//get the values from the additive, just from it's frames...
+	TArray<FFrameNumber> FrameArray = Frames.Array();
+	FrameArray.Sort([](const FFrameNumber& A, const FFrameNumber& B) {
+		return A < B;
+		});
+	TArray<int32> AdditiveKeyIndices;
+	for (int32 Index = 0; Index < FrameArray.Num(); ++Index)
+	{
+		int32 NewIndex = CurrentFrameTimes.CalculateIndex(FrameArray[Index]);
+		if (NewIndex != INDEX_NONE)
+		{
+			AdditiveKeyIndices.Add(NewIndex);
+		}
+	}
+	AnimLayerItem.SetSectionsActive(true);
+	FEvalHelpers::CalculateWorldTransforms(World, Sequencer, CurrentFrameTimes,
+		AdditiveKeyIndices,
+		EvaluatingActors, EvaluatingControlRigs);
+	
+	
+	FFrameRate TickResolution = Sequencer->GetFocusedTickResolution();
+	FMovieSceneSequenceTransform RootToLocalTransform = Sequencer->GetFocusedMovieSceneSequenceTransform();
+	constexpr ERigTransformType::Type TransformType = ERigTransformType::CurrentGlobal;
+	FMovieSceneInverseSequenceTransform LocalToRootTransform = RootToLocalTransform.Inverse();
+
+	FRigControlModifiedContext Context;
+	Context.SetKey = EControlRigSetKey::Always;
+	Context.KeyMask = (uint32)EControlRigContextChannelToKey::AllTransform;
+
+	TArray <TPair<FFrameNumber, FTransformAsArray>> Keys;
+	TArray<FTransformAsArray> ValueDifferences;
+	ValueDifferences.SetNum(CurrentFrameTimes.NumFrames);
+	TArray<TArray<double>> PercentageDifferences;
+	FTransformAsArray StartValue, EndValue;
+
+	for (const TPair< UControlRig*, FControlRigAndWorldTransforms>& BaseControlRigAndWorld : BaseEvaluatingControlRigs)
+	{
+		UControlRig* ControlRig = BaseControlRigAndWorld.Key;
+		if (FControlRigAndWorldTransforms* EvalCRandWT = EvaluatingControlRigs.Find(ControlRig))
+		{
+			for (int32 KeyIndex = 0; KeyIndex < AdditiveKeyIndices.Num() - 1; ++KeyIndex)
+			{
+				int32 FrameStartIndex = AdditiveKeyIndices[KeyIndex];
+				int32 FrameEndIndex = AdditiveKeyIndices[KeyIndex + 1];
+				for (const TPair<FName, TSharedPtr<FArrayOfTransforms>>& Pair : BaseControlRigAndWorld.Value.ControlAndWorldTransforms)
+				{
+					if (FRigControlElement* ControlElement = ControlRig->FindControl(Pair.Key))
+					{
+						if (ControlTypeValidForAdjustment(ControlElement) == false)
+						{
+							continue;
+						}
+						if (TSharedPtr<FArrayOfTransforms>* EvalTransforms = EvalCRandWT->ControlAndWorldTransforms.Find(Pair.Key))
+						{
+							//get base layer values as diffs and percentages
+							GetWorldTransformPercentageOfChange(FrameStartIndex, FrameEndIndex, Pair, ValueDifferences, PercentageDifferences);
+							//get new layer change
+							FTransform ST = (*EvalTransforms)->Transforms[FrameStartIndex];
+							ST = ST.GetRelativeTransform((*EvalCRandWT->ParentTransforms)[FrameStartIndex]);
+							FTransformAsArray StartTransform(ST);
+							
+							FTransform ET = (*EvalTransforms)->Transforms[FrameEndIndex];
+							ET = ET.GetRelativeTransform((*EvalCRandWT->ParentTransforms)[FrameEndIndex]);
+							FTransformAsArray EndTransform(ET);
+							
+							FTransformAsArray CurrentValue;
+							FTransformAsArray PreviousValue;
+							FTransformAsArray TotalPoseLayerChange;
+							for (int32 ChannelIndex = 0; ChannelIndex < 9; ++ChannelIndex)
+							{
+								TotalPoseLayerChange[ChannelIndex] = FMath::Abs(EndTransform[ChannelIndex] - StartTransform[ChannelIndex]);
+								PreviousValue[ChannelIndex] = StartTransform[ChannelIndex];
+							}
+							for (int32 FrameIndex = FrameStartIndex; FrameIndex <= FrameEndIndex; ++FrameIndex)
+							{
+								const FFrameNumber FrameNumber = CurrentFrameTimes.CalculateFrame(FrameIndex);
+								FFrameTime GlobalTime(FrameNumber);
+								GlobalTime = LocalToRootTransform.TryTransformTime(GlobalTime).Get(GlobalTime); //player evals in root time so need to go back to it.
+
+								FMovieSceneContext MovieSceneContext = FMovieSceneContext(FMovieSceneEvaluationRange(GlobalTime, TickResolution), Sequencer->GetPlaybackStatus()).SetHasJumped(true);
+								Sequencer->GetEvaluationTemplate().EvaluateSynchronousBlocking(MovieSceneContext);
+								ControlRig->Evaluate_AnyThread();
+								for (int32 ChannelIndex = 0; ChannelIndex < 9; ++ChannelIndex)
+								{
+									if(FrameIndex < PercentageDifferences[ChannelIndex].Num())
+									{
+										double Value = PercentageDifferences[ChannelIndex][FrameIndex];
+										const double ValueDelta = (TotalPoseLayerChange[ChannelIndex] / 100.0) * Value;
+										CurrentValue[ChannelIndex] = (EndTransform[ChannelIndex] > StartTransform[ChannelIndex]) ?
+											PreviousValue[ChannelIndex] + ValueDelta : PreviousValue[ChannelIndex] - ValueDelta;
+										PreviousValue[ChannelIndex] = CurrentValue[ChannelIndex];
+									}
+									else
+									{
+										CurrentValue[ChannelIndex] = PreviousValue[ChannelIndex];
+									}
+								}
+								FTransform NewTransform = CurrentValue.AsTransform();
+								Context.LocalTime = TickResolution.AsSeconds(FFrameTime(FrameNumber));
+								FRigControlValue Value = ControlRig->GetControlValueFromGlobalTransform(ControlElement->GetKey().Name, NewTransform, TransformType);
+								NewTransform = Value.GetAsTransform(ControlElement->Settings.ControlType, ControlElement->Settings.PrimaryAxis);
+								FEulerTransform EulerTransform(NewTransform);;
+								UE::AIE::FSetTransformHelpers::SetControlTransform(EulerTransform, ControlRig, ControlElement, Context);
+
+							}
+						}
+					}
+				}
+			}
+		}
+	}
+	
+	//make sure to reset things
+	Sequencer->ForceEvaluate();
+	//we now need to make sure the control rigs are back up to date.
+	//we do this by going through the control rig tracks that the actors/rigs are dependent up
+
+	for (const TPair< UControlRig*, FControlRigAndWorldTransforms>& ControlRigAndWorldAgain : EvaluatingControlRigs)
+	{
+		if (UControlRig * ControlRig = ControlRigAndWorldAgain.Key)
+		{
+			ControlRig->Evaluate_AnyThread();
+			if (ControlRig->GetObjectBinding())
+			{
+				ControlRig->EvaluateSkeletalMeshComponent(0.0);
+			}
+		}
+	}
+}
+
+} //AIE
+} //UE
+
+
 bool UAnimLayers::AdjustmentBlendLayers(ISequencer* InSequencer, int32 LayerIndex)
 {
+	using namespace UE::AIE;
 	if (InSequencer == nullptr)
 	{
 		return false;
@@ -3352,11 +3801,11 @@ bool UAnimLayers::AdjustmentBlendLayers(ISequencer* InSequencer, int32 LayerInde
 	}
 	UAnimLayer* BaseLayer = AnimLayers[0];
 	UAnimLayer* AnimLayer = AnimLayers[LayerIndex];
-	const FFrameNumber Min =  TNumericLimits<FFrameNumber>::Lowest();
-	const FFrameNumber Max =  TNumericLimits<FFrameNumber>::Max();
+	const FFrameNumber Min = TNumericLimits<FFrameNumber>::Lowest();
+	const FFrameNumber Max = TNumericLimits<FFrameNumber>::Max();
 	TRange<FFrameNumber> Range(Min, Max);
 	FScopedTransaction Transaction(LOCTEXT("AdjustmentBlendLayer", "Adjustment Blend layer"), !GIsTransacting);
-	
+
 	Modify();
 	AnimLayer->Modify();
 
@@ -3365,35 +3814,47 @@ bool UAnimLayers::AdjustmentBlendLayers(ISequencer* InSequencer, int32 LayerInde
 	{
 		if (Pair.Key != nullptr)
 		{
-			for (FAnimLayerSectionItem& SectionItem : Pair.Value.SectionItems)
+			if (Pair.Key->IsA<UControlRig>())
 			{
-				if (SectionItem.Section.IsValid())
+				if (FAnimLayerItem* BaseLayerItem = BaseLayer->AnimLayerItems.Find(Pair.Key))
 				{
-					FAnimLayerSectionItem* BaseSectionItem = nullptr;
-					if (FAnimLayerItem* Owner = BaseLayer->AnimLayerItems.Find(Pair.Key))
-					{
-						BaseSectionItem = Owner->FindMatchingSectionItem(SectionItem.Section.Get());
+					AdjustmentBlendWorldSpace(InSequencer, *BaseLayerItem, Pair.Value);
 
-					}
-					if (BaseSectionItem && BaseSectionItem->Section.IsValid())
+				}
+			}
+			else
+			{
+
+				for (FAnimLayerSectionItem& SectionItem : Pair.Value.SectionItems)
+				{
+					if (SectionItem.Section.IsValid())
 					{
-						if (SectionItem.Section->IsActive())//active sections merge them
+						FAnimLayerSectionItem* BaseSectionItem = nullptr;
+						if (FAnimLayerItem* Owner = BaseLayer->AnimLayerItems.Find(Pair.Key))
 						{
-							TArrayView<FMovieSceneFloatChannel*> BaseFloatChannels = BaseSectionItem->Section->GetChannelProxy().GetChannels<FMovieSceneFloatChannel>();
-							TArrayView<FMovieSceneDoubleChannel*> BaseDoubleChannels = BaseSectionItem->Section->GetChannelProxy().GetChannels<FMovieSceneDoubleChannel>();
-							TArrayView<FMovieSceneFloatChannel*> LayerFloatChannels = SectionItem.Section->GetChannelProxy().GetChannels<FMovieSceneFloatChannel>();
-							TArrayView<FMovieSceneDoubleChannel*> LayerDoubleChannels = SectionItem.Section->GetChannelProxy().GetChannels<FMovieSceneDoubleChannel>();
-							SectionItem.Section->Modify();
-							
-							if (BaseDoubleChannels.Num() > 0)
-							{
-								AdjustmentBlend<FMovieSceneDoubleChannel>(SectionItem.Section.Get(), BaseDoubleChannels,LayerDoubleChannels,  InSequencer);
-							}
-							else if (BaseFloatChannels.Num() > 0)
+							BaseSectionItem = Owner->FindMatchingSectionItem(SectionItem.Section.Get());
+
+						}
+						if (BaseSectionItem && BaseSectionItem->Section.IsValid())
+						{
+							if (SectionItem.Section->IsActive())//active sections merge them
 							{
-								AdjustmentBlend<FMovieSceneFloatChannel>(SectionItem.Section.Get(), BaseFloatChannels,LayerFloatChannels, InSequencer);
+								TArrayView<FMovieSceneFloatChannel*> BaseFloatChannels = BaseSectionItem->Section->GetChannelProxy().GetChannels<FMovieSceneFloatChannel>();
+								TArrayView<FMovieSceneDoubleChannel*> BaseDoubleChannels = BaseSectionItem->Section->GetChannelProxy().GetChannels<FMovieSceneDoubleChannel>();
+								TArrayView<FMovieSceneFloatChannel*> LayerFloatChannels = SectionItem.Section->GetChannelProxy().GetChannels<FMovieSceneFloatChannel>();
+								TArrayView<FMovieSceneDoubleChannel*> LayerDoubleChannels = SectionItem.Section->GetChannelProxy().GetChannels<FMovieSceneDoubleChannel>();
+								SectionItem.Section->Modify();
+
+								if (BaseDoubleChannels.Num() > 0)
+								{
+									AdjustmentBlend<FMovieSceneDoubleChannel>(SectionItem.Section.Get(), BaseDoubleChannels, LayerDoubleChannels, InSequencer);
+								}
+								else if (BaseFloatChannels.Num() > 0)
+								{
+									AdjustmentBlend<FMovieSceneFloatChannel>(SectionItem.Section.Get(), BaseFloatChannels, LayerFloatChannels, InSequencer);
+								}
 							}
-						}	
+						}
 					}
 				}
 			}
@@ -3403,7 +3864,6 @@ bool UAnimLayers::AdjustmentBlendLayers(ISequencer* InSequencer, int32 LayerInde
 	InSequencer->NotifyMovieSceneDataChanged(EMovieSceneDataChangeType::MovieSceneStructureItemsChanged);
 	return true;
 }
-
 #undef LOCTEXT_NAMESPACE
 
 
diff --git a/Engine/Plugins/Animation/ControlRig/Source/ControlRigEditor/Private/Sequencer/AnimLayers/AnimLayers.h b/Engine/Plugins/Animation/ControlRig/Source/ControlRigEditor/Private/Sequencer/AnimLayers/AnimLayers.h
--- a/Engine/Plugins/Animation/ControlRig/Source/ControlRigEditor/Private/Sequencer/AnimLayers/AnimLayers.h
+++ b/Engine/Plugins/Animation/ControlRig/Source/ControlRigEditor/Private/Sequencer/AnimLayers/AnimLayers.h
@@ -127,6 +127,8 @@ struct FAnimLayerItem
 
 	//find section that matches based upon incoming section (same type and property)
 	FAnimLayerSectionItem* FindMatchingSectionItem(UMovieSceneSection* InMovieSceneSection);
+
+	void SetSectionsActive(bool bIsActive);
 };
 
 UENUM()
diff --git a/Engine/Plugins/Animation/ControlRig/Source/ControlRigEditor/Private/Sequencer/AnimLayers/SAnimLayers.cpp b/Engine/Plugins/Animation/ControlRig/Source/ControlRigEditor/Private/Sequencer/AnimLayers/SAnimLayers.cpp
--- a/Engine/Plugins/Animation/ControlRig/Source/ControlRigEditor/Private/Sequencer/AnimLayers/SAnimLayers.cpp
+++ b/Engine/Plugins/Animation/ControlRig/Source/ControlRigEditor/Private/Sequencer/AnimLayers/SAnimLayers.cpp
@@ -1460,15 +1460,15 @@ TSharedPtr<SWidget> FAnimLayerSourcesView::OnSourceConstructContextMenu()
 					const FText ToolTipText = LOCTEXT("SetPassthroughKeytooltip", "Set zero key(Additive) or previous value(Override)");
 					MenuBuilder.AddMenuEntry(Label, ToolTipText, FSlateIcon(), Action);
 				}
-				/* adjustment blending not in yet
+				//adjustment blending 
 				if (Selected.Num() == 1 && Selected[0]->GetAnimLayerIndex(AnimLayers) != 0)
 				{
 					FUIAction Action = FUIAction(FExecuteAction::CreateRaw((this), &FAnimLayerSourcesView::AdjustmentBlend));
-					const FText Label = LOCTEXT("AdjustmentBlend", "Adjustment blend");
+					const FText Label = LOCTEXT("AdjustmentBlend", "Adjustment Blend");
 					const FText ToolTipText = LOCTEXT("AdjustmentBlendtooltip", "AdustmentBlend");
 					MenuBuilder.AddMenuEntry(Label, ToolTipText, FSlateIcon(), Action);
 				}
-				*/
+				
 				MenuBuilder.AddMenuSeparator();
 				if (Selected.Num() == 1)
 				{
