Modular Rig & Constraints: fix constraints using the long name version of  modular rig controls when switching from 5.5 to 5.6
- note: 43508504 was a first attempt to fix this but this is not enough, the controls are spawned after the animation patch is being applied. (Original fix from Benoit)

#rnx

diff --git a/Engine/Plugins/Animation/ControlRig/Source/ControlRig/Private/Constraints/ControlRigTransformableHandle.cpp b/Engine/Plugins/Animation/ControlRig/Source/ControlRig/Private/Constraints/ControlRigTransformableHandle.cpp
--- a/Engine/Plugins/Animation/ControlRig/Source/ControlRig/Private/Constraints/ControlRigTransformableHandle.cpp
+++ b/Engine/Plugins/Animation/ControlRig/Source/ControlRig/Private/Constraints/ControlRigTransformableHandle.cpp
@@ -65,6 +65,21 @@ namespace UE::Private
 		const USceneComponent* Component = nullptr;
 		int32 PoseVersion = INDEX_NONE;
 	};
+
+	void FixModularRigControlName(UControlRig* InControlRig, UTransformableControlHandle* InOutHandle)
+	{
+		if (UModularRig* ModularRig = Cast<UModularRig>(InControlRig))
+		{
+			if (const FRigControlElement* Control = ModularRig->FindControl(InOutHandle->ControlName))
+			{
+				const FName& ControlName = Control->GetFName();
+				if (ControlName != InOutHandle->ControlName)
+				{
+					InOutHandle->ControlName = ControlName;
+				}
+			}
+		}
+	}
 }
 
 /**
@@ -559,6 +574,8 @@ void UTransformableControlHandle::RegisterDelegates()
 	
 	if (ControlRig.IsValid())
 	{
+		UE::Private::FixModularRigControlName(ControlRig.Get(), this);
+		
 		if (URigHierarchy* Hierarchy = ControlRig->GetHierarchy())
 		{
 			Hierarchy->OnModified().AddUObject(this, &UTransformableControlHandle::OnHierarchyModified);
@@ -599,6 +616,11 @@ void UTransformableControlHandle::OnHierarchyModified(
 
 	switch (InNotif)
 	{
+		case ERigHierarchyNotification::ElementAdded:
+		{
+			UE::Private::FixModularRigControlName(ControlRig.Get(), this);
+			break;
+		}
 		case ERigHierarchyNotification::ElementRemoved:
 		{
 			// FIXME this leaves the constraint invalid as the element won't exist anymore
