Anim In Engine: Handle Control Type changes between Transform Control Types and Location/Rotator/Scale Types by adding a new scripting function.

diff --git a/Engine/Plugins/Animation/ControlRig/Source/ControlRig/Private/Sequencer/MovieSceneControlRigParameterSection.cpp b/Engine/Plugins/Animation/ControlRig/Source/ControlRig/Private/Sequencer/MovieSceneControlRigParameterSection.cpp
--- a/Engine/Plugins/Animation/ControlRig/Source/ControlRig/Private/Sequencer/MovieSceneControlRigParameterSection.cpp
+++ b/Engine/Plugins/Animation/ControlRig/Source/ControlRig/Private/Sequencer/MovieSceneControlRigParameterSection.cpp
@@ -1602,7 +1602,6 @@ void UMovieSceneControlRigParameterSection::AddTransformParameter(FName InParame
 			ExistingCurves->Scale[1].SetDefault(0.0f);
 			ExistingCurves->Scale[2].SetDefault(0.0f);
 		}
-
 		MaskOutIfThereAreMaskedControls(InParameterName);
 
 		if(bReconstructChannel)
@@ -1612,6 +1611,62 @@ void UMovieSceneControlRigParameterSection::AddTransformParameter(FName InParame
 	}
 }
 
+bool UMovieSceneControlRigParameterSection::CopyVectorParameterCurvesToTransform(const FName& InParameterName, ERigControlType InType)
+{
+	if (HasVectorParameter(InParameterName))
+	{
+		for (FVectorParameterNameAndCurves& Vector : GetVectorParameterNamesAndCurves())
+		{
+			if (Vector.ParameterName == InParameterName)
+			{
+				FTransformParameterNameAndCurves* ExistingCurves = nullptr;
+				int32 Index = INDEX_NONE;
+				for (int32 IndexCount = 0; IndexCount < TransformParameterNamesAndCurves.Num(); ++ IndexCount)
+				{
+					FTransformParameterNameAndCurves& TransformParameterNamesAndCurve = TransformParameterNamesAndCurves[IndexCount];
+					if (TransformParameterNamesAndCurve.ParameterName == InParameterName)
+					{
+						Index = IndexCount;
+						break;
+					}
+				}
+				if (Index == INDEX_NONE) // no transform so add
+				{
+					Index = TransformParameterNamesAndCurves.Add(FTransformParameterNameAndCurves(InParameterName));
+
+				}
+				if (Index != INDEX_NONE)
+				{
+					ExistingCurves = &TransformParameterNamesAndCurves[Index];
+					if (InType == ERigControlType::Position)
+					{
+						ExistingCurves->Translation[0] = Vector.XCurve;
+						ExistingCurves->Translation[1] = Vector.YCurve;
+						ExistingCurves->Translation[2] = Vector.ZCurve;
+						return true;
+					}
+					else if (InType == ERigControlType::Rotator)
+					{
+						ExistingCurves->Rotation[0] = Vector.XCurve;
+						ExistingCurves->Rotation[1] = Vector.YCurve;
+						ExistingCurves->Rotation[2] = Vector.ZCurve;
+						return true;
+					}
+					else if (InType == ERigControlType::Scale)
+					{
+						ExistingCurves->Scale[0] = Vector.XCurve;
+						ExistingCurves->Scale[1] = Vector.YCurve;
+						ExistingCurves->Scale[2] = Vector.ZCurve;
+						return true;
+					}
+				}
+			}
+		}
+	}
+	return false;
+}
+
+
 // only allow creation of space channels onto non-parented Controls
 bool UMovieSceneControlRigParameterSection::CanCreateSpaceChannel(FName InControlName) const
 {
diff --git a/Engine/Plugins/Animation/ControlRig/Source/ControlRig/Public/Sequencer/MovieSceneControlRigParameterSection.h b/Engine/Plugins/Animation/ControlRig/Source/ControlRig/Public/Sequencer/MovieSceneControlRigParameterSection.h
--- a/Engine/Plugins/Animation/ControlRig/Source/ControlRig/Public/Sequencer/MovieSceneControlRigParameterSection.h
+++ b/Engine/Plugins/Animation/ControlRig/Source/ControlRig/Public/Sequencer/MovieSceneControlRigParameterSection.h
@@ -279,6 +279,8 @@ public:
 
 
 	UE_API bool RenameParameterName(const FName& OldParameterName, const FName& NewParameterName, TOptional<ERigControlType> ControlType);
+	/** Copy existing internal vector paramter of Position/Rotator/Scale type to a transform parameter type*/
+	UE_API bool  CopyVectorParameterCurvesToTransform(const FName& InControlName, ERigControlType InType);
 
 #if WITH_EDITOR
 	UE_API void OnControlRigEditorSettingChanged(UObject* InSettingsChanged, struct FPropertyChangedEvent& InPropertyChangedEvent);
@@ -574,7 +576,6 @@ public:
 	/** Add Space Parameter for a specified Control, no Default since that is Parent space*/
 	UE_API void AddSpaceChannel(FName InControlName, bool bReconstructChannel);
 
-
 	/** Clear Everything Out*/
 	UE_API void ClearAllParameters();
 
diff --git a/Engine/Plugins/Animation/ControlRig/Source/ControlRigEditor/Private/ControlRigSequencerEditorLibrary.cpp b/Engine/Plugins/Animation/ControlRig/Source/ControlRigEditor/Private/ControlRigSequencerEditorLibrary.cpp
--- a/Engine/Plugins/Animation/ControlRig/Source/ControlRigEditor/Private/ControlRigSequencerEditorLibrary.cpp
+++ b/Engine/Plugins/Animation/ControlRig/Source/ControlRigEditor/Private/ControlRigSequencerEditorLibrary.cpp
@@ -2864,6 +2864,36 @@ bool UControlRigSequencerEditorLibrary::ImportFBXToControlRigTrack(UWorld* World
 	return bValid;
 }
 
+bool UControlRigSequencerEditorLibrary::ControlRigCopyVectorParameterCurvesToTransform(UMovieSceneControlRigParameterTrack* InTrack,
+	const FName& InControlName, ERigControlType InType)
+{
+	if (InTrack == nullptr)
+	{
+		UE_LOG(LogControlRig, Error, TEXT("UControlRigSequencerEditorLibrary::ControlRigCopyVectorParameterCurvesToTransform:  Track is null"));
+		return false;
+	}
+	bool bPassed = false;
+	for (UMovieSceneSection* Section : InTrack->GetAllSections())
+	{
+		if (Section)
+		{
+			UMovieSceneControlRigParameterSection* CRSection = Cast<UMovieSceneControlRigParameterSection>(Section);
+			if (CRSection)
+			{
+				if (CRSection->CopyVectorParameterCurvesToTransform(InControlName, InType))
+				{
+					bPassed = true;
+				}
+			}
+		}
+	}
+	if (bPassed == false)
+	{
+		UE_LOG(LogControlRig, Error, TEXT("UControlRigSequencerEditorLibrary::ControlRigCopyVectorParameterCurvesToTransform:  No Sections Converted"));
+	}
+	return bPassed;
+}
+
 bool UControlRigSequencerEditorLibrary::ExportFBXFromControlRigSection(ULevelSequence* Sequence, const UMovieSceneControlRigParameterSection* Section, const UMovieSceneUserExportFBXControlRigSettings* ExportFBXControlRigSettings)
 {
 	if (!Sequence || !ExportFBXControlRigSettings ||
diff --git a/Engine/Plugins/Animation/ControlRig/Source/ControlRigEditor/Public/ControlRigSequencerEditorLibrary.h b/Engine/Plugins/Animation/ControlRig/Source/ControlRigEditor/Public/ControlRigSequencerEditorLibrary.h
--- a/Engine/Plugins/Animation/ControlRig/Source/ControlRigEditor/Public/ControlRigSequencerEditorLibrary.h
+++ b/Engine/Plugins/Animation/ControlRig/Source/ControlRigEditor/Public/ControlRigSequencerEditorLibrary.h
@@ -161,6 +161,17 @@ public:
 		FFrameNumber InStartFrame, EMovieSceneTimeUnit TimeUnit = EMovieSceneTimeUnit::DisplayRate, bool bKeyReduce = false, float Tolerance = 0.001f,
 		EMovieSceneKeyInterpolation Interpolation = EMovieSceneKeyInterpolation::SmartAuto, bool bResetControls = true);
 
+	/** Copy the internal parameter curves from vector control rig curves to transform curves, this is needed if a vector type is changed to transform type
+	* You would run this script first, resave, and then on next reload the curves should now show up on correct curves in the transform control
+	* @param InTrack Control Rig track to convert
+	* @param InControlName The name of the existing Vector type control
+	* @param InType The type of the control to convert, ERigControlType::Position,Rotator, or Scale
+	* @param returns true if at least one section was converted, false if nothing was converted
+	*/
+	UFUNCTION(BlueprintCallable, Category = "Editor Scripting | Sequencer Tools | Control Rig")
+	static UE_API bool ControlRigCopyVectorParameterCurvesToTransform(UMovieSceneControlRigParameterTrack* InTrack,
+		const FName& InControlName, ERigControlType InType);
+
 	/**
 	* Bake the current animation in the binding to a Control Rig track
 	* @param World The active world
