ControlRig: remove deprecated rig options

#rnx

diff --git a/Engine/Plugins/Animation/ControlRig/Source/ControlRigEditor/Private/EditMode/SControlRigEditModeTools.cpp b/Engine/Plugins/Animation/ControlRig/Source/ControlRigEditor/Private/EditMode/SControlRigEditModeTools.cpp
--- a/Engine/Plugins/Animation/ControlRig/Source/ControlRigEditor/Private/EditMode/SControlRigEditModeTools.cpp
+++ b/Engine/Plugins/Animation/ControlRig/Source/ControlRigEditor/Private/EditMode/SControlRigEditModeTools.cpp
@@ -56,11 +56,8 @@ void SControlRigEditModeTools::Cleanup()
 		{
 			SettingsDetailsView->SetKeyframeHandler(nullptr);
 		}
-		if (RigOptionsDetailsView)
-		{
-			RigOptionsDetailsView->SetKeyframeHandler(nullptr);
-		}
 	}
+	
 	for (TPair<FDelegateHandle, TWeakObjectPtr<UControlRig>>& Handles : HandlesToClear)
 	{
 		if (Handles.Value.IsValid())
@@ -72,6 +69,7 @@ void SControlRigEditModeTools::Cleanup()
 			Handles.Key.Reset();
 		}
 	}
+	
 	HandlesToClear.Reset();
 }
 
@@ -101,15 +99,6 @@ void SControlRigEditModeTools::SetControlRigs(const TArrayView<TWeakObjectPtr<UC
 		}
 	}
 
-	//mz todo handle multiple rigs
-	TArray<TWeakObjectPtr<>> Objects;
-	if (ControlRigs.Num() > 0)
-	{
-		UControlRig* Rig = ControlRigs[0].Get();
-		Objects.Add(Rig);
-	}
-	RigOptionsDetailsView->SetObjects(Objects);
-
 	UpdateOverridesDetailsView();
 
 
@@ -210,10 +199,6 @@ void SControlRigEditModeTools::Construct(const FArguments& InArgs, TSharedPtr<FC
 	ControlVector2DDetailsView->SetGenericLayoutDetailsDelegate(FOnGetDetailCustomizationInstance::CreateStatic(&FControlRigEditModeGenericDetails::MakeInstance, ModeTools));
 #endif
 
-	RigOptionsDetailsView = FModuleManager::GetModuleChecked<FPropertyEditorModule>("PropertyEditor").CreateDetailView(DetailsViewArgs);
-	RigOptionsDetailsView->SetKeyframeHandler(SharedThis(this));
-	RigOptionsDetailsView->OnFinishedChangingProperties().AddSP(this, &SControlRigEditModeTools::OnRigOptionFinishedChange);
-
 	OverridesDetailsView = FModuleManager::GetModuleChecked<FPropertyEditorModule>("PropertyEditor").CreateDetailView(DetailsViewArgs);
 	OverridesDetailsView->OnFinishedChangingProperties().AddSP(this, &SControlRigEditModeTools::OnOverrideOptionFinishedChange);
 	OverridesDetailsView->SetIsPropertyVisibleDelegate(FIsPropertyVisible::CreateSP(this, &SControlRigEditModeTools::ShouldShowOverrideProperty));
@@ -500,21 +485,6 @@ void SControlRigEditModeTools::Construct(const FArguments& InArgs, TSharedPtr<FC
 				})
 			]
 			
-			+ SVerticalBox::Slot()
-			.AutoHeight()
-			[
-				SAssignNew(RigOptionExpander, SExpandableArea)
-				.InitiallyCollapsed(false)
-				.Visibility(this, &SControlRigEditModeTools::GetRigOptionExpanderVisibility)
-				.AreaTitle(LOCTEXT("RigOption_Header", "Rig Options"))
-				.AreaTitleFont(FAppStyle::GetFontStyle("DetailsView.CategoryFontStyle"))
-				.BorderBackgroundColor(FLinearColor(.6f, .6f, .6f))
-				.BodyContent()
-				[
-					RigOptionsDetailsView.ToSharedRef()
-				]
-			]
-
 			+SVerticalBox::Slot()
 			.AutoHeight()
 			[
@@ -1254,32 +1224,6 @@ FReply SControlRigEditModeTools::HandleAddConstraintClicked()
 	return FReply::Handled();
 }
 
-EVisibility SControlRigEditModeTools::GetRigOptionExpanderVisibility() const
-{
-	for (const TWeakObjectPtr<UControlRig>& ControlRig : ControlRigs)
-	{
-		if (ControlRig.IsValid())
-		{
-			if (Cast<UFKControlRig>(ControlRig))
-			{
-				return EVisibility::Visible;
-			}
-		}
-	}
-	return EVisibility::Collapsed;
-}
-
-void SControlRigEditModeTools::OnRigOptionFinishedChange(const FPropertyChangedEvent& PropertyChangedEvent)
-{
-	TArray<TWeakObjectPtr<UControlRig>> ControlRigsCopy = ControlRigs;
-	SetControlRigs(ControlRigsCopy);
-
-	if (FControlRigEditMode* EditMode = static_cast<FControlRigEditMode*>(ModeTools->GetActiveMode(FControlRigEditMode::ModeName)))
-	{
-		EditMode->SetObjects_Internal();
-	}
-}
-
 void SControlRigEditModeTools::CustomizeToolBarPalette(FToolBarBuilder& ToolBarBuilder)
 {
 	//TOGGLE SELECTED RIG CONTROLS
diff --git a/Engine/Plugins/Animation/ControlRig/Source/ControlRigEditor/Private/EditMode/SControlRigEditModeTools.h b/Engine/Plugins/Animation/ControlRig/Source/ControlRigEditor/Private/EditMode/SControlRigEditModeTools.h
--- a/Engine/Plugins/Animation/ControlRig/Source/ControlRigEditor/Private/EditMode/SControlRigEditModeTools.h
+++ b/Engine/Plugins/Animation/ControlRig/Source/ControlRigEditor/Private/EditMode/SControlRigEditModeTools.h
@@ -83,12 +83,6 @@ private:
 	TSharedPtr<IDetailsView> ControlEnumDetailsView;
 	TSharedPtr<IDetailsView> ControlVector2DDetailsView;
 	TSharedPtr<IDetailsView> ControlVectorDetailsView;
-#endif
-	/** Expander to interact with the options of the rig  */
-	TSharedPtr<SExpandableArea> RigOptionExpander;
-	TSharedPtr<IDetailsView> RigOptionsDetailsView;
-#if USE_LOCAL_DETAILS
-	/** Hierarchy picker for controls*/
 	TSharedPtr<SRigHierarchyTreeView> HierarchyTreeView;
 #endif
 
@@ -123,10 +117,6 @@ private:
 	void Compensate(TOptional<FFrameNumber> OptionalKeyTime, bool bSetPreviousTick);
 	bool ReadyForBakeOrCompensation() const;
 	FReply HandleAddConstraintClicked();
-
-	EVisibility GetRigOptionExpanderVisibility() const;
-
-	void OnRigOptionFinishedChange(const FPropertyChangedEvent& PropertyChangedEvent);
 	
 	/** constraint type to show selection*/
 	void OnSelectShowConstraints(int32 Index);
