Control Rig: Improvements to tween overlay widget:
- Switched implementation to SDraggableBoxOverlay, which fixed several bugs and simplified the code.
- Add ability to reposition the widget using Alt+U.
- Known bug: When you enter Animation mode, the tween widget is initially hidden, and you chord Alt+U to position the widget, the widget is painted at the position saved in the config for 1 frame and then next frame jumps to the clicked location. Intentionally leaving this unfixed for now (time-tradeoff). The reason is that the tick space geometry is not yet set until the first time the widget is painted.

diff --git a/Engine/Plugins/Animation/ControlRig/Source/ControlRigEditor/Private/EditMode/ControlRigEditModeCommands.cpp b/Engine/Plugins/Animation/ControlRig/Source/ControlRigEditor/Private/EditMode/ControlRigEditModeCommands.cpp
--- a/Engine/Plugins/Animation/ControlRig/Source/ControlRigEditor/Private/EditMode/ControlRigEditModeCommands.cpp
+++ b/Engine/Plugins/Animation/ControlRig/Source/ControlRigEditor/Private/EditMode/ControlRigEditModeCommands.cpp
@@ -44,6 +44,9 @@ void FControlRigEditModeCommands::RegisterCommands()
 
 	UI_COMMAND(ToggleControlShapeTransformEdit, "Toggle Shape Transform Edit", "Toggle Editing Selected Control's Shape Transform", EUserInterfaceActionType::ToggleButton, FInputChord(EKeys::Period, EModifierKey::Control)); 
 	UI_COMMAND(OpenSpacePickerWidget, "Open the Space Picker", "Allows space switching on the control", EUserInterfaceActionType::Button, FInputChord(EKeys::Tab)); 
+
+	UI_COMMAND(SummonTweenWidget, "Move tween slider to cursor", "Places the tween slider at your mouse position", EUserInterfaceActionType::Button, FInputChord(EKeys::U)); 
+	UI_COMMAND(ToggleTweenWidget, "Show or hide tween slider", "Shows the tween slider if it's hidden, and hides it if it's visible.", EUserInterfaceActionType::Button, FInputChord(EKeys::U, EModifierKey::Alt)); 
 }
 
 #undef LOCTEXT_NAMESPACE
diff --git a/Engine/Plugins/Animation/ControlRig/Source/ControlRigEditor/Private/EditMode/ControlRigEditModeCommands.h b/Engine/Plugins/Animation/ControlRig/Source/ControlRigEditor/Private/EditMode/ControlRigEditModeCommands.h
--- a/Engine/Plugins/Animation/ControlRig/Source/ControlRigEditor/Private/EditMode/ControlRigEditModeCommands.h
+++ b/Engine/Plugins/Animation/ControlRig/Source/ControlRigEditor/Private/EditMode/ControlRigEditModeCommands.h
@@ -106,6 +106,13 @@ public:
 	/** Opens up the space picker widget */
 	TSharedPtr< FUICommandInfo > OpenSpacePickerWidget;
 
+	
+	/** Opens the tween slider widget and positions it at the mouse. */
+	TSharedPtr< FUICommandInfo > SummonTweenWidget;
+
+	/** Toggles the visibility of the tween widget. */
+	TSharedPtr< FUICommandInfo > ToggleTweenWidget;
+
 	/**
 	 * Initialize commands
 	 */
diff --git a/Engine/Plugins/Animation/ControlRig/Source/ControlRigEditor/Private/EditMode/ControlRigEditModeSettings.h b/Engine/Plugins/Animation/ControlRig/Source/ControlRigEditor/Private/EditMode/ControlRigEditModeSettings.h
--- a/Engine/Plugins/Animation/ControlRig/Source/ControlRigEditor/Private/EditMode/ControlRigEditModeSettings.h
+++ b/Engine/Plugins/Animation/ControlRig/Source/ControlRigEditor/Private/EditMode/ControlRigEditModeSettings.h
@@ -3,15 +3,17 @@
 #pragma once
 
 #include "CoreMinimal.h"
+#include "Engine/DeveloperSettings.h"
+#include "Misc/Flyout/FlyoutSavedState.h"
+#include "Overlay/DragBoxPosition.h"
 #include "UObject/Object.h"
-#include "UObject/LazyObjectPtr.h"
 #include "ControlRigEditModeSettings.generated.h"
 
 DECLARE_MULTICAST_DELEGATE_OneParam(FGizmoScaleSet, float /* GizmoScale */);
 
 /** Settings object used to show useful information in the details panel */
-UCLASS(config=EditorPerProjectUserSettings, MinimalAPI)
-class UControlRigEditModeSettings : public UObject
+UCLASS(config=EditorPerProjectUserSettings, MinimalAPI, meta=(DisplayName="Control Rig Edit Mode"))
+class UControlRigEditModeSettings : public UDeveloperSettings
 {
 	GENERATED_BODY()
 
@@ -29,8 +31,12 @@ class UControlRigEditModeSettings : public UObject
 		, bLocalTransformsInEachLocalSpace(true)
 		, GizmoScale(1.0f)
 	{
-		LastInViewportTweenWidgetLocation = FVector2D(EForceInit::ForceInitToZero);
 		DrivenControlColor = FLinearColor::White * FLinearColor(FVector::OneVector * 0.8f);
+		LastTweenOverlayState = FToolWidget_FlyoutSavedState
+		{
+			FToolWidget_DragBoxPosition(FVector2f(20.f, 30.f), HAlign_Right, VAlign_Bottom),
+			false
+		};
 	}
 
 	// UObject interface
@@ -94,7 +100,18 @@ public:
 	float GizmoScale;
 
 	UPROPERTY(config)
-	FVector2D LastInViewportTweenWidgetLocation;
+	FToolWidget_FlyoutSavedState LastTweenOverlayState;
+
+	/** The tint to apply when the tween overlay is out of focus. */
+	UPROPERTY(config, EditAnywhere, Category = "Animation Settings|Tween")
+	FLinearColor TweenOutOfFocusTint = FLinearColor(1.f, 1.f, 1.f, .35f);
+
+	/**
+	 * If enabled, the DragAnimSliderTool hotkey should position the widget at the mouse cursor if the widget was visible when you pressed it.
+	 * If disabled, the DragAnimSliderTool hotkey will leave the widget where you placed it if the widget was visible when you pressed it.
+	 */
+	UPROPERTY(config, EditAnywhere, Category = "Animation Settings|Tween")
+	bool bIndirectSliderMovementShouldSnapSliderToMouse = false;
 
 	/** Delegate broadcasted whenever GizmoScale is modified */
 	FGizmoScaleSet GizmoScaleDelegate;
diff --git a/Engine/Plugins/Animation/ControlRig/Source/ControlRigEditor/Private/EditMode/ControlRigEditModeToolkit.cpp b/Engine/Plugins/Animation/ControlRig/Source/ControlRigEditor/Private/EditMode/ControlRigEditModeToolkit.cpp
--- a/Engine/Plugins/Animation/ControlRig/Source/ControlRigEditor/Private/EditMode/ControlRigEditModeToolkit.cpp
+++ b/Engine/Plugins/Animation/ControlRig/Source/ControlRigEditor/Private/EditMode/ControlRigEditModeToolkit.cpp
@@ -5,6 +5,7 @@
 */
 #include "EditMode/ControlRigEditModeToolkit.h"
 
+#include "ControlRigEditModeCommands.h"
 #include "AnimDetails/Views/SAnimDetailsView.h"
 #include "Widgets/DeclarativeSyntaxSupport.h"
 #include "EditorModes.h"
@@ -26,7 +27,6 @@
 #include "SLevelViewport.h"
 #include "Sequencer/AnimLayers/SAnimLayers.h"
 #include "Sequencer/AnimLayers/AnimLayers.h"
-#include "Tween/SControlRigTweenWidget.h"
 #include "Tools/MotionTrailOptions.h"
 
 class SAnimDetailsView;
@@ -156,14 +156,7 @@ void FControlRigEditModeToolkit::TryInvokeToolkitUI(const FName InName)
 	}
 	else if (InName == TweenOverlayName)
 	{
-		if(TweenWidgetParent)
-		{ 
-			RemoveAndDestroyTweenOverlay();
-		}
-		else
-		{
-			CreateAndShowTweenOverlay();
-		}
+		TweenOverlayManager->ToggleVisibility();
 	}
 }
 
@@ -171,7 +164,7 @@ bool FControlRigEditModeToolkit::IsToolkitUIActive(const FName InName) const
 {
 	if (InName == TweenOverlayName)
 	{
-		return TweenWidgetParent.IsValid();
+		return TweenOverlayManager->IsShowingWidget();
 	}
 
 	TSharedPtr<FAssetEditorModeUILayer> ModeUILayerPtr = ModeUILayer.Pin();
@@ -246,105 +239,6 @@ TSharedRef<SDockTab> SpawnDetailsTab(const FSpawnTabArgs& Args, FControlRigEditM
 		];
 }
 
-
-void FControlRigEditModeToolkit::CreateAndShowTweenOverlay()
-{
-	FVector2D NewTweenWidgetLocation = GetDefault<UControlRigEditModeSettings>()->LastInViewportTweenWidgetLocation;
-
-	if (NewTweenWidgetLocation.IsZero())
-	{
-		const FVector2f ActiveViewportSize = GetToolkitHost()->GetActiveViewportWidgetSize();
-		NewTweenWidgetLocation.X = ActiveViewportSize.X / 2.0f;
-		NewTweenWidgetLocation.Y = FMath::Max(ActiveViewportSize.Y - 100.0f, 0);
-	}
-	
-	UpdateTweenWidgetLocation(NewTweenWidgetLocation);
-
-	SAssignNew(TweenWidgetParent, SHorizontalBox)
-
-		+ SHorizontalBox::Slot()
-		.FillWidth(1.0f)
-		.VAlign(VAlign_Top)
-		.HAlign(HAlign_Left)
-		.Padding(TAttribute<FMargin>(this, &FControlRigEditModeToolkit::GetTweenWidgetPadding))
-		[
-			SAssignNew(TweenWidget, UE::ControlRigEditor::SControlRigTweenWidget)
-			.InOwningToolkit(SharedThis(this))
-			.InOwningEditMode(SharedThis(&EditMode))
-		];
-
-	TryShowTweenOverlay();
-}
-
-void FControlRigEditModeToolkit::TryShowTweenOverlay()
-{
-	if (TweenWidgetParent)
-	{
-		GetToolkitHost()->AddViewportOverlayWidget(TweenWidgetParent.ToSharedRef());
-	}
-}
-
-void FControlRigEditModeToolkit::RemoveAndDestroyTweenOverlay()
-{
-	TryRemoveTweenOverlay();
-	if (TweenWidgetParent)
-	{
-		TweenWidgetParent.Reset();
-		TweenWidget.Reset();
-	}
-}
-
-void FControlRigEditModeToolkit::TryRemoveTweenOverlay()
-{
-	if (IsHosted() && TweenWidgetParent)
-	{
-		if (FLevelEditorModule* LevelEditorModule = FModuleManager::GetModulePtr<FLevelEditorModule>(TEXT("LevelEditor")))
-		{
-			if (TSharedPtr<ILevelEditor> LevelEditor = LevelEditorModule->GetFirstLevelEditor())
-			{
-				for (TSharedPtr<SLevelViewport> LevelViewport : LevelEditor->GetViewports())
-				{
-					if (LevelViewport.IsValid())
-					{
-						if (TweenWidgetParent)
-						{
-							LevelViewport->RemoveOverlayWidget(TweenWidgetParent.ToSharedRef());
-						}
-					}
-				}
-			}
-		}
-	}
-}
-
-void FControlRigEditModeToolkit::UpdateTweenWidgetLocation(const FVector2D InLocation)
-{
-	const FVector2f ActiveViewportSize = GetToolkitHost()->GetActiveViewportWidgetSize();
-	FVector2D ScreenPos = InLocation;
-
-	const float EdgeFactor = 0.97f;
-	const float MinX = ActiveViewportSize.X * (1 - EdgeFactor);
-	const float MinY = ActiveViewportSize.Y * (1 - EdgeFactor);
-	const float MaxX = ActiveViewportSize.X * EdgeFactor;
-	const float MaxY = ActiveViewportSize.Y * EdgeFactor;
-	const bool bOutside = ScreenPos.X < MinX || ScreenPos.X > MaxX || ScreenPos.Y < MinY || ScreenPos.Y > MaxY;
-	if (bOutside)
-	{
-		// reset the location if it was placed out of bounds
-		ScreenPos.X = ActiveViewportSize.X / 2.0f;
-		ScreenPos.Y = FMath::Max(ActiveViewportSize.Y - 100.0f, 0.f);
-	}
-	InViewportTweenWidgetLocation = ScreenPos;
-	UControlRigEditModeSettings* ControlRigEditModeSettings = GetMutableDefault<UControlRigEditModeSettings>();
-	ControlRigEditModeSettings->LastInViewportTweenWidgetLocation = ScreenPos;
-	ControlRigEditModeSettings->SaveConfig();
-}
-
-FMargin FControlRigEditModeToolkit::GetTweenWidgetPadding() const
-{
-	return FMargin(InViewportTweenWidgetLocation.X, InViewportTweenWidgetLocation.Y, 0, 0);
-}
-
 void FControlRigEditModeToolkit::RequestModeUITabs()
 {
 	FModeToolkit::RequestModeUITabs();
@@ -421,10 +315,13 @@ void FControlRigEditModeToolkit::InvokeUI()
 		ModeUILayerPtr->GetTabManager()->TryInvokeTab(UAssetEditorUISubsystem::TopRightTabID);
 		// doesn't work as expected todo ModeUILayerPtr->GetTabManager()->TryInvokeTab(UAssetEditorUISubsystem::TopLeftTabID);
 		ModeUILayerPtr->GetTabManager()->TryInvokeTab(UAssetEditorUISubsystem::BottomRightTabID);
+
+		TweenOverlayManager = MakeUnique<UE::ControlRigEditor::FTweenOverlayManager>(GetToolkitHost(), GetToolkitCommands(), SharedThis(&EditMode));
 		if (bTweenOpen)
 		{
-			CreateAndShowTweenOverlay();
+			TweenOverlayManager->ShowWidget();
 		}
+		
 		if (bAnimLayerTabOpen)
 		{
 			TryInvokeToolkitUI(AnimLayerTabName);
@@ -464,7 +361,6 @@ void FControlRigEditModeToolkit::InvokeUI()
 void FControlRigEditModeToolkit::ShutdownUI()
 {
 	FModeToolkit::ShutdownUI();
-
 	UnregisterAndRemoveFloatingTabs();
 }
 
@@ -472,15 +368,12 @@ void FControlRigEditModeToolkit::UnregisterAndRemoveFloatingTabs()
 {
 	if (FSlateApplication::IsInitialized())
 	{
-		if (TweenWidgetParent)
+		if (TweenOverlayManager)
 		{
-			bTweenOpen = true;
+			bTweenOpen = TweenOverlayManager->IsShowingWidget();
+			TweenOverlayManager.Reset();
 		}
-		else
-		{
-			bTweenOpen = false;
-		}
-		RemoveAndDestroyTweenOverlay();
+		
 		if (ModeUILayer.IsValid())
 		{
 			TSharedPtr<FAssetEditorModeUILayer> ModeUILayerPtr = ModeUILayer.Pin();
diff --git a/Engine/Plugins/Animation/ControlRig/Source/ControlRigEditor/Private/EditMode/ControlRigEditModeToolkit.h b/Engine/Plugins/Animation/ControlRig/Source/ControlRigEditor/Private/EditMode/ControlRigEditModeToolkit.h
--- a/Engine/Plugins/Animation/ControlRig/Source/ControlRigEditor/Private/EditMode/ControlRigEditModeToolkit.h
+++ b/Engine/Plugins/Animation/ControlRig/Source/ControlRigEditor/Private/EditMode/ControlRigEditModeToolkit.h
@@ -11,6 +11,7 @@
 #include "EditMode/SControlRigEditModeTools.h"
 #include "EditMode/ControlRigEditMode.h"
 #include "Tween/ControlRigTweenModel.h"
+#include "Tween/TweenOverlayManager.h"
 
 class SAnimDetailsView;
 class SControlRigOutliner;
@@ -22,6 +23,7 @@ namespace UE::ControlRigEditor
 
 class FControlRigEditModeToolkit : public FModeToolkit
 {
+	using Super = FModeToolkit;
 public:
 	friend UE::ControlRigEditor::SControlRigTweenWidget;
 
@@ -40,11 +42,7 @@ public:
 	virtual TSharedPtr<class SWidget> GetInlineContent() const override { return ModeTools; }
 	virtual bool ProcessCommandBindings(const FKeyEvent& InKeyEvent) const override
 	{
-		if (EditMode.GetCommandBindings() && EditMode.GetCommandBindings()->ProcessCommandBindings(InKeyEvent))
-		{
-			return true;
-		}
-		return false;
+		return EditMode.GetCommandBindings() && EditMode.GetCommandBindings()->ProcessCommandBindings(InKeyEvent);
 	}
 	virtual void Init(const TSharedPtr<IToolkitHost>& InitToolkitHost) override;
 
@@ -77,15 +75,6 @@ public:
 
 protected:
 
-	void CreateAndShowTweenOverlay();
-	void TryShowTweenOverlay();
-	void RemoveAndDestroyTweenOverlay();
-	void TryRemoveTweenOverlay();
-
-	void UpdateTweenWidgetLocation(const FVector2D InLocation);
-
-	FMargin GetTweenWidgetPadding() const;
-
 	/* FModeToolkit Interface */
 	virtual void RequestModeUITabs() override;
 	virtual void InvokeUI() override;
@@ -99,14 +88,13 @@ protected:
 	static bool bPoseTabOpen;
 	static bool bSnapperTabOpen;
 	static bool bTweenOpen;
+	
 private:
 	
 	/** The edit mode we are bound to */
 	FControlRigEditMode& EditMode;
-	TSharedPtr<SWidget> TweenWidgetParent;
-	TSharedPtr<UE::ControlRigEditor::SControlRigTweenWidget> TweenWidget;
+	TUniquePtr<UE::ControlRigEditor::FTweenOverlayManager> TweenOverlayManager;
 
-	FVector2D InViewportTweenWidgetLocation;
 	/** The tools widget */
 	TSharedPtr<SControlRigEditModeTools> ModeTools;
 };
diff --git a/Engine/Plugins/Animation/ControlRig/Source/ControlRigEditor/Private/EditMode/SControlRigEditModeTools.cpp b/Engine/Plugins/Animation/ControlRig/Source/ControlRigEditor/Private/EditMode/SControlRigEditModeTools.cpp
--- a/Engine/Plugins/Animation/ControlRig/Source/ControlRigEditor/Private/EditMode/SControlRigEditModeTools.cpp
+++ b/Engine/Plugins/Animation/ControlRig/Source/ControlRigEditor/Private/EditMode/SControlRigEditModeTools.cpp
@@ -1276,7 +1276,7 @@ void SControlRigEditModeTools::CustomizeToolBarPalette(FToolBarBuilder& ToolBarB
 		),		
 		NAME_None,
 		LOCTEXT("Tweens", "Tweens"),
-		LOCTEXT("TweensTooltip", "Create Tweens"),
+		LOCTEXT("TweensTooltip", "Create Tweens.\n\nDefault bindings:\nU - While hidden, temporarily bring widget to your cursor\nU+LMB - While visible, indirectly move the slider\nAlt+U - Toggle visibility\nShift+U - Cycle function\nCtrl+U - Toggle Overshoot Mode"),
 		FSlateIcon(TEXT("ControlRigEditorStyle"), TEXT("ControlRig.TweenTool")),
 		EUserInterfaceActionType::ToggleButton
 	);
diff --git a/Engine/Plugins/Animation/ControlRig/Source/ControlRigEditor/Private/EditMode/Tween/SControlRigTweenWidget.cpp b/Engine/Plugins/Animation/ControlRig/Source/ControlRigEditor/Private/EditMode/Tween/SControlRigTweenWidget.cpp
deleted file mode 100644
diff --git a/Engine/Plugins/Animation/ControlRig/Source/ControlRigEditor/Private/EditMode/Tween/SControlRigTweenWidget.h b/Engine/Plugins/Animation/ControlRig/Source/ControlRigEditor/Private/EditMode/Tween/SControlRigTweenWidget.h
deleted file mode 100644
diff --git a/Engine/Plugins/Animation/ControlRig/Source/ControlRigEditor/Private/EditMode/Tween/TweenLogic.cpp b/Engine/Plugins/Animation/ControlRig/Source/ControlRigEditor/Private/EditMode/Tween/TweenLogic.cpp
--- a/Engine/Plugins/Animation/ControlRig/Source/ControlRigEditor/Private/EditMode/Tween/TweenLogic.cpp
+++ b/Engine/Plugins/Animation/ControlRig/Source/ControlRigEditor/Private/EditMode/Tween/TweenLogic.cpp
@@ -6,9 +6,6 @@
 #include "ControlRigTweenModels.h"
 #include "EditMode/ControlRigEditMode.h"
 #include "Framework/Commands/UICommandList.h"
-#include "Math/Abstraction/TweenModelArray.h"
-#include "Math/Models/EditorTweenModel.h"
-#include "Misc/SequencerUtils.h"
 #include "Templates/SharedPointer.h"
 
 namespace UE::ControlRigEditor
@@ -19,7 +16,7 @@ FTweenLogic::FTweenLogic(const TAttribute<TWeakPtr<ISequencer>>& InSequencerAttr
 	, Controllers(CommandList, TweenModels, TEXT("ControlRigTween"))
 {}
 
-TSharedRef<SWidget> FTweenLogic::MakeWidget() const
+FTweenLogicWidgets FTweenLogic::MakeWidget() const
 {
 	// Visually move the slider widget when the user uses the U+LMB button command to indirectly move the mouse
 	TAttribute<TOptional<float>> OverrideSliderPosition =
@@ -30,7 +27,9 @@ TSharedRef<SWidget> FTweenLogic::MakeWidget() const
 	
 	FSlimHorizontalToolBarBuilder ToolBarBuilder(CommandList, FMultiBoxCustomization::None);
 	ToolBarBuilder.SetStyle(&FAppStyle::Get(), TEXT("ControlRigTweenToolbar"));
-	Controllers.ToolbarController.AddToToolbar(ToolBarBuilder, { MoveTemp(OverrideSliderPosition) });
-	return ToolBarBuilder.MakeWidget();
+	const TweeningUtilsEditor::FTweenToolbarController::FAddToToolbarResult ToolbarWidgets = Controllers.ToolbarController.AddToToolbar(
+		ToolBarBuilder, { MoveTemp(OverrideSliderPosition) }
+		);
+	return FTweenLogicWidgets(ToolBarBuilder.MakeWidget(), ToolbarWidgets.TweenSlider);
 }
 }
diff --git a/Engine/Plugins/Animation/ControlRig/Source/ControlRigEditor/Private/EditMode/Tween/TweenLogic.h b/Engine/Plugins/Animation/ControlRig/Source/ControlRigEditor/Private/EditMode/Tween/TweenLogic.h
--- a/Engine/Plugins/Animation/ControlRig/Source/ControlRigEditor/Private/EditMode/Tween/TweenLogic.h
+++ b/Engine/Plugins/Animation/ControlRig/Source/ControlRigEditor/Private/EditMode/Tween/TweenLogic.h
@@ -14,6 +14,16 @@ namespace UE::ControlRigEditor
 {
 class FControlRigTweenModels;
 
+struct FTweenLogicWidgets
+{
+	TSharedPtr<SWidget> RootWidget;
+	TSharedPtr<TweeningUtilsEditor::STweenSlider> TweenSlider;
+
+	explicit FTweenLogicWidgets(TSharedPtr<SWidget> InRootWidget, TSharedPtr<TweeningUtilsEditor::STweenSlider> InTweenSlider)
+		: RootWidget(MoveTemp(InRootWidget)), TweenSlider(MoveTemp(InTweenSlider))
+	{}
+};
+
 /** Creates the tweening model specific to ControlRig and provides an interface for interacting with it. */
 class FTweenLogic
 {
@@ -22,7 +32,9 @@ public:
 	explicit FTweenLogic(const TAttribute<TWeakPtr<ISequencer>>& InSequencerAttr, const TSharedRef<FControlRigEditMode>& InOwningEditMode);
 
 	/** Constructs the content for the viewport tweening widget. */
-	TSharedRef<SWidget> MakeWidget() const;
+	FTweenLogicWidgets MakeWidget() const;
+
+	TweeningUtilsEditor::FTweenControllers& GetControllers() { return Controllers; }
 	
 private:
 
@@ -33,6 +45,6 @@ private:
 	const TSharedRef<FControlRigTweenModels> TweenModels;
 
 	/** Common functionality that should be shared consistent with other editor modules, such as Curve Editor. */
-	const TweeningUtilsEditor::FTweenControllers Controllers;
+	TweeningUtilsEditor::FTweenControllers Controllers;
 };
 }
diff --git a/Engine/Plugins/Animation/ControlRig/Source/ControlRigEditor/Private/EditMode/Tween/TweenOverlayManager.cpp b/Engine/Plugins/Animation/ControlRig/Source/ControlRigEditor/Private/EditMode/Tween/TweenOverlayManager.cpp
new file mode 100644
--- /dev/null
+++ b/Engine/Plugins/Animation/ControlRig/Source/ControlRigEditor/Private/EditMode/Tween/TweenOverlayManager.cpp
@@ -0,0 +1,142 @@
+﻿// Copyright Epic Games, Inc. All Rights Reserved.
+
+#include "TweenOverlayManager.h"
+
+#include "EditMode/ControlRigEditModeCommands.h"
+#include "EditMode/ControlRigEditModeSettings.h"
+#include "Editor/ControlRigEditorCommands.h"
+#include "Framework/Commands/UICommandList.h"
+#include "ILevelSequenceEditorToolkit.h"
+#include "LevelSequence.h"
+#include "LevelSequenceEditorBlueprintLibrary.h"
+#include "TweeningUtilsStyle.h"
+#include "Misc/Flyout/Host/ToolkitBasedWidgetHost.h"
+#include "Subsystems/AssetEditorSubsystem.h"
+#include "Toolkits/IToolkitHost.h"
+#include "Widgets/STweenSlider.h"
+#include "Widgets/TweenSliderStyle.h"
+
+namespace UE::ControlRigEditor
+{
+namespace Private
+{
+static void SaveWidgetState(FToolWidget_FlyoutSavedState State)
+{
+	UControlRigEditModeSettings* ControlRigEditModeSettings = GetMutableDefault<UControlRigEditModeSettings>();
+	ControlRigEditModeSettings->LastTweenOverlayState = State;
+	ControlRigEditModeSettings->SaveConfig();
+}
+}
+	
+FTweenOverlayManager::FTweenOverlayManager(
+	const TSharedRef<IToolkitHost>& InToolkitHost,
+	const TSharedRef<FUICommandList>& InToolkitCommandList,
+	const TSharedRef<FControlRigEditMode>& InOwningEditMode
+	)
+	: TweenControllers(
+		TAttribute<TWeakPtr<ISequencer>>::CreateLambda([]
+		{
+			ULevelSequence* LevelSequence = ULevelSequenceEditorBlueprintLibrary::GetCurrentLevelSequence();
+			IAssetEditorInstance* AssetEditor = GEditor->GetEditorSubsystem<UAssetEditorSubsystem>()->FindEditorForAsset(LevelSequence, false);
+			ILevelSequenceEditorToolkit* LevelSequenceEditor = static_cast<ILevelSequenceEditorToolkit*>(AssetEditor);
+			return LevelSequenceEditor ? LevelSequenceEditor->GetSequencer() : nullptr;
+		}),
+		InOwningEditMode
+		)
+	, WidgetHierarchy(TweenControllers.MakeWidget())
+	, FlyoutWidgetManager(
+		FFlyoutWidgetArgs(
+			WidgetHierarchy.RootWidget.ToSharedRef(),
+			TAttribute<TSharedPtr<FUICommandList>>::CreateLambda([WeakList = InToolkitCommandList.ToWeakPtr()]{ return WeakList.Pin(); }),
+			MakeShared<FToolkitBasedWidgetHost>(InToolkitHost)
+			)
+			.SetStateToRestoreFrom(GetMutableDefault<UControlRigEditModeSettings>()->LastTweenOverlayState)
+			.SetSaveStateDelegate(FSaveFlyoutState::CreateStatic(&Private::SaveWidgetState))
+			.SetCanSubdue(TAttribute<bool>::CreateLambda([this]{ return !TweenControllers.GetControllers().MouseSlidingController.IsSliding(); }))
+			.SetSubduedOpacity(TAttribute<FLinearColor>::CreateLambda([]{ return GetMutableDefault<UControlRigEditModeSettings>()->TweenOutOfFocusTint; }))
+			.SetAbsOffsetFromCursor(TAttribute<FVector2f>::CreateRaw(this, &FTweenOverlayManager::ComputeOffsetFromWidgetCenterToSliderPosition))
+			.SetToggleVisibility(FControlRigEditModeCommands::Get().ToggleTweenWidget)
+			.SetSummonToCursor(FControlRigEditModeCommands::Get().SummonTweenWidget)
+		)
+{
+	TweeningUtilsEditor::FTweenControllers& Controllers = TweenControllers.GetControllers();
+	Controllers.ToolbarController.OnOvershootModeCommandInvoked().AddRaw(this, &FTweenOverlayManager::OnTweenCommandInvoked);
+	Controllers.CycleFunctionController.OnFunctionCycleCommandInvoked().AddRaw(this, &FTweenOverlayManager::OnTweenCommandInvoked);
+	
+	TweeningUtilsEditor::FTweenMouseSlidingController& MouseSlider = Controllers.MouseSlidingController;
+	MouseSlider.OnStartSliding().AddRaw(this, &FTweenOverlayManager::OnStartIndirectSliding);
+	MouseSlider.OnStopSliding().AddRaw(this, &FTweenOverlayManager::OnStopIndirectSliding);
+
+	FlyoutWidgetManager.OnVisibilityChanged().AddRaw(this, &FTweenOverlayManager::OnVisibilityChanged);
+	OnVisibilityChanged(IsShowingWidget());
+}
+
+FVector2f FTweenOverlayManager::ComputeOffsetFromWidgetCenterToSliderPosition() const
+{
+	const FGeometry& RootGeometry = WidgetHierarchy.RootWidget->GetTickSpaceGeometry();
+	const FGeometry& TweenSliderGeometry = WidgetHierarchy.TweenSlider->GetTickSpaceGeometry();
+	const bool bHasEverBeenPainted = !RootGeometry.GetLocalSize().IsNearlyZero() && !TweenSliderGeometry.GetLocalSize().IsNearlyZero();
+	if (!ensure(bHasEverBeenPainted)) // FlyoutWidgetManager should have made sure that the widget was painted at least once.
+	{
+		return FVector2f::ZeroVector;
+	}
+	
+	const float AbsCenterOfRoot = RootGeometry.GetAbsolutePosition().X + RootGeometry.GetAbsoluteSize().X / 2.f;
+	const float AbsCenterOfSlider = TweenSliderGeometry.GetAbsolutePosition().X + TweenSliderGeometry.GetAbsoluteSize().X / 2.f;
+	const float CenterOfRootToCenterOfSlider = AbsCenterOfSlider - AbsCenterOfRoot;
+	
+	// No DPI scaling required because we are already computing in absolute screen space.
+	return FVector2f(-CenterOfRootToCenterOfSlider, 0.f);
+}
+
+void FTweenOverlayManager::OnStartIndirectSliding()
+{
+	// The widget should be fully visible while indirectly driving the slider
+	FlyoutWidgetManager.UnsubdueContent();
+	
+	if (GetMutableDefault<UControlRigEditModeSettings>()->bIndirectSliderMovementShouldSnapSliderToMouse)
+	{
+		// For indirect slider manipulation, the user is allowed to have the cursor anywhere in the editor.
+		// We'll just snap the widget to the viewport border if that is the case.
+		constexpr ETemporaryFlyoutPositionFlags Flags = ETemporaryFlyoutPositionFlags::AllowCursorOutsideOfParent;
+		IndirectSlidePositionOverride = FlyoutWidgetManager.TryTemporarilyPositionWidgetAtCursor(Flags);
+	}
+}
+
+void FTweenOverlayManager::OnStopIndirectSliding()
+{
+	IndirectSlidePositionOverride.Reset();
+	
+	// Subdue the content again, but don't force it in case the user was already hovering the widget when they started indirectly sliding it.
+	FlyoutWidgetManager.TrySubdueContent();
+}
+
+void FTweenOverlayManager::OnVisibilityChanged(bool bIsVisible)
+{
+	// FControlRigEditModeCommands::SummonTweenWidget is bound to U (default)
+	// FTweeningUtilsCommands::DragAnimSliderTool is bound to U (default), as well. While pressed, once you LMB, you can indirectly manipulate the slider.
+	//
+	// When visible: Indirect slider manipulation SHOULD occur. SummonTweenWidget will place the widget under the cursor such that the slider is directly under the mouse.
+	// When not visible: Indirect slider manipulation SHOULD NOT occur. DragAnimSliderTool will take precedence over SummonTweenWidget.
+	// See class documentation for command precedence explaination.
+	if (bIsVisible && !FlyoutWidgetManager.IsTemporarilyRepositioned())
+	{
+		TweenControllers.GetControllers().MouseSlidingController.BindCommand();
+	}
+	else
+	{
+		TweenControllers.GetControllers().MouseSlidingController.UnbindCommand();
+	}
+}
+
+void FTweenOverlayManager::OnTweenCommandInvoked()
+{
+	// If the widget is hidden when a command is used on it, temporarily show the widget and move it to the cursor until the user moves the mouse away.
+	// That allows the user to "prepare" the tween widget while they've decided to hide it.
+	if (!IsShowingWidget())
+	{
+		constexpr bool bHideAtEnd = true;
+		FlyoutWidgetManager.SummonToCursorUntilMouseLeave(bHideAtEnd);
+	}
+}
+}
diff --git a/Engine/Plugins/Animation/ControlRig/Source/ControlRigEditor/Private/EditMode/Tween/TweenOverlayManager.h b/Engine/Plugins/Animation/ControlRig/Source/ControlRigEditor/Private/EditMode/Tween/TweenOverlayManager.h
new file mode 100644
--- /dev/null
+++ b/Engine/Plugins/Animation/ControlRig/Source/ControlRigEditor/Private/EditMode/Tween/TweenOverlayManager.h
@@ -0,0 +1,77 @@
+﻿// Copyright Epic Games, Inc. All Rights Reserved.
+
+#pragma once
+
+#include "EditMode/ControlRigBaseDockableView.h"
+#include "Misc/Flyout/FlyoutOverlayManager.h"
+#include "TweenLogic.h"
+
+namespace UE::ControlRigEditor
+{
+/**
+ * Manages all logic around the tween overlay widget.
+ * 
+ * This class is responsible for holding any Control Rig specific state and injects it into a FFlyoutOverlayManager, which handles all the UI
+ * logic for docking the tween slider to the viewport, etc.
+ *
+ * ========== Command precedence ==========
+ * Both FControlRigEditModeCommands::SummonTweenWidget and FTweeningUtilsCommands::DragAnimSliderTool are bound to U by default.
+ * While the flyout widget is visible, DragAnimSliderTool takes precedence. This is achieved by binding:
+ * - DragAnimSliderTool to FControlRigEditModeToolkit::GetCommandBindings().
+ * - SummonTweenWidget to FControlRigEditMode::GetCommandBindings().
+ *
+ * Sequence of input events in the engine is:
+ * 1. FEditorModeTools::InputKey, which invokes FControlRigEditModeToolkit::GetCommandBindings()->ProcessCommandBindings
+ * 2. SEditorViewport::OnKeyDown
+ * 3. SLevelViewport::OnKeyDown
+ * 4. SLevelEditor::OnKeyDown, which invokes FControlRigEditModeToolkit::ProcessCommandBindings
+ * There are other commands bound to e.g. shift+U in SEditorViewport, so it's important the toolkit's tween commands are processed first. So they
+ * must go into FControlRigEditModeToolkit::GetCommandBindings(). This would be fixed if FEditorModeTools::InputKey called
+ * FControlRigEditModeToolkit::ProcessCommandBindings instead.
+ * ========================================
+ */
+class FTweenOverlayManager : public FNoncopyable
+{
+public:
+	
+	explicit FTweenOverlayManager(
+		const TSharedRef<IToolkitHost>& InToolkitHost,
+		const TSharedRef<FUICommandList>& InToolkitCommandList,
+		const TSharedRef<FControlRigEditMode>& InOwningEditMode
+		);
+
+	void ToggleVisibility() { FlyoutWidgetManager.ToggleVisibility(); }
+	void ShowWidget() { FlyoutWidgetManager.ShowWidget(); }
+	void HideWidget() { FlyoutWidgetManager.HideWidget(); }
+	bool IsShowingWidget() const { return FlyoutWidgetManager.IsShowingWidget(); }
+
+private:
+	
+	/** Manages logic for the tweening widget. */
+	FTweenLogic TweenControllers;
+
+	/** The widget hierarchy in the flyout widget manager. */
+	const FTweenLogicWidgets WidgetHierarchy;
+	/** Contains the tween slider widget. Allows it to be dragged around, commands, etc.  */
+	FFlyoutOverlayManager FlyoutWidgetManager;
+	
+	/** Set while the user is sliding using U+LMB. */
+	TSharedPtr<FFlyoutTemporaryPositionOverride> IndirectSlidePositionOverride;
+
+	/** @return Computes the offset the widget should have from the mouse cursor so the center of the slider is under it. */
+	FVector2f ComputeOffsetFromWidgetCenterToSliderPosition() const;
+
+	// When using U+LMB, show move the tween overlay to the mouse.
+	void OnStartIndirectSliding();
+	void OnStopIndirectSliding();
+
+	/**
+	 * While invisible, gives precedence to SummonTweenWidget by deactivating indirect sliding (DragAnimSliderTool, U+LMB)
+	 * because by default both commands are bound to the same hotkey.
+	 */
+	void OnVisibilityChanged(bool bIsVisible);
+
+	/** If the widget is hidden, temporarily shows the widget and snaps it to the cursor until the user moves the mouse away. */
+	void OnTweenCommandInvoked();
+};
+}
diff --git a/Engine/Plugins/Animation/ControlRig/Source/ControlRigEditor/Private/Misc/Flyout/FlyoutOverlayManager.cpp b/Engine/Plugins/Animation/ControlRig/Source/ControlRigEditor/Private/Misc/Flyout/FlyoutOverlayManager.cpp
new file mode 100644
--- /dev/null
+++ b/Engine/Plugins/Animation/ControlRig/Source/ControlRigEditor/Private/Misc/Flyout/FlyoutOverlayManager.cpp
@@ -0,0 +1,474 @@
+﻿// Copyright Epic Games, Inc. All Rights Reserved.
+
+#include "FlyoutOverlayManager.h"
+
+#include "Framework/Application/IInputProcessor.h"
+#include "Framework/Application/SlateApplication.h"
+#include "Host/IWidgetHost.h"
+#include "Overlay/SDraggableBoxOverlay.h"
+#include "Widgets/SNullWidget.h"
+
+namespace UE::ControlRigEditor
+{
+namespace Private
+{
+/** Content placed inside the SDraggableBoxOverlay. This lets us know when the mouse cursor has left the content. */
+class SFlyoutWidgetWrapper : public SCompoundWidget
+{
+	FSimpleDelegate OnMouseEnterDelegate;
+	FSimpleDelegate OnMouseLeftDelegate;
+
+	bool bDetectMouseLeave = false;
+public:
+
+	SLATE_BEGIN_ARGS(SFlyoutWidgetWrapper){}
+		SLATE_NAMED_SLOT(FArguments, Content)
+		SLATE_EVENT(FSimpleDelegate, OnMouseEnter)
+		SLATE_EVENT(FSimpleDelegate, OnMouseLeft)
+	SLATE_END_ARGS()
+
+	void Construct(const FArguments& InArgs)
+	{
+		OnMouseEnterDelegate = InArgs._OnMouseEnter;
+		OnMouseLeftDelegate = InArgs._OnMouseLeft;
+		
+		ChildSlot
+		[
+			InArgs._Content.Widget
+		];
+		
+		SetVisibility(EVisibility::Visible); // Allow clicks to go to child widget.
+	}
+
+	void ClearChildContent()
+	{
+		ChildSlot[ SNullWidget::NullWidget ];
+	}
+
+	void DetectMouseLeave() { bDetectMouseLeave = true; }
+
+	virtual void OnMouseEnter(const FGeometry& MyGeometry, const FPointerEvent& MouseEvent) override
+	{
+		SCompoundWidget::OnMouseEnter(MyGeometry, MouseEvent);
+		OnMouseEnterDelegate.ExecuteIfBound();
+	}
+
+	virtual void OnMouseLeave(const FPointerEvent& MouseEvent) override
+	{
+		SCompoundWidget::OnMouseLeave(MouseEvent);
+		bDetectMouseLeave = false;
+		OnMouseLeftDelegate.ExecuteIfBound();
+	}
+
+	virtual void Tick(const FGeometry& AllottedGeometry, const double InCurrentTime, const float InDeltaTime) override
+	{
+		SCompoundWidget::Tick(AllottedGeometry, InCurrentTime, InDeltaTime);
+
+		if (bDetectMouseLeave && !AllottedGeometry.IsUnderLocation(FSlateApplication::Get().GetCursorPos()))
+		{
+			bDetectMouseLeave = false;
+			OnMouseLeftDelegate.ExecuteIfBound();
+		}
+	}
+};
+
+/** Keeps the flyout widget temporarily repositioned for as long as the user keeps SummonToCursorCommand pressed down. */
+class FSummonToCursorInputProcessor : public IInputProcessor, public TSharedFromThis<FSummonToCursorInputProcessor>
+{
+	const TSharedPtr<FUICommandInfo> SummonToCursorCommand;
+	const TSharedPtr<FFlyoutTemporaryPositionOverride> KeepAliveOverride;
+public:
+
+	explicit FSummonToCursorInputProcessor(
+		TSharedPtr<FUICommandInfo> InSummonToCursorCommand, TSharedPtr<FFlyoutTemporaryPositionOverride> InKeepAliveOverride
+		)
+		: SummonToCursorCommand(MoveTemp(InSummonToCursorCommand))
+		, KeepAliveOverride(MoveTemp(InKeepAliveOverride))
+	{}
+
+	virtual bool HandleKeyUpEvent(FSlateApplication&, const FKeyEvent& InKeyEvent) override
+	{
+		if (!SummonToCursorCommand)
+		{
+			return false;
+		}
+
+		const FKey Key = InKeyEvent.GetKey();
+		const TSharedRef<const FInputChord> PrimaryChord = SummonToCursorCommand->GetActiveChord(EMultipleKeyBindingIndex::Primary);
+		const bool bHasStoppedPressing = PrimaryChord->Key == Key
+			|| (PrimaryChord->NeedsShift() && (Key == EKeys::LeftShift || Key == EKeys::RightShift))
+			|| (PrimaryChord->NeedsAlt() && (Key == EKeys::LeftAlt || Key == EKeys::RightAlt))
+			|| (PrimaryChord->NeedsControl() && (Key == EKeys::LeftControl || Key == EKeys::RightControl))
+			|| (PrimaryChord->NeedsCommand() && (Key == EKeys::LeftCommand || Key == EKeys::RightCommand));
+		if (bHasStoppedPressing)
+		{
+			FSlateApplication::Get().UnregisterInputPreProcessor(SharedThis(this));
+		}
+		
+		return false;
+	}
+
+	virtual void Tick(const float DeltaTime, FSlateApplication& SlateApp, TSharedRef<ICursor> Cursor) override {}
+};
+}
+	
+FFlyoutOverlayManager::FFlyoutOverlayManager(FFlyoutWidgetArgs InArgs)
+	: Args(MoveTemp(InArgs))
+{
+	BindCommands();
+
+	if (InArgs.StateToRestoreFrom && InArgs.StateToRestoreFrom->bWasVisible)
+	{
+		ShowWidget();
+	}
+
+	if (FSlateApplication::IsInitialized())
+	{
+		// This is needed so we can detect when a drop-down "windows" being closed.
+		FSlateApplication::Get().OnWindowBeingDestroyed().AddRaw(this, &FFlyoutOverlayManager::OnWindowDestroyed);
+	}
+}
+
+FFlyoutOverlayManager::~FFlyoutOverlayManager()
+{
+	if (DragWidget)
+	{
+		Args.WidgetHost->RemoveWidgetFromHost();
+		ensureAlwaysMsgf(DragWidget.IsUnique(), TEXT("Widget not properly removed from hierarchy"));
+		
+		// Defensive: if WidgetHost did not properly remove from the hierarchy, the widget life will be extended by our DragWidget referencing it.
+		// Clear it to give it a chance to die.
+		ContentWidget->ClearChildContent(); 
+	}
+
+	if (FSlateApplication::IsInitialized())
+	{
+		FSlateApplication::Get().OnWindowBeingDestroyed().RemoveAll(this);
+	}
+
+	UnbindCommands();
+}
+
+void FFlyoutOverlayManager::ShowWidget()
+{
+	if (bShouldShowWidget)
+	{
+		return;
+	}
+	bShouldShowWidget = true;
+	
+	if (!DragWidget)
+	{
+		DragWidget = SNew(ToolWidgets::SDraggableBoxOverlay)
+			.OnUserDraggedToNewPosition_Raw(this, &FFlyoutOverlayManager::SaveWidgetState)
+			.IsDraggable_Raw(this, &FFlyoutOverlayManager::CanDragWidget)
+			.Visibility_Lambda([this]{ return bShouldShowWidget ? EVisibility::Visible : EVisibility::Collapsed; })
+			.Content()
+			[
+				SAssignNew(ContentWidget, Private::SFlyoutWidgetWrapper)
+				.OnMouseEnter_Raw(this, &FFlyoutOverlayManager::OnMouseEnterContent)
+				.OnMouseLeft_Raw(this, &FFlyoutOverlayManager::OnMouseLeftContent)
+				.Content()[ Args.GetWidgetDelegate.Execute() ]
+			];
+		Args.WidgetHost->AddWidgetToHost(DragWidget.ToSharedRef());
+
+		if (Args.StateToRestoreFrom)
+		{
+			DragWidget->RestoreFromDragBoxPosition(Args.StateToRestoreFrom->Position);
+		}
+
+		TrySubdueContent();
+	}
+
+	// Save that the widget was last visible.
+	SaveWidgetState();
+
+	OnVisibilityChangedDelegate.Broadcast(true);
+}
+
+void FFlyoutOverlayManager::HideWidget()
+{
+	if (DragWidget && bShouldShowWidget)
+	{
+		bShouldShowWidget = false;
+
+		constexpr bool bSkipHide = true; // Avoid recursive call.
+		StopTemporaryWidgetPosition(bSkipHide);
+		constexpr bool bForceSubdue = true; 
+		TrySubdueContent(bForceSubdue);
+		
+		// Save that the widget was last not visible.
+		SaveWidgetState();
+
+		OnVisibilityChangedDelegate.Broadcast(false);
+	}
+}
+
+TSharedPtr<FFlyoutTemporaryPositionOverride> FFlyoutOverlayManager::TryTemporarilyPositionWidgetAtCursor(ETemporaryFlyoutPositionFlags InFlags)
+{
+	// No double starting the operation
+	if (TemporaryMoveOperation)
+	{
+		return nullptr;
+	}
+
+	ShowWidget();
+	if (!ensure(DragWidget))
+	{
+		return nullptr;
+	}
+
+	// The widget's size is required in order to center it around the mouse cursor. For that the widget needs to paint at least once.
+	const FVector2f LocalSize = DragWidget->GetTickSpaceGeometry().GetLocalSize();
+	const bool bWasEverTicked = !LocalSize.IsNearlyZero();
+	if (!bWasEverTicked)
+	{
+		const TSharedRef<FFlyoutTemporaryPositionOverride> LatentOverride = MakeShared<FFlyoutTemporaryPositionOverride>(
+			FFlyoutTemporaryPositionOverride::FPrivateToken(), *this
+			);
+		// ShowTweenWidget call above may have just created the drag widget, in which case GetTickSpaceGeometry() is unset.
+		// We need to wait for the widget to paint at least once so the geometry becomes set.
+		// If that does not happen next tick (it should), disregard this request.
+		ExecuteOnGameThread(TEXT("Defer PositionWidgetAtMouseCursor"), [this, InFlags, WeakWidget = DragWidget.ToWeakPtr(), LatentOverride]
+		{
+			// We assume that the DragWidget is destroyed when this FFlyoutOverlayManager is destroyed.
+			// That may not be the case if RemoveFromOverlayDelegate does not remove the widget properly... in that case we'll get undefined behavior.
+			const TSharedPtr<ToolWidgets::SDraggableBoxOverlay> ThisPin = WeakWidget.Pin();
+			if (!ThisPin)
+			{
+				return;
+			}
+
+			if (!LatentOverride->IsCancelled() && !TemporarilyMoveWidgetToCursor(InFlags))
+			{
+				LatentOverride->Cancel();
+			}
+		});
+		return LatentOverride;
+	}
+
+	if (TemporarilyMoveWidgetToCursor(InFlags))
+	{
+		return MakeShared<FFlyoutTemporaryPositionOverride>(FFlyoutTemporaryPositionOverride::FPrivateToken(), *this);
+	}
+	return nullptr;
+}
+	
+TSharedPtr<FFlyoutTemporaryPositionOverride> FFlyoutOverlayManager::SummonToCursorUntilMouseLeave(bool bHideAtEnd)
+{
+	const ETemporaryFlyoutPositionFlags Flags = bHideAtEnd ? ETemporaryFlyoutPositionFlags::HideAtEnd : ETemporaryFlyoutPositionFlags::None;
+	if (const TSharedPtr<FFlyoutTemporaryPositionOverride> Override = TryTemporarilyPositionWidgetAtCursor(Flags))
+	{
+		FSlateApplication::Get().RegisterInputPreProcessor(MakeShared<Private::FSummonToCursorInputProcessor>(Args.SummonToCursorCommand, Override));
+		CommandTriggeredPositionOverride = Override;
+	}
+	return nullptr;
+}
+	
+void FFlyoutOverlayManager::TrySubdueContent(bool bForce) const
+{
+	if (bForce ||
+		// No subduing while the widget is temporarily moved to the mouse cursor
+		(!TemporaryMoveOperation && Args.CanSubdueAttr.Get())
+		)
+	{
+		ContentWidget->SetColorAndOpacity(Args.SubduedTintAttr.Get());
+	}
+}
+
+void FFlyoutOverlayManager::UnsubdueContent() const
+{
+	ContentWidget->SetColorAndOpacity(FLinearColor(1.f, 1.f, 1.f, 1.f));
+}
+
+void FFlyoutOverlayManager::BindCommands()
+{
+	const TSharedPtr<FUICommandList> CommandList = Args.CommandList.Get();
+	if (!CommandList)
+	{
+		return;
+	}
+
+	// The commands are optional... if the API user does not specify them, they don't want that feature.
+	if (Args.ToggleVisibilityCommand)
+	{
+		CommandList->MapAction(
+			Args.ToggleVisibilityCommand,
+			FExecuteAction::CreateRaw(this, &FFlyoutOverlayManager::ToggleVisibility)
+		);
+	}
+	if (Args.SummonToCursorCommand)
+	{
+		CommandList->MapAction(
+			Args.SummonToCursorCommand,
+			FExecuteAction::CreateRaw(this, &FFlyoutOverlayManager::HandleSummonToCursorCommand)
+		);
+	}
+}
+
+void FFlyoutOverlayManager::UnbindCommands() const
+{
+	const TSharedPtr<FUICommandList> CommandList = Args.CommandList.Get();
+	if (!CommandList) // Command list may have been destroyed already... in that case there is nothing to clean up.
+	{
+		return;
+	}
+
+	if (Args.ToggleVisibilityCommand)
+	{
+		CommandList->UnmapAction(Args.ToggleVisibilityCommand);
+	}
+	if (Args.SummonToCursorCommand)
+	{
+		CommandList->UnmapAction(Args.SummonToCursorCommand);
+	}
+}
+
+void FFlyoutOverlayManager::HandleSummonToCursorCommand()
+{
+	if (!TemporaryMoveOperation && !IsShowingWidget())
+	{
+		SummonToCursorUntilMouseLeave(true);
+	}
+}
+	
+bool FFlyoutOverlayManager::TemporarilyMoveWidgetToCursor(ETemporaryFlyoutPositionFlags InFlags)
+{
+	check(DragWidget);
+
+	// InternalPositionWidgetAt can fail, so get its position before the move so we can save it if the move succeeds. 
+	const FToolWidget_DragBoxPosition SizeBeforeTemporaryMove = DragWidget->GetDragBoxPosition();
+	const bool bHideAtEnd = EnumHasAnyFlags(InFlags, ETemporaryFlyoutPositionFlags::HideAtEnd);
+	const bool bAllowCursorOutsideOfParent = EnumHasAnyFlags(InFlags, ETemporaryFlyoutPositionFlags::AllowCursorOutsideOfParent);
+	
+	const FVector2f CursorPos = FSlateApplication::Get().GetCursorPos();
+	const FToolWidget_DragBoxPosition PositionBeforeMove = DragWidget->GetDragBoxPosition();
+	if (InternalPositionWidgetAt(CursorPos + Args.AbsOffsetFromCursorAttr.Get(), bAllowCursorOutsideOfParent))
+	{
+		// If the user moves the mouse very quickly and summons to cursor, then OnMouseLeave may not fire. To fix this, we'll check every tick.
+		ContentWidget->DetectMouseLeave();
+		
+		// The user wants the widget to open up at the same position if they close it now.
+		SaveWidgetState(SizeBeforeTemporaryMove);
+
+		TemporaryMoveOperation.Emplace(PositionBeforeMove, bHideAtEnd);
+		UnsubdueContent();
+		return true;
+	}
+
+	if (bHideAtEnd) // There was no suitable position to place the widget (e.g. cursor outside of viewport), so hide it again if it was requested.
+	{
+		HideWidget();
+	}
+	return false;
+}
+
+void FFlyoutOverlayManager::StopTemporaryWidgetPosition(bool bSkipHideEvenIfRequested)
+{
+	if (TemporaryMoveOperation)
+	{
+		DragWidget->RestoreFromDragBoxPosition(TemporaryMoveOperation->RestorePosition);
+
+		const bool bHide = TemporaryMoveOperation->bHideWidgetAtEnd; 
+		TemporaryMoveOperation.Reset();
+		
+		if (bHide && !bSkipHideEvenIfRequested)
+		{
+			HideWidget();
+		}
+	}
+}
+
+bool FFlyoutOverlayManager::InternalPositionWidgetAt(const FVector2f& InAbsoluteScreenPosition, bool bAllowCursorOutsideOfParent) const
+{
+	const FVector2f LocalSize = DragWidget->GetTickSpaceGeometry().GetLocalSize();
+	if (!ensure(!LocalSize.IsNearlyZero())) // Widget was never painted, so GetTickSpaceGeometry() is unset. The caller make sure it's valid. 
+	{
+		return false;
+	}
+	const FSlateLayoutTransform AbsoluteToWidget = DragWidget->GetTickSpaceGeometry().GetAccumulatedLayoutTransform().Inverse();
+
+	// If the cursor is outside of the parent widget (e.g. the viewport), there's no point to continue because
+	const FVector2f LocalNormalizedCursorPosition = AbsoluteToWidget.TransformPoint(InAbsoluteScreenPosition) / LocalSize;
+	const bool bClickedOutsideOfParent = LocalNormalizedCursorPosition.X < 0.f || LocalNormalizedCursorPosition.X > 1.f
+		|| LocalNormalizedCursorPosition.Y < 0.f || LocalNormalizedCursorPosition.Y > 1.f;
+	if (bClickedOutsideOfParent && !bAllowCursorOutsideOfParent)
+	{
+		return false;
+	}
+	
+	const FVector2f ContentSize = ContentWidget->GetTickSpaceGeometry().GetAbsoluteSize();
+	const FVector2f AbsFinalPosition = InAbsoluteScreenPosition
+		// Position center of content widget in at the mouse cursor.
+		- FVector2f(0.5f, 0.5f) * ContentSize;
+	const FVector2f LocalFinalPosition = AbsoluteToWidget.TransformPoint(AbsFinalPosition);
+
+	// Normalized coord: (0,0) is top-left, and (1,1) is bottom-right.
+	// bAllowCursorOutsideOfParent must be true here... so we need to make sure to snap the target location to be inside the parent widget (e.g. viewport).
+	const FVector2f NormalizedLocalCoords = FVector2f::Clamp(LocalFinalPosition / LocalSize, FVector2f::ZeroVector, FVector2f(1.f, 1.f));
+	const FVector2f OffsetTopLeftCorner = NormalizedLocalCoords * LocalSize;
+	
+	// Whatever is passed to SetBoxAlignmentOffset is relative to the current HAlign and VAlign, which may be different from the top-left atm.
+	DragWidget->SetBoxHorizontalAlignment(HAlign_Left);
+	DragWidget->SetBoxVerticalAlignment(VAlign_Top);
+	DragWidget->SetBoxAlignmentOffset(OffsetTopLeftCorner);
+
+	return true;
+}
+
+void FFlyoutOverlayManager::SaveWidgetState() const
+{
+	SaveWidgetState(DragWidget->GetDragBoxPosition());
+}
+
+void FFlyoutOverlayManager::SaveWidgetState(const FToolWidget_DragBoxPosition& InPosition) const
+{
+	if (Args.SaveStateDelegate.IsBound() && ensure(DragWidget)
+		// Invariant: We don't expect anything to trigger calling SaveWidgetState while position is temporarily moved.
+		&& ensureAlways(!TemporaryMoveOperation))
+	{
+		const FToolWidget_FlyoutSavedState SavedState { InPosition, bShouldShowWidget };
+		Args.SaveStateDelegate.Execute(SavedState);
+	}
+}
+
+bool FFlyoutOverlayManager::CanDragWidget() const
+{
+	return !TemporaryMoveOperation.IsSet();
+}
+
+void FFlyoutOverlayManager::OnMouseEnterContent()
+{
+	UnsubdueContent();
+}
+
+void FFlyoutOverlayManager::OnMouseLeftContent()
+{
+	// If a menu was opened while the mouse hovered the content widget, we'll assume that the menu belongs to our content widget.
+	// We'll defer hiding finishing the operation until the menu is closed.
+	if (!FSlateApplication::Get().AnyMenusVisible())
+	{
+		// The temporary position set by the command is reset once the mouse leaves the widget bounds... 
+		CommandTriggeredPositionOverride.Reset();
+		// ... however CommandTriggeredPositionOverride may not be the reason its subdued because the API user may have started a
+		// TemporaryMoveOperation themselves: this function will check and hide if applicable.
+		TrySubdueContent();
+	}
+}
+
+void FFlyoutOverlayManager::OnWindowDestroyed(const SWindow&)
+{
+	// If a menu was opened by the content widget and the mouse had left the content bounds (OnMouseLeftContent), then we deferred subduing the
+	// content until the menu is closed...
+	if (ContentWidget
+		// ... the menu may be set up to not hide after you click something in it. That gives the user the chance to hover the content again and
+		// e.g. click the combo button that opened the menu...
+		&& !ContentWidget->IsHovered()
+		// ... and here we handle that it's not a sub-menu that was just closed. 
+		&& !FSlateApplication::Get().AnyMenusVisible())
+	{
+		CommandTriggeredPositionOverride.Reset();
+		TrySubdueContent();
+	}
+}
+}
\ No newline at end of file
diff --git a/Engine/Plugins/Animation/ControlRig/Source/ControlRigEditor/Private/Misc/Flyout/FlyoutOverlayManager.h b/Engine/Plugins/Animation/ControlRig/Source/ControlRigEditor/Private/Misc/Flyout/FlyoutOverlayManager.h
new file mode 100644
--- /dev/null
+++ b/Engine/Plugins/Animation/ControlRig/Source/ControlRigEditor/Private/Misc/Flyout/FlyoutOverlayManager.h
@@ -0,0 +1,166 @@
+﻿// Copyright Epic Games, Inc. All Rights Reserved.
+
+#pragma once
+
+#include "FlyoutWidgetArgs.h"
+#include "Templates/UnrealTemplate.h"
+#include "TemporaryPositionOverride.h"
+
+// This code intentionally has no knowledge about Control Rig at all.
+// The plan is to move it out to ToolWidgets if other tools want to adapt a similar style widget - until then we keep it private to enjoy the
+// flexibility of private code.
+// Please DO NOT introduce ControlRig specific code here; instead invert dependencies by injecting Control Rig specific behaviour.
+
+class IMenu;
+
+namespace UE::ToolWidgets { class SDraggableBoxOverlay; }
+
+namespace UE::ControlRigEditor
+{
+namespace Private { class SFlyoutWidgetWrapper; }
+
+enum class ETemporaryFlyoutPositionFlags : uint8
+{
+	None = 0,
+
+	/**
+	 * If set, the mouse cursor position is snapped to the parent widget draw space.
+	 * If unset, and the mouse cursor is outside the parent widget, the operation is skipped (but a set optional is still returned).
+	 */
+	AllowCursorOutsideOfParent = 1 << 0,
+	
+	/** If set, then the widget will be hidden once the flyout position is moved back to its original position. */
+	HideAtEnd = 1 << 1
+};
+ENUM_CLASS_FLAGS(ETemporaryFlyoutPositionFlags);
+	
+/**
+ * Manages all logic around a flyout widget.
+ * 
+ * A flyout widget is designed to be overlaid with another widget, usually the viewport.
+ * The user can reposition it freely in the parent widget. The following functionality is provided:
+ * - Command: Toggle widget visible / hidden
+ * - Command: Summon widget to cursor. Once your mouse leaves the widget bounds, the widget is positioned back to its original position.
+ * - Drag-drop the widget to reposition it. The widget is "docked" to the corner closest to which the widget is dragged.
+ *	Example: If the user drops the widget at the right-bottom corner, and moves the right side of the viewport to the left, this content widget
+ *	maintains the same offset to the right side.
+ * - Subdue widget while not in use:
+ *	- While hovered, the content widget renders at full opacity.
+ *	- While not hovered, the content widget is subdued.
+ *	- This also handles the edge case of the content widget opening a menu and the user unhovering the content widget. No subduing takes place while
+ *	a menu is being shown.
+ */
+class FFlyoutOverlayManager : public FNoncopyable
+{
+	friend class FFlyoutTemporaryPositionOverride;
+public:
+	
+	explicit FFlyoutOverlayManager(FFlyoutWidgetArgs InArgs);
+	~FFlyoutOverlayManager();
+
+	bool IsShowingWidget() const { return bShouldShowWidget; }
+	void ToggleVisibility() { if (IsShowingWidget()) { HideWidget(); } else { ShowWidget(); } }
+	void ShowWidget();
+	void HideWidget();
+
+	/** @return Whether the widget is temporarily repositioned. */
+	bool IsTemporarilyRepositioned() const { return TemporaryMoveOperation.IsSet(); }
+	/**
+	 * Positions the widget at the mouse cursor until the returned FFlyoutTemporaryPositionOverride is destroyed.
+	 *
+	 * If the widget has never been painted, the repositioning happens next tick. In that case this returns a valid FFlyoutTemporaryPositionOverride
+	 * but that may end up being cancelled next tick (e.g. if the cursor is outside of the viewport).
+	 */
+	TSharedPtr<FFlyoutTemporaryPositionOverride> TryTemporarilyPositionWidgetAtCursor(
+		ETemporaryFlyoutPositionFlags InFlags = ETemporaryFlyoutPositionFlags::None
+		);
+	/**
+	 * Positions the widget at the mouse cursor until the cursor leaves the widget bounds and the returned position override goes out of scope.
+	 *
+	 * If the widget has never been painted, the repositioning happens next tick. In that case this returns a valid FFlyoutTemporaryPositionOverride
+	 * but that may end up being cancelled next tick (e.g. if the cursor is outside of the viewport).
+	 */
+	TSharedPtr<FFlyoutTemporaryPositionOverride> SummonToCursorUntilMouseLeave(bool bHideAtEnd);
+	
+	/** Subdues the content if allowed. */
+	void TrySubdueContent(bool bForce = false) const;
+	/** Shows the content normally again. */
+	void UnsubdueContent() const;
+
+	/** Broadcasts when the visibility of the widget changes. */
+	DECLARE_MULTICAST_DELEGATE_OneParam(FOnVisibilityChanged, bool /*bIsVisible*/);
+	FOnVisibilityChanged& OnVisibilityChanged() { return OnVisibilityChangedDelegate; }
+
+private:
+
+	/** The behaviour this widget was initialized with. */
+	const FFlyoutWidgetArgs Args;
+
+	/** Contains the actual content widget. Nullptr until shown for the first time. */
+	TSharedPtr<ToolWidgets::SDraggableBoxOverlay> DragWidget;
+	/** The widget displayed in DragWidget. This widget wraps the API caller's supplied widget. */
+	TSharedPtr<Private::SFlyoutWidgetWrapper> ContentWidget;
+	
+	/** Whether the widget should be shown. */
+	bool bShouldShowWidget = false;
+
+	struct FTemporaryMoveData
+	{
+		/** The position the box had before it was temporarily moved. */
+		const FToolWidget_DragBoxPosition RestorePosition;
+
+		/** Whether to hide the widget after the operation concludes. */
+		const bool bHideWidgetAtEnd;
+
+		explicit FTemporaryMoveData(const FToolWidget_DragBoxPosition& InRestorePosition, bool bHideWidgetAtEnd)
+			: RestorePosition(InRestorePosition), bHideWidgetAtEnd(bHideWidgetAtEnd)
+		{}
+	};
+	/** Set while the content widget is temporarily moved. */
+	TOptional<FTemporaryMoveData> TemporaryMoveOperation;
+	
+	/** Overrides the position when FFlyoutWidgetArgs::SummonToCursorCommand is invoked. Resets once the mouse leaves the widget bounds. */
+	TSharedPtr<FFlyoutTemporaryPositionOverride> CommandTriggeredPositionOverride;
+
+	/** Broadcasts when the visibility of the widget changes. */
+	FOnVisibilityChanged OnVisibilityChangedDelegate;
+
+	void BindCommands();
+	void UnbindCommands() const;
+	void HandleSummonToCursorCommand();
+
+	/** Moves the widget to the cursor until the cursor leaves the widget bounds. After leaving the widget bounds, move the widget back to where it was. */
+	bool TemporarilyMoveWidgetToCursor(ETemporaryFlyoutPositionFlags InFlags);
+	/** Moves the widget back to its original position. */
+	void StopTemporaryWidgetPosition(bool bSkipHideEvenIfRequested = false);
+	
+	/**
+	 * Places the widget at the given position. If the position is outside the viewport bounds, nothing happens.
+	 * 
+	 * @param bAllowCursorOutsideOfParent If true, InAbsoluteScreenPosition is snapped to the parent widget draw space.
+	 *	If false, nothing happens if the user clicks outside.
+	 * @return Whether the widget was moved. 
+	 */
+	bool InternalPositionWidgetAt(const FVector2f& InAbsoluteScreenPosition, bool bAllowCursorOutsideOfParent = false) const;
+	
+	/** Saves the current widget state, if supported. */
+	void SaveWidgetState() const;
+	void SaveWidgetState(const FToolWidget_DragBoxPosition& InPosition) const;
+	/** @return Whether to allow the content widget to be dragged. False while the widget is temporarily moved. */
+	bool CanDragWidget() const;
+
+	/** Invokes when the mouse hovers the content area. */
+	void OnMouseEnterContent();
+	/** Invoked when the mouse has left the content area. */
+	void OnMouseLeftContent();
+
+	/**
+	 * Used to detect when menu window is destroyed, so we can subdue the content widget.
+	 * 
+	 * This is called shortly after FSlateApplication::OnMenuBeingDestroyed; during OnMenuBeingDestroyed, FSlateApplication::AnyMenusVisible still
+	 * returns true broadcast timing. OnWindowDestroyed is called right after, when AnyMenusVisible returns the updated value.
+	 */
+	void OnWindowDestroyed(const SWindow&);
+};
+}
+
diff --git a/Engine/Plugins/Animation/ControlRig/Source/ControlRigEditor/Private/Misc/Flyout/FlyoutSavedState.h b/Engine/Plugins/Animation/ControlRig/Source/ControlRigEditor/Private/Misc/Flyout/FlyoutSavedState.h
new file mode 100644
--- /dev/null
+++ b/Engine/Plugins/Animation/ControlRig/Source/ControlRigEditor/Private/Misc/Flyout/FlyoutSavedState.h
@@ -0,0 +1,29 @@
+﻿// Copyright Epic Games, Inc. All Rights Reserved.
+
+#pragma once
+
+#include "Delegates/Delegate.h"
+#include "Overlay/DragBoxPosition.h"
+#include "FlyoutSavedState.generated.h"
+
+/** Holds information to save the state of the flyout manager. This is useful for saving in a config file. */
+USTRUCT()
+struct FToolWidget_FlyoutSavedState
+{
+	GENERATED_BODY()
+	
+	/** Where the widget was docked. */
+	UPROPERTY()
+	FToolWidget_DragBoxPosition Position;
+
+	/** Whether to show the widget automatically. */
+	UPROPERTY()
+	bool bWasVisible = false;
+};
+
+namespace UE::ControlRigEditor
+{
+using FFlyoutSavedState = FToolWidget_FlyoutSavedState;
+DECLARE_DELEGATE_OneParam(FSaveFlyoutState, FFlyoutSavedState);
+}
+
diff --git a/Engine/Plugins/Animation/ControlRig/Source/ControlRigEditor/Private/Misc/Flyout/FlyoutWidgetArgs.h b/Engine/Plugins/Animation/ControlRig/Source/ControlRigEditor/Private/Misc/Flyout/FlyoutWidgetArgs.h
new file mode 100644
--- /dev/null
+++ b/Engine/Plugins/Animation/ControlRig/Source/ControlRigEditor/Private/Misc/Flyout/FlyoutWidgetArgs.h
@@ -0,0 +1,92 @@
+﻿// Copyright Epic Games, Inc. All Rights Reserved.
+
+#pragma once
+
+#include "FlyoutSavedState.h"
+#include "Framework/Commands/UICommandList.h"
+#include "Templates/UnrealTemplate.h"
+#include "Tools/AssetDefinition_ControlRigPose.h"
+
+class FUICommandInfo;
+class FUICommandList;
+
+// This code intentionally has NO knowledge about Control Rig at all. The plan is to possible move it out to ToolWidgets.
+// DO NOT introduce ControlRig specific code here (placing it in the ControlRigEditor namespace for now is the only exception as the code standard enforces it).
+
+namespace UE::ControlRigEditor
+{
+class IWidgetHost;
+using FGetFlyoutWidgetContent = FOnGetContent;
+	
+struct FFlyoutWidgetArgs
+{
+	explicit FFlyoutWidgetArgs(
+		TAttribute<TSharedPtr<FUICommandList>> InCommandList,
+		FGetFlyoutWidgetContent InGetWidgetDelegate,
+		const TSharedRef<IWidgetHost>& InWidgetHost
+		)
+		: CommandList(MoveTemp(InCommandList))
+		, GetWidgetDelegate(MoveTemp(InGetWidgetDelegate))
+		, WidgetHost(InWidgetHost)
+	{
+		check(GetWidgetDelegate.IsBound());
+	}
+	
+	explicit FFlyoutWidgetArgs(
+		const TSharedRef<SWidget>& InContentWidget,
+		TAttribute<TSharedPtr<FUICommandList>> InCommandList,
+		const TSharedRef<IWidgetHost>& InWidgetHost
+		)
+		: FFlyoutWidgetArgs(
+			MoveTemp(InCommandList),
+			FGetFlyoutWidgetContent::CreateLambda([InContentWidget]{ return InContentWidget; }),
+			InWidgetHost
+			)
+	{}
+
+	FFlyoutWidgetArgs& SetStateToRestoreFrom(FFlyoutSavedState InState) { StateToRestoreFrom = InState; return *this; }
+	FFlyoutWidgetArgs& SetSaveStateDelegate(FSaveFlyoutState InDelegate) { SaveStateDelegate = MoveTemp(InDelegate); return *this; }
+	FFlyoutWidgetArgs& SetCanSubdue(TAttribute<bool> InCanSubdueAttr) { CanSubdueAttr = MoveTemp(InCanSubdueAttr); return *this; }
+	FFlyoutWidgetArgs& SetSubduedOpacity(TAttribute<FLinearColor> InTintAttr) { SubduedTintAttr = MoveTemp(InTintAttr); return *this; }
+	FFlyoutWidgetArgs& SetAbsOffsetFromCursor(TAttribute<FVector2f> InOffsetAttr) { AbsOffsetFromCursorAttr = MoveTemp(InOffsetAttr); return *this; }
+	FFlyoutWidgetArgs& SetToggleVisibility(TSharedPtr<FUICommandInfo> InCommand) { ToggleVisibilityCommand = MoveTemp(InCommand); return *this; }
+	FFlyoutWidgetArgs& SetSummonToCursor(TSharedPtr<FUICommandInfo> InCommand) { SummonToCursorCommand = MoveTemp(InCommand); return *this; }
+
+	/** Command list that the commands should be bound to. Required. */
+	const TAttribute<TSharedPtr<FUICommandList>> CommandList;
+	/** Creates the widget that is supposed to be displayed. Required. */
+	const FGetFlyoutWidgetContent GetWidgetDelegate;
+	/** Handles adding the widget to some UI hierarchy. Required. */
+	const TSharedPtr<IWidgetHost> WidgetHost;
+
+	/** The initial state that should be applied. You usually pass what you have saved in a user config to this. Optional. */
+	TOptional<FFlyoutSavedState> StateToRestoreFrom;
+	/**
+	 * Saves the state of the flyout widget so it can be restored at a later time, e.g. after an editor restart, etc.
+	 * You usually pass the saved state to RestoredFromState. Optional.
+	 */
+	FSaveFlyoutState SaveStateDelegate;
+
+	/**
+	 * Decides whether the widget can be subdued, i.e. render it so it is less noticeable to the user. Optional.
+	 * 
+	 * This is only invoked when the widget determines now is a good time be hidden, e.g. when the cursor isn't near it.
+	 * In other words: simply returning true does not mean it will be subdued - it just indicates whether right now it is appropriate or not.
+	 */
+	TAttribute<bool> CanSubdueAttr = true;
+	/** The render opacity to set when the widget is out of focus. */
+	TAttribute<FLinearColor> SubduedTintAttr = FLinearColor(1.f, 1.f, 1.f, .35f);
+
+	/**
+	 * By default, the center of the widget is placed at the mouse cursor.
+	 * This is an extra offset that is added to the absolute cursor position. The widget's center is positioned at this final transformed point.
+	 */
+	TAttribute<FVector2f> AbsOffsetFromCursorAttr = FVector2f::ZeroVector;
+
+	/** Shows or hides the widget. Optional.*/
+	TSharedPtr<FUICommandInfo> ToggleVisibilityCommand;
+	/** Summons the widget to the cursor. Once the cursor moves out of the widget bounds, it is placed back ast its original position. Optional.*/
+	TSharedPtr<FUICommandInfo> SummonToCursorCommand;
+};
+}
+
diff --git a/Engine/Plugins/Animation/ControlRig/Source/ControlRigEditor/Private/Misc/Flyout/Host/IWidgetHost.h b/Engine/Plugins/Animation/ControlRig/Source/ControlRigEditor/Private/Misc/Flyout/Host/IWidgetHost.h
new file mode 100644
--- /dev/null
+++ b/Engine/Plugins/Animation/ControlRig/Source/ControlRigEditor/Private/Misc/Flyout/Host/IWidgetHost.h
@@ -0,0 +1,24 @@
+﻿// Copyright Epic Games, Inc. All Rights Reserved.
+
+#pragma once
+
+#include "Templates/SharedPointerFwd.h"
+
+class SWidget;
+
+namespace UE::ControlRigEditor
+{
+/** Hosts a widget in some UI. */
+class IWidgetHost
+{
+public:
+
+	/** Adds the widget to the target UI. The implementation is allowed to keep a strong reference to the widget until RemoveWidgetFromHost is called. */
+	virtual void AddWidgetToHost(const TSharedRef<SWidget>& Widget) = 0;
+
+	/** Removes the widget from the target UI and clears and strong reference to it. */
+	virtual void RemoveWidgetFromHost() = 0;
+
+	virtual ~IWidgetHost() = default;
+};
+}
diff --git a/Engine/Plugins/Animation/ControlRig/Source/ControlRigEditor/Private/Misc/Flyout/Host/ToolkitBasedWidgetHost.cpp b/Engine/Plugins/Animation/ControlRig/Source/ControlRigEditor/Private/Misc/Flyout/Host/ToolkitBasedWidgetHost.cpp
new file mode 100644
--- /dev/null
+++ b/Engine/Plugins/Animation/ControlRig/Source/ControlRigEditor/Private/Misc/Flyout/Host/ToolkitBasedWidgetHost.cpp
@@ -0,0 +1,44 @@
+﻿// Copyright Epic Games, Inc. All Rights Reserved.
+
+#include "ToolkitBasedWidgetHost.h"
+
+#include "Toolkits/IToolkitHost.h"
+
+namespace UE::ControlRigEditor
+{
+FToolkitBasedWidgetHost::FToolkitBasedWidgetHost(const TSharedRef<IToolkitHost>& InToolkitHost)
+	: ToolkitHost(InToolkitHost)
+{
+	ToolkitHost->OnActiveViewportChanged().AddRaw(this, &FToolkitBasedWidgetHost::OnActiveViewportChanged);
+}
+
+FToolkitBasedWidgetHost::~FToolkitBasedWidgetHost()
+{
+	ToolkitHost->OnActiveViewportChanged().RemoveAll(this);
+}
+
+void FToolkitBasedWidgetHost::AddWidgetToHost(const TSharedRef<SWidget>& Widget)
+{
+	ensure(!HostedWidget.IsValid());
+	HostedWidget = Widget;
+	ToolkitHost->AddViewportOverlayWidget(Widget);
+}
+
+void FToolkitBasedWidgetHost::RemoveWidgetFromHost()
+{
+	if (ensure(HostedWidget))
+	{
+		ToolkitHost->RemoveViewportOverlayWidget(HostedWidget.ToSharedRef());
+	}
+	HostedWidget.Reset();
+}
+
+void FToolkitBasedWidgetHost::OnActiveViewportChanged(TSharedPtr<IAssetViewport> OldViewport, TSharedPtr<IAssetViewport> NewViewport) const
+{
+	if (HostedWidget)
+	{
+		ToolkitHost->RemoveViewportOverlayWidget(HostedWidget.ToSharedRef(), OldViewport);
+		ToolkitHost->AddViewportOverlayWidget(HostedWidget.ToSharedRef(), NewViewport);
+	}
+}
+}
diff --git a/Engine/Plugins/Animation/ControlRig/Source/ControlRigEditor/Private/Misc/Flyout/Host/ToolkitBasedWidgetHost.h b/Engine/Plugins/Animation/ControlRig/Source/ControlRigEditor/Private/Misc/Flyout/Host/ToolkitBasedWidgetHost.h
new file mode 100644
--- /dev/null
+++ b/Engine/Plugins/Animation/ControlRig/Source/ControlRigEditor/Private/Misc/Flyout/Host/ToolkitBasedWidgetHost.h
@@ -0,0 +1,38 @@
+﻿// Copyright Epic Games, Inc. All Rights Reserved.
+
+#pragma once
+
+#include "IWidgetHost.h"
+#include "Templates/SharedPointer.h"
+
+class IAssetViewport;
+class IToolkitHost;
+class SWidget;
+
+namespace UE::ControlRigEditor
+{
+/** Adapter which hosts a widget in the IToolkitHost's active viewport. When the viewport changes, the widget is moved to the active viewport. */
+class FToolkitBasedWidgetHost : public IWidgetHost
+{
+public:
+
+	explicit FToolkitBasedWidgetHost(const TSharedRef<IToolkitHost>& InToolkitHost);
+	virtual ~FToolkitBasedWidgetHost() override;
+
+	//~ Begin IWidgetHost Interface
+	virtual void AddWidgetToHost(const TSharedRef<SWidget>& Widget) override;
+	virtual void RemoveWidgetFromHost() override;
+	//~ End IWidgetHost Interface
+
+private:
+
+	/** The toolkit host in whicht the widget is hosted. */
+	TSharedRef<IToolkitHost> ToolkitHost;
+
+	/** The widget that is hosted. Set by AddWidgetToHost. */
+	TSharedPtr<SWidget> HostedWidget;
+
+	/** Moves HostedWidget to the new viewport. */
+	void OnActiveViewportChanged(TSharedPtr<IAssetViewport> OldViewport, TSharedPtr<IAssetViewport> NewViewport) const;
+};
+}
diff --git a/Engine/Plugins/Animation/ControlRig/Source/ControlRigEditor/Private/Misc/Flyout/TemporaryPositionOverride.cpp b/Engine/Plugins/Animation/ControlRig/Source/ControlRigEditor/Private/Misc/Flyout/TemporaryPositionOverride.cpp
new file mode 100644
--- /dev/null
+++ b/Engine/Plugins/Animation/ControlRig/Source/ControlRigEditor/Private/Misc/Flyout/TemporaryPositionOverride.cpp
@@ -0,0 +1,26 @@
+﻿// Copyright Epic Games, Inc. All Rights Reserved.
+
+#include "TemporaryPositionOverride.h"
+
+#include "FlyoutOverlayManager.h"
+
+namespace UE::ControlRigEditor
+{
+FFlyoutTemporaryPositionOverride::FFlyoutTemporaryPositionOverride(FPrivateToken, FFlyoutOverlayManager& InManager)
+	: ManagerToRestore(&InManager)
+	{}
+	
+FFlyoutTemporaryPositionOverride::FFlyoutTemporaryPositionOverride(FFlyoutTemporaryPositionOverride&& Other)
+	: ManagerToRestore(Other.ManagerToRestore)
+{
+	Other.Cancel();
+}
+
+FFlyoutTemporaryPositionOverride::~FFlyoutTemporaryPositionOverride()
+{
+	if (ManagerToRestore)
+	{
+		ManagerToRestore->StopTemporaryWidgetPosition();
+	}
+}
+}
diff --git a/Engine/Plugins/Animation/ControlRig/Source/ControlRigEditor/Private/Misc/Flyout/TemporaryPositionOverride.h b/Engine/Plugins/Animation/ControlRig/Source/ControlRigEditor/Private/Misc/Flyout/TemporaryPositionOverride.h
new file mode 100644
--- /dev/null
+++ b/Engine/Plugins/Animation/ControlRig/Source/ControlRigEditor/Private/Misc/Flyout/TemporaryPositionOverride.h
@@ -0,0 +1,32 @@
+﻿// Copyright Epic Games, Inc. All Rights Reserved.
+
+#pragma once
+
+#include "FlyoutOverlayManager.h"
+#include "Templates/UnrealTemplate.h"
+
+namespace UE::ControlRigEditor
+{
+class FFlyoutOverlayManager;
+
+/** Keeps the flyout widget at the temporary override position until this object is destroyed, at which point it calls StopTemporaryWidgetPosition. */
+class FFlyoutTemporaryPositionOverride : public FNoncopyable
+{
+	friend FFlyoutOverlayManager;
+	struct FPrivateToken { explicit FPrivateToken() = default; };
+	
+	FFlyoutOverlayManager* ManagerToRestore;
+	
+
+	void Cancel() { ManagerToRestore = nullptr; }
+	bool IsCancelled() const { return ManagerToRestore == nullptr; }
+
+public:
+	
+	FFlyoutTemporaryPositionOverride(FPrivateToken, FFlyoutOverlayManager& InManager);
+	FFlyoutTemporaryPositionOverride(FFlyoutTemporaryPositionOverride&& Other);
+	FFlyoutTemporaryPositionOverride& operator=(FFlyoutTemporaryPositionOverride&& Other) = delete;
+
+	~FFlyoutTemporaryPositionOverride();
+};
+}
\ No newline at end of file
diff --git a/Engine/Plugins/Animation/TweeningUtils/Source/TweeningUtilsEditor/Private/Widgets/MVC/CycleFunctionController.cpp b/Engine/Plugins/Animation/TweeningUtils/Source/TweeningUtilsEditor/Private/Widgets/MVC/CycleFunctionController.cpp
--- a/Engine/Plugins/Animation/TweeningUtils/Source/TweeningUtilsEditor/Private/Widgets/MVC/CycleFunctionController.cpp
+++ b/Engine/Plugins/Animation/TweeningUtils/Source/TweeningUtilsEditor/Private/Widgets/MVC/CycleFunctionController.cpp
@@ -44,13 +44,14 @@ FCycleFunctionController::~FCycleFunctionController()
 	CommandList->UnmapAction(CycleCommand);
 }
 
-void FCycleFunctionController::CycleToNextFunction()
+void FCycleFunctionController::CycleToNextFunction() const
 {
 	const int32 Index = TweenModelContainer->IndexOf(*CurrentTweenModelAttr.Get());
 	if (ensure(TweenModelContainer->IsValidIndex(Index)))
 	{
 		const int32 NextIndex = TweenModelContainer->IsValidIndex(Index + 1) ? Index + 1 : 0;
 		HandleTweenChangeDelegate.Execute(*TweenModelContainer->GetModel(NextIndex));
+		OnFunctionCycleCommandInvokedDelegate.Broadcast();
 	}
 }
 }
diff --git a/Engine/Plugins/Animation/TweeningUtils/Source/TweeningUtilsEditor/Private/Widgets/MVC/MouseSlidingController.cpp b/Engine/Plugins/Animation/TweeningUtils/Source/TweeningUtilsEditor/Private/Widgets/MVC/MouseSlidingController.cpp
--- a/Engine/Plugins/Animation/TweeningUtils/Source/TweeningUtilsEditor/Private/Widgets/MVC/MouseSlidingController.cpp
+++ b/Engine/Plugins/Animation/TweeningUtils/Source/TweeningUtilsEditor/Private/Widgets/MVC/MouseSlidingController.cpp
@@ -127,16 +127,31 @@ FMouseSlidingController::FMouseSlidingController(
 {
 	check(InMaxSlideWidthAttr.IsSet() || InMaxSlideWidthAttr.IsBound());
 	check(DragSliderCommand);
-	
-	CommandList->MapAction(
-		DragSliderCommand, FExecuteAction::CreateRaw(this, &FMouseSlidingController::StartListeningForMouseEvents)
-		);
+	BindCommand();
 }
 
 FMouseSlidingController::~FMouseSlidingController()
 {
-	CommandList->UnmapAction(DragSliderCommand);
 	SlidingState.Reset();
+	UnbindCommand();
+}
+
+void FMouseSlidingController::BindCommand()
+{
+	if (!CommandList->IsActionMapped(DragSliderCommand))
+	{
+		CommandList->MapAction(
+			DragSliderCommand, FExecuteAction::CreateRaw(this, &FMouseSlidingController::StartListeningForMouseEvents)
+		);
+	}
+}
+
+void FMouseSlidingController::UnbindCommand() const
+{
+	if (CommandList->IsActionMapped(DragSliderCommand))
+	{
+		CommandList->UnmapAction(DragSliderCommand);
+	}
 }
 
 FMouseSlidingController::FSlidingState::FSlidingState(FMouseSlidingController& InOwner)
@@ -163,7 +178,6 @@ void FMouseSlidingController::StartListeningForMouseEvents()
 	if (!SlidingState)
 	{
 		SlidingState.Emplace(*this);
-		OnStartSlidingDelegate.Broadcast();
 	}
 }
 
diff --git a/Engine/Plugins/Animation/TweeningUtils/Source/TweeningUtilsEditor/Private/Widgets/MVC/STweenView.cpp b/Engine/Plugins/Animation/TweeningUtils/Source/TweeningUtilsEditor/Private/Widgets/MVC/STweenView.cpp
--- a/Engine/Plugins/Animation/TweeningUtils/Source/TweeningUtilsEditor/Private/Widgets/MVC/STweenView.cpp
+++ b/Engine/Plugins/Animation/TweeningUtils/Source/TweeningUtilsEditor/Private/Widgets/MVC/STweenView.cpp
@@ -26,7 +26,7 @@ void STweenView::Construct(const FArguments& InArgs)
 		.HAlign(HAlign_Fill)
 		.VAlign(VAlign_Center) 
 		[
-			SNew(STweenSlider)
+			SAssignNew(TweenSlider, STweenSlider)
 			.SliderIcon(InArgs._SliderIcon)
 			.SliderColor(InArgs._SliderColor)
 			.OverrideSliderPosition(InArgs._OverrideSliderPosition)
diff --git a/Engine/Plugins/Animation/TweeningUtils/Source/TweeningUtilsEditor/Private/Widgets/MVC/TweenMouseSlidingController.cpp b/Engine/Plugins/Animation/TweeningUtils/Source/TweeningUtilsEditor/Private/Widgets/MVC/TweenMouseSlidingController.cpp
--- a/Engine/Plugins/Animation/TweeningUtils/Source/TweeningUtilsEditor/Private/Widgets/MVC/TweenMouseSlidingController.cpp
+++ b/Engine/Plugins/Animation/TweeningUtils/Source/TweeningUtilsEditor/Private/Widgets/MVC/TweenMouseSlidingController.cpp
@@ -7,6 +7,7 @@
 #include "Misc/Attribute.h"
 #include "TweeningUtilsCommands.h"
 #include "TweeningUtilsStyle.h"
+#include "Framework/Application/SlateApplication.h"
 #include "Widgets/TweenSliderStyle.h"
 
 namespace UE::TweeningUtilsEditor
@@ -40,7 +41,22 @@ FTweenMouseSlidingController::FTweenMouseSlidingController(
 
 FTweenMouseSlidingController::FTweenMouseSlidingController(TAttribute<FTweenModel*> InTweenModelAttr, const TSharedRef<FUICommandList>& InCommandList)
 	: FTweenMouseSlidingController(
-		FTweeningUtilsStyle::Get().GetWidgetStyle<FTweenSliderStyle>("TweenSlider").BarDimensions.X,
+		[]
+		{
+			// This FTweenMouseSlidingController constructor overload is for convenience and assumes default values... API user can still use the
+			// other overload if their icon size is different.
+			constexpr float DefaultIconWidth = 20.f;
+			
+			const FTweenSliderStyle TweenStyle = FTweeningUtilsStyle::Get().GetWidgetStyle<FTweenSliderStyle>("TweenSlider");
+			const float ButtonSize = TweenStyle.IconPadding.Left + DefaultIconWidth + TweenStyle.IconPadding.Right;
+			// The bar is padded with half the button width on each side. So the total size is the bar plus the button brush size.
+			const float WidgetWidth = TweenStyle.BarDimensions.X + ButtonSize;
+
+			// The widget draw size needs to be multiplied by DPI to obtain the size the widget would have on screen.
+			const TSharedPtr<SWindow> TopLevelWindow = FSlateApplication::Get().GetActiveTopLevelWindow();
+			const float ScaleDPI = TopLevelWindow ? TopLevelWindow->GetDPIScaleFactor() : 1.f;
+			return WidgetWidth * ScaleDPI;
+		}(),
 		MoveTemp(InTweenModelAttr),
 		InCommandList,
 		FTweeningUtilsCommands::Get().DragAnimSliderTool
diff --git a/Engine/Plugins/Animation/TweeningUtils/Source/TweeningUtilsEditor/Private/Widgets/MVC/TweenToolbarController.cpp b/Engine/Plugins/Animation/TweeningUtils/Source/TweeningUtilsEditor/Private/Widgets/MVC/TweenToolbarController.cpp
--- a/Engine/Plugins/Animation/TweeningUtils/Source/TweeningUtilsEditor/Private/Widgets/MVC/TweenToolbarController.cpp
+++ b/Engine/Plugins/Animation/TweeningUtils/Source/TweeningUtilsEditor/Private/Widgets/MVC/TweenToolbarController.cpp
@@ -101,8 +101,10 @@ FTweenToolbarController::~FTweenToolbarController()
 	}
 }
 
-void FTweenToolbarController::AddToToolbar(FToolBarBuilder& ToolBarBuilder, FMakeWidgetArgs InArgs) const
+FTweenToolbarController::FAddToToolbarResult FTweenToolbarController::AddToToolbar(FToolBarBuilder& ToolBarBuilder, FMakeWidgetArgs InArgs) const
 {
+	FAddToToolbarResult Result;
+	
 	ToolBarBuilder.SetLabelVisibility(EVisibility::Visible);
 	ToolBarBuilder.AddComboButton(
 	   FUIAction(),
@@ -113,21 +115,22 @@ void FTweenToolbarController::AddToToolbar(FToolBarBuilder& ToolBarBuilder, FMak
 	   false, NAME_None, {}, {}, {}, EUserInterfaceActionType::Button, InArgs.FunctionSelectResizeParams
 	   );
 	ToolBarBuilder.SetLabelVisibility(EVisibility::Collapsed);
-		
-	ToolBarBuilder.AddWidget(
-		SNew(STweenView)
+
+	const TSharedRef<STweenView> View = SNew(STweenView)
 		.TweenModel(MakeSelectedTweenModelAttr())
 		.SliderIcon(TAttribute<const FSlateBrush*>::CreateRaw(this, &FTweenToolbarController::GetIconForSlider))
 		.SliderColor(TAttribute<FLinearColor>::CreateRaw(this, &FTweenToolbarController::GetColorForSlider))
-		.OverrideSliderPosition(MoveTemp(InArgs.OverrideSliderPositionAttr)),
-		{}, NAME_None, true, {}, {}, InArgs.SliderResizeParams
-		);
+		.OverrideSliderPosition(MoveTemp(InArgs.OverrideSliderPositionAttr));
+	ToolBarBuilder.AddWidget(View, {}, NAME_None, true, {}, {}, InArgs.SliderResizeParams);
+	Result.TweenSlider = View->GetTweenSlider();
 	
 	ToolBarBuilder.AddToolBarButton(FTweeningUtilsCommands::Get().ToggleOvershootMode,
 		NAME_None, TAttribute<FText>(), TAttribute<FText>(),
 		TAttribute<FSlateIcon>::CreateRaw(this, &FTweenToolbarController::GetOvershootModeIcon),
 		NAME_None, {}, {}, {}, InArgs.OvershootResizeParams
 		);
+
+	return Result;
 }
 
 void FTweenToolbarController::SetSelectedTweenModel(FTweenModel& InTweenModel)
@@ -184,6 +187,8 @@ void FTweenToolbarController::ToggleOvershootMode() const
 			InTweenModel.GetScaleMode() == ETweenScaleMode::Normalized ? ETweenScaleMode::Overshoot : ETweenScaleMode::Normalized
 			);
 	});
+
+	OnOvershootModeCommandInvokedDelegate.Broadcast();
 }
 
 bool FTweenToolbarController::IsOvershootModeEnabled() const
diff --git a/Engine/Plugins/Animation/TweeningUtils/Source/TweeningUtilsEditor/Private/Widgets/STweenSlider.cpp b/Engine/Plugins/Animation/TweeningUtils/Source/TweeningUtilsEditor/Private/Widgets/STweenSlider.cpp
--- a/Engine/Plugins/Animation/TweeningUtils/Source/TweeningUtilsEditor/Private/Widgets/STweenSlider.cpp
+++ b/Engine/Plugins/Animation/TweeningUtils/Source/TweeningUtilsEditor/Private/Widgets/STweenSlider.cpp
@@ -77,12 +77,19 @@ static float AsymmetricToSymmetric(float InValue) { return InValue * 2.0 - 1.0;
 /** Converts [-1,1] to [0,1] */
 static float SymmetricToAsymmetric(float InValue) { return (InValue + 1) / 2.0; }
 	
-static float GetSliderPosition(const FTweenWidgetArgs& InWidgetArgs, float InFallback)
+static float ReconcileSliderPositionOverride(const FTweenWidgetArgs& InWidgetArgs, float InFallback, FTweenSliderDrawArgs& InArgs)
 {
 	const bool bHasOverride = InWidgetArgs.OverrideSliderPositionAttr.IsBound() || InWidgetArgs.OverrideSliderPositionAttr.IsSet();
-	const TOptional<float> OverrideValue = bHasOverride
-		? InWidgetArgs.OverrideSliderPositionAttr.Get() : TOptional<float>();
-	return OverrideValue ? SymmetricToAsymmetric(*OverrideValue) : InFallback;
+	const TOptional<float> OverrideValue = bHasOverride ? InWidgetArgs.OverrideSliderPositionAttr.Get() : TOptional<float>();
+	if (OverrideValue.IsSet())
+	{
+		// While overriding the slider value, something external is driving it... we want to draw its indicators as if it was being driven by mouse.
+		InArgs.bDrawButtonPressed = true;
+		InArgs.bIsDragging = true;
+		InArgs.HoverState.bIsSliderHovered = true;
+		return SymmetricToAsymmetric(*OverrideValue);
+	}
+	return InFallback;
 }
 
 	
@@ -293,21 +300,21 @@ int32 STweenSlider::OnPaint(
 	) const
 {
 	const bool bIsMouseButtonDown = IsMouseDown();
-	const bool bIsDragging = IsDragging();
 	FTweenWidgetArgs& WidgetArgs = Pimpl->WidgetArgs;
 	FTweenSliderHoverState& HoverState = Pimpl->HoverState;
 	
 	TArray<float> NormalizedPositions;
 	FTweenSliderDrawArgs DrawArgs;
 	
-	DrawArgs.bDrawButtonPressed = bIsDragging
+	DrawArgs.bIsDragging = IsDragging();
+	DrawArgs.bDrawButtonPressed = DrawArgs.bIsDragging
 		|| (bIsMouseButtonDown && HoverState.bIsSliderHovered)
 		|| (bIsMouseButtonDown && HoverState.HoveredPointIndex);
-	DrawArgs.bIsDragging = bIsDragging;
 	DrawArgs.HoverState = HoverState;
 	
 	const float SliderPosition = DrawArgs.bDrawButtonPressed
-		? Pimpl->AnimatedCurrentSliderPosition : GetSliderPosition(WidgetArgs, Pimpl->AnimatedCurrentSliderPosition);
+		? Pimpl->AnimatedCurrentSliderPosition : ReconcileSliderPositionOverride(WidgetArgs, Pimpl->AnimatedCurrentSliderPosition, DrawArgs);
+	
 	GetBarGeometry(InAllottedGeometry, WidgetArgs, DrawArgs.BarArea);
 	GetSliderButtonGeometry(
 		SliderPosition, InAllottedGeometry, WidgetArgs, DrawArgs.SliderArea, DrawArgs.IconArea
@@ -316,10 +323,10 @@ int32 STweenSlider::OnPaint(
 		InAllottedGeometry, WidgetArgs, HoverState, bIsMouseButtonDown, DrawArgs.Points, DrawArgs.PointTypes, NormalizedPositions
 		);
 	
-	if (bIsDragging)
+	if (DrawArgs.bIsDragging)
 	{
-		GetPassedPointStates(NormalizedPositions, Pimpl->AnimatedCurrentSliderPosition, DrawArgs.PassedPoints);
-		GetDragValueIndicationGeometry(InAllottedGeometry, WidgetArgs, Pimpl->AnimatedCurrentSliderPosition, DrawArgs.DragValueIndication);
+		GetPassedPointStates(NormalizedPositions, SliderPosition, DrawArgs.PassedPoints);
+		GetDragValueIndicationGeometry(InAllottedGeometry, WidgetArgs, SliderPosition, DrawArgs.DragValueIndication);
 	}
 	
 	return DrawTweenSlider(DrawArgs, WidgetArgs, OutDrawElements, InLayerId, InWidgetStyle);
diff --git a/Engine/Plugins/Animation/TweeningUtils/Source/TweeningUtilsEditor/Private/Widgets/TweenSliderDrawUtils.cpp b/Engine/Plugins/Animation/TweeningUtils/Source/TweeningUtilsEditor/Private/Widgets/TweenSliderDrawUtils.cpp
--- a/Engine/Plugins/Animation/TweeningUtils/Source/TweeningUtilsEditor/Private/Widgets/TweenSliderDrawUtils.cpp
+++ b/Engine/Plugins/Animation/TweeningUtils/Source/TweeningUtilsEditor/Private/Widgets/TweenSliderDrawUtils.cpp
@@ -273,7 +273,9 @@ static void DrawBar(
 	const FTweenSliderStyle* Style = InWidgetArgs.Style;
 	
 	const FSlateBrush* BarBrush = &Style->BarBrush;
-	const FLinearColor FinalBackgroundColor = InWidgetArgs.ColorAndOpacity.Get().GetColor(InWidgetStyle) * BarBrush->TintColor.GetColor(InWidgetStyle);
+	const FLinearColor FinalBackgroundColor = InWidgetArgs.ColorAndOpacity.Get().GetColor(InWidgetStyle)
+		* BarBrush->TintColor.GetColor(InWidgetStyle)
+		* InWidgetStyle.GetColorAndOpacityTint();
 	FSlateDrawElement::MakeBox(
 		OutDrawElements, LayerId, InDrawArgs.BarArea.ToPaintGeometry(), BarBrush, ESlateDrawEffect::None, FinalBackgroundColor
 		);
@@ -296,7 +298,8 @@ static void DrawDragValueIndication(
 	const FLinearColor BarColor =
 		InWidgetArgs.ColorAndOpacity.Get().GetColor(InWidgetStyle)
 		* InWidgetArgs.SliderColor.Get()
-		* Style->PassedValueBackground.TintColor.GetColor(InWidgetStyle);
+		* Style->PassedValueBackground.TintColor.GetColor(InWidgetStyle)
+		* InWidgetStyle.GetColorAndOpacityTint();
 	FSlateDrawElement::MakeBox(
 		OutDrawElements, LayerId, InDrawArgs.DragValueIndication.ToPaintGeometry(), &Style->PassedValueBackground, ESlateDrawEffect::None, BarColor
 		);
@@ -350,7 +353,8 @@ static void DrawBarPoints(
 		
 		const FLinearColor PointColor = UseSliderColor == EUseSliderColor::Yes
 			? ColorAndOpacity * SliderColor * Brush->TintColor.GetColor(InWidgetStyle)
-			: ColorAndOpacity * Brush->TintColor.GetColor(InWidgetStyle);
+			: ColorAndOpacity * Brush->TintColor.GetColor(InWidgetStyle)
+			* InWidgetStyle.GetColorAndOpacityTint();
 		FSlateDrawElement::MakeBox(
 			OutDrawElements, LayerId, PointGeometry.ToPaintGeometry(), Brush, ESlateDrawEffect::None, PointColor
 		);
@@ -377,14 +381,20 @@ static void DrawButton(
 		: InDrawArgs.HoverState.bIsSliderHovered ? Style->HoveredIconTint : Style->NormalIconTint;
 	
 	const FLinearColor ColorAndOpacity = InWidgetArgs.ColorAndOpacity.Get().GetColor(InWidgetStyle);
-	const FLinearColor ButtonTint = ColorAndOpacity * InWidgetArgs.SliderColor.Get() * SliderButtonBrush->TintColor.GetColor(InWidgetStyle);
+	const FLinearColor ButtonTint = ColorAndOpacity
+		* InWidgetArgs.SliderColor.Get()
+		* SliderButtonBrush->TintColor.GetColor(InWidgetStyle) 
+		* InWidgetStyle.GetColorAndOpacityTint();;
 	
 	FSlateDrawElement::MakeBox(
 		OutDrawElements, LayerId, InDrawArgs.SliderArea.ToPaintGeometry(), SliderButtonBrush, ESlateDrawEffect::None, ButtonTint
 		);
 	if (SliderIconBrush)
 	{
-		const FLinearColor IconTint = ColorAndOpacity * BaseIconTint.GetColor(InWidgetStyle) * SliderIconBrush->TintColor.GetColor(InWidgetStyle);
+		const FLinearColor IconTint = ColorAndOpacity
+			* BaseIconTint.GetColor(InWidgetStyle)
+			* SliderIconBrush->TintColor.GetColor(InWidgetStyle)
+			* InWidgetStyle.GetColorAndOpacityTint();
 		FSlateDrawElement::MakeBox(
 			OutDrawElements, LayerId + 1, InDrawArgs.IconArea.ToPaintGeometry(), SliderIconBrush, ESlateDrawEffect::None, IconTint
 		);
diff --git a/Engine/Plugins/Animation/TweeningUtils/Source/TweeningUtilsEditor/Public/Widgets/MVC/CycleFunctionController.h b/Engine/Plugins/Animation/TweeningUtils/Source/TweeningUtilsEditor/Public/Widgets/MVC/CycleFunctionController.h
--- a/Engine/Plugins/Animation/TweeningUtils/Source/TweeningUtilsEditor/Public/Widgets/MVC/CycleFunctionController.h
+++ b/Engine/Plugins/Animation/TweeningUtils/Source/TweeningUtilsEditor/Public/Widgets/MVC/CycleFunctionController.h
@@ -46,6 +46,10 @@ public:
 		const TSharedRef<FUICommandList>& InCommandList
 		); 
 	UE_API virtual ~FCycleFunctionController();
+
+	DECLARE_MULTICAST_DELEGATE(FOnFunctionCycled);
+	/** Invoked when CycleCommand is invoked. */
+	FOnFunctionCycled& OnFunctionCycleCommandInvoked() { return OnFunctionCycleCommandInvokedDelegate; }
 	
 private:
 	
@@ -60,9 +64,12 @@ private:
 	const TSharedRef<FUICommandList> CommandList;
 	/** Invokes CycleToNextFunction. */
 	const TSharedPtr<FUICommandInfo> CycleCommand;
+
+	/** Invoked when CycleCommand is invoked. */
+	FOnFunctionCycled OnFunctionCycleCommandInvokedDelegate;
 	
 	/** Sets the next tween function. */
-	UE_API void CycleToNextFunction();
+	UE_API void CycleToNextFunction() const;
 };
 }
 
diff --git a/Engine/Plugins/Animation/TweeningUtils/Source/TweeningUtilsEditor/Public/Widgets/MVC/MouseSlidingController.h b/Engine/Plugins/Animation/TweeningUtils/Source/TweeningUtilsEditor/Public/Widgets/MVC/MouseSlidingController.h
--- a/Engine/Plugins/Animation/TweeningUtils/Source/TweeningUtilsEditor/Public/Widgets/MVC/MouseSlidingController.h
+++ b/Engine/Plugins/Animation/TweeningUtils/Source/TweeningUtilsEditor/Public/Widgets/MVC/MouseSlidingController.h
@@ -54,6 +54,14 @@ public:
 		);
 	UE_API virtual ~FMouseSlidingController();
 
+	/** Starts listening for the mouse sliding command. */
+	UE_API void BindCommand();
+	/** Stops listening for the mouse sliding command. */
+	UE_API void UnbindCommand() const;
+
+	/** @return Whether the slider operation is currently in progress. */
+	bool IsSliding() const { return SlidingState.IsSet() && SlidingState->InitialMouse.IsSet(); }
+	
 	DECLARE_MULTICAST_DELEGATE_OneParam(FMoveSliderDelegate, float /*InNormalizedRange*/);
 	/** Invoked every tick the slider is moved. The event is throttled, so there's at most one invocation per tick. */
 	FMoveSliderDelegate& OnUpdateSliding() { return OnUpdateSlidingDelegate; }
@@ -100,19 +108,17 @@ private:
 	FSimpleMulticastDelegate OnStopSlidingDelegate;
 
 	/** Once DragSliderCommand has triggered, start listening for mouse down events. */
-	UE_API void StartListeningForMouseEvents();
+	void StartListeningForMouseEvents();
 	/** Once DragSliderCommand is released, stop listening for mouse down events.  */
-	UE_API void StopListeningForMouseEvents(ICursor& InCursor);
+	void StopListeningForMouseEvents(ICursor& InCursor);
 
 	/** Starts sliding from mouse position and locks the mouse; emits OnUpdateSliding calls. Called once DragSliderCommand and LMB are both down. */
-	UE_API void StartSliding(const FVector2D& InInitialScreenLocation, ICursor& InCursor);
+	void StartSliding(const FVector2D& InInitialScreenLocation, ICursor& InCursor);
 	/** Stops emiting emitting OnUpdateSliding calls. Unlocks the cursor. Called once either DragSliderCommand or LMB are released. */
-	UE_API void StopSliding(ICursor& InCursor);
+	void StopSliding(ICursor& InCursor);
 	/** Emits a OnUpdateSliding call. Called once per frame while sliding is active. */
-	UE_API void UpdateSliding(const FVector2D& InScreenLocation);
+	void UpdateSliding(const FVector2D& InScreenLocation);
 };
-
-	
 }
 
-#undef UE_API
+#undef UE_API
\ No newline at end of file
diff --git a/Engine/Plugins/Animation/TweeningUtils/Source/TweeningUtilsEditor/Public/Widgets/MVC/STweenView.h b/Engine/Plugins/Animation/TweeningUtils/Source/TweeningUtilsEditor/Public/Widgets/MVC/STweenView.h
--- a/Engine/Plugins/Animation/TweeningUtils/Source/TweeningUtilsEditor/Public/Widgets/MVC/STweenView.h
+++ b/Engine/Plugins/Animation/TweeningUtils/Source/TweeningUtilsEditor/Public/Widgets/MVC/STweenView.h
@@ -12,6 +12,7 @@ template<typename OptionalType> struct TOptional;
 namespace UE::TweeningUtilsEditor
 {
 class FTweenModel;
+class STweenSlider;
 enum class ETweenScaleMode : uint8;
 	
 /**
@@ -36,8 +37,13 @@ public:
 
 	UE_API void Construct(const FArguments& InArgs);
 
+	const TSharedPtr<STweenSlider>& GetTweenSlider() const { return TweenSlider; }
+
 private:
 
+	/** The embedded tween slider this widget is wrapping. */
+	TSharedPtr<STweenSlider> TweenSlider;
+	
 	/** The model that these controls are driving */
 	TAttribute<FTweenModel*> TweenModelAttr;
 	
diff --git a/Engine/Plugins/Animation/TweeningUtils/Source/TweeningUtilsEditor/Public/Widgets/MVC/TweenToolbarController.h b/Engine/Plugins/Animation/TweeningUtils/Source/TweeningUtilsEditor/Public/Widgets/MVC/TweenToolbarController.h
--- a/Engine/Plugins/Animation/TweeningUtils/Source/TweeningUtilsEditor/Public/Widgets/MVC/TweenToolbarController.h
+++ b/Engine/Plugins/Animation/TweeningUtils/Source/TweeningUtilsEditor/Public/Widgets/MVC/TweenToolbarController.h
@@ -22,8 +22,9 @@ template<typename OptionalType> struct TOptional;
 
 namespace UE::TweeningUtilsEditor
 {
-class ITweenModelContainer;
 class FTweenModel;
+class ITweenModelContainer;
+class STweenSlider;
 	
 /**
  * Binds together a STweenView and FTweenModel in the context of adding the controls to a toolbar widget.
@@ -65,13 +66,18 @@ public:
 		/** Affects how the overshoot button is resized. */
 		FMenuEntryResizeParams OvershootResizeParams;
 	};
+	struct FAddToToolbarResult
+	{
+		/** The tween slider widget that was added to the toolbar. */
+		TSharedPtr<STweenSlider> TweenSlider;
+	};
 	/**
 	 * Extends the toolbar. Safe to call multiple times (although typically you'll call this exactly once).
 	 * 
 	 * @param ToolbarBuilder The toolbar to add the widgets to
 	 * @param InArgs Optional additional args to control how the widget should be constructed.
 	 */
-	UE_API void AddToToolbar(FToolBarBuilder& ToolbarBuilder, FMakeWidgetArgs InArgs = {}) const;
+	UE_API FAddToToolbarResult AddToToolbar(FToolBarBuilder& ToolbarBuilder, FMakeWidgetArgs InArgs = {}) const;
 
 	/** Sets the currently selected tween model. Asserts that the tween model is contained in TweenModels. */
 	UE_API void SetSelectedTweenModel(FTweenModel& InTweenModel);
@@ -84,6 +90,9 @@ public:
 	DECLARE_MULTICAST_DELEGATE_OneParam(FOnTweenFunctionChanged, const FTweenModel&);
 	/** Invoked when the tween function is changed by this controller. */
 	FOnTweenFunctionChanged& OnTweenFunctionChanged() { return OnTweenFunctionChangedDelegate; }
+	/** Invoked when the overshoot mode is changed by this controller. */
+	DECLARE_MULTICAST_DELEGATE(FOnOvershootModeToggled);
+	FOnOvershootModeToggled& OnOvershootModeCommandInvoked() { return OnOvershootModeCommandInvokedDelegate; }
 
 private:
 
@@ -99,6 +108,8 @@ private:
 
 	/** Invoked when the tween function is changed by this controller. */
 	FOnTweenFunctionChanged OnTweenFunctionChangedDelegate;
+	/** Invoked when the overshoot mode is changed by this controller. */
+	FOnOvershootModeToggled OnOvershootModeCommandInvokedDelegate;
 
 	/** Creates the drop-down list for selecting the tween mode. */
 	UE_API TSharedRef<SWidget> MakeTweenModeMenu() const;
diff --git a/Engine/Source/Editor/EditorStyle/Private/StarshipStyle.cpp b/Engine/Source/Editor/EditorStyle/Private/StarshipStyle.cpp
--- a/Engine/Source/Editor/EditorStyle/Private/StarshipStyle.cpp
+++ b/Engine/Source/Editor/EditorStyle/Private/StarshipStyle.cpp
@@ -7715,9 +7715,12 @@ void FStarshipEditorStyle::FStyle::SetupToolkitStyles()
 			CurveEditorToolbar.SetComboLabelMinWidth(LongestTweenName);
 			Set("CurveEditorTweenToolbar", CurveEditorToolbar);
 			
+			FLinearColor ViewportOverlayColor = FStyleColors::Input.GetSpecifiedColor();
+			ViewportOverlayColor.A = 0.75f;
+			Set("ControlRig.OverlayBrush", new FSlateRoundedBoxBrush(ViewportOverlayColor, 8.0f, FStyleColors::Dropdown, 0.0f));
 			CurveEditorToolbar.SetComboLabelMinWidth(LongestTweenName);
 			CurveEditorToolbar.SetAllowWrapButton(false);
-			CurveEditorToolbar.SetBackground(*GetBrush("EditorViewport.OverlayBrush"));
+			CurveEditorToolbar.SetBackground(*GetBrush("ControlRig.OverlayBrush"));
 			Set("ControlRigTweenToolbar", CurveEditorToolbar);
 		}
 
