Control Rig: Modular Rig Event Queue Widget

#rnx

diff --git a/Engine/Plugins/Animation/ControlRig/Source/ControlRig/Private/ModularRig.cpp b/Engine/Plugins/Animation/ControlRig/Source/ControlRig/Private/ModularRig.cpp
--- a/Engine/Plugins/Animation/ControlRig/Source/ControlRig/Private/ModularRig.cpp
+++ b/Engine/Plugins/Animation/ControlRig/Source/ControlRig/Private/ModularRig.cpp
@@ -250,6 +250,24 @@ bool FRigModuleInstance::HasChildModule(const FName& InModuleName, bool bRecursi
 	return false;
 }
 
+uint32 GetTypeHash(const FRigModuleExecutionElement& InElement)
+{
+	uint32 Hash = GetTypeHash(InElement.EventName);
+	Hash = HashCombine(Hash, GetTypeHash(InElement.ModuleName));
+	Hash = HashCombine(Hash, GetTypeHash(InElement.bExecuted));
+	return Hash;
+}
+
+uint32 GetTypeHash(const FRigModuleExecutionQueue& InQueue)
+{
+	uint32 Hash = 0;
+	for (const FRigModuleExecutionElement& Element : InQueue.Elements)
+	{
+		Hash = HashCombine(Hash, GetTypeHash(Element));
+	}
+	return Hash;
+}
+
 void UModularRig::PostInitProperties()
 {
 	Super::PostInitProperties();
@@ -529,6 +547,10 @@ const FModularRigSettings& UModularRig::GetModularRigSettings() const
 
 void UModularRig::ExecuteQueue()
 {
+#if WITH_EDITOR
+	LastExecutedElements.Reset();
+#endif
+	
 	FRigVMExtendedExecuteContext& Context = GetRigVMExtendedExecuteContext();
 	FControlRigExecuteContext& PublicContext = Context.GetPublicDataSafe<FControlRigExecuteContext>();
 	URigHierarchy* Hierarchy = GetHierarchy();
@@ -545,6 +567,14 @@ void UModularRig::ExecuteQueue()
 		FRigModuleExecutionElement& ExecutionElement = ExecutionQueue[ExecutionQueueFront];
 		if (FRigModuleInstance* ModuleInstance = ExecutionElement.ModuleInstance)
 		{
+#if WITH_EDITOR
+			LastExecutedElements.Add(ExecutionElement);
+			LastExecutedElements.Last().bExecuted = false;
+			LastExecutedElements.Last().DurationInMicroSeconds = 0;
+
+			const uint64 StartCycles = FPlatformTime::Cycles64();
+#endif
+			
 			if (UControlRig* ModuleRig = ModuleInstance->GetRig())
 			{
 				// Let the module know of the spawn index the hierarchy is currently going through
@@ -662,6 +692,12 @@ void UModularRig::ExecuteQueue()
 				ModuleRig->Execute_Internal(ExecutionElement.EventName);
 				ExecutionElement.bExecuted = true;
 
+#if WITH_EDITOR
+				const uint64 Cycles = FPlatformTime::Cycles64() - StartCycles;
+				LastExecutedElements.Last().bExecuted = true;
+				LastExecutedElements.Last().DurationInMicroSeconds = static_cast<double>(Cycles) * FPlatformTime::GetSecondsPerCycle() * 1000.0 * 1000.0;
+#endif
+
 				// Copy result of Connection event to the ModularRig's unit context
 				if (ExecutionElement.EventName == FRigUnit_ConnectorExecution::EventName)
 				{
@@ -677,14 +713,37 @@ void UModularRig::ExecuteQueue()
 	{
 		Module->PostConstructionSpawnStartIndex = Hierarchy->GetNextSpawnIndex();
 	}
+
+#if WITH_EDITOR
+	// copy the elements over while applying a lock
+	{
+		FWriteScopeLock WriteLock(LastExecutionQueueLock);
+		LastExecutionQueue.Elements = LastExecutedElements;
+		LastExecutedElements.Reset();
+	}
+#endif
 }
 
 void UModularRig::ResetExecutionQueue()
 {
 	ExecutionQueue.Reset();
 	ExecutionQueueFront = 0;
+#if WITH_EDITOR
+	LastExecutedElements.Reset();
+#endif
 }
 
+#if WITH_EDITOR
+
+FRigModuleExecutionQueue UModularRig::GetLastExecutionQueue() const
+{
+	FReadScopeLock ReadLock(LastExecutionQueueLock);
+	// this performs a copy
+	return LastExecutionQueue;
+}
+
+#endif
+
 void UModularRig::OnObjectsReplaced(const TMap<UObject*, UObject*>& OldToNewInstanceMap)
 {
 	if(Modules.IsEmpty())
@@ -1196,7 +1255,10 @@ FName UModularRig::GetParentModuleName(const FName& InModuleName) const
 
 void UModularRig::ForEachModule(TFunctionRef<bool(FRigModuleInstance*)> PerModuleFunction)
 {
-	TArray<FRigModuleInstance*> ModuleInstances = RootModules;
+	TArray<FRigModuleInstance*> ModuleInstances;
+	ModuleInstances.Reserve(Modules.Num());
+	ModuleInstances.Append(RootModules);
+	
 	for (int32 ModuleIndex = 0; ModuleIndex < ModuleInstances.Num(); ++ModuleIndex)
 	{
 		if (!PerModuleFunction(ModuleInstances[ModuleIndex]))
diff --git a/Engine/Plugins/Animation/ControlRig/Source/ControlRig/Public/ModularRig.h b/Engine/Plugins/Animation/ControlRig/Source/ControlRig/Public/ModularRig.h
--- a/Engine/Plugins/Animation/ControlRig/Source/ControlRig/Public/ModularRig.h
+++ b/Engine/Plugins/Animation/ControlRig/Source/ControlRig/Public/ModularRig.h
@@ -3,6 +3,7 @@
 #pragma once
 #include "ControlRig.h"
 #include "ModularRigModel.h"
+#include "Misc/ScopeRWLock.h"
 #include "ModularRig.generated.h"
 
 #define UE_API CONTROLRIG_API
@@ -111,13 +112,15 @@ struct FRigModuleExecutionElement
 		, ModuleInstance(nullptr)
 		, EventName(NAME_None)
 		, bExecuted(false)
+		, DurationInMicroSeconds(0.0)
 	{}
 
-	FRigModuleExecutionElement(FRigModuleInstance* InModule, FName InEvent)
+	FRigModuleExecutionElement(FRigModuleInstance* InModule, FName InEvent, double InDurationMicroSeconds = 0.0)
 		: ModuleName(InModule->Name)
 		, ModuleInstance(InModule)
 		, EventName(InEvent)
 		, bExecuted(false)
+		, DurationInMicroSeconds(InDurationMicroSeconds)
 	{}
 
 	UPROPERTY()
@@ -129,6 +132,22 @@ struct FRigModuleExecutionElement
 
 	UPROPERTY()
 	bool bExecuted;
+
+	UPROPERTY()
+	double DurationInMicroSeconds;
+
+	UE_API friend uint32 GetTypeHash(const FRigModuleExecutionElement& InElement);
+};
+
+USTRUCT(BlueprintType)
+struct FRigModuleExecutionQueue
+{
+	GENERATED_BODY()
+
+	UPROPERTY()
+	TArray<FRigModuleExecutionElement> Elements;
+
+	UE_API friend uint32 GetTypeHash(const FRigModuleExecutionQueue& InQueue);
 };
 
 /** Runs logic for mapping input data to transforms (the "Rig") */
@@ -181,6 +200,10 @@ public:
 	UE_API void ExecuteQueue();
 	UE_API void ResetExecutionQueue();
 
+#if WITH_EDITOR
+	UE_API FRigModuleExecutionQueue GetLastExecutionQueue() const;
+#endif
+
 	// BEGIN UObject
 	UE_API virtual void BeginDestroy() override;
 	// END UObject
@@ -328,6 +351,10 @@ private:
 
 	TMap<FName, UControlRig*> PreviousModuleRigs;
 
+	TArray<FRigModuleExecutionElement> LastExecutedElements;
+	FRigModuleExecutionQueue LastExecutionQueue;
+	mutable FRWLock LastExecutionQueueLock;
+
 	friend struct FRigModuleInstance;
 };
 
diff --git a/Engine/Plugins/Animation/ControlRig/Source/ControlRigEditor/Private/ControlRigEditorModule.cpp b/Engine/Plugins/Animation/ControlRig/Source/ControlRigEditor/Private/ControlRigEditorModule.cpp
--- a/Engine/Plugins/Animation/ControlRig/Source/ControlRigEditor/Private/ControlRigEditorModule.cpp
+++ b/Engine/Plugins/Animation/ControlRig/Source/ControlRigEditor/Private/ControlRigEditorModule.cpp
@@ -41,6 +41,7 @@
 #include <Editor/ControlRigEditorCommands.h>
 #include "ControlRigHierarchyCommands.h"
 #include "ControlRigModularRigCommands.h"
+#include "Editor/ModularRigEventQueueCommands.h"
 #include "Animation/AnimSequence.h"
 #include "Editor/ControlRigEditorEditMode.h"
 #include "ControlRigElementDetails.h"
@@ -122,6 +123,7 @@ void FControlRigEditorModule::StartupModule()
 	FControlRigEditorCommands::Register();
 	FControlRigHierarchyCommands::Register();
 	FControlRigModularRigCommands::Register();
+	FModularRigEventQueueCommands::Register();
 	FControlRigEditorStyle::Get();
 
 	EdGraphPanelPinFactory = MakeShared<FControlRigGraphPanelPinFactory>();
diff --git a/Engine/Plugins/Animation/ControlRig/Source/ControlRigEditor/Private/Editor/ControlRigEditor.h b/Engine/Plugins/Animation/ControlRig/Source/ControlRigEditor/Private/Editor/ControlRigEditor.h
--- a/Engine/Plugins/Animation/ControlRig/Source/ControlRigEditor/Private/Editor/ControlRigEditor.h
+++ b/Engine/Plugins/Animation/ControlRig/Source/ControlRigEditor/Private/Editor/ControlRigEditor.h
@@ -178,6 +178,7 @@ protected:
 	
 	virtual void OnCloseSuper() = 0;
 
+	friend class SModularRigEventQueueView;
 };
 
 class FControlRigBaseEditor : public IControlRigBaseEditor
diff --git a/Engine/Plugins/Animation/ControlRig/Source/ControlRigEditor/Private/Editor/ControlRigEditorMode.cpp b/Engine/Plugins/Animation/ControlRig/Source/ControlRigEditor/Private/Editor/ControlRigEditorMode.cpp
--- a/Engine/Plugins/Animation/ControlRig/Source/ControlRigEditor/Private/Editor/ControlRigEditorMode.cpp
+++ b/Engine/Plugins/Animation/ControlRig/Source/ControlRigEditor/Private/Editor/ControlRigEditorMode.cpp
@@ -16,6 +16,7 @@
 #include "Editor/RigCurveContainerTabSummoner.h"
 #include "Editor/RigValidationTabSummoner.h"
 #include "Editor/RigAnimAttributeTabSummoner.h"
+#include "Editor/ModularRigEventQueueTabSummoner.h"
 #include "ToolMenus.h"
 #include "Editor/RigVMCompilerResultsTabSummoner.h"
 #include "Editor/RigVMDetailsInspectorTabSummoner.h"
@@ -160,6 +161,7 @@ FModularRigLegacyEditorMode::FModularRigLegacyEditorMode(const TSharedRef<FContr
 
 	TabFactories.RegisterFactory(MakeShared<FModularRigModelTabSummoner>(InControlRigEditor));
 	TabFactories.RegisterFactory(MakeShared<FRigModuleAssetBrowserTabSummoner>(InControlRigEditor));
+	TabFactories.RegisterFactory(MakeShared<FModularRigEventQueueTabSummoner>(InControlRigEditor));
 
 	TabLayout = FTabManager::NewLayout("Standalone_ModularRigEditMode_Layout_v1.3")
 		->AddArea
diff --git a/Engine/Plugins/Animation/ControlRig/Source/ControlRigEditor/Private/Editor/ModularRigEventQueueCommands.cpp b/Engine/Plugins/Animation/ControlRig/Source/ControlRigEditor/Private/Editor/ModularRigEventQueueCommands.cpp
new file mode 100644
--- /dev/null
+++ b/Engine/Plugins/Animation/ControlRig/Source/ControlRigEditor/Private/Editor/ModularRigEventQueueCommands.cpp
@@ -0,0 +1,12 @@
+// Copyright Epic Games, Inc. All Rights Reserved.
+
+#include "Editor/ModularRigEventQueueCommands.h"
+
+#define LOCTEXT_NAMESPACE "ModularRigEventQueueCommands"
+
+void FModularRigEventQueueCommands::RegisterCommands()
+{
+	UI_COMMAND(FocusOnSelection, "Focus On Selection", "Finds the selected event's module in the modular rig tree.", EUserInterfaceActionType::Button, FInputChord());
+}
+
+#undef LOCTEXT_NAMESPACE
diff --git a/Engine/Plugins/Animation/ControlRig/Source/ControlRigEditor/Private/Editor/ModularRigEventQueueCommands.h b/Engine/Plugins/Animation/ControlRig/Source/ControlRigEditor/Private/Editor/ModularRigEventQueueCommands.h
new file mode 100644
--- /dev/null
+++ b/Engine/Plugins/Animation/ControlRig/Source/ControlRigEditor/Private/Editor/ModularRigEventQueueCommands.h
@@ -0,0 +1,27 @@
+// Copyright Epic Games, Inc. All Rights Reserved.
+
+#pragma once
+
+#include "CoreMinimal.h"
+#include "Framework/Commands/Commands.h"
+
+class CONTROLRIGEDITOR_API FModularRigEventQueueCommands : public TCommands<FModularRigEventQueueCommands>
+{
+public:
+	FModularRigEventQueueCommands() : TCommands<FModularRigEventQueueCommands>
+	(
+		"FModularRigEventQueue",
+		NSLOCTEXT("Contexts", "ModularRig", "Event Queue"),
+		NAME_None, // "MainFrame" // @todo Fix this crash
+		"ControlRigEditorStyle" // Icon Style Set
+	)
+	{}
+	
+	/** Focuses on a selected event's module */
+	TSharedPtr< FUICommandInfo > FocusOnSelection;
+
+	/**
+	 * Initialize commands
+	 */
+	virtual void RegisterCommands() override;
+};
diff --git a/Engine/Plugins/Animation/ControlRig/Source/ControlRigEditor/Private/Editor/ModularRigEventQueueTabSummoner.cpp b/Engine/Plugins/Animation/ControlRig/Source/ControlRigEditor/Private/Editor/ModularRigEventQueueTabSummoner.cpp
new file mode 100644
--- /dev/null
+++ b/Engine/Plugins/Animation/ControlRig/Source/ControlRigEditor/Private/Editor/ModularRigEventQueueTabSummoner.cpp
@@ -0,0 +1,27 @@
+// Copyright Epic Games, Inc. All Rights Reserved.
+
+#include "Editor/ModularRigEventQueueTabSummoner.h"
+#include "Editor/SModularRigEventQueueView.h"
+#include "Editor/ControlRigEditor.h"
+
+#define LOCTEXT_NAMESPACE "ModularRigEventQueueTabSummoner"
+
+FModularRigEventQueueTabSummoner::FModularRigEventQueueTabSummoner(const TSharedRef<IControlRigBaseEditor>& InControlRigEditor)
+	: FWorkflowTabFactory(TabID, InControlRigEditor->GetHostingApp())
+	, ControlRigEditor(InControlRigEditor)
+{
+	TabLabel = LOCTEXT("ModularRigEventQueueTabLabel", "Event Queue");
+	TabIcon = FSlateIcon(TEXT("RigVMEditorStyle"), "ExecutionStack.TabIcon"); // use the same icon as for the execution stack
+
+	bIsSingleton = true;
+
+	ViewMenuDescription = LOCTEXT("ModularRigEventQueue_ViewMenu_Desc", "Event Queue");
+	ViewMenuTooltip = LOCTEXT("ModularRigEventQueue_ViewMenu_ToolTip", "Show the Event Queue tab");
+}
+
+TSharedRef<SWidget> FModularRigEventQueueTabSummoner::CreateTabBody(const FWorkflowTabSpawnInfo& Info) const
+{
+	return SNew(SModularRigEventQueueView, ControlRigEditor.Pin().ToSharedRef());
+}
+
+#undef LOCTEXT_NAMESPACE 
diff --git a/Engine/Plugins/Animation/ControlRig/Source/ControlRigEditor/Private/Editor/ModularRigEventQueueTabSummoner.h b/Engine/Plugins/Animation/ControlRig/Source/ControlRigEditor/Private/Editor/ModularRigEventQueueTabSummoner.h
new file mode 100644
--- /dev/null
+++ b/Engine/Plugins/Animation/ControlRig/Source/ControlRigEditor/Private/Editor/ModularRigEventQueueTabSummoner.h
@@ -0,0 +1,27 @@
+// Copyright Epic Games, Inc. All Rights Reserved.
+
+#pragma once
+
+#include "CoreMinimal.h"
+#include "WorkflowOrientedApp/WorkflowTabFactory.h"
+
+#define UE_API CONTROLRIGEDITOR_API
+
+class IControlRigBaseEditor;
+
+struct FModularRigEventQueueTabSummoner : public FWorkflowTabFactory
+{
+public:
+
+	static inline const FLazyName TabID = FLazyName(TEXT("Event Queue"));
+	
+public:
+	UE_API FModularRigEventQueueTabSummoner(const TSharedRef<IControlRigBaseEditor>& InControlRigEditor);
+	
+	UE_API virtual TSharedRef<SWidget> CreateTabBody(const FWorkflowTabSpawnInfo& Info) const override;
+	
+protected:
+	TWeakPtr<IControlRigBaseEditor> ControlRigEditor;
+};
+
+#undef UE_API
diff --git a/Engine/Plugins/Animation/ControlRig/Source/ControlRigEditor/Private/Editor/SModularRigEventQueueView.cpp b/Engine/Plugins/Animation/ControlRig/Source/ControlRigEditor/Private/Editor/SModularRigEventQueueView.cpp
new file mode 100644
--- /dev/null
+++ b/Engine/Plugins/Animation/ControlRig/Source/ControlRigEditor/Private/Editor/SModularRigEventQueueView.cpp
@@ -0,0 +1,525 @@
+// Copyright Epic Games, Inc. All Rights Reserved.
+
+#include "Editor/SModularRigEventQueueView.h"
+#include "Widgets/Layout/SBorder.h"
+#include "Widgets/SBoxPanel.h"
+#include "Widgets/Layout/SScrollBar.h"
+#include "Widgets/Layout/SSpacer.h"
+#include "Widgets/Input/SSearchBox.h"
+#include "Widgets/Input/SNumericEntryBox.h"
+#include "Editor/ControlRigNewEditor.h"
+#include "Editor/ModularRigEventQueueCommands.h"
+#include "Framework/MultiBox/MultiBoxBuilder.h"
+#include "Styling/AppStyle.h"
+#include "DetailLayoutBuilder.h"
+#include "Widgets/Input/SSearchBox.h"
+#include "Dialog/SCustomDialog.h"
+#include "ControlRigEditor.h"
+#include "MVVM/ViewModels/OutlinerColumns/OutlinerColumnTypes.h"
+
+#define LOCTEXT_NAMESPACE "SModularRigEventQueueView"
+
+//////////////////////////////////////////////////////////////
+/// FModularRigEventEntry
+///////////////////////////////////////////////////////////
+FModularRigEventEntry::FModularRigEventEntry(int32 InEventIndex, const FName& InEventName, const FModuleInstanceHandle& InModule, bool InExecuted, double InDurationMicroSeconds)
+	: EventIndex(InEventIndex)
+	, EventName(InEventName)
+	, Module(InModule)
+	, bExecuted(InExecuted)
+{
+	UpdateDurationMicroSeconds(InDurationMicroSeconds);
+}
+
+TSharedRef<ITableRow> FModularRigEventEntry::MakeTreeRowWidget(const TSharedRef<STableViewBase>& InOwnerTable, TSharedRef<FModularRigEventEntry> InEntry, TSharedRef<FUICommandList> InCommandList, TSharedPtr<SModularRigEventQueueView> InEventQueueView)
+{
+	return SNew(SModularRigEventItem, InOwnerTable, InEntry, InCommandList);
+}
+
+void FModularRigEventEntry::UpdateDurationMicroSeconds(double InDurationMicroSeconds)
+{
+	MicroSeconds = 0;
+	if (!bExecuted)
+	{
+		return;
+	}
+		
+	if (MicroSecondsFrames.Num() == 100)
+	{
+		MicroSecondsFrames.RemoveAt(0, EAllowShrinking::No);
+	}
+	MicroSecondsFrames.Add(InDurationMicroSeconds);
+
+	for (const double& Frame : MicroSecondsFrames)
+	{
+		MicroSeconds += Frame;
+	}
+	MicroSeconds /= static_cast<double>(MicroSecondsFrames.Num());
+}
+
+//////////////////////////////////////////////////////////////
+/// SModularRigEventItem
+///////////////////////////////////////////////////////////
+void SModularRigEventItem::Construct(const FArguments& InArgs, const TSharedRef<STableViewBase>& OwnerTable, TSharedRef<FModularRigEventEntry> InEventEntry, TSharedRef<FUICommandList> InCommandList)
+{
+	WeakEventEntry = InEventEntry;
+	WeakCommandList = InCommandList;
+
+	TSharedPtr< STextBlock > NumberWidget;
+	TSharedPtr< STextBlock > TextWidget;
+
+	const FSlateBrush* Icon = FSlateIcon(TEXT("RigVMEditor"), "RigVM.Unit").GetIcon();
+
+	STableRow<TSharedPtr<FModularRigEventEntry>>::Construct(
+		STableRow<TSharedPtr<FModularRigEventEntry>>::FArguments()
+		.Content()
+		[
+			SNew(SHorizontalBox)
+			+ SHorizontalBox::Slot()
+			.AutoWidth()
+			.VAlign(VAlign_Center)
+			.HAlign(HAlign_Left)
+			[
+				SNew(SHorizontalBox)
+				+ SHorizontalBox::Slot()
+				.MaxWidth(35.f)
+				.AutoWidth()
+				.VAlign(VAlign_Center)
+				.HAlign(HAlign_Right)
+				[
+					SAssignNew(NumberWidget, STextBlock)
+					.Text(this, &SModularRigEventItem::GetIndexText)
+					.Font(IDetailLayoutBuilder::GetDetailFont())
+				]
+				+ SHorizontalBox::Slot()
+				.MaxWidth(22.f)
+				.VAlign(VAlign_Center)
+				.HAlign(HAlign_Center)
+				[
+					SNew(SImage)
+					.Image(Icon)
+				]
+				
+				+ SHorizontalBox::Slot()
+				.FillWidth(1.f)
+				.AutoWidth()
+				.VAlign(VAlign_Center)
+				.HAlign(HAlign_Fill)
+				[
+					SAssignNew(TextWidget, STextBlock)
+					.Text(this, &SModularRigEventItem::GetEventLabelText)
+					.Font(this, &SModularRigEventItem::GetEventLabelFont)
+					.ToolTipText(this, &SModularRigEventItem::GetTooltip)
+					.Justification(ETextJustify::Left)
+				]
+
+				+ SHorizontalBox::Slot()
+				.Padding(8.f, 1, 1, 1)
+				.FillWidth(1.f)
+				.AutoWidth()
+				.VAlign(VAlign_Center)
+				.HAlign(HAlign_Fill)
+				[
+					SAssignNew(TextWidget, STextBlock)
+					.Text(this, &SModularRigEventItem::GetModuleLabelText)
+					.Font(this, &SModularRigEventItem::GetModuleLabelFont)
+					.ToolTipText(this, &SModularRigEventItem::GetTooltip)
+					.Justification(ETextJustify::Left)
+				]
+			]
+
+			+ SHorizontalBox::Slot()
+			.FillWidth(1.f)
+			[
+				SNew(SSpacer)
+			]
+
+			+ SHorizontalBox::Slot()
+			.AutoWidth()
+			.VAlign(VAlign_Center)
+			.HAlign(HAlign_Right)
+			[
+				SNew(SHorizontalBox)
+	            
+	            + SHorizontalBox::Slot()
+	            .Padding(20, 0, 0, 0)
+				.AutoWidth()
+	            .VAlign(VAlign_Center)
+	            .HAlign(HAlign_Left)
+	            [
+	                SNew(STextBlock)
+	                .Text(this, &SModularRigEventItem::GetDurationText)
+	                .Font(IDetailLayoutBuilder::GetDetailFont())
+	            ]
+	        ]
+        ], OwnerTable);
+}
+
+FText SModularRigEventItem::GetIndexText() const
+{
+	const FString IndexStr = FString::FromInt(WeakEventEntry.Pin()->EventIndex) + TEXT(".");
+	return FText::FromString(IndexStr);
+}
+
+FText SModularRigEventItem::GetModuleLabelText() const
+{
+	return (FText::FromName(WeakEventEntry.Pin()->Module.GetModuleName()));
+}
+
+FSlateFontInfo SModularRigEventItem::GetModuleLabelFont() const
+{
+	return IDetailLayoutBuilder::GetDetailFontBold();
+}
+
+FText SModularRigEventItem::GetEventLabelText() const
+{
+	return (FText::FromName(WeakEventEntry.Pin()->EventName));
+}
+
+FSlateFontInfo SModularRigEventItem::GetEventLabelFont() const
+{
+	return IDetailLayoutBuilder::GetDetailFont();
+}
+
+FText SModularRigEventItem::GetTooltip() const
+{
+	// empty for now
+	return FText(); 
+}
+
+FText SModularRigEventItem::GetDurationText() const
+{
+	if(WeakEventEntry.IsValid())
+	{
+		if (TSharedPtr<FModularRigEventEntry> Entry = WeakEventEntry.Pin())
+		{
+			if (!Entry->MicroSecondsFrames.IsEmpty())
+			{
+				return FText::FromString(FString::Printf(TEXT("%.02f µs"), (float)Entry->MicroSeconds));
+			}
+		}
+	}
+	return FText();
+}
+
+//////////////////////////////////////////////////////////////
+/// SModularRigEventQueueView
+///////////////////////////////////////////////////////////
+
+SModularRigEventQueueView::~SModularRigEventQueueView()
+{
+	if (TSharedPtr<IControlRigBaseEditor> Editor = WeakEditor.Pin())
+	{
+		if (OnControlRigExecutedHandle.IsValid())
+		{
+			if(URigVMHost* ModularRig = Editor->GetRigVMHost())
+			{
+				ModularRig->OnExecuted_AnyThread().Remove(OnControlRigExecutedHandle);
+			}
+		}
+		if (OnPreviewHostUpdatedHandle.IsValid())
+		{
+			Editor->SharedRigVMEditorRef()->OnPreviewHostUpdated().Remove(OnPreviewHostUpdatedHandle);
+		}
+	}
+}
+
+void SModularRigEventQueueView::Construct( const FArguments& InArgs, TSharedRef<IControlRigBaseEditor> InControlRigEditor)
+{
+	WeakEditor = InControlRigEditor;
+	CommandList = MakeShared<FUICommandList>();
+
+	BindCommands();
+
+	ChildSlot
+	[
+		SNew(SVerticalBox)
+		+SVerticalBox::Slot()
+		.AutoHeight()
+		.VAlign(VAlign_Top)
+		.Padding(0.0f)
+		[
+			SNew(SBorder)
+			.Padding(0.0f)
+			.BorderImage(FAppStyle::GetBrush("DetailsView.CategoryTop"))
+			.BorderBackgroundColor(FLinearColor(0.6f, 0.6f, 0.6f, 1.0f))
+			[
+				SNew(SVerticalBox)
+				+SVerticalBox::Slot()
+				.AutoHeight()
+				.VAlign(VAlign_Top)
+				[
+					SNew(SHorizontalBox)
+					//.Visibility(this, &SRigHierarchy::IsSearchbarVisible)
+					+SHorizontalBox::Slot()
+					.AutoWidth()
+					+SHorizontalBox::Slot()
+					.VAlign(VAlign_Center)
+					.Padding(3.0f, 1.0f)
+					[
+						SAssignNew(FilterBox, SSearchBox)
+						.OnTextChanged(this, &SModularRigEventQueueView::OnFilterTextChanged)
+					]
+				]
+			]
+		]
+		+SVerticalBox::Slot()
+		.Padding(0.0f, 0.0f)
+		[
+			SNew(SBorder)
+			.Padding(2.0f)
+			.BorderImage(FAppStyle::GetBrush("SCSEditor.TreePanel"))
+			[
+				SAssignNew(TreeView, STreeView<TSharedPtr<FModularRigEventEntry>>)
+				.TreeItemsSource(&Events)
+				.SelectionMode(ESelectionMode::Multi)
+				.OnGenerateRow(this, &SModularRigEventQueueView::MakeTableRowWidget)
+				.OnGetChildren(this, &SModularRigEventQueueView::HandleGetChildrenForTree)
+				.OnSelectionChanged(this, &SModularRigEventQueueView::OnSelectionChanged)
+				.OnContextMenuOpening(this, &SModularRigEventQueueView::CreateContextMenu)
+			]
+		]
+	];
+
+	if (TSharedPtr<IControlRigBaseEditor> Editor = WeakEditor.Pin())
+	{
+		OnPreviewHostUpdatedHandle = Editor->SharedRigVMEditorRef()->OnPreviewHostUpdated().AddSP(this, &SModularRigEventQueueView::HandlePreviewHostUpdated);
+		HandlePreviewHostUpdated(&Editor->SharedRigVMEditorRef().Get());
+	}
+}
+
+void SModularRigEventQueueView::OnSelectionChanged(TSharedPtr<FModularRigEventEntry> Selection, ESelectInfo::Type SelectInfo)
+{
+	if (bIsSelecting)
+	{
+		return;
+	}
+	TGuardValue<bool> SelectionGuard(bIsSelecting, true);
+
+	if (!WeakModularRigController.IsValid())
+	{
+		return;
+	}
+
+	UModularRigController* ModularRigController = WeakModularRigController.Get();
+	if (ModularRigController == nullptr)
+	{
+		return;
+	}
+	
+	TArray<TSharedPtr<FModularRigEventEntry>> SelectedItems = TreeView->GetSelectedItems();
+	if (SelectedItems.Num() > 0)
+	{
+		if (TSharedPtr<IControlRigBaseEditor> Editor = WeakEditor.Pin())
+		{
+			TArray<FName> SelectedModules;
+			for (const TSharedPtr<FModularRigEventEntry>& Item : SelectedItems)
+			{
+				SelectedModules.AddUnique(Item->Module.GetModuleName());
+			}
+			
+			const TArray<FName> PreviouslySelectedModules = Editor->GetSelectedModules();
+
+			if (SelectedModules.Num() == PreviouslySelectedModules.Num())
+			{
+				bool bFoundDifference = false;
+				for (const FName& SelectedModule : SelectedModules)
+				{
+					if (!PreviouslySelectedModules.Contains(SelectedModule))
+					{
+						bFoundDifference = true;
+						break;
+					}
+				}
+
+				if (!bFoundDifference)
+				{
+					return;
+				}
+			}
+
+			ModularRigController->SetModuleSelection(SelectedModules);
+		}
+	}
+}
+
+void SModularRigEventQueueView::BindCommands()
+{
+	// create new command
+	const FModularRigEventQueueCommands& Commands = FModularRigEventQueueCommands::Get();
+	CommandList->MapAction(Commands.FocusOnSelection, FExecuteAction::CreateSP(this, &SModularRigEventQueueView::HandleFocusOnSelectedModule));
+}
+
+TSharedRef<ITableRow> SModularRigEventQueueView::MakeTableRowWidget(TSharedPtr<FModularRigEventEntry> InItem, const TSharedRef<STableViewBase>& OwnerTable)
+{
+	return InItem->MakeTreeRowWidget(OwnerTable, InItem.ToSharedRef(), CommandList.ToSharedRef(), SharedThis(this));
+}
+
+void SModularRigEventQueueView::HandleGetChildrenForTree(TSharedPtr<FModularRigEventEntry> InItem,
+	TArray<TSharedPtr<FModularRigEventEntry>>& OutChildren)
+{
+	// nothing to do here
+}
+
+bool SModularRigEventQueueView::PopulateEventQueueView(const UModularRig* InModularRig)
+{
+	if (InModularRig)
+	{
+		// copy the queue
+		const FRigModuleExecutionQueue Queue = InModularRig->GetLastExecutionQueue();
+		const uint32 Hash = GetTypeHash(Queue);
+		if (Hash != LastHash)
+		{
+			Events.Reset();
+
+			for (int32 Index = 0; Index < Queue.Elements.Num(); Index++)
+			{
+				const FRigModuleExecutionElement& Element = Queue.Elements[Index];
+				TSharedPtr<FModularRigEventEntry> NewEntry = MakeShared<FModularRigEventEntry>(Index, Element.EventName, FModuleInstanceHandle(InModularRig, Element.ModuleInstance), Element.bExecuted, Element.DurationInMicroSeconds);
+				Events.Add(NewEntry);
+			}
+
+			LastHash = Hash;
+			return true;
+		}
+		else
+		{
+			check(Events.Num() == Queue.Elements.Num());
+			for (int32 Index = 0; Index < Queue.Elements.Num(); Index++)
+			{
+				Events[Index]->UpdateDurationMicroSeconds(Queue.Elements[Index].DurationInMicroSeconds);
+			}
+		}
+	}
+
+	return false;
+}
+
+void SModularRigEventQueueView::RefreshTreeView(const UModularRig* InModularRig)
+{
+	if (PopulateEventQueueView(InModularRig))
+	{
+		TreeView->RequestTreeRefresh();
+	}
+}
+
+TSharedPtr< SWidget > SModularRigEventQueueView::CreateContextMenu()
+{
+	TArray<TSharedPtr<FModularRigEventEntry>> SelectedItems = TreeView->GetSelectedItems();
+	if(SelectedItems.Num() == 0)
+	{
+		return nullptr;
+	}
+
+	const FModularRigEventQueueCommands& Actions = FModularRigEventQueueCommands::Get();
+
+	const bool CloseAfterSelection = true;
+	FMenuBuilder MenuBuilder(CloseAfterSelection, CommandList);
+	{
+		MenuBuilder.BeginSection("ModularRigEventQueueToolsAction", LOCTEXT("ToolsAction", "Tools"));
+		MenuBuilder.AddMenuEntry(Actions.FocusOnSelection);
+		MenuBuilder.EndSection();
+	}
+
+	return MenuBuilder.MakeWidget();
+}
+
+void SModularRigEventQueueView::HandleFocusOnSelectedModule()
+{
+	OnSelectionChanged(TSharedPtr<FModularRigEventEntry>(), ESelectInfo::Direct);
+}
+
+void SModularRigEventQueueView::OnFilterTextChanged(const FText& SearchText)
+{
+	FilterText = SearchText;
+	if (const UModularRig* ModularRig = WeakModularRig.Get())
+	{
+		RefreshTreeView(ModularRig);
+	}
+}
+
+void SModularRigEventQueueView::HandleModularRigExecuted(class URigVMHost* InModularRig, const FName& InEventName)
+{
+	if (!WeakModularRig.IsValid())
+	{
+		return;
+	}
+
+	if (WeakModularRig.Get() != InModularRig)
+	{
+		return;
+	}
+
+	RefreshTreeView(Cast<UModularRig>(InModularRig));
+}
+
+void SModularRigEventQueueView::HandleModularRigModelModified(EModularRigNotification InNotification, const FRigModuleReference* InModuleReference)
+{
+	if (bIsSelecting)
+	{
+		return;
+	}
+	TGuardValue<bool> SelectionGuard(bIsSelecting, true);
+	
+	if (InNotification == EModularRigNotification::ModuleSelected)
+	{
+		for (const TSharedPtr<FModularRigEventEntry>& Entry : Events)
+		{
+			if (Entry->Module.GetModuleName() == InModuleReference->GetName())
+			{
+				TreeView->SetItemSelection(Entry, true, ESelectInfo::Type::Direct);
+			}
+		}
+	}
+	else if (InNotification == EModularRigNotification::ModuleDeselected)
+	{
+		for (const TSharedPtr<FModularRigEventEntry>& Entry : Events)
+		{
+			if (Entry->Module.GetModuleName() == InModuleReference->GetName())
+			{
+				TreeView->SetItemSelection(Entry, false, ESelectInfo::Type::Direct);
+			}
+		}
+	}
+}
+
+void SModularRigEventQueueView::HandlePreviewHostUpdated(IRigVMEditor* InEditor)
+{
+	if (WeakModularRig.IsValid() && OnControlRigExecutedHandle.IsValid())
+	{
+		if (UModularRig* ModularRig = WeakModularRig.Get())
+		{
+			ModularRig->OnExecuted_AnyThread().Remove(OnControlRigExecutedHandle);
+		}
+	}
+
+	if (WeakModularRigController.IsValid())
+	{
+		if (UModularRigController* Controller = WeakModularRigController.Get())
+		{
+			Controller->OnModified().Remove(OnModularRigModelModifiedHandle);
+		}
+	}
+
+	WeakModularRig.Reset();
+	WeakModularRigController.Reset();
+	
+	if(UModularRig* ModularRig = Cast<UModularRig>(InEditor->GetRigVMHost()))
+	{
+		WeakModularRig = ModularRig;
+		
+		OnControlRigExecutedHandle = ModularRig->OnExecuted_AnyThread().AddSP(this, &SModularRigEventQueueView::HandleModularRigExecuted);
+		RefreshTreeView(ModularRig);
+
+		if (UControlRigBlueprint* ControlRigBlueprint = Cast<UControlRigBlueprint>(ModularRig->GetClass()->ClassGeneratedBy))
+		{
+			if (UModularRigController* Controller = ControlRigBlueprint->GetModularRigController())
+			{
+				WeakModularRigController = Controller;
+				Controller->OnModified().AddSP(this, &SModularRigEventQueueView::HandleModularRigModelModified);
+			}
+		}
+	}
+}
+
+#undef LOCTEXT_NAMESPACE
diff --git a/Engine/Plugins/Animation/ControlRig/Source/ControlRigEditor/Private/Editor/SModularRigEventQueueView.h b/Engine/Plugins/Animation/ControlRig/Source/ControlRigEditor/Private/Editor/SModularRigEventQueueView.h
new file mode 100644
--- /dev/null
+++ b/Engine/Plugins/Animation/ControlRig/Source/ControlRigEditor/Private/Editor/SModularRigEventQueueView.h
@@ -0,0 +1,122 @@
+// Copyright Epic Games, Inc. All Rights Reserved.
+
+#pragma once
+
+#include "CoreMinimal.h"
+#include "ModularRig.h"
+#include "Editor/RigVMEditor.h"
+#include "Widgets/Views/STreeView.h"
+
+#define UE_API CONTROLRIGEDITOR_API
+
+class SModularRigEventQueueView;
+class FUICommandList;
+class SSearchBox;
+class IControlRigBaseEditor;
+
+/** An event in the queue */
+class FModularRigEventEntry : public TSharedFromThis<FModularRigEventEntry>
+{
+public:
+	UE_API FModularRigEventEntry(int32 InEventIndex, const FName& InEventName, const FModuleInstanceHandle& InModule, bool InExecuted, double InDurationMicroSeconds);
+
+	UE_API TSharedRef<ITableRow> MakeTreeRowWidget(const TSharedRef<STableViewBase>& InOwnerTable, TSharedRef<FModularRigEventEntry> InEntry, TSharedRef<FUICommandList> InCommandList, TSharedPtr<SModularRigEventQueueView> InEventQueueView);
+
+	void UpdateDurationMicroSeconds(double InDurationMicroSeconds);
+
+	int32 EventIndex;
+	FName EventName;
+	FModuleInstanceHandle Module;
+	bool bExecuted = false;
+	double MicroSeconds = 0.0;
+	TArray<double> MicroSecondsFrames;
+};
+
+class SModularRigEventItem : public STableRow<TSharedPtr<FModularRigEventEntry>>
+{
+	SLATE_BEGIN_ARGS(SModularRigEventItem) {}
+	SLATE_END_ARGS()
+
+	void Construct(const FArguments& InArgs, const TSharedRef<STableViewBase>& OwnerTable, TSharedRef<FModularRigEventEntry> InEventEntry, TSharedRef<FUICommandList> InCommandList);
+
+private:
+	TWeakPtr<FModularRigEventEntry> WeakEventEntry;
+	TWeakPtr<FUICommandList> WeakCommandList;
+
+	FText GetIndexText() const;
+	FText GetModuleLabelText() const;
+	FSlateFontInfo GetModuleLabelFont() const;
+	FText GetEventLabelText() const;
+	FSlateFontInfo GetEventLabelFont() const;
+	FText GetTooltip() const;
+	FText GetDurationText() const;
+};
+
+class SModularRigEventQueueView : public SCompoundWidget
+{
+public:
+	SLATE_BEGIN_ARGS(SModularRigEventQueueView) {}
+	SLATE_END_ARGS()
+
+	UE_API ~SModularRigEventQueueView();
+
+	UE_API void Construct(const FArguments& InArgs, TSharedRef<IControlRigBaseEditor> InControlRigEditor);
+
+	/** Set Selection Changed */
+	UE_API void OnSelectionChanged(TSharedPtr<FModularRigEventEntry> Selection, ESelectInfo::Type SelectInfo);
+
+	UE_API TSharedPtr< SWidget > CreateContextMenu();
+
+protected:
+
+	/** Rebuild the tree view */
+	UE_API void RefreshTreeView(const UModularRig* InModularRig);
+
+private:
+
+	/** Bind commands that this widget handles */
+	UE_API void BindCommands();
+
+	/** Make a row widget for the table */
+	UE_API TSharedRef<ITableRow> MakeTableRowWidget(TSharedPtr<FModularRigEventEntry> InItem, const TSharedRef<STableViewBase>& OwnerTable);
+
+	/** Get children for the tree */
+	UE_API void HandleGetChildrenForTree(TSharedPtr<FModularRigEventEntry> InItem, TArray<TSharedPtr<FModularRigEventEntry>>& OutChildren);
+
+	/** Focus on the selected module in the rig */
+	UE_API void HandleFocusOnSelectedModule();
+
+	/** Search box widget */
+	TSharedPtr<SSearchBox> FilterBox;
+	FText FilterText;
+	UE_API void OnFilterTextChanged(const FText& SearchText);
+
+	UE_API void HandlePreviewHostUpdated(IRigVMEditor* InEditor);
+	UE_API void HandleModularRigExecuted(class URigVMHost* InModularRig, const FName& InEventName);
+	UE_API void HandleModularRigModelModified(EModularRigNotification InNotification, const FRigModuleReference* InModuleReference);
+
+	/**
+	 * Populate the event queue with descriptive names for each event
+	 * Returns true if the treeview has to refresh.
+	 */
+	UE_API bool PopulateEventQueueView(const UModularRig* InModularRig);
+
+	TSharedPtr<STreeView<TSharedPtr<FModularRigEventEntry>>> TreeView;
+
+	/** Command list we bind to */
+	TSharedPtr<FUICommandList> CommandList;
+
+	TWeakPtr<IControlRigBaseEditor> WeakEditor;
+	TWeakObjectPtr<UModularRig> WeakModularRig;
+	TWeakObjectPtr<UModularRigController> WeakModularRigController;
+
+	TArray<TSharedPtr<FModularRigEventEntry>> Events;
+	uint32 LastHash = 0;
+	bool bIsSelecting = false;
+
+	FDelegateHandle OnPreviewHostUpdatedHandle;
+	FDelegateHandle OnControlRigExecutedHandle;
+	FDelegateHandle OnModularRigModelModifiedHandle;
+};
+
+#undef UE_API
