Control Rig: Event Queue widget adjustments

#rnx

diff --git a/Engine/Plugins/Animation/ControlRig/Source/ControlRigEditor/Private/Editor/SModularRigEventQueueView.cpp b/Engine/Plugins/Animation/ControlRig/Source/ControlRigEditor/Private/Editor/SModularRigEventQueueView.cpp
--- a/Engine/Plugins/Animation/ControlRig/Source/ControlRigEditor/Private/Editor/SModularRigEventQueueView.cpp
+++ b/Engine/Plugins/Animation/ControlRig/Source/ControlRigEditor/Private/Editor/SModularRigEventQueueView.cpp
@@ -101,6 +101,7 @@ void SModularRigEventItem::Construct(const FArguments& InArgs, const TSharedRef<
 				]
 				
 				+ SHorizontalBox::Slot()
+				.Padding(4.f, 1, 1, 1)
 				.FillWidth(1.f)
 				.AutoWidth()
 				.VAlign(VAlign_Center)
@@ -365,9 +366,11 @@ bool SModularRigEventQueueView::PopulateEventQueueView(const UModularRig* InModu
 {
 	if (InModularRig)
 	{
+		const FString FilterString = FilterText.ToString();
+		
 		// copy the queue
 		const FRigModuleExecutionQueue Queue = InModularRig->GetLastExecutionQueue();
-		const uint32 Hash = GetTypeHash(Queue);
+		const uint32 Hash = HashCombine(GetTypeHash(Queue), GetTypeHash(FilterString));
 		if (Hash != LastHash)
 		{
 			Events.Reset();
@@ -375,6 +378,14 @@ bool SModularRigEventQueueView::PopulateEventQueueView(const UModularRig* InModu
 			for (int32 Index = 0; Index < Queue.Elements.Num(); Index++)
 			{
 				const FRigModuleExecutionElement& Element = Queue.Elements[Index];
+				if (!FilterString.IsEmpty())
+				{
+					if (!Element.EventName.ToString().Contains(FilterString, ESearchCase::IgnoreCase) &&
+						!Element.ModuleName.ToString().Contains(FilterString, ESearchCase::IgnoreCase))
+					{
+						continue;
+					}
+				}
 				TSharedPtr<FModularRigEventEntry> NewEntry = MakeShared<FModularRigEventEntry>(Index, Element.EventName, FModuleInstanceHandle(InModularRig, Element.ModuleInstance), Element.bExecuted, Element.DurationInMicroSeconds);
 				Events.Add(NewEntry);
 			}
@@ -384,10 +395,10 @@ bool SModularRigEventQueueView::PopulateEventQueueView(const UModularRig* InModu
 		}
 		else
 		{
-			check(Events.Num() == Queue.Elements.Num());
-			for (int32 Index = 0; Index < Queue.Elements.Num(); Index++)
+			for (int32 Index = 0; Index < Events.Num(); Index++)
 			{
-				Events[Index]->UpdateDurationMicroSeconds(Queue.Elements[Index].DurationInMicroSeconds);
+				const int32 QueueIndex = Events[Index]->EventIndex;
+				Events[Index]->UpdateDurationMicroSeconds(Queue.Elements[QueueIndex].DurationInMicroSeconds);
 			}
 		}
 	}
