Control Rig: Move to depth first traverel for module for each

#rnx

diff --git a/Engine/Plugins/Animation/ControlRig/Source/ControlRig/Private/ModularRig.cpp b/Engine/Plugins/Animation/ControlRig/Source/ControlRig/Private/ModularRig.cpp
--- a/Engine/Plugins/Animation/ControlRig/Source/ControlRig/Private/ModularRig.cpp
+++ b/Engine/Plugins/Animation/ControlRig/Source/ControlRig/Private/ModularRig.cpp
@@ -464,13 +464,30 @@ bool UModularRig::Execute_Internal(const FName& InEventName)
 			return true;
 		});
 
+#if WITH_EDITOR
+		if (EventQueueToRun.IsEmpty() || EventQueueToRun[0] == InEventName)
+		{
+			LastExecutedElements.Reset();
+		}
+#endif
+
 		ExecuteQueue();
 
 		if (bAccumulateTime)
 		{
 			AbsoluteTime += DeltaTime;
 		}
-		
+
+#if WITH_EDITOR
+		// copy the elements over while applying a lock
+		if (EventQueueToRun.IsEmpty() || EventQueueToRun.Last() == InEventName)
+		{
+			FWriteScopeLock WriteLock(LastExecutionQueueLock);
+			LastExecutionQueue.Elements = LastExecutedElements;
+			LastExecutedElements.Reset();
+		}
+#endif
+
 		return true;
 	}
 	return false;
@@ -547,10 +564,6 @@ const FModularRigSettings& UModularRig::GetModularRigSettings() const
 
 void UModularRig::ExecuteQueue()
 {
-#if WITH_EDITOR
-	LastExecutedElements.Reset();
-#endif
-	
 	FRigVMExtendedExecuteContext& Context = GetRigVMExtendedExecuteContext();
 	FControlRigExecuteContext& PublicContext = Context.GetPublicDataSafe<FControlRigExecuteContext>();
 	URigHierarchy* Hierarchy = GetHierarchy();
@@ -713,15 +726,6 @@ void UModularRig::ExecuteQueue()
 	{
 		Module->PostConstructionSpawnStartIndex = Hierarchy->GetNextSpawnIndex();
 	}
-
-#if WITH_EDITOR
-	// copy the elements over while applying a lock
-	{
-		FWriteScopeLock WriteLock(LastExecutionQueueLock);
-		LastExecutionQueue.Elements = LastExecutedElements;
-		LastExecutedElements.Reset();
-	}
-#endif
 }
 
 void UModularRig::ResetExecutionQueue()
@@ -1253,33 +1257,57 @@ FName UModularRig::GetParentModuleName(const FName& InModuleName) const
 	return NAME_None;
 }
 
-void UModularRig::ForEachModule(TFunctionRef<bool(FRigModuleInstance*)> PerModuleFunction)
+void UModularRig::ForEachModule(TFunctionRef<bool(FRigModuleInstance*)> PerModuleFunction, bool bDepthFirst)
 {
-	TArray<FRigModuleInstance*> ModuleInstances;
-	ModuleInstances.Reserve(Modules.Num());
-	ModuleInstances.Append(RootModules);
-	
-	for (int32 ModuleIndex = 0; ModuleIndex < ModuleInstances.Num(); ++ModuleIndex)
+	if (bDepthFirst)
 	{
-		if (!PerModuleFunction(ModuleInstances[ModuleIndex]))
+		for (FRigModuleInstance* Module : RootModules)
 		{
-			break;
+			if (!TraverseModules(Module, PerModuleFunction))
+			{
+				break;
+			}
+		}
+	}
+	else
+	{
+		TArray<FRigModuleInstance*> ModuleInstances;
+		ModuleInstances.Reserve(Modules.Num());
+		ModuleInstances.Append(RootModules);
+		
+		for (int32 ModuleIndex = 0; ModuleIndex < ModuleInstances.Num(); ++ModuleIndex)
+		{
+			if (!PerModuleFunction(ModuleInstances[ModuleIndex]))
+			{
+				break;
+			}
+			ModuleInstances.Append(ModuleInstances[ModuleIndex]->CachedChildren);
 		}
-		ModuleInstances.Append(ModuleInstances[ModuleIndex]->CachedChildren);
 	}
 }
 
-void UModularRig::ForEachModule(TFunctionRef<bool(const FRigModuleInstance*)> PerModuleFunction) const
+void UModularRig::ForEachModule(TFunctionRef<bool(const FRigModuleInstance*)> PerModuleFunction, bool bDepthFirst) const
 {
-	TArray<FRigModuleInstance*> ModuleInstances = RootModules;
-	for (int32 ModuleIndex = 0; ModuleIndex < ModuleInstances.Num(); ++ModuleIndex)
+	const_cast<UModularRig*>(this)->ForEachModule([PerModuleFunction](FRigModuleInstance* InModuleInstance) -> bool
 	{
-		if (!PerModuleFunction(ModuleInstances[ModuleIndex]))
+		return PerModuleFunction(InModuleInstance);
+	}, bDepthFirst);
+}
+
+bool UModularRig::TraverseModules(FRigModuleInstance* InModuleInstance, TFunctionRef<bool(FRigModuleInstance*)> PerModule)
+{
+	if (!PerModule(InModuleInstance))
+	{
+		return false;
+	}
+	for (FRigModuleInstance* ChildModule : InModuleInstance->CachedChildren)
+	{
+		if (!TraverseModules(ChildModule, PerModule))
 		{
-			break;
+			return false;
 		}
-		ModuleInstances.Append(ModuleInstances[ModuleIndex]->CachedChildren);
 	}
+	return true;
 }
 
 void UModularRig::ExecuteConnectorEvent(const FRigElementKey& InConnector, const FRigModuleInstance* InModuleInstance, const FRigElementKeyRedirector* InRedirector, TArray<FRigElementResolveResult>& InOutCandidates)
diff --git a/Engine/Plugins/Animation/ControlRig/Source/ControlRig/Private/ModularRigModel.cpp b/Engine/Plugins/Animation/ControlRig/Source/ControlRig/Private/ModularRigModel.cpp
--- a/Engine/Plugins/Animation/ControlRig/Source/ControlRig/Private/ModularRigModel.cpp
+++ b/Engine/Plugins/Animation/ControlRig/Source/ControlRig/Private/ModularRigModel.cpp
@@ -470,17 +470,46 @@ TArray<const FRigModuleReference*> FModularRigModel::FindModuleInstancesOfClass(
 	return FindModuleInstancesOfClass(*InClass.ToString());
 }
 
-void FModularRigModel::ForEachModule(TFunction<bool(const FRigModuleReference*)> PerModule) const
+void FModularRigModel::ForEachModule(TFunction<bool(const FRigModuleReference*)> PerModule, bool bDepthFirst) const
 {
-	TArray<FRigModuleReference*> ModuleInstances = RootModules;
-	for (int32 Index=0; Index < ModuleInstances.Num(); ++Index)
+	if (bDepthFirst)
 	{
-		if (!PerModule(ModuleInstances[Index]))
+		for(const FRigModuleReference* RootModule : RootModules)
 		{
-			break;
+			if (!TraverseModules(RootModule, PerModule))
+			{
+				break;
+			}
+		}
+	}
+	else
+	{
+		TArray<FRigModuleReference*> ModuleInstances = RootModules;
+		for (int32 Index=0; Index < ModuleInstances.Num(); ++Index)
+		{
+			if (!PerModule(ModuleInstances[Index]))
+			{
+				break;
+			}
+			ModuleInstances.Append(ModuleInstances[Index]->CachedChildren);
+		}
+	}
+}
+
+bool FModularRigModel::TraverseModules(const FRigModuleReference* InModuleReference, TFunction<bool(const FRigModuleReference*)> PerModule)
+{
+	if (!PerModule(InModuleReference))
+	{
+		return false;
+	}
+	for (const FRigModuleReference* ChildModule : InModuleReference->CachedChildren)
+	{
+		if (!TraverseModules(ChildModule, PerModule))
+		{
+			return false;
 		}
-		ModuleInstances.Append(ModuleInstances[Index]->CachedChildren);
 	}
+	return true;
 }
 
 TArray<FName> FModularRigModel::SortModuleNames(const TArray<FName>& InModuleNames) const
diff --git a/Engine/Plugins/Animation/ControlRig/Source/ControlRig/Public/ModularRig.h b/Engine/Plugins/Animation/ControlRig/Source/ControlRig/Public/ModularRig.h
--- a/Engine/Plugins/Animation/ControlRig/Source/ControlRig/Public/ModularRig.h
+++ b/Engine/Plugins/Animation/ControlRig/Source/ControlRig/Public/ModularRig.h
@@ -272,8 +272,8 @@ public:
 	UE_API FName GetParentModuleNameForBP(FName InModuleName) const;
 	UE_API FName GetParentModuleName(const FName& InModuleName) const;
 
-	UE_API void ForEachModule(TFunctionRef<bool(FRigModuleInstance*)> PerModuleFunction);
-	UE_API void ForEachModule(TFunctionRef<bool(const FRigModuleInstance*)> PerModuleFunction) const;
+	UE_API void ForEachModule(TFunctionRef<bool(FRigModuleInstance*)> PerModuleFunction, bool bDepthFirst = true);
+	UE_API void ForEachModule(TFunctionRef<bool(const FRigModuleInstance*)> PerModuleFunction, bool bDepthFirst = true) const;
 
 	UE_API void ExecuteConnectorEvent(const FRigElementKey& InConnector, const FRigModuleInstance* InModuleInstance, const FRigElementKeyRedirector* InRedirector, TArray<FRigElementResolveResult>& InOutCandidates);
 
@@ -349,6 +349,8 @@ private:
 	/** Destroys / discards a module instance rig */
 	static UE_API void DiscardModuleRig(UControlRig* InControlRig);
 
+	static bool TraverseModules(FRigModuleInstance* InModuleInstance, TFunctionRef<bool(FRigModuleInstance*)> PerModule);
+
 	TMap<FName, UControlRig*> PreviousModuleRigs;
 
 	TArray<FRigModuleExecutionElement> LastExecutedElements;
diff --git a/Engine/Plugins/Animation/ControlRig/Source/ControlRig/Public/ModularRigModel.h b/Engine/Plugins/Animation/ControlRig/Source/ControlRig/Public/ModularRigModel.h
--- a/Engine/Plugins/Animation/ControlRig/Source/ControlRig/Public/ModularRigModel.h
+++ b/Engine/Plugins/Animation/ControlRig/Source/ControlRig/Public/ModularRigModel.h
@@ -395,7 +395,7 @@ public:
 	UE_API FRigModuleReference* GetParentModule(const FRigModuleReference* InChildModule);
 	UE_API const FRigModuleReference* GetParentModule(const FRigModuleReference* InChildModule) const;
 
-	UE_API void ForEachModule(TFunction<bool(const FRigModuleReference*)> PerModule) const;
+	UE_API void ForEachModule(TFunction<bool(const FRigModuleReference*)> PerModule, bool bDepthFirst = true) const;
 
 	UE_API TArray<FName> SortModuleNames(const TArray<FName>& InModuleNames) const;
 
@@ -408,6 +408,9 @@ public:
 	UE_API TArray<const FRigModuleReference*> FindModuleInstancesOfClass(TSoftClassPtr<UControlRig> InClass) const;
 
 private:
+
+	static bool TraverseModules(const FRigModuleReference* InModuleReference, TFunction<bool(const FRigModuleReference*)> PerModule);
+
 	TWeakObjectPtr<UObject> OuterClientHost;
 	TArray<FName> SelectedModuleNames;
 
