Control Rig: EngineTest example for programmatic VM / Control Rig

#rnx

diff --git a/Engine/Plugins/Animation/ControlRig/Source/ControlRig/Private/ControlRig.cpp b/Engine/Plugins/Animation/ControlRig/Source/ControlRig/Private/ControlRig.cpp
--- a/Engine/Plugins/Animation/ControlRig/Source/ControlRig/Private/ControlRig.cpp
+++ b/Engine/Plugins/Animation/ControlRig/Source/ControlRig/Private/ControlRig.cpp
@@ -3799,7 +3799,7 @@ void UControlRig::PostInitInstance(URigVMHost* InCDO)
 void UControlRig::SetDynamicHierarchy(TObjectPtr<URigHierarchy> InHierarchy)
 {
 	// Delete any existing hierarchy
-	if (DynamicHierarchy->GetOuter() == this)
+	if (DynamicHierarchy && DynamicHierarchy->GetOuter() == this)
 	{
 		DynamicHierarchy->OnUndoRedo().RemoveAll(this);
 		DynamicHierarchy->Rename(nullptr, GetTransientPackage(), REN_DoNotDirty | REN_DontCreateRedirectors | REN_NonTransactional);
diff --git a/Engine/Plugins/Runtime/RigVM/Source/RigVM/Private/RigVMHost.cpp b/Engine/Plugins/Runtime/RigVM/Source/RigVM/Private/RigVMHost.cpp
--- a/Engine/Plugins/Runtime/RigVM/Source/RigVM/Private/RigVMHost.cpp
+++ b/Engine/Plugins/Runtime/RigVM/Source/RigVM/Private/RigVMHost.cpp
@@ -1389,7 +1389,7 @@ void URigVMHost::InitializeFromCDO()
 		// we initialize all other instances of Control Rig from the CDO here
 		URigVMHost* CDO = GetClass()->GetDefaultObject<URigVMHost>();
 
-		if (VM == nullptr || VM != CDO->VM)
+		if ((VM == nullptr || VM != CDO->VM) && CDO->VM)
 		{
 			VM = CDO->VM;
 		}
diff --git a/Engine/Plugins/Runtime/RigVM/Source/RigVM/Public/RigVMHost.h b/Engine/Plugins/Runtime/RigVM/Source/RigVM/Public/RigVMHost.h
--- a/Engine/Plugins/Runtime/RigVM/Source/RigVM/Public/RigVMHost.h
+++ b/Engine/Plugins/Runtime/RigVM/Source/RigVM/Public/RigVMHost.h
@@ -215,7 +215,7 @@ public:
 	 * @param InEventIndex Deprecated argument. Not used.
 	 */
 	UFUNCTION(BlueprintCallable, Category = "RigVM")
-	UE_API void RequestRunOnceEvent(const FName& InEventName, int32 InEventIndex = -1);
+	UE_API virtual void RequestRunOnceEvent(const FName& InEventName, int32 InEventIndex = -1);
 
 	/** Removes an event running once */
 	UFUNCTION(BlueprintCallable, Category = "RigVM")
