Avoid reserving transform space for non transform elements (curves and connectors) at RigVMHierarchy sort storage function

diff --git a/Engine/Plugins/Animation/ControlRig/Source/ControlRig/Private/Rigs/RigHierarchy.cpp b/Engine/Plugins/Animation/ControlRig/Source/ControlRig/Private/Rigs/RigHierarchy.cpp
--- a/Engine/Plugins/Animation/ControlRig/Source/ControlRig/Private/Rigs/RigHierarchy.cpp
+++ b/Engine/Plugins/Animation/ControlRig/Source/ControlRig/Private/Rigs/RigHierarchy.cpp
@@ -7395,17 +7395,20 @@ bool URigHierarchy::SortElementStorage()
 	int32 NumElementTypes = 0;
 	int32 TotalNumberOfTransformsPerTransformType = 0;
 
-	for(int32 ElementType = 0; ElementType <= LastElementType; ElementType++)
+	for (int32 ElementTypeFlatIndex = 0; ElementTypeFlatIndex <= LastElementType; ElementTypeFlatIndex++)
 	{
-		NumElementsPerType.Add(ElementsPerType[ElementType].Num());
-		PerElementStride.Add(ElementType == ControlElementType ? static_cast<int32>(ERigTransformStorageType::NumStorageTypes) : 1);
+		const ERigElementType ElementTypeFlag = FlatIndexToRigElementType(ElementTypeFlatIndex);
+		const bool bReserveSpace = (((int32)ElementTypeFlag & (int32)ERigElementType::TransformTypes) != 0); // avoid reserving space for non transform types
 
-		if(ElementType > 0)
+		NumElementsPerType.Add(bReserveSpace ? ElementsPerType[ElementTypeFlatIndex].Num() : 0);
+		PerElementStride.Add(ElementTypeFlatIndex == ControlElementType ? static_cast<int32>(ERigTransformStorageType::NumStorageTypes) : 1);
+
+		if (ElementTypeFlatIndex > 0)
 		{
-			OffsetPerElementType[ElementType] = OffsetPerElementType[ElementType-1] + NumElementsPerType[ElementType-1] * PerElementStride[ElementType-1]; 
+			OffsetPerElementType[ElementTypeFlatIndex] = OffsetPerElementType[ElementTypeFlatIndex - 1] + NumElementsPerType[ElementTypeFlatIndex - 1] * PerElementStride[ElementTypeFlatIndex - 1];
 		}
-		
-		NumElementTypes += ElementsPerType[ElementType].IsEmpty() ? 0 : 1;
+
+		NumElementTypes += (bReserveSpace && ElementsPerType[ElementTypeFlatIndex].Num() > 0) ? 1 : 0;
 		TotalNumberOfTransformsPerTransformType += NumElementsPerType.Last() * PerElementStride.Last();
 	}
 
diff --git a/Engine/Plugins/Animation/ControlRig/Source/ControlRig/Public/Rigs/RigHierarchyDefines.h b/Engine/Plugins/Animation/ControlRig/Source/ControlRig/Public/Rigs/RigHierarchyDefines.h
--- a/Engine/Plugins/Animation/ControlRig/Source/ControlRig/Public/Rigs/RigHierarchyDefines.h
+++ b/Engine/Plugins/Animation/ControlRig/Source/ControlRig/Public/Rigs/RigHierarchyDefines.h
@@ -43,6 +43,7 @@ enum class ERigElementType : uint8
 	First = Bone UMETA(Hidden), 
 	Last = Socket UMETA(Hidden), 
 	ToResetAfterConstructionEvent = Bone | Control | Curve | Socket UMETA(Hidden),
+	TransformTypes = Bone | Null | Control | Reference | Socket UMETA(Hidden),
 };
 
 UENUM(BlueprintType)
