Fix missing required forward declare for Clang
#rnx

diff --git a/Engine/Plugins/Animation/ControlRig/Source/ControlRigEditor/Private/Misc/Flyout/FlyoutOverlayManager.h b/Engine/Plugins/Animation/ControlRig/Source/ControlRigEditor/Private/Misc/Flyout/FlyoutOverlayManager.h
--- a/Engine/Plugins/Animation/ControlRig/Source/ControlRigEditor/Private/Misc/Flyout/FlyoutOverlayManager.h
+++ b/Engine/Plugins/Animation/ControlRig/Source/ControlRigEditor/Private/Misc/Flyout/FlyoutOverlayManager.h
@@ -11,6 +11,7 @@
 // flexibility of private code.
 // Please DO NOT introduce ControlRig specific code here; instead invert dependencies by injecting Control Rig specific behaviour.
 
+class FFlyoutTemporaryPositionOverride;
 class IMenu;
 
 namespace UE::ToolWidgets { class SDraggableBoxOverlay; }
