Fix MSVC being overly pedantic
#rnx

diff --git a/Engine/Plugins/Animation/ControlRig/Source/ControlRigEditor/Private/Misc/Flyout/TemporaryPositionOverride.h b/Engine/Plugins/Animation/ControlRig/Source/ControlRigEditor/Private/Misc/Flyout/TemporaryPositionOverride.h
--- a/Engine/Plugins/Animation/ControlRig/Source/ControlRigEditor/Private/Misc/Flyout/TemporaryPositionOverride.h
+++ b/Engine/Plugins/Animation/ControlRig/Source/ControlRigEditor/Private/Misc/Flyout/TemporaryPositionOverride.h
@@ -12,7 +12,7 @@ class FFlyoutOverlayManager;
 /** Keeps the flyout widget at the temporary override position until this object is destroyed, at which point it calls StopTemporaryWidgetPosition. */
 class FFlyoutTemporaryPositionOverride : public FNoncopyable
 {
-	friend FFlyoutOverlayManager;
+	friend class FFlyoutOverlayManager;
 	struct FPrivateToken { explicit FPrivateToken() = default; };
 	
 	FFlyoutOverlayManager* ManagerToRestore;
