Control Rig: Return empty array for unresolved array connectors

#rnx

diff --git a/Engine/Plugins/Animation/ControlRig/Source/ControlRig/Private/Units/Execution/RigUnit_RigModules.cpp b/Engine/Plugins/Animation/ControlRig/Source/ControlRig/Private/Units/Execution/RigUnit_RigModules.cpp
--- a/Engine/Plugins/Animation/ControlRig/Source/ControlRig/Private/Units/Execution/RigUnit_RigModules.cpp
+++ b/Engine/Plugins/Animation/ControlRig/Source/ControlRig/Private/Units/Execution/RigUnit_RigModules.cpp
@@ -64,6 +64,10 @@ FRigUnit_ResolveArrayConnector_Execute()
 	if(Result.Num() == 1)
 	{
 		bIsConnected = Result[0] != Connector;
+		if (!bIsConnected)
+		{
+			Result.Reset();
+		}
 	}
 	else if(Result.Num() > 1)
 	{
