Control Rig: Fix modular rigs not resolving secondary connectos correctly
#rnx

diff --git a/Engine/Plugins/Animation/ControlRig/Source/ControlRig/Private/ModularRigController.cpp b/Engine/Plugins/Animation/ControlRig/Source/ControlRig/Private/ModularRigController.cpp
--- a/Engine/Plugins/Animation/ControlRig/Source/ControlRig/Private/ModularRigController.cpp
+++ b/Engine/Plugins/Animation/ControlRig/Source/ControlRig/Private/ModularRigController.cpp
@@ -364,12 +364,14 @@ bool UModularRigController::ConnectConnectorToElements(const FRigElementKey& InC
 		}
 	}
 
+	Notify(EModularRigNotification::ConnectionChanged, Module);
+	
 #if WITH_EDITOR
 	if (UControlRig* RigCDO = Module->Class->GetDefaultObject<UControlRig>())
 	{
 		if (ModularRig)
 		{
-			TGuardValue<bool> NotificationsGuard(bSuspendNotifications, true);
+			//TGuardValue<bool> NotificationsGuard(bSuspendNotifications, true);
 
 			if (URigHierarchy* Hierarchy = ModularRig->GetHierarchy())
 			{
diff --git a/Engine/Plugins/Animation/ControlRig/Source/ControlRigEditor/Private/Editor/ControlRigEditor.cpp b/Engine/Plugins/Animation/ControlRig/Source/ControlRigEditor/Private/Editor/ControlRigEditor.cpp
--- a/Engine/Plugins/Animation/ControlRig/Source/ControlRigEditor/Private/Editor/ControlRigEditor.cpp
+++ b/Engine/Plugins/Animation/ControlRig/Source/ControlRigEditor/Private/Editor/ControlRigEditor.cpp
@@ -6308,7 +6308,7 @@ void FControlRigBaseEditor::OnPostConstruction_AnyThread(UControlRig* InRig, con
 							{
 								if(UModularRigController* ModularRigController = RigBlueprint->ModularRigModel.GetController())
 								{
-									(void)ModularRigController->ConnectConnectorToElement(PrimaryConnector->GetKey(), RootBone->GetKey(), false);
+									(void)ModularRigController->ConnectConnectorToElement(PrimaryConnector->GetKey(), RootBone->GetKey(), false, false, false);
 								}
 							}
 						}
@@ -6324,7 +6324,7 @@ void FControlRigBaseEditor::OnPostConstruction_AnyThread(UControlRig* InRig, con
 							{
 								if(UModularRigController* ModularRigController = RigBlueprint->ModularRigModel.GetController())
 								{
-									(void)ModularRigController->ConnectConnectorToElement(PrimaryConnector->GetKey(), RootBone->GetKey(), false);
+									(void)ModularRigController->ConnectConnectorToElement(PrimaryConnector->GetKey(), RootBone->GetKey(), false, false, false);
 								}
 							}
 						}
