Anim Layers: Added layer not showing up in sub scene. Need to make sure HandleOnAnimListChanged is set on the correct AnimLayer Object

diff --git a/Engine/Plugins/Animation/ControlRig/Source/ControlRigEditor/Private/Sequencer/AnimLayers/SAnimLayers.cpp b/Engine/Plugins/Animation/ControlRig/Source/ControlRigEditor/Private/Sequencer/AnimLayers/SAnimLayers.cpp
--- a/Engine/Plugins/Animation/ControlRig/Source/ControlRigEditor/Private/Sequencer/AnimLayers/SAnimLayers.cpp
+++ b/Engine/Plugins/Animation/ControlRig/Source/ControlRigEditor/Private/Sequencer/AnimLayers/SAnimLayers.cpp
@@ -1790,6 +1790,11 @@ void SAnimLayers::Tick(const FGeometry& AllottedGeometry, const double InCurrent
 		if (LastMovieSceneSig != CurrentMovieSceneSig)
 		{
 			LastMovieSceneSig = CurrentMovieSceneSig;
+			if (UAnimLayers* AnimLayersPtr = UAnimLayers::GetAnimLayers(Sequencer))
+			{
+				AnimLayersPtr->AnimLayerListChanged().AddRaw(AnimLayerController.Get(), &FAnimLayerController::HandleOnAnimLayerListChanged);
+				AnimLayers = AnimLayersPtr;
+			}
 			AnimLayerController->RefreshSelectionData();
 		}
 	}
