Fixed circular includes by removing FlyoutOverlayManager include in TemporaryPositionOverride

diff --git a/Engine/Plugins/Animation/ControlRig/Source/ControlRigEditor/Private/Misc/Flyout/TemporaryPositionOverride.h b/Engine/Plugins/Animation/ControlRig/Source/ControlRigEditor/Private/Misc/Flyout/TemporaryPositionOverride.h
--- a/Engine/Plugins/Animation/ControlRig/Source/ControlRigEditor/Private/Misc/Flyout/TemporaryPositionOverride.h
+++ b/Engine/Plugins/Animation/ControlRig/Source/ControlRigEditor/Private/Misc/Flyout/TemporaryPositionOverride.h
@@ -1,8 +1,7 @@
-﻿// Copyright Epic Games, Inc. All Rights Reserved.
+// Copyright Epic Games, Inc. All Rights Reserved.
 
 #pragma once
 
-#include "FlyoutOverlayManager.h"
 #include "Templates/UnrealTemplate.h"
 
 namespace UE::ControlRigEditor
