Anim In Engine: Add Spawn Restore State to Snapper/Motion trails. Also since we evaluate motion trails over many frames we need to disable Sequencer::PostUndo since that's slow and refreshes the UI, which makes the motion trail calculation look buggy, so added some ISequencer functions and corresponding RestoreState to handle that. Also fixed issue with snapping over a frame when doing full frames.

diff --git a/Engine/Plugins/Animation/ControlRig/Source/ControlRigEditor/Private/Tools/ControlRigSnapper.cpp b/Engine/Plugins/Animation/ControlRig/Source/ControlRigEditor/Private/Tools/ControlRigSnapper.cpp
--- a/Engine/Plugins/Animation/ControlRig/Source/ControlRigEditor/Private/Tools/ControlRigSnapper.cpp
+++ b/Engine/Plugins/Animation/ControlRig/Source/ControlRigEditor/Private/Tools/ControlRigSnapper.cpp
@@ -172,7 +172,7 @@ bool FControlRigSnapper::GetControlRigControlTransforms(ISequencer* Sequencer,
 	{
 		FMovieSceneSequenceIDRef Template = Sequencer->GetFocusedTemplateID();
 		FMovieSceneSequenceTransform RootToLocalTransform = Sequencer->GetFocusedMovieSceneSequenceTransform();
-		TOptional<FFrameNumber> FrameNumber = Sequencer->GetLocalTime().Time.RoundToFrame();
+		TOptional<FFrameNumber> FrameNumber;
 		return LocalGetControlRigControlTransforms(Sequencer, FrameNumber, Sequencer->GetFocusedMovieSceneSequence(), Template, RootToLocalTransform,
 			ControlRig, ControlName, Frames, ParentTransforms, OutTransforms);
 	
@@ -546,22 +546,29 @@ bool FControlRigSnapper::SnapIt(FFrameNumber StartFrame, FFrameNumber EndFrame,c
 			UE_LOG(LogControlRig, Error, TEXT("No Frames To Snap, probably no keys within time range"));
 			return false;
 		}
-		Sequencer->ForceEvaluate(); // force an evaluate so any control rig get's binding setup
+
+		//if in subsquence and we want to, turn on ShouldEvaluateSubSequencesInIsolation, use the following object for that
+		//disabling this for nowUE::Sequencer::FSequencerBakingSetupRestore RestoreBaking(Sequencer);
+		FScopedTransaction ScopedTransaction(LOCTEXT("SnapAnimation", "Snap Animation"));
+
+		FSpawnableRestoreState SpawnableRestoreState(MovieScene, Sequencer->GetSharedPlaybackState().ToSharedPtr());
+
 		TArray<FTransform> WorldTransformToSnap;
 		bool bSnapToFirstFrameNotParents = !CalculateWorldTransformsFromParents(Sequencer.Get(), ParentToSnap, Frames, WorldTransformToSnap);
 		if (Frames.Num() != WorldTransformToSnap.Num())
 		{
 			UE_LOG(LogControlRig, Error, TEXT("Number of Frames %d to Snap different than Parent Frames %d"), Frames.Num(),WorldTransformToSnap.Num());
+			ScopedTransaction.Cancel();
 			return false;
 		}
 
-		//if in subsquence and we want to, turn on ShouldEvaluateSubSequencesInIsolation, use the following object for that
-		//disabling this for nowUE::Sequencer::FSequencerBakingSetupRestore RestoreBaking(Sequencer);
+		const bool bAvoidEvaluates = (Frames.Num() == 0) || (Frames.Num() == 1 && Frames[0] == Sequencer->GetLocalTime().Time.RoundToFrame());
 
-		FScopedTransaction ScopedTransaction(LOCTEXT("SnapAnimation", "Snap Animation"));
-	
 		//need to make sure binding is setup
-		Sequencer->ForceEvaluate();
+		if (bAvoidEvaluates == false)
+		{
+			Sequencer->ForceEvaluate();
+		}
 		
 		TArray<FGuidAndActor > ActorsToSnap;
 		//There may be Actors here not in Sequencer so we add them to sequencer also
diff --git a/Engine/Plugins/Animation/ControlRig/Source/ControlRigEditor/Private/Tools/EvaluateSequencerTools.cpp b/Engine/Plugins/Animation/ControlRig/Source/ControlRigEditor/Private/Tools/EvaluateSequencerTools.cpp
--- a/Engine/Plugins/Animation/ControlRig/Source/ControlRigEditor/Private/Tools/EvaluateSequencerTools.cpp
+++ b/Engine/Plugins/Animation/ControlRig/Source/ControlRigEditor/Private/Tools/EvaluateSequencerTools.cpp
@@ -7,6 +7,7 @@
 #include "MovieSceneSequence.h"
 #include "MovieScene.h"
 #include "MovieSceneToolHelpers.h"
+#include "EntitySystem/MovieSceneSharedPlaybackState.h"
 #include "ControlRigObjectBinding.h"
 #include "Evaluation/MovieSceneEvaluationTemplateInstance.h"
 #include "ILevelSequenceEditorToolkit.h"
@@ -199,8 +200,13 @@ namespace AIE
 		{
 			return false;
 		}
+		UMovieScene* MovieScene = Sequencer->GetFocusedMovieSceneSequence()->GetMovieScene();
+		
+		FSpawnableRestoreState SpawnableRestoreState(MovieScene, Sequencer->GetSharedPlaybackState().ToSharedPtr());
+
 		const FQualifiedFrameTime QFrameTime = Sequencer->GetGlobalTime();
 		FMovieSceneContext Context = FMovieSceneContext(FMovieSceneEvaluationRange(QFrameTime.Time, QFrameTime.Rate), Sequencer->GetPlaybackStatus()).SetHasJumped(true);
+
 		//have no way to just evaluate specific tracks in Sequencer so do all of them.
 		Sequencer->GetEvaluationTemplate().EvaluateSynchronousBlocking(Context);
 		const FConstraintsManagerController& Controller = FConstraintsManagerController::Get(World);
@@ -240,6 +246,8 @@ namespace AIE
 		FFrameRate TickResolution = MovieScene->GetTickResolution();
 		FFrameRate DisplayRate = MovieScene->GetDisplayRate();
 
+		FSpawnableRestoreState SpawnableRestoreState(MovieScene, Sequencer->GetSharedPlaybackState().ToSharedPtr());
+
 		TSet<UControlRig*> ControlsToRevalaute;
 		FControlRigEditMode::FTurnOffPosePoseUpdate  TurnOff; //stop flashing
 		for (int32 Index : Indices)
@@ -257,7 +265,7 @@ namespace AIE
 				FTransform& Transform = Actor.WorldTransforms->Transforms[Index];
 				FTransform& ParentTransform = Actor.ParentTransforms->Transforms[Index];
 
-				GetActorWorldTransform(Actor.Actor, Transform,ParentTransform);
+				GetActorWorldTransform(Actor.Actor, Transform, ParentTransform);
 			}
 
 			const float DeltaTime = 0.033f;
@@ -287,7 +295,6 @@ namespace AIE
 				}
 			}
 		}
-		
 		return true;
 	}
 
diff --git a/Engine/Source/Editor/MovieSceneTools/Private/MovieSceneToolHelpers.cpp b/Engine/Source/Editor/MovieSceneTools/Private/MovieSceneToolHelpers.cpp
--- a/Engine/Source/Editor/MovieSceneTools/Private/MovieSceneToolHelpers.cpp
+++ b/Engine/Source/Editor/MovieSceneTools/Private/MovieSceneToolHelpers.cpp
@@ -4682,10 +4682,18 @@ FSpawnableRestoreState::~FSpawnableRestoreState()
 	}
 
 	// Restore modified spawned sections
-	bool bOrigSquelchTransactionNotification = GEditor->bSquelchTransactionNotification;
+	const bool bOrigSquelchTransactionNotification = GEditor->bSquelchTransactionNotification;
 	GEditor->bSquelchTransactionNotification = true;
+	const bool bOrigIgnoreSelectionChange = GEditor->bIgnoreSelectionChange;
+	GEditor->bIgnoreSelectionChange = true;
+	const bool bOrigSuspendBroadcastPostUndoRedo = GEditor->bSuspendBroadcastPostUndoRedo;
+	GEditor->bSuspendBroadcastPostUndoRedo = true;
+
 	GEditor->UndoTransaction(false);
+
 	GEditor->bSquelchTransactionNotification = bOrigSquelchTransactionNotification;
+	GEditor->bIgnoreSelectionChange = bOrigIgnoreSelectionChange;
+	GEditor->bSuspendBroadcastPostUndoRedo = bOrigSuspendBroadcastPostUndoRedo;
 }
 
 
