Control Rig: Selection now uses undo command instead of Modify call. This reduces undo stack size..

diff --git a/Engine/Plugins/Animation/ControlRig/Source/ControlRig/Private/Commands/ControlRigCommandChange.cpp b/Engine/Plugins/Animation/ControlRig/Source/ControlRig/Private/Commands/ControlRigCommandChange.cpp
new file mode 100644
--- /dev/null
+++ b/Engine/Plugins/Animation/ControlRig/Source/ControlRig/Private/Commands/ControlRigCommandChange.cpp
@@ -0,0 +1,23 @@
+﻿// Copyright Epic Games, Inc. All Rights Reserved.
+
+#include "ControlRigCommandChange.h"
+
+#include "CoreGlobals.h"
+#include "Misc/ITransaction.h"
+#if WITH_EDITOR
+#include "ScopedTransaction.h"
+#endif
+
+namespace UE::ControlRig
+{
+void FControlRigCommandChange::StoreUndo(UObject* Object, TUniquePtr<FControlRigCommandChange> Change, const FText& Description)
+{
+#if WITH_EDITOR
+	if (GUndo && Object && Change)
+	{
+		const FScopedTransaction Transaction(Description);
+		GUndo->StoreUndo(Object, MoveTemp(Change));
+	}
+#endif
+}
+}
diff --git a/Engine/Plugins/Animation/ControlRig/Source/ControlRig/Private/Commands/ControlRigCommandChange.h b/Engine/Plugins/Animation/ControlRig/Source/ControlRig/Private/Commands/ControlRigCommandChange.h
new file mode 100644
--- /dev/null
+++ b/Engine/Plugins/Animation/ControlRig/Source/ControlRig/Private/Commands/ControlRigCommandChange.h
@@ -0,0 +1,22 @@
+﻿// Copyright Epic Games, Inc. All Rights Reserved.
+
+#pragma once
+
+#include "Internationalization/Text.h"
+#include "Misc/Change.h"
+
+class UObject;
+
+namespace UE::ControlRig
+{
+/** Base class for command changes in control rig. */
+class FControlRigCommandChange : public FCommandChange
+{
+public:
+	
+	virtual FString ToString() const override { return TEXT("ControlRigCommand"); }
+
+	/** Adds the command to the undo buffer. */
+	static void StoreUndo(UObject* Object, TUniquePtr<FControlRigCommandChange> Change, const FText& Description = FText::GetEmpty());
+};
+}
diff --git a/Engine/Plugins/Animation/ControlRig/Source/ControlRig/Private/Commands/ElementSelectionCommand.cpp b/Engine/Plugins/Animation/ControlRig/Source/ControlRig/Private/Commands/ElementSelectionCommand.cpp
new file mode 100644
--- /dev/null
+++ b/Engine/Plugins/Animation/ControlRig/Source/ControlRig/Private/Commands/ElementSelectionCommand.cpp
@@ -0,0 +1,55 @@
+﻿// Copyright Epic Games, Inc. All Rights Reserved.
+
+#include "ElementSelectionCommand.h"
+
+#include "Rigs/RigHierarchy.h"
+#include "Rigs/RigHierarchyController.h"
+
+#include <type_traits>
+
+namespace UE::ControlRig
+{
+namespace ElementSelectionDetail
+{
+template<typename TConvertSelectionState> requires std::is_invocable_r_v<bool, TConvertSelectionState, bool /*bIsSelectedWhenRedoing*/>
+static void CombineData(
+	const FElementSelectionArray& InChanges, URigHierarchyController& InController,
+	TConvertSelectionState&& InConvertSelectionState
+	)
+{
+	for (const FElementSelectionData& ChangeData : InChanges)
+	{
+		const bool bIsSelected = InConvertSelectionState(ChangeData.bSelect);
+		constexpr bool bClearSelection = false,  bSetupUndo = false;
+		InController.SelectHierarchyKey(ChangeData.Key, bIsSelected, bClearSelection, bSetupUndo);
+	}
+}
+
+static URigHierarchyController* GetController(UObject* InUndoObject)
+{
+	URigHierarchy* Hierarchy = Cast<URigHierarchy>(InUndoObject);
+	URigHierarchyController* Controller = Hierarchy ? Hierarchy->GetController() : nullptr;
+	return Controller;
+}
+}
+	
+void FElementSelectionCommand::Apply(UObject* Object)
+{
+	if (URigHierarchyController* Controller = ElementSelectionDetail::GetController(Object))
+	{
+		ElementSelectionDetail::CombineData(Changes, *Controller,
+			[](bool bSelectionStateWhenApplying) { return bSelectionStateWhenApplying; }
+			);
+	}
+}
+
+void FElementSelectionCommand::Revert(UObject* Object)
+{
+	if (URigHierarchyController* Controller = ElementSelectionDetail::GetController(Object))
+	{
+		ElementSelectionDetail::CombineData(Changes, *Controller,
+			[](bool bSelectionStateWhenApplying) { return !bSelectionStateWhenApplying; }
+			);
+	}
+}
+}
diff --git a/Engine/Plugins/Animation/ControlRig/Source/ControlRig/Private/Commands/ElementSelectionCommand.h b/Engine/Plugins/Animation/ControlRig/Source/ControlRig/Private/Commands/ElementSelectionCommand.h
new file mode 100644
--- /dev/null
+++ b/Engine/Plugins/Animation/ControlRig/Source/ControlRig/Private/Commands/ElementSelectionCommand.h
@@ -0,0 +1,48 @@
+﻿// Copyright Epic Games, Inc. All Rights Reserved.
+
+#pragma once
+
+#include "ControlRigCommandChange.h"
+#include "Rigs/RigHierarchyDefines.h"
+#include "UObject/WeakObjectPtrTemplates.h"
+
+class URigHierarchy;
+
+namespace UE::ControlRig
+{
+struct FElementSelectionData
+{
+	/** The key whose selection changed. */
+	FRigHierarchyKey Key;
+	/** When applied (redo), whether the key was selected. */
+	bool bSelect;
+
+	FElementSelectionData(const FRigHierarchyKey& Key, bool bSelect) : Key(Key) , bSelect(bSelect) {}
+};
+
+// Avoid heap allocation for the most common case of 1 element
+using FElementSelectionArray = TArray<FElementSelectionData, TInlineAllocator<1>>;
+	
+/** Changes the selection of a URigHierarchy. */
+class FElementSelectionCommand : public FControlRigCommandChange
+{
+public:
+	
+	explicit FElementSelectionCommand(FElementSelectionArray InChanges) : Changes(MoveTemp(InChanges)) {}
+	FElementSelectionCommand(const FElementSelectionData& InSingleData)
+		: FElementSelectionCommand(FElementSelectionArray( { InSingleData }))
+	{}
+
+	//~ Begin FChange Interface
+	virtual void Apply(UObject* Object) override;
+	virtual void Revert(UObject* Object) override;
+	//~ End FChange Interface
+
+private:
+
+	/** The elements and their changed state */
+	FElementSelectionArray Changes;
+};
+}
+
+
diff --git a/Engine/Plugins/Animation/ControlRig/Source/ControlRig/Private/Rigs/RigHierarchyController.cpp b/Engine/Plugins/Animation/ControlRig/Source/ControlRig/Private/Rigs/RigHierarchyController.cpp
--- a/Engine/Plugins/Animation/ControlRig/Source/ControlRig/Private/Rigs/RigHierarchyController.cpp
+++ b/Engine/Plugins/Animation/ControlRig/Source/ControlRig/Private/Rigs/RigHierarchyController.cpp
@@ -19,6 +19,7 @@
 #endif
 
 #include "Algo/Count.h"
+#include "Commands/ElementSelectionCommand.h"
 #include "Engine/SkeletalMeshSocket.h"
 
 #include UE_INLINE_GENERATED_CPP_BY_NAME(RigHierarchyController)
@@ -97,6 +98,8 @@ bool URigHierarchyController::SelectComponent(FRigComponentKey InKey, bool bSele
 
 bool URigHierarchyController::SelectHierarchyKey(FRigHierarchyKey InKey, bool bSelect, bool bClearSelection, bool bSetupUndo)
 {
+	using namespace UE::ControlRig;
+	
 	if(!IsValid())
 	{
 		return false;
@@ -172,7 +175,9 @@ bool URigHierarchyController::SelectHierarchyKey(FRigHierarchyKey InKey, bool bS
 		// If recording into transaction buffer, modify the hierarchy so undoing selection works correctly.
 		if (bSetupUndo && GUndo)
 		{
-			Hierarchy->Modify();
+			FControlRigCommandChange::StoreUndo(
+				Hierarchy, MakeUnique<FElementSelectionCommand>(FElementSelectionData(Key, bSelect))
+				);
 		}
 		
 		Element->bSelected = bSelect;
@@ -215,7 +220,9 @@ bool URigHierarchyController::SelectHierarchyKey(FRigHierarchyKey InKey, bool bS
 		// If recording into transaction buffer, modify the hierarchy so undoing selection works correctly.
 		if (bSetupUndo && GUndo)
 		{
-			Hierarchy->Modify();
+			FControlRigCommandChange::StoreUndo(
+				Hierarchy, MakeUnique<FElementSelectionCommand>(FElementSelectionData(Key, bSelect))
+				);
 		}
 		
 		Component->bSelected = bSelect;
