Fixed crash in Control Rig editor when the preview floor mesh is not available, now it will trigger an error log and continue to load the asset

diff --git a/Engine/Plugins/Animation/ControlRig/Source/ControlRigEditor/Private/Editor/ControlRigEditor.cpp b/Engine/Plugins/Animation/ControlRig/Source/ControlRigEditor/Private/Editor/ControlRigEditor.cpp
--- a/Engine/Plugins/Animation/ControlRig/Source/ControlRigEditor/Private/Editor/ControlRigEditor.cpp
+++ b/Engine/Plugins/Animation/ControlRig/Source/ControlRigEditor/Private/Editor/ControlRigEditor.cpp
@@ -3134,32 +3134,40 @@ void FControlRigBaseEditor::HandlePreviewSceneCreated(const TSharedRef<IPersonaP
 
 	// load a ground mesh
 	static const TCHAR* GroundAssetPath = TEXT("/Engine/MapTemplates/SM_Template_Map_Floor.SM_Template_Map_Floor");
-	UStaticMesh* FloorMesh = Cast<UStaticMesh>(StaticLoadObject(UStaticMesh::StaticClass(), NULL, GroundAssetPath, NULL, LOAD_None, NULL));
-	UMaterial* DefaultMaterial = UMaterial::GetDefaultMaterial(MD_Surface);
-	check(FloorMesh);
-	check(DefaultMaterial);
+	if (UStaticMesh* FloorMesh = Cast<UStaticMesh>(StaticLoadObject(UStaticMesh::StaticClass(), NULL, GroundAssetPath, NULL, LOAD_None, NULL)))
+	{
+		UMaterial* DefaultMaterial = UMaterial::GetDefaultMaterial(MD_Surface);
+		ensureMsgf(DefaultMaterial, TEXT("Control Rig editor default material can not be created"));
+
+		// create ground mesh actor
+		AStaticMeshActor* GroundActor = InPersonaPreviewScene->GetWorld()->SpawnActor<AStaticMeshActor>(AStaticMeshActor::StaticClass(), FTransform::Identity);
+		GroundActor->SetFlags(RF_Transient);
+		GroundActor->GetStaticMeshComponent()->SetStaticMesh(FloorMesh);
+		if (DefaultMaterial)
+		{
+			GroundActor->GetStaticMeshComponent()->SetMaterial(0, DefaultMaterial);
+		}
+		GroundActor->SetMobility(EComponentMobility::Static);
+		GroundActor->GetStaticMeshComponent()->SetCollisionEnabled(ECollisionEnabled::QueryAndPhysics);
+		GroundActor->GetStaticMeshComponent()->SetCollisionProfileName(UCollisionProfile::BlockAll_ProfileName);
+		GroundActor->GetStaticMeshComponent()->bSelectable = false;
+		// this will be an invisible collision box that users can use to test traces
+		GroundActor->GetStaticMeshComponent()->SetVisibility(false);
+
+		WeakGroundActorPtr = GroundActor;
+	}
+	else
+	{
+		UE_LOG(LogControlRig, Error, TEXT("Control Rig editor ground asset not found :%s"), GroundAssetPath);
+	}
 
 	// leave some metadata on the world used for debug object labeling
-	if(FWorldContext* WorldContext = GEngine->GetWorldContextFromWorld(InPersonaPreviewScene->GetWorld()))
+	if (FWorldContext* WorldContext = GEngine->GetWorldContextFromWorld(InPersonaPreviewScene->GetWorld()))
 	{
 		static constexpr TCHAR Format[] = TEXT("ControlRigEditor (%s)");
 		WorldContext->CustomDescription = FString::Printf(Format, *GetRigVMBlueprint()->GetName());
 	}
 
-	// create ground mesh actor
-	AStaticMeshActor* GroundActor = InPersonaPreviewScene->GetWorld()->SpawnActor<AStaticMeshActor>(AStaticMeshActor::StaticClass(), FTransform::Identity);
-	GroundActor->SetFlags(RF_Transient);
-	GroundActor->GetStaticMeshComponent()->SetStaticMesh(FloorMesh);
-	GroundActor->GetStaticMeshComponent()->SetMaterial(0, DefaultMaterial);
-	GroundActor->SetMobility(EComponentMobility::Static);
-	GroundActor->GetStaticMeshComponent()->SetCollisionEnabled(ECollisionEnabled::QueryAndPhysics);
-	GroundActor->GetStaticMeshComponent()->SetCollisionProfileName(UCollisionProfile::BlockAll_ProfileName);
-	GroundActor->GetStaticMeshComponent()->bSelectable = false;
-	// this will be an invisible collision box that users can use to test traces
-	GroundActor->GetStaticMeshComponent()->SetVisibility(false);
-
-	WeakGroundActorPtr = GroundActor;
-
 	AAnimationEditorPreviewActor* Actor = InPersonaPreviewScene->GetWorld()->SpawnActor<AAnimationEditorPreviewActor>(AAnimationEditorPreviewActor::StaticClass(), FTransform::Identity);
 	Actor->SetFlags(RF_Transient);
 	InPersonaPreviewScene->SetActor(Actor);
