Avoid a crash in Control Rig when adding an empty rule to a connector

diff --git a/Engine/Plugins/Animation/ControlRig/Source/ControlRig/Private/ModularRigRuleManager.cpp b/Engine/Plugins/Animation/ControlRig/Source/ControlRig/Private/ModularRigRuleManager.cpp
--- a/Engine/Plugins/Animation/ControlRig/Source/ControlRig/Private/ModularRigRuleManager.cpp
+++ b/Engine/Plugins/Animation/ControlRig/Source/ControlRig/Private/ModularRigRuleManager.cpp
@@ -264,6 +264,10 @@ void UModularRigRuleManager::FilterByConnectorRules(FWorkData& InOutWorkData)
 	
 	for(const FRigConnectionRuleStash& Stash : Rules)
 	{
+		if (!Stash.IsValid())
+		{
+			continue;
+		}
 		TSharedPtr<FStructOnScope> Storage;
 		const FRigConnectionRule* Rule = Stash.Get(Storage);
 
diff --git a/Engine/Plugins/Animation/ControlRig/Source/ControlRig/Private/Rigs/RigConnectionRules.cpp b/Engine/Plugins/Animation/ControlRig/Source/ControlRig/Private/Rigs/RigConnectionRules.cpp
--- a/Engine/Plugins/Animation/ControlRig/Source/ControlRig/Private/Rigs/RigConnectionRules.cpp
+++ b/Engine/Plugins/Animation/ControlRig/Source/ControlRig/Private/Rigs/RigConnectionRules.cpp
@@ -216,6 +216,10 @@ FRigElementResolveResult FRigAndConnectionRule::Resolve(const FRigBaseElement* I
 	TSharedPtr<FStructOnScope> Storage;
 	for(const FRigConnectionRuleStash& ChildRule : ChildRules)
 	{
+		if (!ChildRule.IsValid())
+		{
+			continue;
+		}
 		if(const FRigConnectionRule* Rule = ChildRule.Get(Storage))
 		{
 			Result = Rule->Resolve(InTarget, InRuleInput);
@@ -241,6 +245,10 @@ FRigElementResolveResult FRigOrConnectionRule::Resolve(const FRigBaseElement* In
 	TSharedPtr<FStructOnScope> Storage;
 	for(const FRigConnectionRuleStash& ChildRule : ChildRules)
 	{
+		if (!ChildRule.IsValid())
+		{
+			continue;
+		}
 		if(const FRigConnectionRule* Rule = ChildRule.Get(Storage))
 		{
 			Result = Rule->Resolve(InTarget, InRuleInput);
