Control Rig: Polish for EventQueueView

#rnx

diff --git a/Engine/Plugins/Animation/ControlRig/Source/ControlRigEditor/Private/Editor/SModularRigEventQueueView.cpp b/Engine/Plugins/Animation/ControlRig/Source/ControlRigEditor/Private/Editor/SModularRigEventQueueView.cpp
--- a/Engine/Plugins/Animation/ControlRig/Source/ControlRigEditor/Private/Editor/SModularRigEventQueueView.cpp
+++ b/Engine/Plugins/Animation/ControlRig/Source/ControlRigEditor/Private/Editor/SModularRigEventQueueView.cpp
@@ -277,7 +277,6 @@ void SModularRigEventQueueView::Construct( const FArguments& InArgs, TSharedRef<
 				.OnGenerateRow(this, &SModularRigEventQueueView::MakeTableRowWidget)
 				.OnGetChildren(this, &SModularRigEventQueueView::HandleGetChildrenForTree)
 				.OnSelectionChanged(this, &SModularRigEventQueueView::OnSelectionChanged)
-				.OnContextMenuOpening(this, &SModularRigEventQueueView::CreateContextMenu)
 			]
 		]
 	];
@@ -414,27 +413,6 @@ void SModularRigEventQueueView::RefreshTreeView(const UModularRig* InModularRig)
 	}
 }
 
-TSharedPtr< SWidget > SModularRigEventQueueView::CreateContextMenu()
-{
-	TArray<TSharedPtr<FModularRigEventEntry>> SelectedItems = TreeView->GetSelectedItems();
-	if(SelectedItems.Num() == 0)
-	{
-		return nullptr;
-	}
-
-	const FModularRigEventQueueCommands& Actions = FModularRigEventQueueCommands::Get();
-
-	const bool CloseAfterSelection = true;
-	FMenuBuilder MenuBuilder(CloseAfterSelection, CommandList);
-	{
-		MenuBuilder.BeginSection("ModularRigEventQueueToolsAction", LOCTEXT("ToolsAction", "Tools"));
-		MenuBuilder.AddMenuEntry(Actions.FocusOnSelection);
-		MenuBuilder.EndSection();
-	}
-
-	return MenuBuilder.MakeWidget();
-}
-
 void SModularRigEventQueueView::HandleFocusOnSelectedModule()
 {
 	OnSelectionChanged(TSharedPtr<FModularRigEventEntry>(), ESelectInfo::Direct);
diff --git a/Engine/Plugins/Animation/ControlRig/Source/ControlRigEditor/Private/Editor/SModularRigEventQueueView.h b/Engine/Plugins/Animation/ControlRig/Source/ControlRigEditor/Private/Editor/SModularRigEventQueueView.h
--- a/Engine/Plugins/Animation/ControlRig/Source/ControlRigEditor/Private/Editor/SModularRigEventQueueView.h
+++ b/Engine/Plugins/Animation/ControlRig/Source/ControlRigEditor/Private/Editor/SModularRigEventQueueView.h
@@ -65,8 +65,6 @@ public:
 	/** Set Selection Changed */
 	UE_API void OnSelectionChanged(TSharedPtr<FModularRigEventEntry> Selection, ESelectInfo::Type SelectInfo);
 
-	UE_API TSharedPtr< SWidget > CreateContextMenu();
-
 protected:
 
 	/** Rebuild the tree view */
