Control Rig: Fix config values for array properties

#rnx

diff --git a/Engine/Plugins/Animation/ControlRig/Source/ControlRigEditor/Private/ControlRigModuleDetails.cpp b/Engine/Plugins/Animation/ControlRig/Source/ControlRigEditor/Private/ControlRigModuleDetails.cpp
--- a/Engine/Plugins/Animation/ControlRig/Source/ControlRigEditor/Private/ControlRigModuleDetails.cpp
+++ b/Engine/Plugins/Animation/ControlRig/Source/ControlRigEditor/Private/ControlRigModuleDetails.cpp
@@ -952,9 +952,22 @@ void FRigModuleInstanceDetails::OnConfigValueChanged(const FPropertyChangedEvent
 		return;
 	}
 
-	FString PropertyPathString;
+	FStringBuilderBase PathBuilder;
+
+	auto AppendProperty = [&PathBuilder](const FProperty* Property, int32 InArrayIndex)
+	{
+		if (PathBuilder.Len() > 0)
+		{
+			PathBuilder.Append(TEXT("->"));
+		}
+		PathBuilder.Append(Property->GetFName().ToString());
+		if (InArrayIndex >= 0)
+		{
+			PathBuilder.Appendf(TEXT("[%d]"), InArrayIndex);
+		}
+	};
+
 	{
-		TSharedRef< FPropertyPath > PropertyPath = FPropertyPath::CreateEmpty();
 		FProperty* Property = InPropertyChangedEvent.Property;
 		if(!ensure(Property))
 		{
@@ -972,7 +985,7 @@ void FRigModuleInstanceDetails::OnConfigValueChanged(const FPropertyChangedEvent
 			// the property names in the array are provided in reverse order
 			for(int32 Index = PropertyNames.Num() - 1; Index >= 0; Index--)
 			{
-				if(!PropertyPath->IsValid())
+				if(PathBuilder.Len() == 0)
 				{
 					ensure(InPropertyChangedEvent.GetNumObjectsBeingEdited() > 0);
 					Property = FControlRigOverrideValue::FindProperty(InPropertyChangedEvent.GetObjectBeingEdited(0)->GetClass(), PropertyNames[Index]);
@@ -1000,27 +1013,17 @@ void FRigModuleInstanceDetails::OnConfigValueChanged(const FPropertyChangedEvent
 						return;
 					}
 				}
-				
-				PropertyPath->AddProperty(
-					{
-						TWeakFieldPtr<FProperty>(Property),
-						InPropertyChangedEvent.GetArrayIndex(PropertyNames[Index])
-					}
-				);
+
+				AppendProperty(Property, PropertyNameStack.FindChecked(PropertyNames[Index]));
 			}
 		}
 		else
 		{
-			PropertyPath->AddProperty(
-				{
-					TWeakFieldPtr<FProperty>(InPropertyChangedEvent.MemberProperty),
-					InPropertyChangedEvent.GetArrayIndex(InPropertyChangedEvent.GetMemberPropertyName().ToString())
-				}
-			);
+			AppendProperty(InPropertyChangedEvent.MemberProperty, InPropertyChangedEvent.GetArrayIndex(InPropertyChangedEvent.GetMemberPropertyName().ToString()));
 		}
-
-		PropertyPathString = PropertyPath->ToString();
 	}
+
+	FString PropertyPathString = PathBuilder.ToString();
 	
 	// we need to shorten the path if we already have data on something above this.
 	bool bFoundValueWithShorterPath = true;
