Anim in Engine: Improve Anim Details UX

* Editing multiple values no longer displays the first property value
* Releasing focus on multiple values no longer sets the first property value
* Shift clicking to the left of a property name no longer resets the property
* Added support for math expressions with multiple operators
* Removed support for =operator syntax instead use operator=, e.g *=3 to multiply the current value by 3
* Fix navigation doesn't work in some cases

diff --git a/Engine/Plugins/Animation/ControlRig/Source/ControlRigEditor/Private/AnimDetails/AnimDetailsMathOperation.h b/Engine/Plugins/Animation/ControlRig/Source/ControlRigEditor/Private/AnimDetails/AnimDetailsMathOperation.h
deleted file mode 100644
diff --git a/Engine/Plugins/Animation/ControlRig/Source/ControlRigEditor/Private/AnimDetails/AnimDetailsMultiEditUtil.cpp b/Engine/Plugins/Animation/ControlRig/Source/ControlRigEditor/Private/AnimDetails/AnimDetailsMultiEditUtil.cpp
--- a/Engine/Plugins/Animation/ControlRig/Source/ControlRigEditor/Private/AnimDetails/AnimDetailsMultiEditUtil.cpp
+++ b/Engine/Plugins/Animation/ControlRig/Source/ControlRigEditor/Private/AnimDetails/AnimDetailsMultiEditUtil.cpp
@@ -9,8 +9,8 @@
 #include "AnimDetails/AnimDetailsProxyManager.h"
 #include "AnimDetails/AnimDetailsSelection.h"
 #include "AnimDetails/Proxies/AnimDetailsProxyBase.h"
-#include "AnimDetailsMathOperation.h"
 #include "Editor.h"
+#include "Math/BasicMathExpressionEvaluator.h"
 #include "PropertyHandle.h"
 #include "ScopedTransaction.h"
 
@@ -52,63 +52,6 @@ namespace UE::ControlRigEditor
 				 PropertyHandle->NotifyPostChange(EPropertyChangeType::ValueSet);
 			}
 		}
-
-		/** Applies a math operation to the property */
-		template <typename ValueType>
-		static void ApplyMathOperation(const TSharedRef<IPropertyHandle>& PropertyHandle, FAnimDetailsMathOperation<ValueType> MathOperation)
-		{
-			TArray<FString> PerObjectValues;
-			if (PropertyHandle->GetPerObjectValues(PerObjectValues) == FPropertyAccess::Success)
-			{
-				for (int32 ValueIndex = 0; ValueIndex < PerObjectValues.Num(); ++ValueIndex)
-				{
-					ValueType LHSValue;
-					TTypeFromString<ValueType>::FromString(LHSValue, *PerObjectValues[ValueIndex]);
-
-					switch (MathOperation.MathOperator)
-					{
-					case EAnimDetailsMathOperator::Add:
-						LHSValue += MathOperation.RHSValue;
-						break;
-
-					case EAnimDetailsMathOperator::Substract:
-						LHSValue -= MathOperation.RHSValue;
-						break;
-
-					case EAnimDetailsMathOperator::Multiply:
-						LHSValue *= MathOperation.RHSValue;
-						break;
-
-					case EAnimDetailsMathOperator::Divide:
-
-						// Avoid divisions by zero
-						if constexpr (std::is_floating_point_v<ValueType>)
-						{
-							if (FMath::IsNearlyEqual(MathOperation.RHSValue, 0.f))
-							{
-								return;
-							}
-						}
-						else if (MathOperation.RHSValue == 0)
-						{
-							return;
-						}
-
-						LHSValue /= MathOperation.RHSValue;
-						break;
-
-					default:
-						checkf(0, TEXT("Unhandled enum value"));
-					}
-
-					PerObjectValues[ValueIndex] = TTypeToString<ValueType>::ToSanitizedString(LHSValue);
-				}
-
-				PropertyHandle->NotifyPreChange();
-				PropertyHandle->SetPerObjectValues(PerObjectValues);
-				PropertyHandle->NotifyPostChange(EPropertyChangeType::ValueSet);
-			}
-		}
 	}
 
 	TUniquePtr<FAnimDetailsMultiEditUtil> FAnimDetailsMultiEditUtil::Instance;
@@ -204,58 +147,81 @@ namespace UE::ControlRigEditor
 		for (const TSharedRef<IPropertyHandle>& PropertyHandle : GetPropertiesBeingEdited<ValueType>(ProxyManager, InstigatorProperty))
 		{
 			PropertyHandle->NotifyPreChange();
+
 			PropertyHandle->SetValue(Value);
+
 			PropertyHandle->NotifyPostChange(EPropertyChangeType::ValueSet);
 		}
 	}
 
-	template void FAnimDetailsMultiEditUtil::MultiEditSet(const UAnimDetailsProxyManager& ProxyManager, const double Value, const TSharedRef<IPropertyHandle>& InstigatorProperty);
-	template void FAnimDetailsMultiEditUtil::MultiEditSet(const UAnimDetailsProxyManager& ProxyManager, const int64 Value, const TSharedRef<IPropertyHandle>& InstigatorProperty);
-	template void FAnimDetailsMultiEditUtil::MultiEditSet(const UAnimDetailsProxyManager& ProxyManager, const bool Value, const TSharedRef<IPropertyHandle>& InstigatorProperty);
+	template void FAnimDetailsMultiEditUtil::MultiEditSet<double>(const UAnimDetailsProxyManager& ProxyManager, const double Value, const TSharedRef<IPropertyHandle>& InstigatorProperty);
+	template void FAnimDetailsMultiEditUtil::MultiEditSet<int64>(const UAnimDetailsProxyManager& ProxyManager, const int64 Value, const TSharedRef<IPropertyHandle>& InstigatorProperty);
+	template void FAnimDetailsMultiEditUtil::MultiEditSet<bool>(const UAnimDetailsProxyManager& ProxyManager, const bool Value, const TSharedRef<IPropertyHandle>& InstigatorProperty);
 
 	template <typename ValueType>
-	void FAnimDetailsMultiEditUtil::MultiEditMath(
-		const UAnimDetailsProxyManager& ProxyManager, 
-		const FAnimDetailsMathOperation<ValueType>& MathOperation, 
+	bool FAnimDetailsMultiEditUtil::MultiEditMath(
+		const UAnimDetailsProxyManager& ProxyManager,
+		FString Expression,
 		const TSharedRef<IPropertyHandle>& InstigatorProperty)
 	{
-		if (!ensureMsgf(MathOperation.MathOperator != EAnimDetailsMathOperator::None, TEXT("Unexpected trying to perform a math operation but no operator is defined")))
+		// Plain numeric values are considered valid mathematical expressions by FBasicMathExpressionEvaluator.
+		// Don't handle these here. Instead these should be set with MultiEditSet.
+		if (Expression.IsNumeric())
 		{
-			return;
+			return false;
 		}
 
-		const FScopedTransaction ScopedTransaction(LOCTEXT("MultiEditMathTransaction", "Set Property Value"));
+		TMap<TSharedRef<IPropertyHandle>, TArray<double>> PropertyToPerObjectValuesMap;
+
+		const FBasicMathExpressionEvaluator Parser;
 
 		for (const TSharedRef<IPropertyHandle>& PropertyHandle : GetPropertiesBeingEdited<ValueType>(ProxyManager, InstigatorProperty))
 		{
-			const FFieldClass* PropertyClass = PropertyHandle->GetPropertyClass();
-
-			if constexpr (std::is_floating_point_v<ValueType>)
-			{
-				if (PropertyClass == FDoubleProperty::StaticClass())
-				{
-					PropertyUtils::ApplyMathOperation<double>(PropertyHandle, MathOperation);
-				}
-			}
-			else if constexpr (std::is_integral_v<ValueType>)
+			TArray<double>& PerObjectValues = PropertyToPerObjectValuesMap.Add(PropertyHandle);
+			TArray<FString> PerObjectValueStrings;
+			if (PropertyHandle->GetPerObjectValues(PerObjectValueStrings) == FPropertyAccess::Success)
 			{
-				if (PropertyClass == FInt64Property::StaticClass())
+				for (int32 ValueIndex = 0; ValueIndex < PerObjectValueStrings.Num(); ++ValueIndex)
 				{
-					PropertyUtils::ApplyMathOperation<int64>(PropertyHandle, MathOperation);
+					ValueType OldValue;
+					TTypeFromString<ValueType>::FromString(OldValue, *PerObjectValueStrings[ValueIndex]);
+
+					const TValueOrError<double, FExpressionError> ValueOrError = Parser.Evaluate(*Expression, OldValue);
+					if (const double* ValuePtr = ValueOrError.TryGetValue())
+					{
+						PerObjectValues.Add(*ValuePtr);
+					}
+					else
+					{
+						// Does not evaulate to a math expression
+						return false;
+					}
 				}
 			}
-			else
-			{
-				[] <bool SupportedType = false>()
+		}
+
+		// Apply a math operation if the string could be parsed as such
+		for (const TTuple<TSharedRef<IPropertyHandle>, TArray<double>>& PropertyToPerObjectValuesPair : PropertyToPerObjectValuesMap)
+		{
+			TArray<FString> PerObjectValueStrings;
+			Algo::Transform(PropertyToPerObjectValuesPair.Value, PerObjectValueStrings,
+				[](const double& Value)
 				{
-					static_assert(SupportedType, "Unsupported type in FAnimDetailsMultiEditUtil.");
-				}();
-			}
+					return TTypeToString<ValueType>::ToSanitizedString(Value);
+				});
+
+			PropertyToPerObjectValuesPair.Key->NotifyPreChange();
+
+			PropertyToPerObjectValuesPair.Key->SetPerObjectValues(PerObjectValueStrings);
+
+			PropertyToPerObjectValuesPair.Key->NotifyPostChange(EPropertyChangeType::ValueSet);
 		}
-	}
 
-	template void FAnimDetailsMultiEditUtil::MultiEditMath(const UAnimDetailsProxyManager& ProxyManager, const FAnimDetailsMathOperation<double>& MathOperation, const TSharedRef<IPropertyHandle>& InstigatorProperty);
-	template void FAnimDetailsMultiEditUtil::MultiEditMath(const UAnimDetailsProxyManager& ProxyManager, const FAnimDetailsMathOperation<int64>& MathOperation, const TSharedRef<IPropertyHandle>& InstigatorProperty);
+		return true;
+	}
+	
+	template bool FAnimDetailsMultiEditUtil::MultiEditMath<double>(const UAnimDetailsProxyManager& ProxyManager, FString StringValue, const TSharedRef<IPropertyHandle>& InstigatorProperty);
+	template bool FAnimDetailsMultiEditUtil::MultiEditMath<int64>(const UAnimDetailsProxyManager& ProxyManager, FString StringValue, const TSharedRef<IPropertyHandle>& InstigatorProperty);
 
 	template <typename ValueType>
 	void FAnimDetailsMultiEditUtil::MultiEditChange(
diff --git a/Engine/Plugins/Animation/ControlRig/Source/ControlRigEditor/Private/AnimDetails/AnimDetailsMultiEditUtil.h b/Engine/Plugins/Animation/ControlRig/Source/ControlRigEditor/Private/AnimDetails/AnimDetailsMultiEditUtil.h
--- a/Engine/Plugins/Animation/ControlRig/Source/ControlRigEditor/Private/AnimDetails/AnimDetailsMultiEditUtil.h
+++ b/Engine/Plugins/Animation/ControlRig/Source/ControlRigEditor/Private/AnimDetails/AnimDetailsMultiEditUtil.h
@@ -31,24 +31,26 @@ namespace UE::ControlRigEditor
 		void Leave(const TWeakPtr<IPropertyHandle>& WeakPropertyHandle);
 
 		/** 
-		 * Sets a value to all selected properties that joined this multi edit util. 
+		 * Sets a value to all selected properties that joined this multi edit util. Supports math operations.
 		 * 
 		 * @param ProxyManager			The proxy manager that owns the proxies and their properties
-		 * @param Value					The value to set for the properties
+		 * @param StringValue			The value to set for the properties
 		 * @param InstigatorProperty	The property that instigated this change
 		¨*/
-		template <typename ValueType>
+		template <typename ValueType> 
 		void MultiEditSet(const UAnimDetailsProxyManager& ProxyManager, const ValueType Value, const TSharedRef<IPropertyHandle>& InstigatorProperty);
 
 		/**
-		 * Applies a math operation to all selected properties that joined this multi edit util.
+		 * Sets a value to all selected properties that joined this multi edit util. Supports math operations.
 		 *
 		 * @param ProxyManager			The proxy manager that owns the proxies and their properties
-		 * @param MathOperation			The math operation to perform.
+		 * @param Expression			The mathematical expression to apply for the properties
 		 * @param InstigatorProperty	The property that instigated this change
+		 * 
+		 * @return						True if a math operation was successfully applied.
 		¨*/
 		template <typename ValueType>
-		void MultiEditMath(const UAnimDetailsProxyManager& ProxyManager, const FAnimDetailsMathOperation<ValueType>& MathOperation, const TSharedRef<IPropertyHandle>& InstigatorProperty);
+		bool MultiEditMath(const UAnimDetailsProxyManager& ProxyManager, FString Expression, const TSharedRef<IPropertyHandle>& InstigatorProperty);
 
 		/** 
 		 * Applies a delta to all selected properties that joined this multi edit util. 
diff --git a/Engine/Plugins/Animation/ControlRig/Source/ControlRigEditor/Private/AnimDetails/AnimDetailsNavigableWidgetRegistry.cpp b/Engine/Plugins/Animation/ControlRig/Source/ControlRigEditor/Private/AnimDetails/AnimDetailsNavigableWidgetRegistry.cpp
--- a/Engine/Plugins/Animation/ControlRig/Source/ControlRigEditor/Private/AnimDetails/AnimDetailsNavigableWidgetRegistry.cpp
+++ b/Engine/Plugins/Animation/ControlRig/Source/ControlRigEditor/Private/AnimDetails/AnimDetailsNavigableWidgetRegistry.cpp
@@ -38,32 +38,19 @@ namespace UE::ControlRigEditor
 
 	bool FAnimDetailsNavigableWidgetRegistry::IsWidgetVisible(const TSharedRef<SWidget>& Widget) const
 	{
-		if (Widget->GetVisibility() == EVisibility::Collapsed ||
-			Widget->GetVisibility() == EVisibility::Hidden)
-		{
-			return false;
-		}
-
-		// In some cases details views just minimize their children, but don't flag them hidden or collapsed
-		if (Widget->GetCachedGeometry().GetAbsoluteSize().IsNearlyZero())
-		{
-			return false;
-		}
-
-		// In some cases details views just remove their widgets from the window, but don't flag them hidden or collapsed
-		const TSharedPtr<SWidget> ActiveTopLevelWindow = FSlateApplication::Get().GetActiveTopLevelWindow();
-		TSharedPtr<SWidget> VisibilityTestWidget = Widget->GetParentWidget();
+		TSharedPtr<SWidget> VisibilityTestWidget = Widget;
 		while (VisibilityTestWidget.IsValid())
 		{
-			if (VisibilityTestWidget == ActiveTopLevelWindow)
+			if (VisibilityTestWidget->GetVisibility() == EVisibility::Collapsed ||
+				VisibilityTestWidget->GetVisibility() == EVisibility::Hidden)
 			{
-				return true;
+				return false;
 			}
 
 			VisibilityTestWidget = VisibilityTestWidget->GetParentWidget();
 		}
 
-		return false;
+		return true;
 	}
 
 	void FAnimDetailsNavigableWidgetRegistry::Register(const TSharedRef<SWidget>& InWidget)
diff --git a/Engine/Plugins/Animation/ControlRig/Source/ControlRigEditor/Private/AnimDetails/Customizations/AnimDetailsProxyDetails.cpp b/Engine/Plugins/Animation/ControlRig/Source/ControlRigEditor/Private/AnimDetails/Customizations/AnimDetailsProxyDetails.cpp
--- a/Engine/Plugins/Animation/ControlRig/Source/ControlRigEditor/Private/AnimDetails/Customizations/AnimDetailsProxyDetails.cpp
+++ b/Engine/Plugins/Animation/ControlRig/Source/ControlRigEditor/Private/AnimDetails/Customizations/AnimDetailsProxyDetails.cpp
@@ -101,8 +101,6 @@ namespace UE::ControlRigEditor
 				continue;
 			}
 
-			Property->SetPropertyFlags(CPF_TextExportTransient); // Hack to turn off shift copy/paste
-
 			uint32 NumChildren;
 			PropertyHandle->GetNumChildren(NumChildren);
 
diff --git a/Engine/Plugins/Animation/ControlRig/Source/ControlRigEditor/Private/AnimDetails/Customizations/AnimDetailsValueCustomization.cpp b/Engine/Plugins/Animation/ControlRig/Source/ControlRigEditor/Private/AnimDetails/Customizations/AnimDetailsValueCustomization.cpp
--- a/Engine/Plugins/Animation/ControlRig/Source/ControlRigEditor/Private/AnimDetails/Customizations/AnimDetailsValueCustomization.cpp
+++ b/Engine/Plugins/Animation/ControlRig/Source/ControlRigEditor/Private/AnimDetails/Customizations/AnimDetailsValueCustomization.cpp
@@ -14,12 +14,69 @@
 #include "EditMode/ControlRigEditMode.h"
 #include "EditorModeManager.h"
 #include "IDetailChildrenBuilder.h"
+#include "PropertyEditorClipboard.h"
 #include "Widgets/Input/SNumericEntryBox.h"
 
 #define LOCTEXT_NAMESPACE "AnimDetailsValueCustomization"
 
 namespace UE::ControlRigEditor
 {
+	namespace Private
+	{
+		FAnimDetailsChildWidgetVisiblityHandler::FAnimDetailsChildWidgetVisiblityHandler(const TSharedRef<IPropertyHandle>& InStructPropertyHandle, const TArray<TWeakPtr<SWidget>>& InWeakWidgets)
+			: WeakStructPropertyHandle(InStructPropertyHandle)
+			, WeakWidgets(InWeakWidgets)
+		{
+			if (WeakWidgets.IsEmpty())
+			{
+				return;
+			}
+
+			TickerHandle = FTSTicker::GetCoreTicker().AddTicker(FTickerDelegate::CreateLambda(
+				[this](float DeltaTime)
+				{
+					check(!WeakWidgets.IsEmpty());
+
+					const TSharedPtr<SWidget> FirstWidget = WeakWidgets[0].Pin();
+					if (!FirstWidget.IsValid())
+					{
+						return false;
+					}
+
+					if (!WeakStructPropertyHandle.IsValid())
+					{
+						return false;
+					}
+
+					const EVisibility Visibility = WeakStructPropertyHandle.Pin()->IsExpanded() ?
+						EVisibility::Visible :
+						EVisibility::Collapsed;
+
+					if (FirstWidget->GetVisibility() != Visibility)
+					{
+						for (const TWeakPtr<SWidget>& WeaWidget : WeakWidgets)
+						{
+							if (WeaWidget.IsValid())
+							{
+								WeaWidget.Pin()->SetVisibility(Visibility);
+							}
+						}
+					}
+
+					return true;
+				}
+			));
+		}
+
+		FAnimDetailsChildWidgetVisiblityHandler::~FAnimDetailsChildWidgetVisiblityHandler()
+		{
+			if (TickerHandle.IsValid())
+			{
+				FTSTicker::GetCoreTicker().RemoveTicker(TickerHandle);
+			}
+		}
+	}
+
 	TSharedRef<IPropertyTypeCustomization> FAnimDetailsValueCustomization::MakeInstance()
 	{
 		return MakeShared<FAnimDetailsValueCustomization>();
@@ -40,8 +97,9 @@ namespace UE::ControlRigEditor
 		DetailBuilder = &CategoryBuilder.GetParentLayout();
 
 		StructPropertyHandle = InStructPropertyHandle;
-
+		
 		// Show custom children if expanded, allowing to select individual controls in the property name row
+		TArray<TWeakPtr<SWidget>> ChildContentWidgets;
 		for (int32 ChildIndex = 0; ChildIndex < SortedChildHandles.Num(); ++ChildIndex)
 		{
 			TSharedRef<IPropertyHandle> ChildHandle = SortedChildHandles[ChildIndex];
@@ -56,6 +114,13 @@ namespace UE::ControlRigEditor
 
 			const TSharedRef<SWidget> ValueWidget = MakeChildWidget(InStructPropertyHandle, ChildHandle);
 			
+			// Handle visibilty for child content widgets. This is required so that  
+			// FAnimDetailsNavigableWidgetRegistry can find the right widget to navigate to.
+			if (const bool bCanBeExpanded = SortedChildHandles.Num() > 1)
+			{
+				ChildContentWidgets.Add(ValueWidget);
+			}
+
 			InStructBuilder.AddProperty(ChildHandle)
 				.CustomWidget()
 				.NameContent()
@@ -85,20 +150,52 @@ namespace UE::ControlRigEditor
 				[
 					ChildHandle->CreateDefaultPropertyButtonWidgets()
 				]
+				.CopyAction(
+					FUIAction(
+						FExecuteAction::CreateLambda([WeakChildHandle = ChildHandle->AsWeak()]
+							{
+								const TSharedPtr<IPropertyHandle> PropertyHandle = WeakChildHandle.IsValid() ? WeakChildHandle.Pin() : nullptr;
+
+								FString Value;
+								if (PropertyHandle.IsValid() &&
+									PropertyHandle->IsValidHandle() &&
+									PropertyHandle->GetValueAsFormattedString(Value, PPF_Copy) == FPropertyAccess::Success)
+								{
+									FPropertyEditorClipboard::ClipboardCopy(*Value);
+								}
+							}),
+						FCanExecuteAction::CreateLambda([]()
+							{
+								return true;
+							})
+					)
+				)
 				.PasteAction(
 					FUIAction(
-						FExecuteAction::CreateLambda([]() 
-							{ 
-								checkNoEntry(); 
+						FExecuteAction::CreateLambda([WeakChildHandle = ChildHandle->AsWeak()]
+							{								
+								const TSharedPtr<IPropertyHandle> PropertyHandle = WeakChildHandle.IsValid() ? WeakChildHandle.Pin() : nullptr;
+
+								FString Value;
+								if (PropertyHandle.IsValid() &&
+									PropertyHandle->IsValidHandle())
+								{
+									FString ClipboardContent;
+									FPropertyEditorClipboard::ClipboardPaste(ClipboardContent);
+
+									PropertyHandle->SetValueFromFormattedString(ClipboardContent);
+								}
 							}),
-						FCanExecuteAction::CreateLambda([]() 
-							{ 
-								checkNoEntry(); 
-								return false; 
+						FCanExecuteAction::CreateLambda([]()
+							{
+								const bool bShiftDown = FSlateApplication::Get().GetModifierKeys().IsShiftDown();
+								return !bShiftDown;
 							})
 					)
 				);
 		}
+
+		Private::FAnimDetailsChildWidgetVisiblityHandler(InStructPropertyHandle, ChildContentWidgets);
 	}
 
 	void FAnimDetailsValueCustomization::MakeHeaderRow(TSharedRef<IPropertyHandle>& InStructPropertyHandle, FDetailWidgetRow& InRow)
@@ -151,8 +248,6 @@ namespace UE::ControlRigEditor
 				continue;
 			}
 
-			Property->SetPropertyFlags(CPF_TextExportTransient); // Hack to turn off shift copy/paste
-
 			const TSharedRef<SWidget> ChildWidget = MakeChildWidget(InStructPropertyHandle, ChildHandle);
 
 			// Always display childs in the struct row but disable them if they're filtered out
diff --git a/Engine/Plugins/Animation/ControlRig/Source/ControlRigEditor/Private/AnimDetails/Customizations/AnimDetailsValueCustomization.h b/Engine/Plugins/Animation/ControlRig/Source/ControlRigEditor/Private/AnimDetails/Customizations/AnimDetailsValueCustomization.h
--- a/Engine/Plugins/Animation/ControlRig/Source/ControlRigEditor/Private/AnimDetails/Customizations/AnimDetailsValueCustomization.h
+++ b/Engine/Plugins/Animation/ControlRig/Source/ControlRigEditor/Private/AnimDetails/Customizations/AnimDetailsValueCustomization.h
@@ -2,6 +2,7 @@
 
 #pragma once
 
+#include "Containers/Ticker.h"
 #include "Customizations/MathStructCustomizations.h"
 
 class IDetailLayoutBuilder;
@@ -11,6 +12,28 @@ class SWidget;
 
 namespace UE::ControlRigEditor
 {
+	namespace Private
+	{
+		/**
+		 * When structs are collapsed, the visibility of the widgets is not updated, and their visibility attribute is not executed.
+		 *
+		 * This struct serves to update the widget's visibility regardless of this state. Required to deduce if the widget is navigable.
+		 */
+		struct FAnimDetailsChildWidgetVisiblityHandler
+		{
+			FAnimDetailsChildWidgetVisiblityHandler() = default;
+			FAnimDetailsChildWidgetVisiblityHandler(const TSharedRef<IPropertyHandle>& InStructPropertyHandle, const TArray<TWeakPtr<SWidget>>& InWeakWidgets);
+			
+			~FAnimDetailsChildWidgetVisiblityHandler();
+
+		private:
+			TWeakPtr<IPropertyHandle> WeakStructPropertyHandle;
+			TArray<TWeakPtr<SWidget>> WeakWidgets;
+
+			FTSTicker::FDelegateHandle TickerHandle;
+		};
+	}
+
 	/** Property type customization for struct properties such as FAnimDetailsBool or FAnimDetailsTransform */
 	class FAnimDetailsValueCustomization
 		: public FMathStructCustomization
@@ -42,8 +65,8 @@ namespace UE::ControlRigEditor
 		/** Gets a color from the property */
 		FLinearColor GetColorFromProperty(const FName& PropertyName) const;
 
-		// use the property row to get expansion state to hide/show widgets, and use the builder/structhandler to get it
-		mutable IDetailPropertyRow* DetailPropertyRow = nullptr;
+		/** Handles visibility for child widgets in this customization */
+		Private::FAnimDetailsChildWidgetVisiblityHandler ChildContentVisibilityHandler;
 
 		/** Pointer to the detail builder, or nullptr if not initialized */
 		IDetailLayoutBuilder* DetailBuilder = nullptr;
diff --git a/Engine/Plugins/Animation/ControlRig/Source/ControlRigEditor/Private/AnimDetails/Widgets/SAnimDetailsValueBoolean.cpp b/Engine/Plugins/Animation/ControlRig/Source/ControlRigEditor/Private/AnimDetails/Widgets/SAnimDetailsValueBoolean.cpp
--- a/Engine/Plugins/Animation/ControlRig/Source/ControlRigEditor/Private/AnimDetails/Widgets/SAnimDetailsValueBoolean.cpp
+++ b/Engine/Plugins/Animation/ControlRig/Source/ControlRigEditor/Private/AnimDetails/Widgets/SAnimDetailsValueBoolean.cpp
@@ -94,7 +94,7 @@ namespace UE::ControlRigEditor
 		if (WeakProxyManager.IsValid() && WeakPropertyHandle.IsValid())
 		{
 			const bool bEnabled = CheckBoxState == ECheckBoxState::Checked;
-			UE::ControlRigEditor::FAnimDetailsMultiEditUtil::Get().MultiEditSet<bool>(*WeakProxyManager.Get(), bEnabled, WeakPropertyHandle.Pin().ToSharedRef());
+			UE::ControlRigEditor::FAnimDetailsMultiEditUtil::Get().MultiEditSet(*WeakProxyManager.Get(), bEnabled, WeakPropertyHandle.Pin().ToSharedRef());
 		}
 	}
 }
diff --git a/Engine/Plugins/Animation/ControlRig/Source/ControlRigEditor/Private/AnimDetails/Widgets/SAnimDetailsValueNumeric.cpp b/Engine/Plugins/Animation/ControlRig/Source/ControlRigEditor/Private/AnimDetails/Widgets/SAnimDetailsValueNumeric.cpp
--- a/Engine/Plugins/Animation/ControlRig/Source/ControlRigEditor/Private/AnimDetails/Widgets/SAnimDetailsValueNumeric.cpp
+++ b/Engine/Plugins/Animation/ControlRig/Source/ControlRigEditor/Private/AnimDetails/Widgets/SAnimDetailsValueNumeric.cpp
@@ -2,7 +2,6 @@
 
 #include "SAnimDetailsValueNumeric.h"
 
-#include "AnimDetails/AnimDetailsMathOperation.h"
 #include "AnimDetails/AnimDetailsMultiEditUtil.h"
 #include "AnimDetails/AnimDetailsProxyManager.h"
 #include "AnimDetails/AnimDetailsSettings.h"
@@ -119,46 +118,27 @@ namespace UE::ControlRigEditor
 		: public TDefaultNumericTypeInterface<NumericType>
 	{
 	public:
-		/** Constructs this type interface with a multi-edit context */
-		TAnimDetailsNumericTypeInterface(const TWeakObjectPtr<UAnimDetailsProxyManager>& InProxyManager, const TWeakPtr<IPropertyHandle>& InPropertyHandle)
-			: WeakPropertyHandle(InPropertyHandle)
-			, WeakProxyManager(InProxyManager)
+		/** Constructs this type interface */
+		TAnimDetailsNumericTypeInterface(const TSharedRef<SAnimDetailsValueNumeric<NumericType>>& InValueWidget)
+			: WeakValueWidget(InValueWidget)
 		{}
 
 		//~ Begin TDefaultNumericTypeInterface interface
 		virtual TOptional<NumericType> FromString(const FString& InString, const NumericType& InExistingValue) override
 		{
-			using namespace UE::ControlRigEditor;
-
-			const TSharedPtr<IPropertyHandle> PropertyHandle = WeakPropertyHandle.Pin();
-			if (WeakProxyManager.IsValid() && PropertyHandle.IsValid() && PropertyHandle->IsValidHandle())
+			if (!WeakValueWidget.IsValid())
 			{
-				NumericType Value;
-				if (const TOptional<FAnimDetailsMathOperation<NumericType>> OptionalMathOperation = FAnimDetailsMathParser::FromString<NumericType>(InString))
-				{
-					FAnimDetailsMultiEditUtil::Get().MultiEditMath(*WeakProxyManager.Get(), OptionalMathOperation.GetValue(), PropertyHandle.ToSharedRef());
-				}
-				else if (LexTryParseString(Value, *InString))
-				{
-					FAnimDetailsMultiEditUtil::Get().MultiEditSet(*WeakProxyManager.Get(), Value, PropertyHandle.ToSharedRef());
-				}
-				NumericType NewValue;
-				if (PropertyHandle->GetValue(NewValue) == FPropertyAccess::Success)
-				{
-					return NewValue;
-				}
+				return TOptional<NumericType>();
 			}
-
-			return TOptional<NumericType>();
+			WeakValueWidget.Pin()->SetDisplayText(FText::FromString(InString));
+			
+			return TDefaultNumericTypeInterface<NumericType>::FromString(InString, InExistingValue);
 		}
 		//~ End TDefaultNumericTypeInterface interface
 
 	private:
 		/** The weak property to operate on */
-		const TWeakPtr<IPropertyHandle> WeakPropertyHandle;
-
-		/** The weak proxy manager that owns the property */
-		const TWeakObjectPtr<UAnimDetailsProxyManager> WeakProxyManager;
+		const TWeakPtr<SAnimDetailsValueNumeric<NumericType>> WeakValueWidget;
 	};
 
 
@@ -188,7 +168,7 @@ namespace UE::ControlRigEditor
 		LabelLocation = InArgs._LabelLocation;
 		LabelPadding = InArgs._LabelPadding;
 
-		TypeInterface = MakeShared<TAnimDetailsNumericTypeInterface<NumericType>>(ProxyManager, InPropertyHandle);
+		TypeInterface = MakeShared<TAnimDetailsNumericTypeInterface<NumericType>>(SharedThis(this));
 
 		ForceRefresh();
 
@@ -266,6 +246,7 @@ namespace UE::ControlRigEditor
 			.Font(FAppStyle::GetFontStyle("PropertyWindow.NormalFont"))
 			.Value(this, &SAnimDetailsValueNumeric<NumericType>::GetValue)
 			.OnGetDisplayValue(this, &SAnimDetailsValueNumeric<NumericType>::OnGetDisplayValue)
+			.OnEditEditableText(this, &SAnimDetailsValueNumeric<NumericType>::OnGetDisplayValue)
 			.OnValueChanged(this, &SAnimDetailsValueNumeric<NumericType>::OnValueChanged)
 			.OnValueCommitted(this, &SAnimDetailsValueNumeric<NumericType>::OnValueCommitted)
 			.OnBeginSliderMovement(this, &SAnimDetailsValueNumeric<NumericType>::OnBeginSliderMovement)
@@ -283,6 +264,7 @@ namespace UE::ControlRigEditor
 			.SliderExponent(PropertyMetaData.SliderExponent)
 			.Delta(PropertyMetaData.Delta)
 			.LinearDeltaSensitivity(1)
+			.ClearKeyboardFocusOnCommit(true)
 			.TypeInterface(TypeInterface);
 
 		NavigableWidgetRegistrar.Register(*ProxyManager, SpinBox.ToSharedRef());
@@ -342,6 +324,12 @@ namespace UE::ControlRigEditor
 		];
 	}
 
+	template <typename NumericType>
+	void SAnimDetailsValueNumeric<NumericType>::SetDisplayText(const FText& Text)
+	{
+		DisplayText = Text;
+	}
+
 	template <typename NumericType>
 	NumericType SAnimDetailsValueNumeric<NumericType>::GetValue() const
 	{
@@ -438,12 +426,23 @@ namespace UE::ControlRigEditor
 	template <typename NumericType>
 	void SAnimDetailsValueNumeric<NumericType>::OnValueCommitted(NumericType Value, ETextCommit::Type CommitType)
 	{
-		// Note setters and math are handled in the numeric type interface
 		if (bIsUsingSlider)
 		{
 			constexpr bool bInteractive = false;
 			MultiEditChangePropertyValue(Value, bInteractive);
 		}
+		else if (WeakPropertyHandle.IsValid())
+		{
+			const TSharedRef<IPropertyHandle> PropertyHandle = WeakPropertyHandle.Pin().ToSharedRef();
+
+			// When pressing enter, two commit calls will occur, one with ETextCommit::OnEnter 
+			// and one with ETextCommit::OnCleared. Only handle the later that is raised in any case.
+			if (CommitType != ETextCommit::OnEnter &&
+				!FAnimDetailsMultiEditUtil::Get().MultiEditMath<NumericType>(*WeakProxyManager.Get(), DisplayText.ToString(), PropertyHandle))
+			{
+				FAnimDetailsMultiEditUtil::Get().MultiEditSet(*WeakProxyManager.Get(), Value, PropertyHandle);
+			}
+		}
 	}
 
 	template <typename NumericType>
diff --git a/Engine/Plugins/Animation/ControlRig/Source/ControlRigEditor/Private/AnimDetails/Widgets/SAnimDetailsValueNumeric.h b/Engine/Plugins/Animation/ControlRig/Source/ControlRigEditor/Private/AnimDetails/Widgets/SAnimDetailsValueNumeric.h
--- a/Engine/Plugins/Animation/ControlRig/Source/ControlRigEditor/Private/AnimDetails/Widgets/SAnimDetailsValueNumeric.h
+++ b/Engine/Plugins/Animation/ControlRig/Source/ControlRigEditor/Private/AnimDetails/Widgets/SAnimDetailsValueNumeric.h
@@ -62,6 +62,9 @@ namespace UE::ControlRigEditor
 		/** Constructs this widget. Edits the provided property handle */
 		void Construct(const FArguments& InArgs, const TSharedRef<IPropertyHandle>& InPropertyHandle);
 
+		/** Sets the currently displayed text */
+		void SetDisplayText(const FText& Text);
+
 		/** Builds a text label */
 		static TSharedRef<SWidget> BuildLabel(TAttribute<FText> LabelText, const FSlateColor& ForegroundColor, const FSlateColor& BackgroundColor)
 		{
@@ -133,6 +136,9 @@ namespace UE::ControlRigEditor
 		/** The displayed property */
 		TWeakPtr<IPropertyHandle> WeakPropertyHandle;
 
+		/** The current display text */
+		FText DisplayText;
+
 		/** Timer handle for the refresh methods */
 		FTimerHandle RefreshTimerHandle;
 
