Anim in Engine: Fix Anim Details properties could not be navigated when navigating from properties in a collapsed to properties in an expanded struct

#rnx

diff --git a/Engine/Plugins/Animation/ControlRig/Source/ControlRigEditor/Private/AnimDetails/Customizations/AnimDetailsValueCustomization.cpp b/Engine/Plugins/Animation/ControlRig/Source/ControlRigEditor/Private/AnimDetails/Customizations/AnimDetailsValueCustomization.cpp
--- a/Engine/Plugins/Animation/ControlRig/Source/ControlRigEditor/Private/AnimDetails/Customizations/AnimDetailsValueCustomization.cpp
+++ b/Engine/Plugins/Animation/ControlRig/Source/ControlRigEditor/Private/AnimDetails/Customizations/AnimDetailsValueCustomization.cpp
@@ -195,7 +195,7 @@ namespace UE::ControlRigEditor
 				);
 		}
 
-		Private::FAnimDetailsChildWidgetVisiblityHandler(InStructPropertyHandle, ChildContentWidgets);
+		ChildContentVisibilityHandler = MakeUnique<Private::FAnimDetailsChildWidgetVisiblityHandler>(InStructPropertyHandle, ChildContentWidgets);
 	}
 
 	void FAnimDetailsValueCustomization::MakeHeaderRow(TSharedRef<IPropertyHandle>& InStructPropertyHandle, FDetailWidgetRow& InRow)
diff --git a/Engine/Plugins/Animation/ControlRig/Source/ControlRigEditor/Private/AnimDetails/Customizations/AnimDetailsValueCustomization.h b/Engine/Plugins/Animation/ControlRig/Source/ControlRigEditor/Private/AnimDetails/Customizations/AnimDetailsValueCustomization.h
--- a/Engine/Plugins/Animation/ControlRig/Source/ControlRigEditor/Private/AnimDetails/Customizations/AnimDetailsValueCustomization.h
+++ b/Engine/Plugins/Animation/ControlRig/Source/ControlRigEditor/Private/AnimDetails/Customizations/AnimDetailsValueCustomization.h
@@ -66,7 +66,7 @@ namespace UE::ControlRigEditor
 		FLinearColor GetColorFromProperty(const FName& PropertyName) const;
 
 		/** Handles visibility for child widgets in this customization */
-		Private::FAnimDetailsChildWidgetVisiblityHandler ChildContentVisibilityHandler;
+		TUniquePtr<Private::FAnimDetailsChildWidgetVisiblityHandler> ChildContentVisibilityHandler;
 
 		/** Pointer to the detail builder, or nullptr if not initialized */
 		IDetailLayoutBuilder* DetailBuilder = nullptr;
