Make the animation show flags of the control rig editor never clip off of the viewport toolbar.

[RN] minor editor

diff --git a/Engine/Plugins/Animation/ControlRig/Source/ControlRigEditor/Private/Editor/ControlRigViewportToolbarExtensions.cpp b/Engine/Plugins/Animation/ControlRig/Source/ControlRigEditor/Private/Editor/ControlRigViewportToolbarExtensions.cpp
--- a/Engine/Plugins/Animation/ControlRig/Source/ControlRigEditor/Private/Editor/ControlRigViewportToolbarExtensions.cpp
+++ b/Engine/Plugins/Animation/ControlRig/Source/ControlRigEditor/Private/Editor/ControlRigViewportToolbarExtensions.cpp
@@ -1462,6 +1462,7 @@ void UE::ControlRig::PopulateControlRigViewportToolbarShowSubmenu(const FName In
 	// Show this in the top-level to highlight it for Animation Mode users.
 	AnimationSubmenu.SetShowInToolbarTopLevel(true);
 	AnimationSubmenu.InsertPosition.Position = EToolMenuInsertType::First;
+	AnimationSubmenu.ToolBarData.ResizeParams.AllowClipping = false;
 	AllShowFlagsSection.AddEntry(AnimationSubmenu);
 }
 
