Control Rig: Make sure to auto tangents when baking, this got removed somehow.

diff --git a/Engine/Plugins/Animation/ControlRig/Source/ControlRig/Private/Sequencer/MovieSceneControlRigParameterSection.cpp b/Engine/Plugins/Animation/ControlRig/Source/ControlRig/Private/Sequencer/MovieSceneControlRigParameterSection.cpp
--- a/Engine/Plugins/Animation/ControlRig/Source/ControlRig/Private/Sequencer/MovieSceneControlRigParameterSection.cpp
+++ b/Engine/Plugins/Animation/ControlRig/Source/ControlRig/Private/Sequencer/MovieSceneControlRigParameterSection.cpp
@@ -4206,7 +4206,13 @@ bool UMovieSceneControlRigParameterSection::LoadAnimSequenceIntoThisSection(UAni
 			Channel->Optimize(Params);
 		}
 	}
-	
+	else
+	{
+		for (FMovieSceneFloatChannel* Channel : FloatChannels)
+		{
+			Channel->AutoSetTangents();
+		}
+	}
 	return true;
 }
 
