Anim in Engine: Fix Anim Details committing property values with enter no longer committed the value

#rnx

diff --git a/Engine/Plugins/Animation/ControlRig/Source/ControlRigEditor/Private/AnimDetails/Widgets/SAnimDetailsValueNumeric.cpp b/Engine/Plugins/Animation/ControlRig/Source/ControlRigEditor/Private/AnimDetails/Widgets/SAnimDetailsValueNumeric.cpp
--- a/Engine/Plugins/Animation/ControlRig/Source/ControlRigEditor/Private/AnimDetails/Widgets/SAnimDetailsValueNumeric.cpp
+++ b/Engine/Plugins/Animation/ControlRig/Source/ControlRigEditor/Private/AnimDetails/Widgets/SAnimDetailsValueNumeric.cpp
@@ -435,10 +435,7 @@ namespace UE::ControlRigEditor
 		{
 			const TSharedRef<IPropertyHandle> PropertyHandle = WeakPropertyHandle.Pin().ToSharedRef();
 
-			// When pressing enter, two commit calls will occur, one with ETextCommit::OnEnter 
-			// and one with ETextCommit::OnCleared. Only handle the later that is raised in any case.
-			if (CommitType != ETextCommit::OnEnter &&
-				!FAnimDetailsMultiEditUtil::Get().MultiEditMath<NumericType>(*WeakProxyManager.Get(), DisplayText.ToString(), PropertyHandle))
+			if (!FAnimDetailsMultiEditUtil::Get().MultiEditMath<NumericType>(*WeakProxyManager.Get(), DisplayText.ToString(), PropertyHandle))
 			{
 				FAnimDetailsMultiEditUtil::Get().MultiEditSet(*WeakProxyManager.Get(), Value, PropertyHandle);
 			}
