Control Rig: Fix crash when removing public variable in module

#rnx

diff --git a/Engine/Plugins/Animation/ControlRig/Source/ControlRig/Private/ControlRigOverride.cpp b/Engine/Plugins/Animation/ControlRig/Source/ControlRig/Private/ControlRigOverride.cpp
--- a/Engine/Plugins/Animation/ControlRig/Source/ControlRig/Private/ControlRigOverride.cpp
+++ b/Engine/Plugins/Animation/ControlRig/Source/ControlRig/Private/ControlRigOverride.cpp
@@ -255,7 +255,10 @@ uint8* FControlRigOverrideValue::SubjectPtrToValuePtr(const void* InSubjectPtr,
 	for (const FPropertyInfo& Info : Properties)
 	{
 		const FProperty* Property = Info.Property.Get();
-		check(Property);
+		if (!Property)
+		{
+			return nullptr;
+		}
 
 		const FArrayProperty* ArrayProperty = CastField<FArrayProperty>(Property);
 		if(ArrayProperty == nullptr)
diff --git a/Engine/Plugins/Animation/ControlRig/Source/ControlRig/Private/ModularRigController.cpp b/Engine/Plugins/Animation/ControlRig/Source/ControlRig/Private/ModularRigController.cpp
--- a/Engine/Plugins/Animation/ControlRig/Source/ControlRig/Private/ModularRigController.cpp
+++ b/Engine/Plugins/Animation/ControlRig/Source/ControlRig/Private/ModularRigController.cpp
@@ -2139,6 +2139,16 @@ void UModularRigController::RefreshModuleVariables(const FRigModuleReference* In
 	{
 		// Make sure the Override properties belong to the module class
 		FControlRigOverrideValue CorrectedOverride = Override;
+		if (CorrectedOverride.GetRootProperty() == nullptr)
+		{
+			UE_LOG(LogControlRig, Display, TEXT("Cannot determine root property '%s' in instance class %s."), *CorrectedOverride.GetPath(), *Module->Class.ToString());
+			continue;
+		}
+		if (CorrectedOverride.GetLeafProperty() == nullptr)
+		{
+			UE_LOG(LogControlRig, Display, TEXT("Cannot determine leaf property '%s' in instance class %s."), *CorrectedOverride.GetPath(), *Module->Class.ToString());
+			continue;
+		}
 		if(Override.GetRootProperty()->GetOwnerUObject() != Module->Class.Get())
 		{
 			for (int32 i=0; i<Override.Properties.Num(); i++)
