Control Rig: Fix 2 transactions being generated when dragging transform handle

diff --git a/Engine/Plugins/Animation/ControlRig/Source/ControlRigEditor/Private/EditMode/ControlRigEditMode.cpp b/Engine/Plugins/Animation/ControlRig/Source/ControlRigEditor/Private/EditMode/ControlRigEditMode.cpp
--- a/Engine/Plugins/Animation/ControlRig/Source/ControlRigEditor/Private/EditMode/ControlRigEditMode.cpp
+++ b/Engine/Plugins/Animation/ControlRig/Source/ControlRigEditor/Private/EditMode/ControlRigEditMode.cpp
@@ -1294,12 +1294,6 @@ bool FControlRigEditMode::HandleEndTransform(FEditorViewportClient* InViewportCl
 	
 	if (InteractionScopes.Num() > 0)
 	{		
-		if (bManipulatorMadeChange)
-		{
-			bManipulatorMadeChange = false;
-			GEditor->EndTransaction();
-		}
-
 		TArray<TWeakObjectPtr<UControlRig>> RigsToTick;
 		RigsToTick.Reserve(InteractionScopes.Num());
 
@@ -1334,6 +1328,13 @@ bool FControlRigEditMode::HandleEndTransform(FEditorViewportClient* InViewportCl
 
 		Keyframer.Finalize(InViewportClient->GetWorld());
 		Keyframer.Reset();
+
+		// Keyframer may cause additional changes, such as auto-keying in Sequencer, etc. Those changes shoudl be part of the current transaction.
+		if (bManipulatorMadeChange)
+		{
+			bManipulatorMadeChange = false;
+			GEditor->EndTransaction();
+		}
 		
 		InteractionScopes.Reset();
 
