Control Rig: Move max name length to a property of the hierarchy

#rnx

diff --git a/Engine/Plugins/Animation/ControlRig/Source/ControlRig/Private/Rigs/FKControlRig.cpp b/Engine/Plugins/Animation/ControlRig/Source/ControlRig/Private/Rigs/FKControlRig.cpp
--- a/Engine/Plugins/Animation/ControlRig/Source/ControlRig/Private/Rigs/FKControlRig.cpp
+++ b/Engine/Plugins/Animation/ControlRig/Source/ControlRig/Private/Rigs/FKControlRig.cpp
@@ -363,6 +363,27 @@ void UFKControlRig::CreateRigElements(const FReferenceSkeleton& InReferenceSkele
 	GetHierarchy()->Reset();
 	if (URigHierarchyController* Controller = GetHierarchy()->GetController(true))
 	{
+		// determine the max element name length
+		int32 MaxElementName = GetHierarchy()->GetMaxNameLength();
+		for (int32 Index = 0; Index < InitializationOptions.CurveNames.Num(); ++Index)
+		{
+			MaxElementName = FMath::Max<int32>(MaxElementName, InitializationOptions.CurveNames[Index].GetStringLength() + 14 ); // _CURVE_CONTROL suffix  
+		}
+		if (InSkeleton)
+		{
+			InSkeleton->ForEachCurveMetaData([&MaxElementName](const FName& InCurveName, const FCurveMetaData&)
+			{
+				MaxElementName = FMath::Max<int32>(MaxElementName, InCurveName.GetStringLength() + 14 ); // _CURVE_CONTROL suffix  
+			});
+		}
+
+		const TArray<FMeshBoneInfo> BoneInfos = InReferenceSkeleton.GetRefBoneInfo();
+		for (const FMeshBoneInfo& BoneInfo : BoneInfos)
+		{
+			MaxElementName = FMath::Max<int32>(MaxElementName, BoneInfo.Name.GetStringLength() + 8 ); // _CONTROL suffix  
+		}
+		GetHierarchy()->SetMaxNameLength(MaxElementName);
+		
 		Controller->ImportBones(InReferenceSkeleton, NAME_None, false, false, true, false);
 
 		if (InitializationOptions.bGenerateCurveControls)
diff --git a/Engine/Plugins/Animation/ControlRig/Source/ControlRig/Private/Rigs/RigHierarchy.cpp b/Engine/Plugins/Animation/ControlRig/Source/ControlRig/Private/Rigs/RigHierarchy.cpp
--- a/Engine/Plugins/Animation/ControlRig/Source/ControlRig/Private/Rigs/RigHierarchy.cpp
+++ b/Engine/Plugins/Animation/ControlRig/Source/ControlRig/Private/Rigs/RigHierarchy.cpp
@@ -117,6 +117,7 @@ URigHierarchy::URigHierarchy()
 , ElementIndexLookup()
 , TransformStackIndex(0)
 , bTransactingForTransformChange(false)
+, MaxNameLength(GetDefaultMaxNameLength())
 , bIsInteracting(false)
 , LastInteractedKey()
 , bSuspendNotifications(false)
@@ -760,6 +761,9 @@ void URigHierarchy::Load(FArchive& Ar)
 		IncrementTopologyVersion();
 	}
 
+	// set the max name length but respect the length of elements + components
+	SetMaxNameLength(GetDefaultMaxNameLength());
+	
 	(void)SortElementStorage();
 }
 
@@ -868,6 +872,8 @@ void URigHierarchy::Reset_Impl(bool bResetElements)
 	{
 		HierarchyForCacheValidation->Reset();
 	}
+
+	SetMaxNameLength(GetDefaultMaxNameLength());
 }
 
 #if WITH_EDITOR
@@ -2500,7 +2506,32 @@ bool URigHierarchy::SplitNameSpace_Deprecated(const FRigName& InNameSpacedPath,
 	return false;
 }
 
-void URigHierarchy::SanitizeName(FRigName& InOutName, bool bAllowNameSpaces)
+int32 URigHierarchy::SetMaxNameLength(int32 InMaxNameLength)
+{
+	MaxNameLength = InMaxNameLength + UE_RIGHIERARCHY_NUMERIC_SUFFIX_LEN;
+
+	// ensure the new setting suffices current element names
+	for(const FRigBaseElement* Element : Elements)
+	{
+		MaxNameLength = FMath::Max<int32>(MaxNameLength, Element->GetName().Len());
+	}
+	
+	// ensure the new setting suffices current component names
+	for(FInstancedStruct& InstancedStruct : ElementComponents)
+	{
+		if(!InstancedStruct.IsValid())
+		{
+			continue;
+		}
+
+		FRigBaseComponent* Component = InstancedStruct.GetMutablePtr<FRigBaseComponent>();
+		MaxNameLength = FMath::Max<int32>(MaxNameLength, Component->GetName().Len());
+	}
+
+	return MaxNameLength;
+}
+
+void URigHierarchy::SanitizeName(FRigName& InOutName, bool bAllowNameSpaces, int32 InMaxNameLength)
 {
 	// Sanitize the name
 	FString SanitizedNameString = InOutName.GetName();
@@ -2526,9 +2557,9 @@ void URigHierarchy::SanitizeName(FRigName& InOutName, bool bAllowNameSpaces)
 		}
 	}
 
-	if (SanitizedNameString.Len() > GetMaxNameLength())
+	if (SanitizedNameString.Len() > InMaxNameLength)
 	{
-		SanitizedNameString.LeftChopInline(SanitizedNameString.Len() - GetMaxNameLength());
+		SanitizedNameString.LeftChopInline(SanitizedNameString.Len() - InMaxNameLength);
 		bChangedSomething = true;
 	}
 
@@ -2538,10 +2569,10 @@ void URigHierarchy::SanitizeName(FRigName& InOutName, bool bAllowNameSpaces)
 	}
 }
 
-FRigName URigHierarchy::GetSanitizedName(const FRigName& InName, bool bAllowNameSpaces)
+FRigName URigHierarchy::GetSanitizedName(const FRigName& InName, bool bAllowNameSpaces, int32 InMaxNameLength)
 {
 	FRigName Name = InName;
-	SanitizeName(Name, bAllowNameSpaces);
+	SanitizeName(Name, bAllowNameSpaces, InMaxNameLength);
 	return Name;
 }
 
@@ -2612,7 +2643,7 @@ bool URigHierarchy::IsNameAvailable(const FRigName& InPotentialNewName, ERigElem
 	}
 
 	FRigName SanitizedName = UnsanitizedName;
-	SanitizeName(SanitizedName, bAllowNameSpaces);
+	SanitizeName(SanitizedName, bAllowNameSpaces, GetMaxNameLength());
 
 	if (SanitizedName != UnsanitizedName)
 	{
@@ -2672,7 +2703,7 @@ bool URigHierarchy::IsDisplayNameAvailable(const FRigElementKey& InParentElement
 	}
 
 	FRigName SanitizedName = UnsanitizedName;
-	SanitizeName(SanitizedName, true);
+	SanitizeName(SanitizedName, true, GetMaxNameLength());
 
 	if (SanitizedName != UnsanitizedName)
 	{
@@ -2702,7 +2733,7 @@ bool URigHierarchy::IsComponentNameAvailable(const FRigElementKey& InElementKey,
 FRigName URigHierarchy::GetSafeNewName(const FRigName& InPotentialNewName, ERigElementType InType, bool bAllowNameSpace) const
 {
 	FRigName SanitizedName = InPotentialNewName;
-	SanitizeName(SanitizedName, bAllowNameSpace);
+	SanitizeName(SanitizedName, bAllowNameSpace, GetMaxNameLength());
 
 	bAllowNameSpaceWhenSanitizingName = bAllowNameSpace;
 	if(ExecuteContext)
@@ -2721,9 +2752,9 @@ FRigName URigHierarchy::GetSafeNewName(const FRigName& InPotentialNewName, ERigE
 		}
 	}
 
-	if (SanitizedName.Len() > GetMaxNameLength() - 4)
+	if (SanitizedName.Len() > GetMaxNameLength() - UE_RIGHIERARCHY_NUMERIC_SUFFIX_LEN)
 	{
-		SanitizedName.SetName(SanitizedName.GetName().LeftChop(SanitizedName.Len() - (GetMaxNameLength() - 4)));
+		SanitizedName.SetName(SanitizedName.GetName().LeftChop(SanitizedName.Len() - (GetMaxNameLength() - UE_RIGHIERARCHY_NUMERIC_SUFFIX_LEN)));
 	}
 
 	const FName UniqueName = FRigName(GetUniqueName(*SanitizedName.GetName(), [this, InType](const FName& InName) -> bool
@@ -2745,7 +2776,7 @@ FRigName URigHierarchy::GetSafeNewDisplayName(const FRigElementKey& InParentElem
 	static const int32 ControlElementIndex = RigElementTypeToFlatIndex(ERigElementType::Control); 
 
 	FRigName SanitizedName = InPotentialNewDisplayName;
-	SanitizeName(SanitizedName);
+	SanitizeName(SanitizedName, true, GetMaxNameLength());
 
 	// this can be nullptr - and that's ok.
 	const FRigBaseElement* ParentElement = Find(InParentElement);
@@ -2836,7 +2867,7 @@ FRigName URigHierarchy::GetSafeNewComponentName(const FRigElementKey& InElementK
 	}
 
 	FRigName SanitizedName = InPotentialNewName;
-	SanitizeName(SanitizedName);
+	SanitizeName(SanitizedName, true, GetMaxNameLength());
 
 	if (SanitizedName.Len() > GetMaxNameLength() - 4)
 	{
diff --git a/Engine/Plugins/Animation/ControlRig/Source/ControlRig/Private/Rigs/RigHierarchyController.cpp b/Engine/Plugins/Animation/ControlRig/Source/ControlRig/Private/Rigs/RigHierarchyController.cpp
--- a/Engine/Plugins/Animation/ControlRig/Source/ControlRig/Private/Rigs/RigHierarchyController.cpp
+++ b/Engine/Plugins/Animation/ControlRig/Source/ControlRig/Private/Rigs/RigHierarchyController.cpp
@@ -2871,7 +2871,7 @@ int32 URigHierarchyController::AddElement(FRigBaseElement* InElementToAdd, FRigB
 	Hierarchy->IncrementTopologyVersion();
 
 	FRigName DesiredName = InDesiredName;
-	URigHierarchy::SanitizeName(DesiredName);
+	URigHierarchy::SanitizeName(DesiredName,true, Hierarchy->GetMaxNameLength());
 
 	const FRigHierarchyModulePath ModulePath(InDesiredName.ToString());
 	if(ModulePath.IsValid())
diff --git a/Engine/Plugins/Animation/ControlRig/Source/ControlRig/Public/Rigs/RigHierarchy.h b/Engine/Plugins/Animation/ControlRig/Source/ControlRig/Public/Rigs/RigHierarchy.h
--- a/Engine/Plugins/Animation/ControlRig/Source/ControlRig/Public/Rigs/RigHierarchy.h
+++ b/Engine/Plugins/Animation/ControlRig/Source/ControlRig/Public/Rigs/RigHierarchy.h
@@ -37,6 +37,9 @@ DECLARE_MULTICAST_DELEGATE_ThreeParams(FRigHierarchyMetadataTagChangedDelegate,
 extern CONTROLRIG_API TAutoConsoleVariable<bool> CVarControlRigHierarchyEnableRotationOrder;
 extern CONTROLRIG_API TAutoConsoleVariable<bool> CVarControlRigHierarchyEnableModules;
 
+#define UE_RIGHIERARCHY_DEFAULT_MAX_NAME_LENGTH 100 
+#define UE_RIGHIERARCHY_NUMERIC_SUFFIX_LEN 4 
+
 UENUM()
 enum ERigTransformStackEntryType : int
 {
@@ -2009,20 +2012,33 @@ public:
 	 * Returns the max allowed length for a name within the hierarchy.
 	 * @return Returns the max allowed length for a name within the hierarchy.
 	 */
-	static int32 GetMaxNameLength() { return 200; }
+	static int32 GetDefaultMaxNameLength() { return UE_RIGHIERARCHY_DEFAULT_MAX_NAME_LENGTH; }
+
+	/**
+	 * Returns the max allowed length for a name within the hierarchy.
+	 * @return Returns the max allowed length for a name within the hierarchy.
+	 */
+	int32 GetMaxNameLength() const { return MaxNameLength; }
+
+	/**
+	 * Sets the maximum allowed length for a name within the hierarchy.
+	 * This value is corrected to the upperbound between desired max and current max.
+	 * @return Returns the max allowed length for a name within the hierarchy.
+	 */
+	CONTROLRIG_API int32 SetMaxNameLength(int32 InMaxNameLength);
 
 	/**
 	 * Sanitizes a name by removing invalid characters.
 	 * @param InOutName The name to sanitize in place.
 	 */
-	static CONTROLRIG_API void SanitizeName(FRigName& InOutName, bool bAllowNameSpaces = true);
+	static CONTROLRIG_API void SanitizeName(FRigName& InOutName, bool bAllowNameSpaces = true, int32 InMaxNameLength = UE_RIGHIERARCHY_DEFAULT_MAX_NAME_LENGTH);
 
 	/**
 	 * Sanitizes a name by removing invalid characters.
 	 * @param InName The name to sanitize.
 	 * @return The sanitized name.
  	 */
-	static CONTROLRIG_API FRigName GetSanitizedName(const FRigName& InName, bool bAllowNameSpaces = true);
+	static CONTROLRIG_API FRigName GetSanitizedName(const FRigName& InName, bool bAllowNameSpaces = true, int32 InMaxNameLength = UE_RIGHIERARCHY_DEFAULT_MAX_NAME_LENGTH);
 
 	/**
 	 * Returns true if a given name is available.
@@ -4856,6 +4872,12 @@ private:
 	void NotifyPostUndoSelectionChanges();
 #endif
 
+	/**
+	 * The max length for an element / component name within this hierarchy.
+	 * Defaults to GetDefaultMaxNameLength().
+	 */
+	int32 MaxNameLength;
+	
 	/**
 	 * Manages merging transform actions into one during an interaction
 	 */
