Adjustment Blending: Change world space algorithm to work in the additive layer space, not fully in world space, fixes some calculation errors.

diff --git a/Engine/Plugins/Animation/ControlRig/Source/ControlRigEditor/Private/Sequencer/AnimLayers/AnimLayers.cpp b/Engine/Plugins/Animation/ControlRig/Source/ControlRigEditor/Private/Sequencer/AnimLayers/AnimLayers.cpp
--- a/Engine/Plugins/Animation/ControlRig/Source/ControlRigEditor/Private/Sequencer/AnimLayers/AnimLayers.cpp
+++ b/Engine/Plugins/Animation/ControlRig/Source/ControlRigEditor/Private/Sequencer/AnimLayers/AnimLayers.cpp
@@ -3508,7 +3508,22 @@ struct FTransformAsArray
 		return Values[Index];
 	}
 
-	
+	FTransformAsArray Abs() const
+	{
+		FTransformAsArray Val;
+		for (int32 Index = 0; Index < 3; ++Index)
+		{
+			const int32 RealIndex = Index * 3;
+			const int32 RealIndex1 = RealIndex + 1;
+			const int32 RealIndex2 = RealIndex + 2;
+
+			Val.Values[RealIndex]  = FMath::Abs(Values[RealIndex]);
+			Val.Values[RealIndex1] = FMath::Abs(Values[RealIndex1]);
+			Val.Values[RealIndex2] = FMath::Abs(Values[RealIndex2]);
+		}
+		return Val;
+	}
+
 	FTransformAsArray operator+=(const FTransformAsArray& RHS)
 	{
 		for (int32 Index = 0; Index < 3; ++Index)
@@ -3540,6 +3555,21 @@ FTransformAsArray operator-(const FTransformAsArray& LHS, const FTransformAsArra
 	return ReturnVal;
 }
 
+FTransformAsArray operator+(const FTransformAsArray& LHS, const FTransformAsArray& RHS) {
+	FTransformAsArray ReturnVal;
+	for (int32 Index = 0; Index < 3; ++Index)
+	{
+		const int32 RealIndex = Index * 3;
+		const int32 RealIndex1 = RealIndex + 1;
+		const int32 RealIndex2 = RealIndex + 2;
+
+		ReturnVal[RealIndex] = LHS[RealIndex] + RHS[RealIndex];
+		ReturnVal[RealIndex1] = LHS[RealIndex1] + RHS[RealIndex1];
+		ReturnVal[RealIndex2] = LHS[RealIndex2] + RHS[RealIndex2];
+	}
+	return ReturnVal;
+}
+
 void GetWorldTransformPercentageOfChange(int32 FrameStartIndex, int32 FrameEndIndex, const TPair<FName, TSharedPtr<UE::AIE::FArrayOfTransforms>>& Pair, TArray<FTransformAsArray>& ValueDifferences, TArray<TArray<double>>& PercentageDifferences)
 {
 	const int32 NumKeys = FrameEndIndex - FrameStartIndex + 1;
@@ -3551,9 +3581,10 @@ void GetWorldTransformPercentageOfChange(int32 FrameStartIndex, int32 FrameEndIn
 		CurrentArray = NextArray;
 		FTransform& NextTransform = (*Pair.Value)[Index + 1];
 		NextArray = FTransformAsArray(NextTransform);
-		ValueDifferences[Index] = NextArray - CurrentArray;
+		ValueDifferences[Index] = (NextArray - CurrentArray).Abs();
 		TotalChange += ValueDifferences[Index];
 	}
+	PercentageDifferences.Reset(0);
 	PercentageDifferences.SetNum(9); //one per channel
 	for (int32 ChannelIndex = 0; ChannelIndex < 9; ++ChannelIndex)
 	{
@@ -3671,7 +3702,7 @@ void AdjustmentBlendWorldSpace(ISequencer* Sequencer,  FAnimLayerItem& BaseAnimL
 	}
 	AnimLayerItem.SetSectionsActive(true);
 	FEvalHelpers::CalculateWorldTransforms(World, Sequencer, CurrentFrameTimes,
-		AdditiveKeyIndices,
+		CurrentFrameTimes.GetFullIndexArray(),
 		EvaluatingActors, EvaluatingControlRigs);
 	
 	
@@ -3699,37 +3730,49 @@ void AdjustmentBlendWorldSpace(ISequencer* Sequencer,  FAnimLayerItem& BaseAnimL
 			{
 				int32 FrameStartIndex = AdditiveKeyIndices[KeyIndex];
 				int32 FrameEndIndex = AdditiveKeyIndices[KeyIndex + 1];
-				for (const TPair<FName, TSharedPtr<FArrayOfTransforms>>& Pair : BaseControlRigAndWorld.Value.ControlAndWorldTransforms)
+				if (FrameStartIndex == FrameEndIndex)
 				{
-					if (FRigControlElement* ControlElement = ControlRig->FindControl(Pair.Key))
+					continue;  //should never happen
+				}
+				for (const TPair<FName, TSharedPtr<FArrayOfTransforms>>& BasePair : BaseControlRigAndWorld.Value.ControlAndWorldTransforms)
+				{
+					if (FRigControlElement* ControlElement = ControlRig->FindControl(BasePair.Key))
 					{
 						if (ControlTypeValidForAdjustment(ControlElement) == false)
 						{
 							continue;
 						}
-						if (TSharedPtr<FArrayOfTransforms>* EvalTransforms = EvalCRandWT->ControlAndWorldTransforms.Find(Pair.Key))
+						if (TSharedPtr<FArrayOfTransforms>* EvalTransforms = EvalCRandWT->ControlAndWorldTransforms.Find(BasePair.Key))
 						{
+							ValueDifferences.Reset();
+							ValueDifferences.SetNum(CurrentFrameTimes.NumFrames);
 							//get base layer values as diffs and percentages
-							GetWorldTransformPercentageOfChange(FrameStartIndex, FrameEndIndex, Pair, ValueDifferences, PercentageDifferences);
+							GetWorldTransformPercentageOfChange(FrameStartIndex, FrameEndIndex, BasePair, ValueDifferences, PercentageDifferences);
 							//get new layer change
-							FTransform ST = (*EvalTransforms)->Transforms[FrameStartIndex];
-							ST = ST.GetRelativeTransform((*EvalCRandWT->ParentTransforms)[FrameStartIndex]);
-							FTransformAsArray StartTransform(ST);
-							
-							FTransform ET = (*EvalTransforms)->Transforms[FrameEndIndex];
-							ET = ET.GetRelativeTransform((*EvalCRandWT->ParentTransforms)[FrameEndIndex]);
-							FTransformAsArray EndTransform(ET);
+							FTransform STBase = BasePair.Value->Transforms[FrameStartIndex];
+							FTransformAsArray StartTransformBase(STBase);
+							FTransformAsArray StartTransformSubtractBase((*EvalTransforms)->Transforms[FrameStartIndex]);
+							StartTransformSubtractBase = StartTransformSubtractBase - StartTransformBase;
 							
+							FTransform ETBase = BasePair.Value->Transforms[FrameEndIndex];
+							FTransformAsArray EndTransformBase(ETBase);
+							FTransformAsArray EndTransformSubtractBase((*EvalTransforms)->Transforms[FrameEndIndex]);
+							EndTransformSubtractBase = EndTransformSubtractBase - EndTransformBase;
+				
 							FTransformAsArray CurrentValue;
 							FTransformAsArray PreviousValue;
 							FTransformAsArray TotalPoseLayerChange;
+
 							for (int32 ChannelIndex = 0; ChannelIndex < 9; ++ChannelIndex)
 							{
-								TotalPoseLayerChange[ChannelIndex] = FMath::Abs(EndTransform[ChannelIndex] - StartTransform[ChannelIndex]);
-								PreviousValue[ChannelIndex] = StartTransform[ChannelIndex];
+								TotalPoseLayerChange[ChannelIndex] = FMath::Abs(StartTransformSubtractBase[ChannelIndex] - EndTransformSubtractBase[ChannelIndex]);
+								PreviousValue[ChannelIndex] = StartTransformSubtractBase[ChannelIndex];
 							}
-							for (int32 FrameIndex = FrameStartIndex; FrameIndex <= FrameEndIndex; ++FrameIndex)
+							FTransform CT,NT;
+							FTransformAsArray CurrentTransform, NewTransform;
+							for (int32 FrameIndex = FrameStartIndex + 1; FrameIndex <= FrameEndIndex; ++FrameIndex)
 							{
+
 								const FFrameNumber FrameNumber = CurrentFrameTimes.CalculateFrame(FrameIndex);
 								FFrameTime GlobalTime(FrameNumber);
 								GlobalTime = LocalToRootTransform.TryTransformTime(GlobalTime).Get(GlobalTime); //player evals in root time so need to go back to it.
@@ -3739,24 +3782,34 @@ void AdjustmentBlendWorldSpace(ISequencer* Sequencer,  FAnimLayerItem& BaseAnimL
 								ControlRig->Evaluate_AnyThread();
 								for (int32 ChannelIndex = 0; ChannelIndex < 9; ++ChannelIndex)
 								{
-									if(FrameIndex < PercentageDifferences[ChannelIndex].Num())
+									if (FMath::IsNearlyZero(TotalPoseLayerChange[ChannelIndex], (1.e-4f)))
 									{
-										double Value = PercentageDifferences[ChannelIndex][FrameIndex];
+										CurrentValue[ChannelIndex] =PreviousValue[ChannelIndex];
+									}
+									else if (FrameIndex < PercentageDifferences[ChannelIndex].Num())
+									{
+										double Value = PercentageDifferences[ChannelIndex][FrameIndex - 1];
 										const double ValueDelta = (TotalPoseLayerChange[ChannelIndex] / 100.0) * Value;
-										CurrentValue[ChannelIndex] = (EndTransform[ChannelIndex] > StartTransform[ChannelIndex]) ?
+										CurrentValue[ChannelIndex] = (EndTransformSubtractBase[ChannelIndex] > StartTransformSubtractBase[ChannelIndex]) ?
 											PreviousValue[ChannelIndex] + ValueDelta : PreviousValue[ChannelIndex] - ValueDelta;
-										PreviousValue[ChannelIndex] = CurrentValue[ChannelIndex];
 									}
 									else
 									{
-										CurrentValue[ChannelIndex] = PreviousValue[ChannelIndex];
+										CurrentValue[ChannelIndex] = StartTransformSubtractBase[ChannelIndex] + (EndTransformSubtractBase[ChannelIndex] - StartTransformSubtractBase[ChannelIndex]) * (FrameIndex - FrameStartIndex) / (FrameEndIndex - FrameStartIndex);
 									}
+									PreviousValue[ChannelIndex] = CurrentValue[ChannelIndex];
+
 								}
-								FTransform NewTransform = CurrentValue.AsTransform();
+								//get current transform of the base
+								CT = BasePair.Value->Transforms[FrameIndex];
+								CT = CT.GetRelativeTransform((*EvalCRandWT->ParentTransforms)[FrameIndex]);
+								CurrentTransform = FTransformAsArray(CT);
+								NewTransform = CurrentTransform + CurrentValue;
+								NT = NewTransform.AsTransform();;
 								Context.LocalTime = TickResolution.AsSeconds(FFrameTime(FrameNumber));
-								FRigControlValue Value = ControlRig->GetControlValueFromGlobalTransform(ControlElement->GetKey().Name, NewTransform, TransformType);
-								NewTransform = Value.GetAsTransform(ControlElement->Settings.ControlType, ControlElement->Settings.PrimaryAxis);
-								FEulerTransform EulerTransform(NewTransform);;
+								FRigControlValue Value = ControlRig->GetControlValueFromGlobalTransform(ControlElement->GetKey().Name, NT, TransformType);
+								NT = Value.GetAsTransform(ControlElement->Settings.ControlType, ControlElement->Settings.PrimaryAxis);
+								FEulerTransform EulerTransform(NT);;
 								UE::AIE::FSetTransformHelpers::SetControlTransform(EulerTransform, ControlRig, ControlElement, Context);
 
 							}
