Anim In Engine: Space Switching: Remove extra keys when clicking a space key by checking to see if we are changing into the same space in which case we can delete the extra key and the space key

diff --git a/Engine/Plugins/Animation/ControlRig/Source/ControlRigEditor/Private/Sequencer/ControlRigSpaceChannelEditors.cpp b/Engine/Plugins/Animation/ControlRig/Source/ControlRigEditor/Private/Sequencer/ControlRigSpaceChannelEditors.cpp
--- a/Engine/Plugins/Animation/ControlRig/Source/ControlRigEditor/Private/Sequencer/ControlRigSpaceChannelEditors.cpp
+++ b/Engine/Plugins/Animation/ControlRig/Source/ControlRigEditor/Private/Sequencer/ControlRigSpaceChannelEditors.cpp
@@ -368,14 +368,35 @@ FKeyHandle FControlRigSpaceChannelHelpers::SequencerKeyControlRigSpaceChannel(UC
 		TSortedMap<FFrameNumber, FFrameNumber> ExtraFrames;
 		FControlRigSpaceChannelHelpers::GetFramesInThisSpaceAfterThisTime(ControlRig, ControlKey.Name, ExistingValue,
 			Channel, SectionToKey, Time, ExtraFrames);
+		SectionToKey->Modify();
+		UMovieSceneControlRigParameterSection* ControlRigSection = Cast<UMovieSceneControlRigParameterSection>(SectionToKey);
+
 		if (ExtraFrames.Num() > 0)
 		{
 			for (const TPair<FFrameNumber, FFrameNumber>& Frame : ExtraFrames)
 			{
 				Frames.Add(Frame.Value);
 			}
+			//if the next space is the same as the new one, then we delete those transform frame and space key if it exists
+			FMovieSceneControlRigSpaceBaseKey NextValue;
+			using namespace UE::MovieScene;
+			FFrameNumber NextSpaceTime = Frames[Frames.Num() - 1] + 1;
+			EvaluateChannel(Channel, NextSpaceTime, NextValue);
+			if (NextValue == Value)
+			{
+				int32 ExistingIndex = ChannelInterface.FindKey(NextSpaceTime);
+				if (ExistingIndex != INDEX_NONE)
+				{
+					ChannelInterface.RemoveKey(ExistingIndex);
+					if (ControlRigSection)
+					{
+						FControlRigSpaceChannelHelpers::DeleteTransformKeysAtThisTime(ControlRig, ControlRigSection, ControlKey.Name, NextSpaceTime -1);
+						Frames.SetNum(Frames.Num() - 1);
+					}
+				}
+			}
+
 		}
-		SectionToKey->Modify();
 		const FFrameNumber StartFrame = Sequencer->GetFocusedMovieSceneSequence()->GetMovieScene()->GetPlaybackRange().GetLowerBoundValue();
 		if (StartFrame == Time)
 		{
@@ -403,7 +424,6 @@ FKeyHandle FControlRigSpaceChannelHelpers::SequencerKeyControlRigSpaceChannel(UC
 
 		// store tangents to keep the animation when switching 
 		TArray<FMovieSceneTangentData> Tangents;
-		UMovieSceneControlRigParameterSection* ControlRigSection = Cast<UMovieSceneControlRigParameterSection>(SectionToKey);
 		if (ControlRigSection)
 		{
 			FChannelMapInfo* pChannelIndex = nullptr;
