Anim In Engine: Pivot tool fixes: Constraints work,  Fix shift click so always goes to last selected object, added hot key to toggle pivot(and motion trail0 modes. Contains commented out code that is supposed to let you scale the widget based upon preference but it's not working(for editor/gizmo team).

diff --git a/Engine/Plugins/Animation/ControlRig/Source/ControlRigEditor/Private/ControlRigSequencerEditorLibrary.cpp b/Engine/Plugins/Animation/ControlRig/Source/ControlRigEditor/Private/ControlRigSequencerEditorLibrary.cpp
--- a/Engine/Plugins/Animation/ControlRig/Source/ControlRigEditor/Private/ControlRigSequencerEditorLibrary.cpp
+++ b/Engine/Plugins/Animation/ControlRig/Source/ControlRigEditor/Private/ControlRigSequencerEditorLibrary.cpp
@@ -52,6 +52,7 @@
 #include "SequencerTools.h"
 #include "EditMode/Tween/ControlRigTweenModel.h"
 #include "Math/Abstraction/KeyBlendingAbstraction.h"
+#include "Tools/EvaluateSequencerTools.h"
 #include "Sequencer/AnimLayers/AnimLayers.h"
 
 
@@ -1037,8 +1038,8 @@ TArray<FTransform> UControlRigSequencerEditorLibrary::GetControlRigWorldTransfor
 static void LocalSetControlRigWorldTransforms(ULevelSequence* LevelSequence,UControlRig* ControlRig, FName ControlName, EControlRigSetKey SetKey, const TArray<FFrameNumber>& InFrames, 
 	const TArray<FTransform>& WorldTransforms, EMovieSceneTimeUnit TimeUnit)
 {
-	TWeakPtr<ISequencer> WeakSequencer = GetSequencerFromAsset();
-	if (WeakSequencer.IsValid() && ControlRig)
+	TSharedPtr<ISequencer> SequencerPtr = GetSequencerFromAsset();
+	if (SequencerPtr && ControlRig)
 	{
 		if (TSharedPtr<IControlRigObjectBinding> ObjectBinding = ControlRig->GetObjectBinding())
 		{
@@ -1048,38 +1049,47 @@ static void LocalSetControlRigWorldTransforms(ULevelSequence* LevelSequence,UCon
 				AActor* Actor = Component->GetTypedOuter< AActor >();
 				if (Actor)
 				{
-					TArray<FFrameNumber> Frames;
-					ConvertFramesToTickResolution(LevelSequence->GetMovieScene(), InFrames, TimeUnit, Frames);
+					if (FRigControlElement* ControlElement = ControlRig->GetHierarchy()->Find<FRigControlElement>((FRigElementKey(ControlName, ERigElementType::Control))))
+					{
 
-					UMovieScene* MovieScene = WeakSequencer.Pin().Get()->GetFocusedMovieSceneSequence()->GetMovieScene();
-					FFrameRate TickResolution = MovieScene->GetTickResolution();
-					FRigControlModifiedContext Context;
-					Context.SetKey = SetKey;
+						TArray<FFrameNumber> Frames;
+						ConvertFramesToTickResolution(LevelSequence->GetMovieScene(), InFrames, TimeUnit, Frames);
 
-					ControlRig->Modify();
-					TArray<FTransform> ControlRigParentWorldTransforms;
-					FActorForWorldTransforms ControlRigActorSelection;
-					ControlRigActorSelection.Actor = Actor;
-					MovieSceneToolHelpers::GetActorWorldTransforms(WeakSequencer.Pin().Get(), ControlRigActorSelection, Frames, ControlRigParentWorldTransforms);
-					FControlRigSnapper Snapper;
-					
-					TArray<FFrameNumber> OneFrame;
-					OneFrame.SetNum(1);
-					TArray<FTransform> CurrentControlRigTransform, CurrentParentWorldTransform;
-					CurrentControlRigTransform.SetNum(1);
-					CurrentParentWorldTransform.SetNum(1);
-					
-					for (int32 Index = 0; Index < WorldTransforms.Num(); ++Index)
-					{
-						OneFrame[0] = Frames[Index];
-						CurrentParentWorldTransform[0] = ControlRigParentWorldTransforms[Index];
-						//this will evaluate at the current frame which we want
-						Snapper.GetControlRigControlTransforms(WeakSequencer.Pin().Get(), ControlRig, ControlName, OneFrame, CurrentParentWorldTransform, CurrentControlRigTransform);
-							
-						const FFrameNumber& FrameNumber = Frames[Index];
-						Context.LocalTime = TickResolution.AsSeconds(FFrameTime(FrameNumber));
-						FTransform GlobalTransform = WorldTransforms[Index].GetRelativeTransform(ControlRigParentWorldTransforms[Index]);
-						ControlRig->SetControlGlobalTransform(ControlName, GlobalTransform, true, Context, false /*undo*/, false /*bPrintPython*/, true/* bFixEulerFlips*/);
+						UMovieScene* MovieScene = SequencerPtr->GetFocusedMovieSceneSequence()->GetMovieScene();
+						FFrameRate TickResolution = MovieScene->GetTickResolution();
+						FRigControlModifiedContext Context;
+						Context.SetKey = SetKey;
+
+						ControlRig->Modify();
+						TArray<FTransform> ControlRigParentWorldTransforms;
+						FActorForWorldTransforms ControlRigActorSelection;
+						ControlRigActorSelection.Actor = Actor;
+						MovieSceneToolHelpers::GetActorWorldTransforms(SequencerPtr.Get(), ControlRigActorSelection, Frames, ControlRigParentWorldTransforms);
+						FControlRigSnapper Snapper;
+
+						TArray<FFrameNumber> OneFrame;
+						OneFrame.SetNum(1);
+						TArray<FTransform> CurrentControlRigTransform, CurrentParentWorldTransform;
+						CurrentControlRigTransform.SetNum(1);
+						CurrentParentWorldTransform.SetNum(1);
+
+						for (int32 Index = 0; Index < WorldTransforms.Num(); ++Index)
+						{
+							OneFrame[0] = Frames[Index];
+							CurrentParentWorldTransform[0] = ControlRigParentWorldTransforms[Index];
+							//this will evaluate at the current frame which we want
+							Snapper.GetControlRigControlTransforms(SequencerPtr.Get(), ControlRig, ControlName, OneFrame, CurrentParentWorldTransform, CurrentControlRigTransform);
+
+							const FFrameNumber& FrameNumber = Frames[Index];
+							Context.LocalTime = TickResolution.AsSeconds(FFrameTime(FrameNumber));
+							FTransform GlobalTransform = WorldTransforms[Index].GetRelativeTransform(ControlRigParentWorldTransforms[Index]);
+
+							constexpr ERigTransformType::Type TransformType = ERigTransformType::CurrentGlobal;
+							FRigControlValue Value = ControlRig->GetControlValueFromGlobalTransform(ControlName, GlobalTransform, TransformType);
+							GlobalTransform = Value.GetAsTransform(ControlElement->Settings.ControlType, ControlElement->Settings.PrimaryAxis);
+							FEulerTransform EulerTransform(GlobalTransform);
+							UE::AIE::FSetTransformHelpers::SetControlTransform(EulerTransform, ControlRig, ControlElement, Context);
+						}
 					}
 				}
 			}
diff --git a/Engine/Plugins/Animation/ControlRig/Source/ControlRigEditor/Private/EditMode/ControlRigEditMode.cpp b/Engine/Plugins/Animation/ControlRig/Source/ControlRigEditor/Private/EditMode/ControlRigEditMode.cpp
--- a/Engine/Plugins/Animation/ControlRig/Source/ControlRigEditor/Private/EditMode/ControlRigEditMode.cpp
+++ b/Engine/Plugins/Animation/ControlRig/Source/ControlRigEditor/Private/EditMode/ControlRigEditMode.cpp
@@ -86,6 +86,7 @@
 #include "Slate/SceneViewport.h"
 #include "EditorGizmos/EditorTransformGizmoUtil.h"
 #include "Sequencer/AnimLayers/AnimLayers.h"
+#include "Tools/MotionTrailOptions.h"
 
 #include UE_INLINE_GENERATED_CPP_BY_NAME(ControlRigEditMode)
 
@@ -3125,6 +3126,14 @@ void FControlRigEditMode::BindCommands()
 	CommandBindings->MapAction(
 		Commands.OpenSpacePickerWidget,
 		FExecuteAction::CreateRaw(this, &FControlRigEditMode::OpenSpacePickerWidget));
+
+	CommandBindings->MapAction(
+		Commands.TogglePivotMode,
+		FExecuteAction::CreateRaw(this, &FControlRigEditMode::TogglePivotMode));
+
+	CommandBindings->MapAction(
+		Commands.ToggleMotionTrails,
+		FExecuteAction::CreateRaw(this, &FControlRigEditMode::ToggleMotionTrails));
 }
 
 bool FControlRigEditMode::IsControlSelected(const bool bUseShapes) const
@@ -3545,6 +3554,43 @@ void FControlRigEditMode::OpenSpacePickerWidget()
 
 	PickerWidget->OpenDialog(false);
 }
+static void ToggleTool(const FString& InToolName)
+{
+	if (FLevelEditorModule* LevelEditorModule = FModuleManager::GetModulePtr<FLevelEditorModule>(TEXT("LevelEditor")))
+	{
+		TSharedPtr<ILevelEditor> LevelEditorPtr = LevelEditorModule->GetLevelEditorInstance().Pin();
+
+		if (LevelEditorPtr.IsValid())
+		{
+			FString ActiveToolName = LevelEditorPtr->GetEditorModeManager().GetInteractiveToolsContext()->ToolManager->GetActiveToolName(EToolSide::Left);
+			if (ActiveToolName == InToolName)
+			{
+				LevelEditorPtr->GetEditorModeManager().GetInteractiveToolsContext()->ToolManager->DeactivateTool(EToolSide::Left, EToolShutdownType::Completed);
+			}
+			else
+			{
+				LevelEditorPtr->GetEditorModeManager().GetInteractiveToolsContext()->ToolManager->SelectActiveToolType(EToolSide::Left, InToolName);
+				LevelEditorPtr->GetEditorModeManager().GetInteractiveToolsContext()->ToolManager->ActivateTool(EToolSide::Left);
+			}
+		}
+	}
+
+}
+void FControlRigEditMode::TogglePivotMode()
+{
+	FString Tool = TEXT("SequencerPivotTool");
+	ToggleTool(Tool);
+}
+
+void FControlRigEditMode::ToggleMotionTrails()
+{
+	if (UMotionTrailToolOptions* Settings = GetMutableDefault<UMotionTrailToolOptions>())
+	{
+		Settings->bShowTrails = !Settings->bShowTrails;
+		FPropertyChangedEvent ShowTrailEvent(UMotionTrailToolOptions::StaticClass()->FindPropertyByName(GET_MEMBER_NAME_CHECKED(UMotionTrailToolOptions, bShowTrails)));
+		Settings->PostEditChangeProperty(ShowTrailEvent);
+	}
+}
 
 FText FControlRigEditMode::GetToggleControlShapeTransformEditHotKey() const
 {
diff --git a/Engine/Plugins/Animation/ControlRig/Source/ControlRigEditor/Private/EditMode/ControlRigEditModeCommands.cpp b/Engine/Plugins/Animation/ControlRig/Source/ControlRigEditor/Private/EditMode/ControlRigEditModeCommands.cpp
--- a/Engine/Plugins/Animation/ControlRig/Source/ControlRigEditor/Private/EditMode/ControlRigEditModeCommands.cpp
+++ b/Engine/Plugins/Animation/ControlRig/Source/ControlRigEditor/Private/EditMode/ControlRigEditModeCommands.cpp
@@ -42,6 +42,9 @@ void FControlRigEditModeCommands::RegisterCommands()
 	UI_COMMAND(SelectMirrorPose, "Select Mirrored Saved Pose", "Select mirrored Controls on saved Pose", EUserInterfaceActionType::Button, FInputChord());
 	UI_COMMAND(PasteMirrorPose, "Paste Mirrored Saved Pose", "Paste mirrored Controls from saved Pose", EUserInterfaceActionType::Button, FInputChord());
 
+	UI_COMMAND(TogglePivotMode, "Toggle Pivot Mode Tool", "Toggle Pivot Mode Tool", EUserInterfaceActionType::ToggleButton, FInputChord());
+	UI_COMMAND(ToggleMotionTrails, "Toggle Motion Paths Tool", "Toggle Motion paths Tool", EUserInterfaceActionType::ToggleButton, FInputChord());
+
 	UI_COMMAND(ToggleControlShapeTransformEdit, "Toggle Shape Transform Edit", "Toggle Editing Selected Control's Shape Transform", EUserInterfaceActionType::ToggleButton, FInputChord(EKeys::Period, EModifierKey::Control)); 
 	UI_COMMAND(OpenSpacePickerWidget, "Open the Space Picker", "Allows space switching on the control", EUserInterfaceActionType::Button, FInputChord(EKeys::Tab)); 
 
diff --git a/Engine/Plugins/Animation/ControlRig/Source/ControlRigEditor/Private/EditMode/ControlRigEditModeCommands.h b/Engine/Plugins/Animation/ControlRig/Source/ControlRigEditor/Private/EditMode/ControlRigEditModeCommands.h
--- a/Engine/Plugins/Animation/ControlRig/Source/ControlRigEditor/Private/EditMode/ControlRigEditModeCommands.h
+++ b/Engine/Plugins/Animation/ControlRig/Source/ControlRigEditor/Private/EditMode/ControlRigEditModeCommands.h
@@ -101,7 +101,11 @@ public:
 	/** Paste saved mirror Pose */
 	TSharedPtr< FUICommandInfo > PasteMirrorPose;
 
+	/**Toggle Pivot Mode*/
+	TSharedPtr< FUICommandInfo> TogglePivotMode;
 
+	/**Toggle Motion Trail Mode*/
+	TSharedPtr< FUICommandInfo> ToggleMotionTrails;
 
 	/** Opens up the space picker widget */
 	TSharedPtr< FUICommandInfo > OpenSpacePickerWidget;
diff --git a/Engine/Plugins/Animation/ControlRig/Source/ControlRigEditor/Public/EditMode/ControlRigEditMode.h b/Engine/Plugins/Animation/ControlRig/Source/ControlRigEditor/Public/EditMode/ControlRigEditMode.h
--- a/Engine/Plugins/Animation/ControlRig/Source/ControlRigEditor/Public/EditMode/ControlRigEditMode.h
+++ b/Engine/Plugins/Animation/ControlRig/Source/ControlRigEditor/Public/EditMode/ControlRigEditMode.h
@@ -375,6 +375,12 @@ public:
 	/** Invert Input Pose for this Control Rig Controls based upon selection and channel states*/
 	static UE_API void InvertInputPose(UControlRig* ControlRig, const FRigControlModifiedContext& Context, bool bSelectionOnly, bool bIncludeChannels);
 
+	/** Toggle Motion Trail Mode*/
+	UE_API void ToggleMotionTrails();
+
+	/** Toggle Pivot Mode*/
+	UE_API void TogglePivotMode();
+
 private:
 	
 	/** Whether or not we should Frame Selection or not*/
diff --git a/Engine/Plugins/Editor/SequencerAnimTools/Source/SequencerAnimTools/Private/SequencerAnimEditPivotTool.cpp b/Engine/Plugins/Editor/SequencerAnimTools/Source/SequencerAnimTools/Private/SequencerAnimEditPivotTool.cpp
--- a/Engine/Plugins/Editor/SequencerAnimTools/Source/SequencerAnimTools/Private/SequencerAnimEditPivotTool.cpp
+++ b/Engine/Plugins/Editor/SequencerAnimTools/Source/SequencerAnimTools/Private/SequencerAnimEditPivotTool.cpp
@@ -315,6 +315,7 @@ void USequencerPivotTool::Setup()
 	OnEditorSelectionChangedHandle = USelection::SelectionChangedEvent.AddUObject(this, &USequencerPivotTool::OnEditorSelectionChanged);
 
 	SaveLastSelected();
+	SavePivotTransforms();
 
 	CommandBindings = MakeShareable(new FUICommandList);
 
@@ -414,7 +415,6 @@ void USequencerPivotTool::UpdateTransformAndSelectionOnEntering()
 		}
 		SetGizmoBasedOnSelection(!bShiftPressedWhenStarted && !bCtrlPressedWhenStarted);
 	}
-	SavePivotTransforms();
 }
 
 bool USequencerPivotTool::ProcessCommandBindings(const FKey Key, const bool bRepeat) const
@@ -960,8 +960,17 @@ void USequencerPivotTool::UpdateGizmoVisibility()
 
 void USequencerPivotTool::UpdateGizmoTransform()
 {
+	FEditorModeTools& ModeManager = GLevelEditorModeTools(); //mz need to fix for persona/ animation edit mode
+	/* This should work but doesn't. It affects the 'highlighting' correctly but it doesn't render correctly
+	#jira UE-306751
+	const double Scale = (double)ModeManager.GetWidgetScale();
+	const FVector3d ScaleVector(Scale, Scale, Scale);
+	FTransform GizmoTransformWithScale(GizmoTransform);
+	GizmoTransformWithScale.SetScale3D(ScaleVector);
+	*/
 	if (TransformGizmo)
 	{
+		//TransformGizmo->ReinitializeGizmoTransform(GizmoTransformWithScale);
 		TransformGizmo->ReinitializeGizmoTransform(GizmoTransform);
 	}
 
@@ -969,6 +978,7 @@ void USequencerPivotTool::UpdateGizmoTransform()
 	{
 		// note: this is a trick to avoid cycling notifications... (cf. UCombinedTransformGizmo::ReinitializeGizmoTransform)
 		TGuardValue<bool> PivotModeGuard(TransformProxy->bSetPivotMode, true);
+		//TransformProxy->SetTransform(GizmoTransformWithScale);
 		TransformProxy->SetTransform(GizmoTransform);
 	}
 
@@ -1347,3 +1357,4 @@ void USequencerPivotTool::UpdatePivotOverlayLocation(const FVector2D InLocation,
 
 #undef LOCTEXT_NAMESPACE
 
+
