Anim in Engine: Anim Details now shows short property names for control rig controls, attributes are now grouped

#rnx

diff --git a/Engine/Plugins/Animation/ControlRig/Source/ControlRigEditor/Private/AnimDetails/Customizations/AnimDetailsProxyDetails.cpp b/Engine/Plugins/Animation/ControlRig/Source/ControlRigEditor/Private/AnimDetails/Customizations/AnimDetailsProxyDetails.cpp
--- a/Engine/Plugins/Animation/ControlRig/Source/ControlRigEditor/Private/AnimDetails/Customizations/AnimDetailsProxyDetails.cpp
+++ b/Engine/Plugins/Animation/ControlRig/Source/ControlRigEditor/Private/AnimDetails/Customizations/AnimDetailsProxyDetails.cpp
@@ -53,7 +53,8 @@ namespace UE::ControlRigEditor
 		FText DisplayNameText;
 		if (AllProxies.Num() == 1 && AllProxies[0])
 		{
-			DisplayNameText = AllProxies[0]->GetDisplayNameText();
+			// Always use the long name for the category
+			DisplayNameText = AllProxies[0]->GetDisplayNameText(EElementNameDisplayMode::ForceLong);
 			AllProxies[0]->UpdatePropertyDisplayNames(DetailBuilder);
 		}
 		else
diff --git a/Engine/Plugins/Animation/ControlRig/Source/ControlRigEditor/Private/AnimDetails/Customizations/AnimDetailsProxyManagerDetails.cpp b/Engine/Plugins/Animation/ControlRig/Source/ControlRigEditor/Private/AnimDetails/Customizations/AnimDetailsProxyManagerDetails.cpp
--- a/Engine/Plugins/Animation/ControlRig/Source/ControlRigEditor/Private/AnimDetails/Customizations/AnimDetailsProxyManagerDetails.cpp
+++ b/Engine/Plugins/Animation/ControlRig/Source/ControlRigEditor/Private/AnimDetails/Customizations/AnimDetailsProxyManagerDetails.cpp
@@ -2,78 +2,205 @@
 
 #include "AnimDetailsProxyManagerDetails.h"
 
+#include "Algo/AnyOf.h"
 #include "AnimDetails/AnimDetailsProxyManager.h"
 #include "AnimDetails/Proxies/AnimDetailsProxyBase.h"
+#include "ControlRig.h"
 #include "DetailCategoryBuilder.h"
 #include "DetailLayoutBuilder.h"
+#include "IDetailGroup.h"
 #include "ISequencer.h"
 
+#define LOCTEXT_NAMESPACE "AnimDetailsProxyManagerDetails"
+
 namespace UE::ControlRigEditor
 {
-	TSharedRef<IDetailCustomization> FAnimDetailProxyManagerDetails::MakeInstance()
-	{
-		return MakeShared<FAnimDetailProxyManagerDetails>();
-	}
-
-	void FAnimDetailProxyManagerDetails::CustomizeDetails(IDetailLayoutBuilder& DetailLayout)
+	namespace Private
 	{
-		TArray<TWeakObjectPtr<UObject>> ObjectsBeingCustomized;
-		DetailLayout.GetObjectsBeingCustomized(ObjectsBeingCustomized);
-
-		for (const TWeakObjectPtr<UObject>& ObjectBeingCustomized : ObjectsBeingCustomized)
+		/** Adds proxies into their respective groups */
+		void AddProxiesToDetails(IDetailLayoutBuilder& DetailLayout, const TArray<UAnimDetailsProxyBase*>& Proxies)
 		{
-			UAnimDetailsProxyManager* ProxyManager = Cast<UAnimDetailsProxyManager>(ObjectBeingCustomized.Get());
-			if (!ProxyManager)
+			if (Proxies.IsEmpty())
 			{
-				continue;
+				return;
 			}
 
-			FAnimDetailsFilter& Filter = ProxyManager->GetAnimDetailsFilter();
-			const TArray<UAnimDetailsProxyBase*> FilteredProxies = Filter.GetFilteredProxies();
+			/** A group of attribute proxies, defined by their parent name */
+			struct FAnimDetailsGroupedAttributeProxies
+			{
+				/** The parent name */
+				FName ParentName;
+
+				/** The property names to display in this group */
+				TSet<FName> PropertyNames;
+
+				/** The proxies in this group */
+				TMap<FName, TArray<UObject*>> DetailRowIDToProxiesMap;
+			};
 
-			TMap<FName, TArray<UObject*>> DetailRowIDToGroupedProxiesMap;
-			TMap<FName, TArray<UObject*>> DetailRowIDToIndividualProxiesMap;
-			for (UAnimDetailsProxyBase* Proxy : FilteredProxies)
+			TMap<FName, TArray<UObject*>> TopLevelDetailRowIDToProxiesMap;
+			TArray<FAnimDetailsGroupedAttributeProxies> AttributeGroups;
+
+			for (UAnimDetailsProxyBase* Proxy : Proxies)
 			{
 				if (!Proxy)
 				{
 					continue;
 				}
-
+				
 				const FName DetailsRowID = Proxy->GetDetailRowID();
 				if (Proxy->bIsIndividual)
 				{
-					DetailRowIDToIndividualProxiesMap.FindOrAdd(DetailsRowID).Add(Proxy);
+					const FName DetailRowID = Proxy->GetDetailRowID();
+					const FRigControlElement* ControlElement = Proxy->GetControlElement();
+					const FRigElementKey& ControlElementKey = Proxy->GetControlElementKey();
+
+					TOptional<FName> OptionalParentName;
+					if (ControlElement &&
+						ControlElementKey.IsValid())
+					{
+						const UControlRig* ControlRig = Proxy->GetControlRig();
+						const URigHierarchy* Hierarchy = ControlRig ? ControlRig->GetHierarchy() : nullptr;
+						if (Hierarchy)
+						{
+							FName ParentName = Hierarchy->GetModuleFName(ControlElementKey);
+							if (ParentName != NAME_None)
+							{
+								OptionalParentName = ParentName;
+							}
+							else if (const FRigBaseElement* Parent = Hierarchy->GetFirstParent(ControlElement))
+							{
+								ParentName = Parent ? Parent->GetDisplayName() : NAME_None;
+								if (ParentName != NAME_None)
+								{
+									OptionalParentName = ParentName;
+								}
+							}
+						}
+					}
+
+					FAnimDetailsGroupedAttributeProxies* AttributeGroupPtr = [&AttributeGroups, &OptionalParentName]() -> FAnimDetailsGroupedAttributeProxies*
+						{
+							if (OptionalParentName.IsSet())
+							{
+								return AttributeGroups.FindByPredicate(
+									[&OptionalParentName](const FAnimDetailsGroupedAttributeProxies& AttributeProxies)
+									{
+										return AttributeProxies.ParentName == OptionalParentName.GetValue();
+									});
+							}
+
+							return nullptr;
+						}();
+
+		
+					if (AttributeGroupPtr)
+					{
+						for (const FName& PropertyName : Proxy->GetPropertyNames())
+						{
+							AttributeGroupPtr->PropertyNames.Add(PropertyName);
+						}
+
+						AttributeGroupPtr->DetailRowIDToProxiesMap.FindOrAdd(Proxy->GetDetailRowID()).Add(Proxy);
+					}
+					else
+					{
+						FAnimDetailsGroupedAttributeProxies NewGroup;
+						NewGroup.ParentName = OptionalParentName.IsSet() ? OptionalParentName.GetValue() : NAME_None;
+						NewGroup.DetailRowIDToProxiesMap.Add(Proxy->GetDetailRowID(), { Proxy });
+						
+						for (const FName& PropertyName : Proxy->GetPropertyNames())
+						{
+							NewGroup.PropertyNames.Add(PropertyName);
+						}
+
+						AttributeGroups.Add(NewGroup);
+					}
 				}
 				else
 				{
-					DetailRowIDToGroupedProxiesMap.FindOrAdd(DetailsRowID).Add(Proxy);
+					TopLevelDetailRowIDToProxiesMap.FindOrAdd(Proxy->GetDetailRowID()).Add(Proxy);
 				}
 			}
 
-			// Mock up the details view as per anim details specifics, whereas grouped proxies show first, followed by individual proxies.
+			// Top level grouped proxies show first
 			IDetailCategoryBuilder& CategoryBuilder_None = DetailLayout.EditCategory("nocategory");
-			for (const TTuple<FName, TArray<UObject*>>& FieldTypeToProxyPair : DetailRowIDToGroupedProxiesMap)
+			for (const TTuple<FName, TArray<UObject*>>& TopLevelDetailRowIDToProxiesPair : TopLevelDetailRowIDToProxiesMap)
 			{
 				CategoryBuilder_None.AddExternalObjects(
-					FieldTypeToProxyPair.Value,
+					TopLevelDetailRowIDToProxiesPair.Value,
 					EPropertyLocation::Default,
 					FAddPropertyParams()
 					.HideRootObjectNode(true)
 				);
 			}
 
-			// Note, individual proxies are displayed individually, but are still multi-edited across multiple control rigs
-			IDetailCategoryBuilder& CategoryBuilder_Attributes = DetailLayout.EditCategory("Attributes");
-			for (const TTuple<FName, TArray<UObject*>>& ControlNameToIndividualProxiesPair : DetailRowIDToIndividualProxiesMap)
+			// Attribute proxies are displayed individually, but are still multi-edited across multiple control rig
+			if (AttributeGroups.Num() == 1)
 			{
-				CategoryBuilder_Attributes.AddExternalObjects(
-					ControlNameToIndividualProxiesPair.Value,
-					EPropertyLocation::Default,
-					FAddPropertyParams()
-					.HideRootObjectNode(true)
-				);
+				// Add without subgroup if theres only one attribute group
+				const FText CategoryName = LOCTEXT("AttributeCategoryName", "Attributes");
+				IDetailCategoryBuilder& CategoryBuilder_Attributes = DetailLayout.EditCategory("Attributes", CategoryName);
+				const FAnimDetailsGroupedAttributeProxies& AttributeGroup = AttributeGroups[0];
+
+				for (const TTuple<FName, TArray<UObject*>>& DetailRowIDToProxiesPair : AttributeGroup.DetailRowIDToProxiesMap)
+				{
+					CategoryBuilder_Attributes.AddExternalObjects(
+						DetailRowIDToProxiesPair.Value,
+						EPropertyLocation::Default,
+						FAddPropertyParams()
+						.HideRootObjectNode(true)
+					);
+				}
+			}
+			else
+			{
+				for (const FAnimDetailsGroupedAttributeProxies& AttributeGroup : AttributeGroups)
+				{
+					constexpr bool bForAdvanced = false;
+					constexpr bool bStartExpanded = true;
+
+					const FText ModuleNameText = FText::FromString(AttributeGroup.ParentName.ToString());
+
+					const FText CategoryName = FText::Format(LOCTEXT("FormatedAttributeCategoryName", "{0} Attributes"), ModuleNameText);
+					const FName CategoryID = *(CategoryName.ToString() + FGuid::NewGuid().ToString());
+					IDetailCategoryBuilder& CategoryBuilder = DetailLayout.EditCategory(CategoryID, CategoryName);
+
+					for (const TTuple<FName, TArray<UObject*>>& DetailRowIDToProxiesPair : AttributeGroup.DetailRowIDToProxiesMap)
+					{
+						CategoryBuilder.AddExternalObjects(
+							DetailRowIDToProxiesPair.Value,
+							EPropertyLocation::Default,
+							FAddPropertyParams()
+							.HideRootObjectNode(true)
+						);
+					}
+				}
+			}
+		}
+	}
+
+	TSharedRef<IDetailCustomization> FAnimDetailProxyManagerDetails::MakeInstance()
+	{
+		return MakeShared<FAnimDetailProxyManagerDetails>();
+	}
+
+	void FAnimDetailProxyManagerDetails::CustomizeDetails(IDetailLayoutBuilder& DetailLayout)
+	{
+		TArray<TWeakObjectPtr<UObject>> ObjectsBeingCustomized;
+		DetailLayout.GetObjectsBeingCustomized(ObjectsBeingCustomized);
+
+		for (const TWeakObjectPtr<UObject>& ObjectBeingCustomized : ObjectsBeingCustomized)
+		{
+			if (UAnimDetailsProxyManager* ProxyManager = Cast<UAnimDetailsProxyManager>(ObjectBeingCustomized.Get()))
+			{
+				FAnimDetailsFilter& Filter = ProxyManager->GetAnimDetailsFilter();
+				const TArray<UAnimDetailsProxyBase*> FilteredProxies = Filter.GetFilteredProxies();
+
+				Private::AddProxiesToDetails(DetailLayout, FilteredProxies);
 			}
 		}
 	}
 }
+
+#undef LOCTEXT_NAMESPACE
diff --git a/Engine/Plugins/Animation/ControlRig/Source/ControlRigEditor/Private/AnimDetails/Customizations/AnimDetailsValueCustomization.cpp b/Engine/Plugins/Animation/ControlRig/Source/ControlRigEditor/Private/AnimDetails/Customizations/AnimDetailsValueCustomization.cpp
--- a/Engine/Plugins/Animation/ControlRig/Source/ControlRigEditor/Private/AnimDetails/Customizations/AnimDetailsValueCustomization.cpp
+++ b/Engine/Plugins/Animation/ControlRig/Source/ControlRigEditor/Private/AnimDetails/Customizations/AnimDetailsValueCustomization.cpp
@@ -110,7 +110,24 @@ namespace UE::ControlRigEditor
 			}
 
 			const FName PropertyName = ChildHandle->GetProperty()->GetFName();
-			const FText PropertyDisplayText = ChildHandle->GetPropertyDisplayName();
+			const FText PropertyDisplayText = [&ChildHandle, this]()
+				{
+					// Draw the proxy display name if the property has only one child.
+					// This aviods issues where multi-editing would show the non-customized property name when multi editing.
+					if (SortedChildHandles.Num() == 1)
+					{
+						TArray<UObject*> OuterObjects;
+						ChildHandle->GetOuterObjects(OuterObjects);
+
+						UAnimDetailsProxyBase* FirstProxy = OuterObjects.IsEmpty() ? nullptr : Cast<UAnimDetailsProxyBase>(OuterObjects[0]);
+						if (FirstProxy)
+						{
+							return FirstProxy->GetDisplayNameText();
+						}
+					}
+					
+					return ChildHandle->GetPropertyDisplayName();
+				}();
 
 			const TSharedRef<SWidget> ValueWidget = MakeChildWidget(InStructPropertyHandle, ChildHandle);
 			
diff --git a/Engine/Plugins/Animation/ControlRig/Source/ControlRigEditor/Private/AnimDetails/Proxies/AnimDetailsProxyBase.cpp b/Engine/Plugins/Animation/ControlRig/Source/ControlRigEditor/Private/AnimDetails/Proxies/AnimDetailsProxyBase.cpp
--- a/Engine/Plugins/Animation/ControlRig/Source/ControlRigEditor/Private/AnimDetails/Proxies/AnimDetailsProxyBase.cpp
+++ b/Engine/Plugins/Animation/ControlRig/Source/ControlRigEditor/Private/AnimDetails/Proxies/AnimDetailsProxyBase.cpp
@@ -646,7 +646,7 @@ void UAnimDetailsProxyBase::PropagonateValues()
 	}
 }
 
-FText UAnimDetailsProxyBase::GetDisplayNameText() const
+FText UAnimDetailsProxyBase::GetDisplayNameText(const EElementNameDisplayMode ElementNameDisplayMode) const
 {
 	if(!DisplayName.IsEmpty())
 	{
@@ -657,7 +657,7 @@ FText UAnimDetailsProxyBase::GetDisplayNameText() const
 	const FRigControlElement* ControlElement = WeakControlRig.IsValid() ? GetControlElement() : nullptr;
 	if (Hierarchy && ControlElement)
 	{
-		return Hierarchy->GetDisplayNameForUI(ControlElement);
+		return Hierarchy->GetDisplayNameForUI(ControlElement, ElementNameDisplayMode);
 	}
 	else if (const UObject* BoundObject = SequencerItem.GetBoundObject())
 	{
diff --git a/Engine/Plugins/Animation/ControlRig/Source/ControlRigEditor/Private/AnimDetails/Proxies/AnimDetailsProxyBase.h b/Engine/Plugins/Animation/ControlRig/Source/ControlRigEditor/Private/AnimDetails/Proxies/AnimDetailsProxyBase.h
--- a/Engine/Plugins/Animation/ControlRig/Source/ControlRigEditor/Private/AnimDetails/Proxies/AnimDetailsProxyBase.h
+++ b/Engine/Plugins/Animation/ControlRig/Source/ControlRigEditor/Private/AnimDetails/Proxies/AnimDetailsProxyBase.h
@@ -60,7 +60,7 @@ public:
 	void PropagonateValues();
 
 	/** Returns the display name for this proxy as text */
-	FText GetDisplayNameText() const;
+	FText GetDisplayNameText(const EElementNameDisplayMode ElementNameDisplayMode = EElementNameDisplayMode::Auto) const;
 
 	/** Sets a key from current values in sequencer */
 	void SetKey(const IPropertyHandle& KeyedPropertyHandle);
@@ -153,4 +153,3 @@ private:
 	/** Current interaction scopes */
 	TMap<FRigControlElement*, TSharedRef<FControlRigInteractionScope>> InteractionScopes;
 };
-
diff --git a/Engine/Plugins/Animation/ControlRig/Source/ControlRigEditor/Private/AnimDetails/Proxies/AnimDetailsProxyBool.cpp b/Engine/Plugins/Animation/ControlRig/Source/ControlRigEditor/Private/AnimDetails/Proxies/AnimDetailsProxyBool.cpp
--- a/Engine/Plugins/Animation/ControlRig/Source/ControlRigEditor/Private/AnimDetails/Proxies/AnimDetailsProxyBool.cpp
+++ b/Engine/Plugins/Animation/ControlRig/Source/ControlRigEditor/Private/AnimDetails/Proxies/AnimDetailsProxyBool.cpp
@@ -2,6 +2,7 @@
 
 #include "AnimDetailsProxyBool.h"
 
+#include "AnimDetails/AnimDetailsSettings.h"
 #include "ControlRig.h"
 #include "DetailLayoutBuilder.h"
 #include "MovieSceneCommonHelpers.h"
@@ -15,10 +16,13 @@ FName UAnimDetailsProxyBool::GetCategoryName() const
 
 void UAnimDetailsProxyBool::UpdatePropertyDisplayNames(IDetailLayoutBuilder& DetailBuilder)
 {
+	// Use the short name as the property display name when displaying control rig elements
+	constexpr EElementNameDisplayMode NameDisplayMode = EElementNameDisplayMode::ForceShort;
+
 	TSharedPtr<IPropertyHandle> ValuePropertyHandle = DetailBuilder.GetProperty(GET_MEMBER_NAME_CHECKED(UAnimDetailsProxyBool, Bool), GetClass());
 	if (ValuePropertyHandle.IsValid())
 	{
-		ValuePropertyHandle->SetPropertyDisplayName(GetDisplayNameText());
+		ValuePropertyHandle->SetPropertyDisplayName(GetDisplayNameText(NameDisplayMode));
 	}
 
 	if (bIsIndividual)
@@ -27,7 +31,7 @@ void UAnimDetailsProxyBool::UpdatePropertyDisplayNames(IDetailLayoutBuilder& Det
 		ValuePropertyHandle = DetailBuilder.GetProperty(FName(*PropertyPath), GetClass());
 		if (ValuePropertyHandle.IsValid())
 		{
-			ValuePropertyHandle->SetPropertyDisplayName(GetDisplayNameText());
+			ValuePropertyHandle->SetPropertyDisplayName(GetDisplayNameText(NameDisplayMode));
 		}
 	}
 	else
diff --git a/Engine/Plugins/Animation/ControlRig/Source/ControlRigEditor/Private/AnimDetails/Proxies/AnimDetailsProxyEnum.cpp b/Engine/Plugins/Animation/ControlRig/Source/ControlRigEditor/Private/AnimDetails/Proxies/AnimDetailsProxyEnum.cpp
--- a/Engine/Plugins/Animation/ControlRig/Source/ControlRigEditor/Private/AnimDetails/Proxies/AnimDetailsProxyEnum.cpp
+++ b/Engine/Plugins/Animation/ControlRig/Source/ControlRigEditor/Private/AnimDetails/Proxies/AnimDetailsProxyEnum.cpp
@@ -2,6 +2,7 @@
 
 #include "AnimDetailsProxyEnum.h"
 
+#include "AnimDetails/AnimDetailsSettings.h"
 #include "ControlRig.h"
 #include "DetailLayoutBuilder.h"
 #include "MovieSceneCommonHelpers.h"
@@ -25,10 +26,13 @@ FName UAnimDetailsProxyEnum::GetDetailRowID() const
 
 void UAnimDetailsProxyEnum::UpdatePropertyDisplayNames(IDetailLayoutBuilder& DetailBuilder)
 {
+	// Use the short name as the property display name when displaying control rig elements
+	constexpr EElementNameDisplayMode NameDisplayMode = EElementNameDisplayMode::ForceShort;
+
 	TSharedPtr<IPropertyHandle> ValuePropertyHandle = DetailBuilder.GetProperty(GET_MEMBER_NAME_CHECKED(UAnimDetailsProxyEnum, Enum), GetClass());
 	if (ValuePropertyHandle.IsValid())
 	{
-		ValuePropertyHandle->SetPropertyDisplayName(GetDisplayNameText());
+		ValuePropertyHandle->SetPropertyDisplayName(GetDisplayNameText(NameDisplayMode));
 	}
 
 	if (SequencerItem.IsValid())
@@ -46,7 +50,7 @@ void UAnimDetailsProxyEnum::UpdatePropertyDisplayNames(IDetailLayoutBuilder& Det
 		ValuePropertyHandle = DetailBuilder.GetProperty(FName(*PropertyPath), GetClass());
 		if (ValuePropertyHandle.IsValid())
 		{
-			ValuePropertyHandle->SetPropertyDisplayName(GetDisplayNameText());
+			ValuePropertyHandle->SetPropertyDisplayName(GetDisplayNameText(NameDisplayMode));
 		}
 	}
 }
diff --git a/Engine/Plugins/Animation/ControlRig/Source/ControlRigEditor/Private/AnimDetails/Proxies/AnimDetailsProxyFloat.cpp b/Engine/Plugins/Animation/ControlRig/Source/ControlRigEditor/Private/AnimDetails/Proxies/AnimDetailsProxyFloat.cpp
--- a/Engine/Plugins/Animation/ControlRig/Source/ControlRigEditor/Private/AnimDetails/Proxies/AnimDetailsProxyFloat.cpp
+++ b/Engine/Plugins/Animation/ControlRig/Source/ControlRigEditor/Private/AnimDetails/Proxies/AnimDetailsProxyFloat.cpp
@@ -2,6 +2,7 @@
 
 #include "AnimDetailsProxyFloat.h"
 
+#include "AnimDetails/AnimDetailsSettings.h"
 #include "ControlRig.h"
 #include "DetailLayoutBuilder.h"
 #include "MovieSceneCommonHelpers.h"
@@ -15,10 +16,13 @@ FName UAnimDetailsProxyFloat::GetCategoryName() const
 
 void UAnimDetailsProxyFloat::UpdatePropertyDisplayNames(IDetailLayoutBuilder& DetailBuilder)
 {
+	// Use the short name as the property display name when displaying control rig elements
+	constexpr EElementNameDisplayMode NameDisplayMode = EElementNameDisplayMode::ForceShort;
+
 	TSharedPtr<IPropertyHandle> ValuePropertyHandle = DetailBuilder.GetProperty(GET_MEMBER_NAME_CHECKED(UAnimDetailsProxyFloat, Float), GetClass());
 	if (ValuePropertyHandle.IsValid())
 	{
-		ValuePropertyHandle->SetPropertyDisplayName(GetDisplayNameText());
+		ValuePropertyHandle->SetPropertyDisplayName(GetDisplayNameText(NameDisplayMode));
 	}
 
 	if (bIsIndividual)
@@ -27,7 +31,7 @@ void UAnimDetailsProxyFloat::UpdatePropertyDisplayNames(IDetailLayoutBuilder& De
 		ValuePropertyHandle = DetailBuilder.GetProperty(FName(*PropertyPath), GetClass());
 		if (ValuePropertyHandle.IsValid())
 		{
-			ValuePropertyHandle->SetPropertyDisplayName(GetDisplayNameText());
+			ValuePropertyHandle->SetPropertyDisplayName(GetDisplayNameText(NameDisplayMode));
 		}
 	}
 	else if (SequencerItem.IsValid())
diff --git a/Engine/Plugins/Animation/ControlRig/Source/ControlRigEditor/Private/AnimDetails/Proxies/AnimDetailsProxyInteger.cpp b/Engine/Plugins/Animation/ControlRig/Source/ControlRigEditor/Private/AnimDetails/Proxies/AnimDetailsProxyInteger.cpp
--- a/Engine/Plugins/Animation/ControlRig/Source/ControlRigEditor/Private/AnimDetails/Proxies/AnimDetailsProxyInteger.cpp
+++ b/Engine/Plugins/Animation/ControlRig/Source/ControlRigEditor/Private/AnimDetails/Proxies/AnimDetailsProxyInteger.cpp
@@ -2,6 +2,7 @@
 
 #include "AnimDetailsProxyInteger.h"
 
+#include "AnimDetails/AnimDetailsSettings.h"
 #include "ControlRig.h"
 #include "DetailLayoutBuilder.h"
 #include "MovieSceneCommonHelpers.h"
@@ -15,10 +16,13 @@ FName UAnimDetailsProxyInteger::GetCategoryName() const
 
 void UAnimDetailsProxyInteger::UpdatePropertyDisplayNames(IDetailLayoutBuilder& DetailBuilder)
 {
+	// Use the short name as the property display name when displaying control rig elements
+	constexpr EElementNameDisplayMode NameDisplayMode = EElementNameDisplayMode::ForceShort;
+
 	TSharedPtr<IPropertyHandle> ValuePropertyHandle = DetailBuilder.GetProperty(GET_MEMBER_NAME_CHECKED(UAnimDetailsProxyInteger, Integer), GetClass());
 	if (ValuePropertyHandle.IsValid())
 	{
-		ValuePropertyHandle->SetPropertyDisplayName(GetDisplayNameText());
+		ValuePropertyHandle->SetPropertyDisplayName(GetDisplayNameText(NameDisplayMode));
 	}
 
 	if (SequencerItem.IsValid())
@@ -36,7 +40,7 @@ void UAnimDetailsProxyInteger::UpdatePropertyDisplayNames(IDetailLayoutBuilder&
 		ValuePropertyHandle = DetailBuilder.GetProperty(FName(*PropertyPath), GetClass());
 		if (ValuePropertyHandle.IsValid())
 		{
-			ValuePropertyHandle->SetPropertyDisplayName(GetDisplayNameText());
+			ValuePropertyHandle->SetPropertyDisplayName(GetDisplayNameText(NameDisplayMode));
 		}
 	}
 }
