Anim in Engine: Fix Anim Details shows long instead of short property names for Attributes

#rnx

diff --git a/Engine/Plugins/Animation/ControlRig/Source/ControlRigEditor/Private/AnimDetails/Customizations/AnimDetailsValueCustomization.cpp b/Engine/Plugins/Animation/ControlRig/Source/ControlRigEditor/Private/AnimDetails/Customizations/AnimDetailsValueCustomization.cpp
--- a/Engine/Plugins/Animation/ControlRig/Source/ControlRigEditor/Private/AnimDetails/Customizations/AnimDetailsValueCustomization.cpp
+++ b/Engine/Plugins/Animation/ControlRig/Source/ControlRigEditor/Private/AnimDetails/Customizations/AnimDetailsValueCustomization.cpp
@@ -122,7 +122,7 @@ namespace UE::ControlRigEditor
 						UAnimDetailsProxyBase* FirstProxy = OuterObjects.IsEmpty() ? nullptr : Cast<UAnimDetailsProxyBase>(OuterObjects[0]);
 						if (FirstProxy)
 						{
-							return FirstProxy->GetDisplayNameText();
+							return FirstProxy->GetDisplayNameText(EElementNameDisplayMode::ForceShort);
 						}
 					}
 					
