Control Rig: Refactor hierarchy dependency tracking

* move data to context
* move logic to RigDependency.cpp
* introduce support for metadata tracking

outstanding:
* support for metadata
* support for modular rigs

#rnx

diff --git a/Engine/Plugins/Animation/ControlRig/Source/ControlRig/Private/Constraints/ControlRigTransformableHandle.cpp b/Engine/Plugins/Animation/ControlRig/Source/ControlRig/Private/Constraints/ControlRigTransformableHandle.cpp
--- a/Engine/Plugins/Animation/ControlRig/Source/ControlRig/Private/Constraints/ControlRigTransformableHandle.cpp
+++ b/Engine/Plugins/Animation/ControlRig/Source/ControlRig/Private/Constraints/ControlRigTransformableHandle.cpp
@@ -408,7 +408,7 @@ bool UTransformableControlHandle::HasDirectDependencyWith(const UTransformableHa
 			if (InControlRig && InRigHierarchy)
 			{
 #if WITH_EDITOR
-				const FRigVMDependenciesProvider SameModuleDependencyProvider(InRigHierarchy, InControlRig->GetVM());
+				const FRigDependenciesProviderForControlRig SameModuleDependencyProvider(InControlRig);
 				return InRigHierarchy->IsParentedTo(ControlElement, PossibleParent, SameModuleDependencyProvider);
 #else
 				return InRigHierarchy->IsParentedTo(ControlElement, PossibleParent);
diff --git a/Engine/Plugins/Animation/ControlRig/Source/ControlRig/Private/ControlRig.cpp b/Engine/Plugins/Animation/ControlRig/Source/ControlRig/Private/ControlRig.cpp
--- a/Engine/Plugins/Animation/ControlRig/Source/ControlRig/Private/ControlRig.cpp
+++ b/Engine/Plugins/Animation/ControlRig/Source/ControlRig/Private/ControlRig.cpp
@@ -114,7 +114,7 @@ void UControlRig::ResetRecordedTransforms(const FName& InEventName)
 {
 	if(const URigHierarchy* Hierarchy = GetHierarchy())
 	{
-		if(Hierarchy->bRecordTransformsAtRuntime)
+		if(Hierarchy->bRecordInstructionsAtRuntime)
 		{
 			bool bResetRecordedTransforms = SupportsEvent(InEventName);
 			if(InEventName == FRigUnit_PostBeginExecution::EventName ||
@@ -125,8 +125,9 @@ void UControlRig::ResetRecordedTransforms(const FName& InEventName)
 			
 			if(bResetRecordedTransforms)
 			{
-				Hierarchy->ReadTransformsAtRuntime.Reset();
-				Hierarchy->WrittenTransformsAtRuntime.Reset();
+				FControlRigExecuteContext& PublicData = GetRigVMExtendedExecuteContext().GetPublicData<FControlRigExecuteContext>();
+				PublicData.ReadRecords.Reset();
+				PublicData.WrittenRecords.Reset();
 			}
 		}
 	}
@@ -585,7 +586,7 @@ void UControlRig::Evaluate_AnyThread()
 				if (FRigBaseElement* NewParent = Hierarchy->Find(Pair.Value.NewParent))
 				{
 #if WITH_EDITOR
-					FRigVMDependenciesProvider DependencyProvider(Hierarchy, GetVM());
+					FRigDependenciesProviderForControlRig DependencyProvider(this);
 					Hierarchy->SwitchToParent(Element, NewParent, Pair.Value.bInitial, Pair.Value.bAffectChildren, DependencyProvider, nullptr);
 #else
 					Hierarchy->SwitchToParent(Element, NewParent, Pair.Value.bInitial, Pair.Value.bAffectChildren);
@@ -1111,21 +1112,21 @@ bool UControlRig::Execute(const FName& InEventName)
 		Hierarchy->UpdateReferences(&PublicContext);
 
 #if WITH_EDITOR
-		bool bRecordTransformsAtRuntime = true;
+		bool bRecordInstructionsAtRuntime = true;
 		if(const UObject* Outer = GetOuter())
 		{
 			if(Outer->IsA<UControlRigComponent>())
 			{
-				bRecordTransformsAtRuntime = false;
+				bRecordInstructionsAtRuntime = false;
 			}
 		}
 		
 		if (IsNativeRig())
 		{
-			bRecordTransformsAtRuntime = false;
+			bRecordInstructionsAtRuntime = false;
 		}
 		
-		RecordTransformsPerInstructionGuard = MakeShared<TGuardValue<bool>>(Hierarchy->bRecordTransformsAtRuntime, bRecordTransformsAtRuntime);
+		RecordTransformsPerInstructionGuard = MakeShared<TGuardValue<bool>>(Hierarchy->bRecordInstructionsAtRuntime, bRecordInstructionsAtRuntime);
 		ResetRecordedTransforms(InEventName);
 #endif
 	}
@@ -2402,7 +2403,7 @@ void UControlRig::SwitchToParent(const FRigElementKey& InElementKey, const FRigE
 	else
 	{
 #if WITH_EDITOR
-		FRigVMDependenciesProvider DependenciesProvider(DynamicHierarchy, GetVM());
+		FRigDependenciesProviderForControlRig DependenciesProvider(this);
 		DynamicHierarchy->SwitchToParent(InElementKey, InNewParentKey, bInitial, bAffectChildren, DependenciesProvider, nullptr);
 #else
 		DynamicHierarchy->SwitchToParent(InElementKey, InNewParentKey, bInitial, bAffectChildren);
diff --git a/Engine/Plugins/Animation/ControlRig/Source/ControlRig/Private/Rigs/FKControlRig.cpp b/Engine/Plugins/Animation/ControlRig/Source/ControlRig/Private/Rigs/FKControlRig.cpp
--- a/Engine/Plugins/Animation/ControlRig/Source/ControlRig/Private/Rigs/FKControlRig.cpp
+++ b/Engine/Plugins/Animation/ControlRig/Source/ControlRig/Private/Rigs/FKControlRig.cpp
@@ -123,7 +123,7 @@ bool UFKControlRig::Execute_Internal(const FName& InEventName)
 		{
 			Hierarchy->StorePoseForTrace(TEXT("UFKControlRig::BeginExecuteUnits"));
 		}
-		Hierarchy->bRecordTransformsAtRuntime = false;
+		Hierarchy->bRecordInstructionsAtRuntime = false;
 	}
 #endif	
 
@@ -298,7 +298,7 @@ void UFKControlRig::Initialize(bool bInitRigUnits /*= true*/)
 #if WITH_EDITOR
 	if(URigHierarchy* Hierarchy = GetHierarchy())
 	{
-		Hierarchy->bRecordTransformsAtRuntime = false;
+		Hierarchy->bRecordInstructionsAtRuntime = false;
 	}
 #endif // WITH_EDITOR
 
diff --git a/Engine/Plugins/Animation/ControlRig/Source/ControlRig/Private/Rigs/RigConnectionRules.cpp b/Engine/Plugins/Animation/ControlRig/Source/ControlRig/Private/Rigs/RigConnectionRules.cpp
--- a/Engine/Plugins/Animation/ControlRig/Source/ControlRig/Private/Rigs/RigConnectionRules.cpp
+++ b/Engine/Plugins/Animation/ControlRig/Source/ControlRig/Private/Rigs/RigConnectionRules.cpp
@@ -341,9 +341,9 @@ FRigElementResolveResult FRigChildOfPrimaryConnectionRule::Resolve(const FRigBas
 		}
 	}
 
-	static const FNoDependenciesProvider NoDependencyProvider;
+	static const FEmptyRigDependenciesProvider EmptyRigDependenciesProvider;
 	if(!InRuleInput.GetHierarchy()->IsParentedTo(
-		const_cast<FRigBaseElement*>(InTarget), const_cast<FRigTransformElement*>(PrimaryTarget), NoDependencyProvider))
+		const_cast<FRigBaseElement*>(InTarget), const_cast<FRigTransformElement*>(PrimaryTarget), EmptyRigDependenciesProvider))
 	{
 		Result.SetInvalidTarget(FText::FromString(FString::Printf(IsNotAChildOfFormat, *InTarget->GetKey().ToString())));
 	}
diff --git a/Engine/Plugins/Animation/ControlRig/Source/ControlRig/Private/Rigs/RigDependency.cpp b/Engine/Plugins/Animation/ControlRig/Source/ControlRig/Private/Rigs/RigDependency.cpp
--- a/Engine/Plugins/Animation/ControlRig/Source/ControlRig/Private/Rigs/RigDependency.cpp
+++ b/Engine/Plugins/Animation/ControlRig/Source/ControlRig/Private/Rigs/RigDependency.cpp
@@ -1,52 +1,258 @@
 // Copyright Epic Games, Inc. All Rights Reserved.
 
 #include "Rigs/RigDependency.h"
+#include "ModularRig.h"
+#include "Units/Execution/RigUnit_BeginExecution.h"
 
-#include "ControlRig.h"
-#include "Rigs/RigHierarchy.h"
-
-const TMap<int32, TArray<int32>>& FNoDependenciesProvider::GetDependencies() const
+const TRigHierarchyDependencyMap& FEmptyRigDependenciesProvider::GetRigHierarchyDependencies() const
 {
-	static const TMap<int32, TArray<int32>> Empty;
+	static const TRigHierarchyDependencyMap Empty;
 	return Empty;
 }
 
-FRigVMDependenciesProvider::FRigVMDependenciesProvider(const URigHierarchy* InHierarchy, URigVM* InVM, FName InEventName)
-	: WeakHierarchy(InHierarchy)
-	, WeakRigVM(InVM)
+FRigDependenciesProviderForVM::FRigDependenciesProviderForVM(UControlRig* InControlRig, FName InEventName)
+	: WeakControlRig(InControlRig)
 	, EventName(InEventName)
 {}
-	
-const TMap<int32, TArray<int32>>& FRigVMDependenciesProvider::GetDependencies() const
+
+// UE_DISABLE_OPTIMIZATION
+
+const TRigHierarchyDependencyMap& FRigDependenciesProviderForVM::GetRigHierarchyDependencies() const
 {
-	// NOTE: read URigHierarchy::GetDependenciesForVM about using this function with modular rigs.
+	static const TRigHierarchyDependencyMap Empty;
 
 #if WITH_EDITOR
-	const URigHierarchy* Hierarchy = WeakHierarchy.Get();
-	if (ensure(Hierarchy))
+	UControlRig* ControlRig = GetControlRig();
+	if (!ensure(ControlRig))
+	{
+		return Empty;
+	}
+
+	const URigHierarchy* Hierarchy = ControlRig->GetHierarchy();
+	if (!ensure(ControlRig))
+	{
+		return Empty;
+	}
+
+	const URigVM* VM = ControlRig->GetVM();
+	if (!ensure(VM))
+	{
+		return Empty;
+	}
+
+	uint32 NewHash = 0; 
+	auto LazyGetTopologyHash = [Hierarchy, &NewHash]()
+	{
+		if (NewHash == 0)
+		{
+			// NOTE: GetTopologyVersion() is used to track simple topology changes (GetTopologyHash() is probably more complete but also slower)
+			NewHash = Hierarchy->GetTopologyVersion();
+		}
+		return NewHash;
+	};
+
+	const FRigVMExtendedExecuteContext& ExtendedExecuteContext = ControlRig->GetRigVMExtendedExecuteContext();
+	const FControlRigExecuteContext& Context = ExtendedExecuteContext.GetPublicData<FControlRigExecuteContext>();
+	
+	if (!CachedDependencies.IsEmpty() && TopologyHash == LazyGetTopologyHash())
+	{
+		return CachedDependencies;
+	}
+	
+	FName EventNameToUse = EventName;
+	if(EventNameToUse.IsNone())
+	{
+		EventNameToUse = FRigUnit_BeginExecution().GetEventName();
+	}
+
+	CachedDependencies.Reset();
+	
+	const FRigVMInstructionArray Instructions = VM->GetByteCode().GetInstructions();
+
+	// if the VM doesn't implement the given event
+	if(!VM->ContainsEntry(EventNameToUse))
+	{
+		return Empty;
+	}
+
+	// only represent instruction for a given event
+	const int32 EntryIndex = VM->GetByteCode().FindEntryIndex(EventNameToUse);
+	const int32 EntryInstructionIndex = VM->GetByteCode().GetEntry(EntryIndex).InstructionIndex;
+
+	TMap<FRigVMOperand, TArray<int32>> OperandToInstructions;
+
+	for(int32 InstructionIndex = EntryInstructionIndex; InstructionIndex < Instructions.Num(); InstructionIndex++)
+	{
+		// early exit since following instructions belong to another event
+		if(Instructions[InstructionIndex].OpCode == ERigVMOpCode::Exit)
+		{
+			break;
+		}
+
+		const FRigVMOperandArray InputOperands = VM->GetByteCode().GetInputOperands(InstructionIndex);
+		for(const FRigVMOperand InputOperand : InputOperands)
+		{
+			const FRigVMOperand InputOperandNoRegisterOffset(InputOperand.GetMemoryType(), InputOperand.GetRegisterIndex());
+			OperandToInstructions.FindOrAdd(InputOperandNoRegisterOffset).Add(InstructionIndex);
+		}
+	}
+
+	TArray<TArray<FInstructionRecord>> ReadRecordsPerInstruction, WrittenRecordsPerInstruction;
+	
+	// Find the max instruction index
+	{
+		/** TODO: UE-207320
+		 * This function needs to be rewritten to account for modular rigs. Rig Hierarchies will have reads/writes coming from
+		 * different VMs, so we cannot depend on the algorithm used in this function to find dependencies. */
+		
+		int32 MaxInstructionIndex = Instructions.Num();
+		for(int32 RecordType = 0; RecordType < 2; RecordType++)
+		{
+			const TArray<FInstructionRecord>& Records = RecordType == 0 ? Context.ReadRecords : Context.WrittenRecords; 
+			for(int32 RecordIndex = 0; RecordIndex < Records.Num(); RecordIndex++)
+			{
+				const FInstructionRecord& Record = Records[RecordIndex];
+				MaxInstructionIndex = FMath::Max(MaxInstructionIndex, Record.InstructionIndex);
+			}
+		}
+		ReadRecordsPerInstruction.AddZeroed(MaxInstructionIndex+1);
+		WrittenRecordsPerInstruction.AddZeroed(MaxInstructionIndex+1);
+	}
+
+	// fill lookup tables per instruction / element
+	for(int32 RecordType = 0; RecordType < 2; RecordType++)
+	{
+		const TArray<FInstructionRecord>& Records = RecordType == 0 ? Context.ReadRecords : Context.WrittenRecords; 
+		TArray<TArray<FInstructionRecord>>& PerInstruction = RecordType == 0 ? ReadRecordsPerInstruction : WrittenRecordsPerInstruction;
+
+		for(int32 RecordIndex = 0; RecordIndex < Records.Num(); RecordIndex++)
+		{
+			const FInstructionRecord& Record = Records[RecordIndex];
+			PerInstruction[Record.InstructionIndex].Add(Record);
+		}
+	}
+
+	// for each read transform on an instruction
+	// follow the operands to the next instruction affected by it
+	TArray<FInstructionRecord> FilteredRecords;
+	TArray<int32> InstructionsToVisit;
+
+	for(int32 InstructionIndex = EntryInstructionIndex; InstructionIndex < Instructions.Num(); InstructionIndex++)
 	{
-		uint32 NewHash = 0; 
-		auto LazyGetTopologyHash = [Hierarchy, &NewHash]()
+		// early exit since following instructions belong to another event
+		if(Instructions[InstructionIndex].OpCode == ERigVMOpCode::Exit)
+		{
+			break;
+		}
+
+		const TArray<FInstructionRecord>& ReadRecords = ReadRecordsPerInstruction[InstructionIndex];
+		if(ReadRecords.IsEmpty())
 		{
-			if (NewHash == 0)
+			continue;
+		}
+
+		FilteredRecords.Reset();
+
+		for(int32 ReadRecordIndex = 0; ReadRecordIndex < ReadRecords.Num(); ReadRecordIndex++)
+		{
+			InstructionsToVisit.Reset();
+			InstructionsToVisit.Add(InstructionIndex);
+
+			for(int32 InstructionToVisitIndex = 0; InstructionToVisitIndex < InstructionsToVisit.Num(); InstructionToVisitIndex++)
 			{
-				// NOTE: GetTopologyVersion() is used to track simple topology changes (GetTopologyHash() is probably more complete but also slower)
-				NewHash = Hierarchy->GetTopologyVersion();
+				const int32 InstructionToVisit = InstructionsToVisit[InstructionToVisitIndex];
+
+				const TArray<FInstructionRecord>& WrittenRecords = WrittenRecordsPerInstruction[InstructionToVisit];
+				for(int32 WrittenRecordIndex = 0; WrittenRecordIndex < WrittenRecords.Num(); WrittenRecordIndex++)
+				{
+					const FInstructionRecord& WrittenRecord = WrittenRecords[WrittenRecordIndex];
+
+					// for the first instruction in this pass let's only care about
+					// written transforms which have not been read before
+					if(InstructionToVisitIndex == 0)
+					{
+						if(ReadRecords.Contains(WrittenRecord))
+						{
+							continue;
+						}
+					}
+
+					// store the filtered records for any instruction here
+					// by clearing the instruction index field.
+					FInstructionRecord SimplifiedRecord = WrittenRecord;
+					SimplifiedRecord.InstructionIndex = INDEX_NONE;
+					FilteredRecords.AddUnique(SimplifiedRecord);
+				}
+
+				FRigVMOperandArray OutputOperands = VM->GetByteCode().GetOutputOperands(InstructionToVisit);
+				for(const FRigVMOperand OutputOperand : OutputOperands)
+				{
+					const FRigVMOperand OutputOperandNoRegisterOffset(OutputOperand.GetMemoryType(), OutputOperand.GetRegisterIndex());
+
+					if(const TArray<int32>* InstructionsWithInputOperand = OperandToInstructions.Find(OutputOperandNoRegisterOffset))
+					{
+						for(int32 InstructionWithInputOperand : *InstructionsWithInputOperand)
+						{
+							InstructionsToVisit.AddUnique(InstructionWithInputOperand);
+						}
+					}
+				}
 			}
-			return NewHash;
-		};
+		}
 		
-		if (CachedDependencies.IsEmpty() || TopologyHash != LazyGetTopologyHash())
+		for(const FInstructionRecord& ReadRecord : ReadRecords)
 		{
-			URigVM* RigVM = WeakRigVM.Get();
-			if (ensure(RigVM))
+			for(const FInstructionRecord& FilteredRecord : FilteredRecords)
 			{
-				CachedDependencies = Hierarchy->GetDependenciesForVM(RigVM, EventName);
-				TopologyHash = NewHash;
+				// only create dependencies for reads and writes that are on the same slice
+				if(ReadRecord != FilteredRecord && ReadRecord.SliceIndex == FilteredRecord.SliceIndex)
+				{
+					const FRigHierarchyRecord SourceHierarchyRecord(ReadRecord.ElementIndex, ReadRecord.MetadataName);
+					const FRigHierarchyRecord TargetHierarchyRecord(FilteredRecord.ElementIndex, FilteredRecord.MetadataName);
+					CachedDependencies.FindOrAdd(TargetHierarchyRecord).AddUnique(SourceHierarchyRecord);
+				}
 			}
 		}
 	}
+
+	TopologyHash = NewHash;
+	
 #endif
 
 	return CachedDependencies;
-}
\ No newline at end of file
+}
+
+UControlRig* FRigDependenciesProviderForVM::GetControlRig() const
+{
+	if (WeakControlRig.IsValid())
+	{
+		return WeakControlRig.Get();
+	}
+	return nullptr;
+}
+
+FRigDependenciesProviderForControlRig::FRigDependenciesProviderForControlRig(UControlRig* InControlRig, FName InEventName)
+	: FRigDependenciesProviderForVM(InControlRig, InEventName)
+{
+	if (const UControlRig* ControlRig = GetControlRig())
+	{
+		if (ControlRig->IsA<UModularRig>())
+		{
+			// todo
+		}
+	}
+}
+
+const TRigHierarchyDependencyMap& FRigDependenciesProviderForControlRig::GetRigHierarchyDependencies() const
+{
+	if (const UControlRig* ControlRig = GetControlRig())
+	{
+		if (ControlRig->IsA<UModularRig>())
+		{
+			// todo
+		}
+	}
+	return FRigDependenciesProviderForVM::GetRigHierarchyDependencies();
+}
+
+// UE_ENABLE_OPTIMIZATION
diff --git a/Engine/Plugins/Animation/ControlRig/Source/ControlRig/Private/Rigs/RigHierarchy.cpp b/Engine/Plugins/Animation/ControlRig/Source/ControlRig/Private/Rigs/RigHierarchy.cpp
--- a/Engine/Plugins/Animation/ControlRig/Source/ControlRig/Private/Rigs/RigHierarchy.cpp
+++ b/Engine/Plugins/Animation/ControlRig/Source/ControlRig/Private/Rigs/RigHierarchy.cpp
@@ -138,7 +138,7 @@ URigHierarchy::URigHierarchy()
 , bAllowNameSpaceWhenSanitizingName(false)
 , ExecuteContext(nullptr)
 #if WITH_EDITOR
-, bRecordTransformsAtRuntime(true)
+, bRecordInstructionsAtRuntime(true)
 #endif
 , ElementKeyRedirector(nullptr)
 , ElementBeingDestroyed(nullptr)
@@ -3839,7 +3839,7 @@ bool URigHierarchy::SetParentWeightArray(FRigBaseElement* InChild,  const TArray
 	return false;
 }
 
-bool URigHierarchy::CanSwitchToParent(FRigElementKey InChild, FRigElementKey InParent, const IDependenciesProvider& InDependencyProvider, FString* OutFailureReason)
+bool URigHierarchy::CanSwitchToParent(FRigElementKey InChild, FRigElementKey InParent, const IRigDependenciesProvider& InDependencyProvider, FString* OutFailureReason)
 {
 	InParent = PreprocessParentElementKeyForSpaceSwitching(InChild, InParent);
 
@@ -3906,14 +3906,14 @@ bool URigHierarchy::CanSwitchToParent(FRigElementKey InChild, FRigElementKey InP
 	return true;
 }
 
-bool URigHierarchy::SwitchToParent(FRigElementKey InChild, FRigElementKey InParent, bool bInitial, bool bAffectChildren, const IDependenciesProvider& InDependencyProvider, FString* OutFailureReason)
+bool URigHierarchy::SwitchToParent(FRigElementKey InChild, FRigElementKey InParent, bool bInitial, bool bAffectChildren, const IRigDependenciesProvider& InDependencyProvider, FString* OutFailureReason)
 {
 	InParent = PreprocessParentElementKeyForSpaceSwitching(InChild, InParent);
 	return SwitchToParent(Find(InChild), Find(InParent), bInitial, bAffectChildren, InDependencyProvider, OutFailureReason);
 }
 
 bool URigHierarchy::SwitchToParent(FRigBaseElement* InChild, FRigBaseElement* InParent, bool bInitial,
-	bool bAffectChildren, const IDependenciesProvider& InDependencyProvider, FString* OutFailureReason)
+	bool bAffectChildren, const IRigDependenciesProvider& InDependencyProvider, FString* OutFailureReason)
 {
 	FRigHierarchyEnableControllerBracket EnableController(this, true);
 
@@ -3931,19 +3931,22 @@ bool URigHierarchy::SwitchToParent(FRigBaseElement* InChild, FRigBaseElement* In
 	}
 
 	// rely on the VM's dependency map if there's currently an available context.
-	const IDependenciesProvider* DependencyProviderPtr = &InDependencyProvider;
+	const IRigDependenciesProvider* DependencyProviderPtr = &InDependencyProvider;
 
 #if WITH_EDITOR
 	// Keep this in function scope, since we might be taking a pointer to it.
-	FRigVMDependenciesProvider DependencyProviderFromVM;
+	FRigDependenciesProviderForControlRig DependencyProviderFromVM;
 	{
 		UE::TScopeLock Lock(ExecuteContextLock);
-		if(ExecuteContext != nullptr && DependencyProviderPtr->GetDependencies().IsEmpty())
+		if(ExecuteContext != nullptr && DependencyProviderPtr->GetRigHierarchyDependencies().IsEmpty())
 		{
 			if(ExecuteContext->VM)
 			{
-				DependencyProviderFromVM = FRigVMDependenciesProvider(this, ExecuteContext->VM);
-				DependencyProviderPtr = &DependencyProviderFromVM;
+				if(UControlRig* Rig = Cast<UControlRig>(GetOuter()))
+				{
+					DependencyProviderFromVM = FRigDependenciesProviderForControlRig(Rig);
+					DependencyProviderPtr = &DependencyProviderFromVM;
+				}
 			}
 		}
 	}
@@ -4955,22 +4958,25 @@ FTransform URigHierarchy::GetTransform(FRigTransformElement* InTransformElement,
 
 #if WITH_EDITOR
 	{
-		if(bRecordTransformsAtRuntime)
+		if(bRecordInstructionsAtRuntime)
 		{
 			UE::TScopeLock Lock(ExecuteContextLock);
 			if (ExecuteContext != nullptr)
 			{
-				ReadTransformsAtRuntime.Emplace(
-					ExecuteContext->GetPublicData<>().GetInstructionIndex(),
+				const FControlRigExecuteContext& PublicData = ExecuteContext->GetPublicData<FControlRigExecuteContext>();
+				const FName VMName = PublicData.GetRigModuleInstance() ? PublicData.GetRigModuleInstance()->Name : NAME_None;
+				PublicData.ReadRecords.Emplace(
+					PublicData.GetInstructionIndex(),
 					ExecuteContext->GetSlice().GetIndex(),
 					InTransformElement->GetIndex(),
-					InTransformType
+					InTransformType,
+					VMName
 				);
 			}			
 		}
 	}
 	
-	TGuardValue<bool> RecordTransformsPerInstructionGuard(bRecordTransformsAtRuntime, false);
+	TGuardValue<bool> RecordTransformsPerInstructionGuard(bRecordInstructionsAtRuntime, false);
 	
 #endif
 	
@@ -5115,18 +5121,20 @@ void URigHierarchy::SetTransform(FRigTransformElement* InTransformElement, const
 
 	// lock execute context scope
 	{
-		if(bRecordTransformsAtRuntime)
+		if(bRecordInstructionsAtRuntime)
 		{
 			UE::TScopeLock Lock(ExecuteContextLock);
 			if (ExecuteContext)
 			{
-				const FRigVMExecuteContext& PublicData = ExecuteContext->GetPublicData<>();
+				const FControlRigExecuteContext& PublicData = ExecuteContext->GetPublicData<FControlRigExecuteContext>();
 				const FRigVMSlice& Slice = ExecuteContext->GetSlice();
-				WrittenTransformsAtRuntime.Emplace(
+				const FName VMName = PublicData.GetRigModuleInstance() ? PublicData.GetRigModuleInstance()->Name : NAME_None;
+				PublicData.WrittenRecords.Emplace(
 					PublicData.GetInstructionIndex(),
 					Slice.GetIndex(),
 					InTransformElement->GetIndex(),
-					InTransformType
+					InTransformType,
+					VMName
 				);
 
 				// if we are setting a control / null parent after a child here - let's let the user know
@@ -5139,9 +5147,9 @@ void URigHierarchy::SetTransform(FRigTransformElement* InTransformElement, const
 						{
 							for (FRigBaseElement* Child : GetChildren(InTransformElement))
 							{
-								const bool bChildFound = WrittenTransformsAtRuntime.ContainsByPredicate([Child](const TInstructionSliceElement& Entry) -> bool
+								const bool bChildFound = PublicData.WrittenRecords.ContainsByPredicate([Child](const FInstructionRecord& Record) -> bool
 								{
-									return Entry.Get<2>() == Child->GetIndex();
+									return Record.ElementIndex == Child->GetIndex();
 								});
 
 								if(bChildFound)
@@ -5169,7 +5177,7 @@ void URigHierarchy::SetTransform(FRigTransformElement* InTransformElement, const
 		}
 	}
 	
-	TGuardValue<bool> RecordTransformsPerInstructionGuard(bRecordTransformsAtRuntime, false);
+	TGuardValue<bool> RecordTransformsPerInstructionGuard(bRecordInstructionsAtRuntime, false);
 	
 #endif
 
@@ -5298,22 +5306,26 @@ FTransform URigHierarchy::GetControlOffsetTransform(FRigControlElement* InContro
 
 	// lock execute context scope
 	{
-		if(bRecordTransformsAtRuntime) 
+		if(bRecordInstructionsAtRuntime) 
 		{
 			UE::TScopeLock Lock(ExecuteContextLock);
 			if (ExecuteContext)
 			{
-				ReadTransformsAtRuntime.Emplace(
-					ExecuteContext->GetPublicData<>().GetInstructionIndex(),
+				const FControlRigExecuteContext& PublicData = ExecuteContext->GetPublicData<FControlRigExecuteContext>();
+				const FName VMName = PublicData.GetRigModuleInstance() ? PublicData.GetRigModuleInstance()->Name : NAME_None;
+
+				PublicData.ReadRecords.Emplace(
+					PublicData.GetInstructionIndex(),
 					ExecuteContext->GetSlice().GetIndex(),
 					InControlElement->GetIndex(),
-					InTransformType
+					InTransformType,
+					VMName
 				);
 			}
 		}
 	}
 	
-	TGuardValue<bool> RecordTransformsPerInstructionGuard(bRecordTransformsAtRuntime, false);
+	TGuardValue<bool> RecordTransformsPerInstructionGuard(bRecordInstructionsAtRuntime, false);
 	
 #endif
 
@@ -5391,22 +5403,25 @@ void URigHierarchy::SetControlOffsetTransform(FRigControlElement* InControlEleme
 
 	// lock execute context scope
 	{
-		if(bRecordTransformsAtRuntime)
+		if(bRecordInstructionsAtRuntime)
 		{
 			UE::TScopeLock Lock(ExecuteContextLock);
 			if (ExecuteContext)
 			{
-				WrittenTransformsAtRuntime.Emplace(
-					ExecuteContext->GetPublicData<>().GetInstructionIndex(),
+				const FControlRigExecuteContext& PublicData = ExecuteContext->GetPublicData<FControlRigExecuteContext>();
+				const FName VMName = PublicData.GetRigModuleInstance() ? PublicData.GetRigModuleInstance()->Name : NAME_None;
+				PublicData.WrittenRecords.Emplace(
+					PublicData.GetInstructionIndex(),
 					ExecuteContext->GetSlice().GetIndex(),
 					InControlElement->GetIndex(),
-					InTransformType
+					InTransformType,
+					VMName
 				);
 			}
 		}
 	}
 	
-	TGuardValue<bool> RecordTransformsPerInstructionGuard(bRecordTransformsAtRuntime, false);
+	TGuardValue<bool> RecordTransformsPerInstructionGuard(bRecordInstructionsAtRuntime, false);
 	
 #endif
 
@@ -6282,17 +6297,17 @@ FRigHierarchyKey URigHierarchy::GetPreviousHierarchyParent(const FRigHierarchyKe
 	return FRigHierarchyKey();
 }
 
-bool URigHierarchy::IsParentedTo(FRigBaseElement* InChild, FRigBaseElement* InParent, const IDependenciesProvider& InDependencyProvider) const
+bool URigHierarchy::IsParentedTo(FRigBaseElement* InChild, FRigBaseElement* InParent, const IRigDependenciesProvider& InDependencyProvider) const
 {
 	ElementDependencyVisited.Reset();
-	if(!InDependencyProvider.GetDependencies().IsEmpty())
+	if(!InDependencyProvider.GetRigHierarchyDependencies().IsEmpty())
 	{
 		ElementDependencyVisited.SetNumZeroed(Elements.Num());
 	}
 	return IsDependentOn(InChild, InParent, InDependencyProvider, true);
 }
 
-bool URigHierarchy::IsDependentOn(FRigBaseElement* InDependent, FRigBaseElement* InDependency, const IDependenciesProvider& InDependencyProvider, bool bIsOnActualTopology) const
+bool URigHierarchy::IsDependentOn(FRigBaseElement* InDependent, FRigBaseElement* InDependency, const IRigDependenciesProvider& InDependencyProvider, bool bIsOnActualTopology) const
 {
 	if((InDependent == nullptr) || (InDependency == nullptr))
 	{
@@ -6376,14 +6391,19 @@ bool URigHierarchy::IsDependentOn(FRigBaseElement* InDependent, FRigBaseElement*
 	}
 
 	// check the optional dependency map
-	const TElementDependencyMap& DependencyMap = InDependencyProvider.GetDependencies();
-	if(const TArray<int32>* DependentIndicesPtr = DependencyMap.Find(InDependent->GetIndex()))
+	const TRigHierarchyDependencyMap& DependencyMap = InDependencyProvider.GetRigHierarchyDependencies();
+	if(const TArray<FRigHierarchyRecord>* DependentRecordsPtr = DependencyMap.Find(InDependent->GetIndex()))
 	{
-		const TArray<int32>& DependentIndices = *DependentIndicesPtr;
-		for(const int32 DependentIndex : DependentIndices)
+		const TArray<FRigHierarchyRecord>& DependentRecords = *DependentRecordsPtr;
+		for(const FRigHierarchyRecord& DependentRecord : DependentRecords)
 		{
-			ensure(Elements.IsValidIndex(DependentIndex));
-			if(IsDependentOn(Elements[DependentIndex], InDependency, InDependencyProvider, false))
+			if (!DependentRecord.IsElement())
+			{
+				continue;
+			}
+			
+			ensure(Elements.IsValidIndex(DependentRecord.Index));
+			if(IsDependentOn(Elements[DependentRecord.Index], InDependency, InDependencyProvider, false))
 			{
 				// we'll only update the caches if we are following edges on the actual topology
 				if(bIsOnActualTopology)
@@ -8239,173 +8259,6 @@ FName URigHierarchy::GetUniqueName(const FName& InName, TFunction<bool(const FNa
 	return Name;
 }
 
-#if WITH_EDITOR
-
-URigHierarchy::TElementDependencyMap URigHierarchy::GetDependenciesForVM(const URigVM* InVM, FName InEventName) const
-{
-	LLM_SCOPE_BYNAME(TEXT("Animation/ControlRig"));
-	check(InVM);
-
-	if(InEventName.IsNone())
-	{
-		InEventName = FRigUnit_BeginExecution().GetEventName();
-	}
-
-	URigHierarchy::TElementDependencyMap Dependencies;
-	const FRigVMInstructionArray Instructions = InVM->GetByteCode().GetInstructions();
-
-	// if the VM doesn't implement the given event
-	if(!InVM->ContainsEntry(InEventName))
-	{
-		return Dependencies;
-	}
-
-	// only represent instruction for a given event
-	const int32 EntryIndex = InVM->GetByteCode().FindEntryIndex(InEventName);
-	const int32 EntryInstructionIndex = InVM->GetByteCode().GetEntry(EntryIndex).InstructionIndex;
-
-	TMap<FRigVMOperand, TArray<int32>> OperandToInstructions;
-
-	for(int32 InstructionIndex = EntryInstructionIndex; InstructionIndex < Instructions.Num(); InstructionIndex++)
-	{
-		// early exit since following instructions belong to another event
-		if(Instructions[InstructionIndex].OpCode == ERigVMOpCode::Exit)
-		{
-			break;
-		}
-
-		const FRigVMOperandArray InputOperands = InVM->GetByteCode().GetInputOperands(InstructionIndex);
-		for(const FRigVMOperand InputOperand : InputOperands)
-		{
-			const FRigVMOperand InputOperandNoRegisterOffset(InputOperand.GetMemoryType(), InputOperand.GetRegisterIndex());
-			OperandToInstructions.FindOrAdd(InputOperandNoRegisterOffset).Add(InstructionIndex);
-		}
-	}
-
-	typedef TTuple<int32, int32> TInt32Tuple;
-	TArray<TArray<TInt32Tuple>> ReadTransformPerInstruction, WrittenTransformsPerInstruction;
-	
-	// Find the max instruction index
-	{
-		/** TODO: UE-207320
-		 * This function needs to be rewritten to account for modular rigs. Rig Hierarchies will have reads/writes coming from
-		 * different VMs, so we cannot depend on the algorithm used in this function to find dependencies. */
-		
-		int32 MaxInstructionIndex = Instructions.Num();
-		for(int32 RecordType = 0; RecordType < 2; RecordType++)
-		{
-			const TArray<TInstructionSliceElement>& Records = RecordType == 0 ? ReadTransformsAtRuntime : WrittenTransformsAtRuntime; 
-			for(int32 RecordIndex = 0; RecordIndex < Records.Num(); RecordIndex++)
-			{
-				const TInstructionSliceElement& Record = Records[RecordIndex];
-				const int32& InstructionIndex = Record.Get<0>();
-				MaxInstructionIndex = FMath::Max(MaxInstructionIndex, InstructionIndex);
-			}
-		}
-		ReadTransformPerInstruction.AddZeroed(MaxInstructionIndex+1);
-		WrittenTransformsPerInstruction.AddZeroed(MaxInstructionIndex+1);
-	}
-
-	// fill lookup tables per instruction / element
-	for(int32 RecordType = 0; RecordType < 2; RecordType++)
-	{
-		const TArray<TInstructionSliceElement>& Records = RecordType == 0 ? ReadTransformsAtRuntime : WrittenTransformsAtRuntime; 
-		TArray<TArray<TInt32Tuple>>& PerInstruction = RecordType == 0 ? ReadTransformPerInstruction : WrittenTransformsPerInstruction;
-
-		for(int32 RecordIndex = 0; RecordIndex < Records.Num(); RecordIndex++)
-		{
-			const TInstructionSliceElement& Record = Records[RecordIndex];
-			const int32& InstructionIndex = Record.Get<0>();
-			const int32& SliceIndex = Record.Get<1>();
-			const int32& TransformIndex = Record.Get<2>();
-			PerInstruction[InstructionIndex].Emplace(SliceIndex, TransformIndex);
-		}
-	}
-
-	// for each read transform on an instruction
-	// follow the operands to the next instruction affected by it
-	TArray<TInt32Tuple> FilteredTransforms;
-	TArray<int32> InstructionsToVisit;
-
-	for(int32 InstructionIndex = EntryInstructionIndex; InstructionIndex < Instructions.Num(); InstructionIndex++)
-	{
-		// early exit since following instructions belong to another event
-		if(Instructions[InstructionIndex].OpCode == ERigVMOpCode::Exit)
-		{
-			break;
-		}
-
-		const TArray<TInt32Tuple>& ReadTransforms = ReadTransformPerInstruction[InstructionIndex];
-		if(ReadTransforms.IsEmpty())
-		{
-			continue;
-		}
-
-		FilteredTransforms.Reset();
-
-		for(int32 ReadTransformIndex = 0; ReadTransformIndex < ReadTransforms.Num(); ReadTransformIndex++)
-		{
-			const TInt32Tuple& ReadTransform = ReadTransforms[ReadTransformIndex];
-			
-			InstructionsToVisit.Reset();
-			InstructionsToVisit.Add(InstructionIndex);
-
-			for(int32 InstructionToVisitIndex = 0; InstructionToVisitIndex < InstructionsToVisit.Num(); InstructionToVisitIndex++)
-			{
-				const int32 InstructionToVisit = InstructionsToVisit[InstructionToVisitIndex];
-
-				const TArray<TInt32Tuple>& WrittenTransforms = WrittenTransformsPerInstruction[InstructionToVisit];
-				for(int32 WrittenTransformIndex = 0; WrittenTransformIndex < WrittenTransforms.Num(); WrittenTransformIndex++)
-				{
-					const TInt32Tuple& WrittenTransform = WrittenTransforms[WrittenTransformIndex];
-
-					// for the first instruction in this pass let's only care about
-					// written transforms which have not been read before
-					if(InstructionToVisitIndex == 0)
-					{
-						if(ReadTransforms.Contains(WrittenTransform))
-						{
-							continue;
-						}
-					}
-					
-					FilteredTransforms.AddUnique(WrittenTransform);
-				}
-
-				FRigVMOperandArray OutputOperands = InVM->GetByteCode().GetOutputOperands(InstructionToVisit);
-				for(const FRigVMOperand OutputOperand : OutputOperands)
-				{
-					const FRigVMOperand OutputOperandNoRegisterOffset(OutputOperand.GetMemoryType(), OutputOperand.GetRegisterIndex());
-
-					if(const TArray<int32>* InstructionsWithInputOperand = OperandToInstructions.Find(OutputOperandNoRegisterOffset))
-					{
-						for(int32 InstructionWithInputOperand : *InstructionsWithInputOperand)
-						{
-							InstructionsToVisit.AddUnique(InstructionWithInputOperand);
-						}
-					}
-				}
-			}
-		}
-		
-		for(const TInt32Tuple& ReadTransform : ReadTransforms)
-		{
-			for(const TInt32Tuple& FilteredTransform : FilteredTransforms)
-			{
-				// only create dependencies for reads and writes that are on the same slice
-				if(ReadTransform != FilteredTransform && ReadTransform.Get<0>() == FilteredTransform.Get<0>())
-				{
-					Dependencies.FindOrAdd(FilteredTransform.Get<1>()).AddUnique(ReadTransform.Get<1>());
-				}
-			}
-		}
-	}
-
-	return Dependencies;
-}
-
-#endif
-
 void URigHierarchy::UpdateVisibilityOnProxyControls()
 {
 	URigHierarchy* HierarchyForSelection = HierarchyForSelectionPtr.Get();
diff --git a/Engine/Plugins/Animation/ControlRig/Source/ControlRig/Private/Units/Execution/RigUnit_DynamicHierarchy.cpp b/Engine/Plugins/Animation/ControlRig/Source/ControlRig/Private/Units/Execution/RigUnit_DynamicHierarchy.cpp
--- a/Engine/Plugins/Animation/ControlRig/Source/ControlRig/Private/Units/Execution/RigUnit_DynamicHierarchy.cpp
+++ b/Engine/Plugins/Animation/ControlRig/Source/ControlRig/Private/Units/Execution/RigUnit_DynamicHierarchy.cpp
@@ -215,7 +215,7 @@ FRigUnit_SwitchParent_Execute()
 		default:
 		{
 			FString FailureReason;
-			static const FNoDependenciesProvider NoDependencyProvider;
+			static const FEmptyRigDependenciesProvider NoDependencyProvider;
 			if(!ExecuteContext.Hierarchy->SwitchToParent(ChildElement, ParentElement, false, true, NoDependencyProvider, &FailureReason))
 			{
 				if(!FailureReason.IsEmpty())
diff --git a/Engine/Plugins/Animation/ControlRig/Source/ControlRig/Public/Rigs/RigDependency.h b/Engine/Plugins/Animation/ControlRig/Source/ControlRig/Public/Rigs/RigDependency.h
--- a/Engine/Plugins/Animation/ControlRig/Source/ControlRig/Public/Rigs/RigDependency.h
+++ b/Engine/Plugins/Animation/ControlRig/Source/ControlRig/Public/Rigs/RigDependency.h
@@ -2,58 +2,76 @@
 
 #pragma once
 
-#include "Containers/Array.h"
-#include "Containers/Map.h"
 #include "UObject/NameTypes.h"
 #include "UObject/WeakObjectPtr.h"
 #include "UObject/WeakObjectPtrTemplates.h"
+#include "RigDependencyRecords.h"
 
 #define UE_API CONTROLRIG_API
 
-class URigHierarchy;
+class UControlRig;
 class URigVM;
 
 /**
- * IDependenciesProvider provides an interface for constructing and passing elements dependencies to functions that need to perform dependency tests.
+ * IRigDependenciesProvider provides an interface for constructing and passing elements dependencies to functions that need to perform dependency tests.
  * This is currently mainly used to ensure that certain rig elements are not dependent on each other (parent switching, constraints ordering, etc.) 
  */
 
-struct IDependenciesProvider
+struct IRigDependenciesProvider
 {
-	virtual ~IDependenciesProvider() = default;
-	virtual const TMap<int32, TArray<int32>>& GetDependencies() const = 0;
+	virtual ~IRigDependenciesProvider() = default;
+	virtual const TRigHierarchyDependencyMap& GetRigHierarchyDependencies() const = 0;
 };
 
 /**
- * FNoDependenciesProvider is the default dependency provider (no dependency at all) 
+ * FEmptyRigDependenciesProvider is the default dependency provider (no dependency at all) 
  */
 
-struct FNoDependenciesProvider : public IDependenciesProvider
+struct FEmptyRigDependenciesProvider : public IRigDependenciesProvider
 {
-	virtual ~FNoDependenciesProvider() override = default;
-	UE_API virtual const TMap<int32, TArray<int32>>& GetDependencies() const override;
+	virtual ~FEmptyRigDependenciesProvider() override = default;
+	UE_API virtual const TRigHierarchyDependencyMap& GetRigHierarchyDependencies() const override;
 };
 
 /**
- * FRigVMDependenciesProvider builds a dependency map based on the provided RigVM's instructions
+ * FRigDependenciesProviderForVM builds a dependency map based on the provided RigVM's instructions
  */
 
-struct FRigVMDependenciesProvider : public IDependenciesProvider
+struct FRigDependenciesProviderForVM : public IRigDependenciesProvider
 {
-	FRigVMDependenciesProvider() = default;
-	UE_API FRigVMDependenciesProvider(const URigHierarchy* InHierarchy, URigVM* InVM, FName InEventName = NAME_None);
-	virtual ~FRigVMDependenciesProvider() override = default;
+	FRigDependenciesProviderForVM() = default;
+	UE_API FRigDependenciesProviderForVM(UControlRig* InControlRig, FName InEventName = NAME_None);
+	virtual ~FRigDependenciesProviderForVM() override = default;
 	
-	UE_API virtual const TMap<int32, TArray<int32>>& GetDependencies() const override;
-	
-private:
-		
-	TWeakObjectPtr<const URigHierarchy> WeakHierarchy = nullptr;
-	TWeakObjectPtr<URigVM> WeakRigVM = nullptr;
+	UE_API virtual const TRigHierarchyDependencyMap& GetRigHierarchyDependencies() const override;
+
+	UE_API UControlRig* GetControlRig() const;
+
+protected:
+
+	TWeakObjectPtr<UControlRig> WeakControlRig = nullptr;
 	FName EventName = NAME_None;
 	
-	mutable TMap<int32, TArray<int32>> CachedDependencies;
+	mutable TRigHierarchyDependencyMap CachedDependencies;
 	mutable uint32 TopologyHash = 0;
 };
 
+/**
+ * FRigDependenciesProviderForControlRig builds a dependency map based on the provided rig.
+ * This works both for graph based control rigs or modular control rigs.
+ */
+
+struct FRigDependenciesProviderForControlRig : public FRigDependenciesProviderForVM
+{
+	FRigDependenciesProviderForControlRig() = default;
+	UE_API FRigDependenciesProviderForControlRig(UControlRig* InControlRig, FName InEventName = NAME_None);
+	virtual ~FRigDependenciesProviderForControlRig() override = default;
+	
+	UE_API virtual const TRigHierarchyDependencyMap& GetRigHierarchyDependencies() const override;
+
+protected:
+
+	TMap<FName, TSharedPtr<FRigDependenciesProviderForVM>> ModuleDependencyProviders;
+};
+
 #undef UE_API
diff --git a/Engine/Plugins/Animation/ControlRig/Source/ControlRig/Public/Rigs/RigDependencyRecords.h b/Engine/Plugins/Animation/ControlRig/Source/ControlRig/Public/Rigs/RigDependencyRecords.h
new file mode 100644
--- /dev/null
+++ b/Engine/Plugins/Animation/ControlRig/Source/ControlRig/Public/Rigs/RigDependencyRecords.h
@@ -0,0 +1,228 @@
+// Copyright Epic Games, Inc. All Rights Reserved.
+
+#pragma once
+
+#include "Containers/Array.h"
+#include "Containers/Map.h"
+#include "RigHierarchyElements.h"
+
+#define UE_API CONTROLRIG_API
+
+struct FRigHierarchyRecord
+{
+	FRigHierarchyRecord()
+		: Index(INDEX_NONE)
+		, MetadataName(NAME_None)
+	{
+	}
+
+	FRigHierarchyRecord(int32 InIndex, const FName& InMetadataName = NAME_None)
+		: Index(InIndex)
+		, MetadataName(InMetadataName)
+	{
+	}
+	
+	bool IsValid() const
+	{
+		return Index != INDEX_NONE;
+	}
+
+	bool IsElement() const
+	{
+		return IsValid() && MetadataName.IsNone();
+	}
+
+	bool IsMetadata() const
+	{
+		return IsValid() && !MetadataName.IsNone();
+	}
+
+	bool operator==(const FRigHierarchyRecord& InRecord) const
+	{
+		return Index == InRecord.Index && MetadataName == InRecord.MetadataName;
+	}
+
+	bool operator <(const FRigHierarchyRecord& InRecord) const
+	{
+		if (Index < InRecord.Index)
+		{
+			return true;
+		}
+		if (Index > InRecord.Index)
+		{
+			return false;
+		}
+		return MetadataName.Compare(InRecord.MetadataName) < 0;
+	}
+
+	bool operator >(const FRigHierarchyRecord& InRecord) const
+	{
+		if (Index > InRecord.Index)
+		{
+			return true;
+		}
+		if (Index < InRecord.Index)
+		{
+			return false;
+		}
+		return MetadataName.Compare(InRecord.MetadataName) > 0;
+	}
+
+	int32 Index;
+	FName MetadataName;
+};
+
+typedef TMap<FRigHierarchyRecord, TArray<FRigHierarchyRecord>> TRigHierarchyDependencyMap;
+typedef TPair<FRigHierarchyRecord, TArray<FRigHierarchyRecord>> TRigHierarchyDependencyMapPair;
+
+inline uint32 GetTypeHash(const FRigHierarchyRecord& InRecord)
+{
+	return HashCombine(GetTypeHash(InRecord.Index), GetTypeHash(InRecord.MetadataName));
+}
+
+struct FInstructionRecord
+{
+	FInstructionRecord()
+		: InstructionIndex(INDEX_NONE)
+		, SliceIndex(INDEX_NONE)
+		, ElementIndex(INDEX_NONE)
+		, TransformType(ERigTransformType::CurrentLocal)
+		, VM(NAME_None)
+		, MetadataName(NAME_None)
+	{
+	}
+
+	FInstructionRecord(int32 InInstructionIndex, int32 InSliceIndex, int32 InElementIndex, ERigTransformType::Type InTransformType, const FName& InVM = NAME_None, const FName& InMetadataName = NAME_None)
+		:InstructionIndex(InInstructionIndex)
+		, SliceIndex(InSliceIndex)
+		, ElementIndex(InElementIndex)
+		, TransformType(InTransformType)
+		, VM(InVM)
+		, MetadataName(InMetadataName)
+	{
+	}
+
+	bool operator==(const FInstructionRecord& InRecord) const
+	{
+		return InstructionIndex == InRecord.InstructionIndex &&
+			SliceIndex == InRecord.SliceIndex &&
+			ElementIndex == InRecord.ElementIndex &&
+			TransformType == InRecord.TransformType &&
+			VM == InRecord.VM &&
+			MetadataName == InRecord.MetadataName;
+	}
+
+	bool operator <(const FInstructionRecord& InRecord) const
+	{
+		if (InstructionIndex < InRecord.InstructionIndex)
+		{
+			return true;
+		}
+		if (InstructionIndex > InRecord.InstructionIndex)
+		{
+			return false;
+		}
+		if (SliceIndex < InRecord.SliceIndex)
+		{
+			return true;
+		}
+		if (SliceIndex > InRecord.SliceIndex)
+		{
+			return false;
+		}
+		if (ElementIndex < InRecord.ElementIndex)
+		{
+			return true;
+		}
+		if (ElementIndex > InRecord.ElementIndex)
+		{
+			return false;
+		}
+		if (TransformType < InRecord.TransformType)
+		{
+			return true;
+		}
+		if (TransformType > InRecord.TransformType)
+		{
+			return false;
+		}
+		const int32 VMComp = VM.Compare(InRecord.VM);
+		if (VMComp < 0)
+		{
+			return true;
+		}
+		if (VMComp > 0)
+		{
+			return false;
+		}
+		return MetadataName.Compare(InRecord.MetadataName) < 0;
+	}
+
+	bool operator >(const FInstructionRecord& InRecord) const
+	{
+		if (InstructionIndex > InRecord.InstructionIndex)
+		{
+			return true;
+		}
+		if (InstructionIndex < InRecord.InstructionIndex)
+		{
+			return false;
+		}
+		if (SliceIndex > InRecord.SliceIndex)
+		{
+			return true;
+		}
+		if (SliceIndex < InRecord.SliceIndex)
+		{
+			return false;
+		}
+		if (ElementIndex > InRecord.ElementIndex)
+		{
+			return true;
+		}
+		if (ElementIndex < InRecord.ElementIndex)
+		{
+			return false;
+		}
+		if (TransformType > InRecord.TransformType)
+		{
+			return true;
+		}
+		if (TransformType < InRecord.TransformType)
+		{
+			return false;
+		}
+		const int32 VMComp = VM.Compare(InRecord.VM);
+		if (VMComp > 0)
+		{
+			return true;
+		}
+		if (VMComp < 0)
+		{
+			return false;
+		}
+		return MetadataName.Compare(InRecord.MetadataName) > 0;
+	}
+
+	int32 InstructionIndex;
+	int32 SliceIndex;
+	int32 ElementIndex;
+	ERigTransformType::Type TransformType;
+	FName VM;
+	FName MetadataName;
+};
+
+inline uint32 GetTypeHash(const FInstructionRecord& InRecord)
+{
+	uint32 Hash = HashCombine(GetTypeHash(InRecord.InstructionIndex), GetTypeHash(InRecord.SliceIndex));
+	Hash = HashCombine(Hash, GetTypeHash(InRecord.ElementIndex));
+	Hash = HashCombine(Hash, GetTypeHash(static_cast<int32>(InRecord.TransformType)));
+	Hash = HashCombine(Hash, GetTypeHash(InRecord.VM));
+	Hash = HashCombine(Hash, GetTypeHash(InRecord.MetadataName));
+	return Hash;
+}
+
+typedef TMap<FInstructionRecord, TArray<FInstructionRecord>> TInstructionRecordMap;
+typedef TPair<FInstructionRecord, TArray<FInstructionRecord>> TInstructionRecordMapPair;
+
+#undef UE_API
diff --git a/Engine/Plugins/Animation/ControlRig/Source/ControlRig/Public/Rigs/RigHierarchy.h b/Engine/Plugins/Animation/ControlRig/Source/ControlRig/Public/Rigs/RigHierarchy.h
--- a/Engine/Plugins/Animation/ControlRig/Source/ControlRig/Public/Rigs/RigHierarchy.h
+++ b/Engine/Plugins/Animation/ControlRig/Source/ControlRig/Public/Rigs/RigHierarchy.h
@@ -169,9 +169,6 @@ class URigHierarchy : public UObject
 
 public:
 
-	typedef TMap<int32, TArray<int32>> TElementDependencyMap;
-	typedef TPair<int32, TArray<int32>> TElementDependencyMapPair;
-	typedef TTuple<int32, int32, int32, ERigTransformType::Type> TInstructionSliceElement;
 	inline static const FLazyName TagMetadataName = TEXT("Tags");
 	inline static const FLazyName ShortModuleNameMetadataName_Deprecated = TEXT("ShortModuleName");
 	inline static const FLazyName DesiredNameMetadataName = TEXT("DesiredName");
@@ -3591,7 +3588,7 @@ public:
 	CONTROLRIG_API bool CanSwitchToParent(
 		FRigElementKey InChild,
 		FRigElementKey InParent,
-		const IDependenciesProvider& InDependencyProvider = FNoDependenciesProvider(),
+		const IRigDependenciesProvider& InDependencyProvider = FEmptyRigDependenciesProvider(),
 		FString* OutFailureReason = nullptr);
 
 	/**
@@ -3607,14 +3604,14 @@ public:
 	UFUNCTION(BlueprintCallable, Category = URigHierarchy)
 	bool SwitchToParent(FRigElementKey InChild, FRigElementKey InParent, bool bInitial = false, bool bAffectChildren = true)
 	{
-		return SwitchToParent(InChild, InParent, bInitial, bAffectChildren, FNoDependenciesProvider(), nullptr);
+		return SwitchToParent(InChild, InParent, bInitial, bAffectChildren, FEmptyRigDependenciesProvider(), nullptr);
 	}
 
 	CONTROLRIG_API bool SwitchToParent(FRigElementKey InChild,
 		FRigElementKey InParent,
 		bool bInitial,
 		bool bAffectChildren,
-		const IDependenciesProvider& InDependencyProvider,
+		const IRigDependenciesProvider& InDependencyProvider,
 		FString* OutFailureReason);
 
 	/**
@@ -3633,7 +3630,7 @@ public:
 		FRigBaseElement* InParent,
 		bool bInitial = false,
 		bool bAffectChildren = true,
-		const IDependenciesProvider& InDependencyProvider = FNoDependenciesProvider(),
+		const IRigDependenciesProvider& InDependencyProvider = FEmptyRigDependenciesProvider(),
 		FString* OutFailureReason = nullptr);
 
 	/**
@@ -3719,7 +3716,7 @@ public:
 	 * @param InDependencyProvider An additional object providing dependencies to respect
 	 * @return True if the given parent and given child have a parent-child relationship
 	 */
-    bool IsParentedTo(int32 InChildIndex, int32 InParentIndex, const IDependenciesProvider& InDependencyProvider = FNoDependenciesProvider()) const
+    bool IsParentedTo(int32 InChildIndex, int32 InParentIndex, const IRigDependenciesProvider& InDependencyProvider = FEmptyRigDependenciesProvider()) const
 	{
 		if(Elements.IsValidIndex(InChildIndex) && Elements.IsValidIndex(InParentIndex))
 		{
@@ -4451,7 +4448,7 @@ public:
 	 */
 	CONTROLRIG_API bool IsParentedTo(FRigBaseElement* InChild,
 		FRigBaseElement* InParent,
-		const IDependenciesProvider& InDependencyProvider = FNoDependenciesProvider()) const;
+		const IRigDependenciesProvider& InDependencyProvider = FEmptyRigDependenciesProvider()) const;
 
 private:
 	/**
@@ -4464,7 +4461,7 @@ private:
 	 */
 	CONTROLRIG_API bool IsDependentOn(FRigBaseElement* InDependent,
 		FRigBaseElement* InDependency,
-		const IDependenciesProvider& InDependencyProvider = FNoDependenciesProvider(),
+		const IRigDependenciesProvider& InDependencyProvider = FEmptyRigDependenciesProvider(),
 		bool bIsOnActualTopology = true) const;
 
 public:
@@ -5309,12 +5306,7 @@ private:
 	const FRigVMExtendedExecuteContext* ExecuteContext;
 
 #if WITH_EDITOR
-	mutable bool bRecordTransformsAtRuntime;
-	mutable TArray<TInstructionSliceElement> ReadTransformsAtRuntime;
-	mutable TArray<TInstructionSliceElement> WrittenTransformsAtRuntime;
-public:
-
-	CONTROLRIG_API TElementDependencyMap GetDependenciesForVM(const URigVM* InVM, FName InEventName = NAME_None) const;
+	mutable bool bRecordInstructionsAtRuntime;
 
 private:
 
diff --git a/Engine/Plugins/Animation/ControlRig/Source/ControlRig/Public/Units/RigUnitContext.h b/Engine/Plugins/Animation/ControlRig/Source/ControlRig/Public/Units/RigUnitContext.h
--- a/Engine/Plugins/Animation/ControlRig/Source/ControlRig/Public/Units/RigUnitContext.h
+++ b/Engine/Plugins/Animation/ControlRig/Source/ControlRig/Public/Units/RigUnitContext.h
@@ -9,6 +9,7 @@
 #include "Animation/AttributesRuntime.h"
 #include "ControlRigAssetUserData.h"
 #include "RigVMCore/RigVMExecuteContext.h"
+#include "Rigs/RigDependencyRecords.h"
 #include "RigUnitContext.generated.h"
 
 #define UE_API CONTROLRIG_API
@@ -240,6 +241,9 @@ public:
 	FRigUnitContext UnitContext;
 	URigHierarchy* Hierarchy;
 	UControlRig* ControlRig;
+	
+	mutable TArray<FInstructionRecord> ReadRecords;
+	mutable TArray<FInstructionRecord> WrittenRecords;
 
 #if WITH_EDITOR
 	virtual void Report(const FRigVMLogSettings& InLogSettings, const FName& InFunctionName, int32 InInstructionIndex, const FString& InMessage) const override
diff --git a/Engine/Plugins/Animation/ControlRig/Source/ControlRigDeveloper/Private/ControlRigVisualGraphUtils.cpp b/Engine/Plugins/Animation/ControlRig/Source/ControlRigDeveloper/Private/ControlRigVisualGraphUtils.cpp
--- a/Engine/Plugins/Animation/ControlRig/Source/ControlRigDeveloper/Private/ControlRigVisualGraphUtils.cpp
+++ b/Engine/Plugins/Animation/ControlRig/Source/ControlRigDeveloper/Private/ControlRigVisualGraphUtils.cpp
@@ -229,17 +229,26 @@ FString FControlRigVisualGraphUtils::DumpRigHierarchyToDotGraph(URigHierarchy* I
 	{
 		if(UControlRig* CR = InHierarchy->GetTypedOuter<UControlRig>())
 		{
-			URigVM* VM = CR->GetVM();
+			const FRigDependenciesProviderForControlRig DependenciesProvider(CR, InEventName);
+			TRigHierarchyDependencyMap Dependencies = DependenciesProvider.GetRigHierarchyDependencies();
 
-			URigHierarchy::TElementDependencyMap Dependencies = InHierarchy->GetDependenciesForVM(VM, InEventName);
-
-			for(const URigHierarchy::TElementDependencyMapPair& Dependency : Dependencies)
+			for(const TRigHierarchyDependencyMapPair& Dependency : Dependencies)
 			{
-				const int32 TargetElementIndex = Dependency.Key;
+				if (!Dependency.Key.IsElement())
+				{
+					continue;
+				}
+				
+				const int32 TargetElementIndex = Dependency.Key.Index;
 				
-				for(const int32 SourceElementIndex : Dependency.Value)
+				for(const FRigHierarchyRecord SourceElementRecord : Dependency.Value)
 				{
-					FName EdgeName = *FString::Printf(TEXT("Dependency_%d_%d"), SourceElementIndex, TargetElementIndex);
+					if (!SourceElementRecord.IsElement())
+					{
+						continue;
+					}
+					
+					FName EdgeName = *FString::Printf(TEXT("Dependency_%d_%d"), SourceElementRecord.Index, TargetElementIndex);
 					if(Graph.FindEdge(EdgeName) != INDEX_NONE)
 					{
 						continue;
@@ -248,7 +257,7 @@ FString FControlRigVisualGraphUtils::DumpRigHierarchyToDotGraph(URigHierarchy* I
 					const TOptional<FLinearColor> EdgeColor = FLinearColor::Gray;
 					TOptional<EVisualGraphStyle> Style = EVisualGraphStyle::Dashed;
 
-					const FName SourceNodeName = Local::GetNodeNameForElementIndex(SourceElementIndex);
+					const FName SourceNodeName = Local::GetNodeNameForElementIndex(SourceElementRecord.Index);
 					const FName TargetNodeName = Local::GetNodeNameForElementIndex(TargetElementIndex);
 					const int32 SourceNodeIndex = Graph.FindNode(SourceNodeName);
 					const int32 TargetNodeIndex = Graph.FindNode(TargetNodeName);
diff --git a/Engine/Plugins/Animation/ControlRig/Source/ControlRigEditor/Private/EditMode/ControlRigEditMode.cpp b/Engine/Plugins/Animation/ControlRig/Source/ControlRigEditor/Private/EditMode/ControlRigEditMode.cpp
--- a/Engine/Plugins/Animation/ControlRig/Source/ControlRigEditor/Private/EditMode/ControlRigEditMode.cpp
+++ b/Engine/Plugins/Animation/ControlRig/Source/ControlRigEditor/Private/EditMode/ControlRigEditMode.cpp
@@ -3468,7 +3468,7 @@ void FControlRigEditMode::OpenSpacePickerWidget()
 			{
 				constexpr bool bInitial = false, bAffectChildren = true;
 				const FTransform Transform = InHierarchy->GetGlobalTransform(InControlKey);
-				FRigVMDependenciesProvider DependenciesProvider(InHierarchy, RuntimeRig->GetVM());
+				FRigDependenciesProviderForControlRig DependenciesProvider(RuntimeRig);
 				FString OutFailureReason;
 				if (InHierarchy->SwitchToParent(InControlKey, InSpaceKey, bInitial, bAffectChildren, DependenciesProvider, &OutFailureReason))
 				{
@@ -6775,7 +6775,7 @@ FInteractionDependencyCache& FControlRigEditMode::GetInteractionDependencies(UCo
 		NewDependency.WeakHierarchy = Hierarchy;
 
 		// NOTE: this is not enough for modular rigs since there are several VMs
-		const FRigVMDependenciesProvider DependencyProvider(Hierarchy, InControlRig->GetVM());
+		const FRigDependenciesProviderForControlRig DependencyProvider(InControlRig);
 		
 		for (int32 Index = 0; Index < SelectedControls.Num(); ++Index)
 		{
diff --git a/Engine/Plugins/Animation/ControlRig/Source/ControlRigEditor/Private/EditMode/SControlRigEditModeTools.cpp b/Engine/Plugins/Animation/ControlRig/Source/ControlRigEditor/Private/EditMode/SControlRigEditModeTools.cpp
--- a/Engine/Plugins/Animation/ControlRig/Source/ControlRigEditor/Private/EditMode/SControlRigEditModeTools.cpp
+++ b/Engine/Plugins/Animation/ControlRig/Source/ControlRigEditor/Private/EditMode/SControlRigEditModeTools.cpp
@@ -932,7 +932,7 @@ void SControlRigEditModeTools::HandleActiveSpaceChanged(URigHierarchy* InHierarc
 			if (ControlRig.IsValid() && ControlRig->GetHierarchy() == InHierarchy)
 			{
 				FString FailureReason;
-				FRigVMDependenciesProvider DependencyProvider(InHierarchy, ControlRig->GetVM());
+				FRigDependenciesProviderForControlRig DependencyProvider(ControlRig.Get());
 				if(!InHierarchy->CanSwitchToParent(InControlKey, InSpaceKey, DependencyProvider, &FailureReason))
 				{
 					// notification
diff --git a/Engine/Plugins/Animation/ControlRig/Source/ControlRigEditor/Private/EditMode/SControlRigSpacePicker.cpp b/Engine/Plugins/Animation/ControlRig/Source/ControlRigEditor/Private/EditMode/SControlRigSpacePicker.cpp
--- a/Engine/Plugins/Animation/ControlRig/Source/ControlRigEditor/Private/EditMode/SControlRigSpacePicker.cpp
+++ b/Engine/Plugins/Animation/ControlRig/Source/ControlRigEditor/Private/EditMode/SControlRigSpacePicker.cpp
@@ -161,7 +161,7 @@ void SControlRigSpacePicker::HandleActiveSpaceChanged(URigHierarchy* InHierarchy
 		if (UControlRig* ControlRig = GetControlRig())
 		{
 			FString FailureReason;
-			FRigVMDependenciesProvider DependencyProvider(InHierarchy, ControlRig->GetVM());
+			FRigDependenciesProviderForControlRig DependencyProvider(ControlRig);
 			if (!InHierarchy->CanSwitchToParent(InControlKey, InSpaceKey, DependencyProvider, &FailureReason))
 			{
 				// notification
diff --git a/Engine/Plugins/Animation/ControlRig/Source/ControlRigEditor/Private/Editor/SRigSpacePickerWidget.cpp b/Engine/Plugins/Animation/ControlRig/Source/ControlRigEditor/Private/Editor/SRigSpacePickerWidget.cpp
--- a/Engine/Plugins/Animation/ControlRig/Source/ControlRigEditor/Private/Editor/SRigSpacePickerWidget.cpp
+++ b/Engine/Plugins/Animation/ControlRig/Source/ControlRigEditor/Private/Editor/SRigSpacePickerWidget.cpp
@@ -752,18 +752,18 @@ FReply SRigSpacePickerWidget::HandleAddElementClicked()
 					URigHierarchy* StrongHierarchy = Hierarchy.Get();
 					for(const FRigElementKey& ControlKey : ControlKeys)
 					{
-						FRigVMDependenciesProvider DependencyProvider;
+						FRigDependenciesProviderForControlRig DependencyProvider;
 						FString FailureReason;
 
 						if(UControlRig* ControlRig = StrongHierarchy->GetTypedOuter<UControlRig>())
 						{
-							DependencyProvider = FRigVMDependenciesProvider(StrongHierarchy, ControlRig->GetVM()); 
+							DependencyProvider = FRigDependenciesProviderForControlRig(ControlRig); 
 						}
 						else if(UControlRigBlueprint* RigBlueprint = StrongHierarchy->GetTypedOuter<UControlRigBlueprint>())
 						{
 							if(UControlRig* CDO = Cast<UControlRig>(RigBlueprint->GetRigVMBlueprintGeneratedClass()->GetDefaultObject()))
 							{
-								DependencyProvider = FRigVMDependenciesProvider(StrongHierarchy, CDO->GetVM()); 
+								DependencyProvider = FRigDependenciesProviderForControlRig(CDO); 
 							}
 						}
 						
