Fix for Control Rig curves container losing the user created curves on compile

diff --git a/Engine/Plugins/Animation/ControlRig/Source/ControlRigEditor/Private/Editor/SRigCurveContainer.cpp b/Engine/Plugins/Animation/ControlRig/Source/ControlRigEditor/Private/Editor/SRigCurveContainer.cpp
--- a/Engine/Plugins/Animation/ControlRig/Source/ControlRigEditor/Private/Editor/SRigCurveContainer.cpp
+++ b/Engine/Plugins/Animation/ControlRig/Source/ControlRigEditor/Private/Editor/SRigCurveContainer.cpp
@@ -374,7 +374,7 @@ void SRigCurveContainer::CreateNewNameEntry(const FText& CommittedText, ETextCom
 	FSlateApplication::Get().DismissAllMenus();
 	if (!CommittedText.IsEmpty() && CommitType == ETextCommit::OnEnter)
 	{
-		URigHierarchy* Hierarchy = GetHierarchy();
+		URigHierarchy* Hierarchy = GetDefaultHierarchy();
 		if (Hierarchy)
 		{
 			TGuardValue<bool> GuardReentry(bIsChangingRigHierarchy, true);
