Anim in Engine: Selecting in Anim Details now selects curves of all selected Anim Layers in Curve Editor

diff --git a/Engine/Plugins/Animation/ControlRig/Source/ControlRigEditor/Private/AnimDetails/AnimDetailsSelection.cpp b/Engine/Plugins/Animation/ControlRig/Source/ControlRigEditor/Private/AnimDetails/AnimDetailsSelection.cpp
--- a/Engine/Plugins/Animation/ControlRig/Source/ControlRigEditor/Private/AnimDetails/AnimDetailsSelection.cpp
+++ b/Engine/Plugins/Animation/ControlRig/Source/ControlRigEditor/Private/AnimDetails/AnimDetailsSelection.cpp
@@ -18,6 +18,7 @@
 #include "MVVM/ViewModels/TrackModel.h"
 #include "MVVM/ViewModels/ViewModelIterators.h"
 #include "PropertyHandle.h"
+#include "Sequencer/AnimLayers/AnimLayers.h"
 #include "Sequencer/MovieSceneControlRigParameterTrack.h"
 #include "Tree/SCurveEditorTree.h"
 
@@ -468,25 +469,41 @@ void UAnimDetailsSelection::PropagonateSelectionToCurveEditor()
 				
 					const EControlRigContextChannelToKey ChannelToKeyFromCurve = Proxy->GetChannelToKeyFromChannelName(ChannelName.GetValue());
 					const TViewModelPtr<ICurveEditorTreeItemExtension>& CurveEditorItem = OutlinerExtensionData.CurveEditorItem;
-					if (ChannelToKeyContext != ChannelToKeyFromCurve || ! CurveEditorItem.IsValid())
+					if (ChannelToKeyContext != ChannelToKeyFromCurve || !CurveEditorItem.IsValid())
 					{
 						continue;
 					}
 
 					const FCurveEditorTreeItemID CurveEditorTreeItem = CurveEditorItem->GetCurveEditorItemID();
-					if (CurveEditorTreeItem.IsValid())
+					if (!CurveEditorTreeItem.IsValid())
+					{
+						continue;
+					}
+
+					TArray<UMovieSceneSection*> SectionsToKey;
+					UAnimLayers* AnimLayers = UAnimLayers::GetAnimLayers(Sequencer.Get());
+					if (AnimLayers)
+					{
+						SectionsToKey = AnimLayers->GetSelectedLayerSections();
+					}
+					else
 					{
-						const UMovieSceneSection* SectionToKey = ControlRigParameterTrack->GetSectionToKey(ControlElement->GetFName()) ?
+						UMovieSceneSection* BaseSection = ControlRigParameterTrack->GetSectionToKey(ControlElement->GetFName()) ?
 							ControlRigParameterTrack->GetSectionToKey(ControlElement->GetFName()) :
 							ControlRigParameterTrack->GetAllSections()[0];
-									
+
+						SectionsToKey.Add(BaseSection);
+					}
+
+					for (UMovieSceneSection* SectionToKey : SectionsToKey)
+					{
 						// If there's no section to key we also don't select it
 						const TViewModelPtr<FChannelGroupOutlinerModel>& ChannelModel = OutlinerExtensionData.ChannelModel;
 						if (!ChannelModel.IsValid() || !ChannelModel->GetChannel(SectionToKey))
 						{
 							continue;
 						}
-	
+
 						const bool bPropertyVisible = PropertyData.IsVisible();
 						const bool bPropertySelected = PropertyData.IsSelected();
 
@@ -497,7 +514,7 @@ void UAnimDetailsSelection::PropagonateSelectionToCurveEditor()
 							CurveEditorTreeView->SetItemSelection(CurveEditorTreeItem, bSelected);
 						}
 						else if (const TSharedPtr<FCategoryGroupModel>& CategoryGroupModel = CastViewModel<FCategoryGroupModel>(ChannelModel->GetParent()))
-						{						
+						{
 							// If there is no selection in anim details, select the parent category instead
 							const FCurveEditorTreeItemID ParentCurveEditorTreeItem = CategoryGroupModel->GetCurveEditorItemID();
 							CurveEditorTreeView->SetItemSelection(ParentCurveEditorTreeItem, true);
