Anim in Engine: Fix selecting via Anim Details may instantiate the Anim Layers system and create a base layer, remove obsolete mentions of ControlRigEditModeSettings in AnimDetailsProxyManager

#rnx

diff --git a/Engine/Plugins/Animation/ControlRig/Source/ControlRigEditor/Private/AnimDetails/AnimDetailsProxyManager.cpp b/Engine/Plugins/Animation/ControlRig/Source/ControlRigEditor/Private/AnimDetails/AnimDetailsProxyManager.cpp
--- a/Engine/Plugins/Animation/ControlRig/Source/ControlRigEditor/Private/AnimDetails/AnimDetailsProxyManager.cpp
+++ b/Engine/Plugins/Animation/ControlRig/Source/ControlRigEditor/Private/AnimDetails/AnimDetailsProxyManager.cpp
@@ -16,7 +16,6 @@
 #include "AnimDetails/Proxies/AnimDetailsProxyVector2D.h"
 #include "ControlRig.h"
 #include "EditMode/ControlRigEditMode.h"
-#include "EditMode/ControlRigEditModeSettings.h"
 #include "EditorModeManager.h"
 #include "ISequencer.h"
 #include "LevelEditorViewport.h"
@@ -216,8 +215,6 @@ void UAnimDetailsProxyManager::ForceUpdateProxies()
 	{
 		if (!EditMode->AreEditingControlRigDirectly())
 		{
-			const UControlRigEditModeSettings* Settings = GetDefault<UControlRigEditModeSettings>();
-
 			TMap<UControlRig*, TArray<FRigElementKey>> ControlRigToSelectedKeysMap;
 			EditMode->GetAllSelectedControls(ControlRigToSelectedKeysMap);
 
diff --git a/Engine/Plugins/Animation/ControlRig/Source/ControlRigEditor/Private/AnimDetails/AnimDetailsSelection.cpp b/Engine/Plugins/Animation/ControlRig/Source/ControlRigEditor/Private/AnimDetails/AnimDetailsSelection.cpp
--- a/Engine/Plugins/Animation/ControlRig/Source/ControlRigEditor/Private/AnimDetails/AnimDetailsSelection.cpp
+++ b/Engine/Plugins/Animation/ControlRig/Source/ControlRigEditor/Private/AnimDetails/AnimDetailsSelection.cpp
@@ -481,7 +481,11 @@ void UAnimDetailsSelection::PropagonateSelectionToCurveEditor()
 					}
 
 					TArray<UMovieSceneSection*> SectionsToKey;
-					UAnimLayers* AnimLayers = UAnimLayers::GetAnimLayers(Sequencer.Get());
+
+					// Calling GetAnimLayers without testing HasAnimLayers will create Anim Layers in current version, avoid this
+					UAnimLayers* AnimLayers = UAnimLayers::HasAnimLayers(Sequencer.Get()) ? 
+						UAnimLayers::GetAnimLayers(Sequencer.Get()) : 
+						nullptr;
 					if (AnimLayers)
 					{
 						SectionsToKey = AnimLayers->GetSelectedLayerSections();
