Anim in Engine: Fix a regression where Anim Details did not forward its selection to Curve Editor when Anim Layers were instantiated, but the base layer was never created

#rnx

diff --git a/Engine/Plugins/Animation/ControlRig/Source/ControlRigEditor/Private/AnimDetails/AnimDetailsSelection.cpp b/Engine/Plugins/Animation/ControlRig/Source/ControlRigEditor/Private/AnimDetails/AnimDetailsSelection.cpp
--- a/Engine/Plugins/Animation/ControlRig/Source/ControlRigEditor/Private/AnimDetails/AnimDetailsSelection.cpp
+++ b/Engine/Plugins/Animation/ControlRig/Source/ControlRigEditor/Private/AnimDetails/AnimDetailsSelection.cpp
@@ -490,14 +490,13 @@ void UAnimDetailsSelection::PropagonateSelectionToCurveEditor()
 					{
 						SectionsToKey = AnimLayers->GetSelectedLayerSections();
 					}
-					else
-					{
-						UMovieSceneSection* BaseSection = ControlRigParameterTrack->GetSectionToKey(ControlElement->GetFName()) ?
-							ControlRigParameterTrack->GetSectionToKey(ControlElement->GetFName()) :
-							ControlRigParameterTrack->GetAllSections()[0];
 
-						SectionsToKey.Add(BaseSection);
-					}
+					// In any case add the base section, layers may exist but not hold any sections
+					UMovieSceneSection* BaseSection = ControlRigParameterTrack->GetSectionToKey(ControlElement->GetFName()) ?
+						ControlRigParameterTrack->GetSectionToKey(ControlElement->GetFName()) :
+						ControlRigParameterTrack->GetAllSections()[0];
+
+					SectionsToKey.AddUnique(BaseSection);
 
 					for (UMovieSceneSection* SectionToKey : SectionsToKey)
 					{
