Anim in Engine: Anim Details selection no longer selects the base layer when another anim layer is selected

#rnx

diff --git a/Engine/Plugins/Animation/ControlRig/Source/ControlRigEditor/Private/AnimDetails/AnimDetailsSelection.cpp b/Engine/Plugins/Animation/ControlRig/Source/ControlRigEditor/Private/AnimDetails/AnimDetailsSelection.cpp
--- a/Engine/Plugins/Animation/ControlRig/Source/ControlRigEditor/Private/AnimDetails/AnimDetailsSelection.cpp
+++ b/Engine/Plugins/Animation/ControlRig/Source/ControlRigEditor/Private/AnimDetails/AnimDetailsSelection.cpp
@@ -491,12 +491,15 @@ void UAnimDetailsSelection::PropagonateSelectionToCurveEditor()
 						SectionsToKey = AnimLayers->GetSelectedLayerSections();
 					}
 
-					// In any case add the base section, layers may exist but not hold any sections
-					UMovieSceneSection* BaseSection = ControlRigParameterTrack->GetSectionToKey(ControlElement->GetFName()) ?
-						ControlRigParameterTrack->GetSectionToKey(ControlElement->GetFName()) :
-						ControlRigParameterTrack->GetAllSections()[0];
+					// If nothing is selected in anim layers, select in the base layer
+					if (SectionsToKey.IsEmpty())
+					{
+						UMovieSceneSection* BaseSection = ControlRigParameterTrack->GetSectionToKey(ControlElement->GetFName()) ?
+							ControlRigParameterTrack->GetSectionToKey(ControlElement->GetFName()) :
+							ControlRigParameterTrack->GetAllSections()[0];
 
-					SectionsToKey.AddUnique(BaseSection);
+						SectionsToKey.Add(BaseSection);
+					}
 
 					for (UMovieSceneSection* SectionToKey : SectionsToKey)
 					{
