Anim in Engine: Keying icons no longer consider anim layers that are not currently selected/visible

#rnx

diff --git a/Engine/Plugins/Animation/ControlRig/Source/ControlRigEditor/Private/AnimDetails/Proxies/AnimDetailsProxyBase.cpp b/Engine/Plugins/Animation/ControlRig/Source/ControlRigEditor/Private/AnimDetails/Proxies/AnimDetailsProxyBase.cpp
--- a/Engine/Plugins/Animation/ControlRig/Source/ControlRigEditor/Private/AnimDetails/Proxies/AnimDetailsProxyBase.cpp
+++ b/Engine/Plugins/Animation/ControlRig/Source/ControlRigEditor/Private/AnimDetails/Proxies/AnimDetailsProxyBase.cpp
@@ -22,6 +22,7 @@
 #include "Rigs/RigHierarchyDefines.h"
 #include "ScopedTransaction.h"
 #include "Sequencer/ControlRigSequencerHelpers.h"
+#include "Sequencer/AnimLayers/AnimLayers.h"
 #include "Sequencer/MovieSceneControlRigParameterSection.h"
 #include "Sequencer/MovieSceneControlRigParameterTrack.h"
 #include "SequencerAddKeyOperation.h"
@@ -390,6 +391,7 @@ namespace UE::ControlRigEditor
 		}
 
 		static EPropertyKeyedStatus GetKeyedStatusInTrack(
+			ISequencer& Sequencer,
 			const UControlRig* ControlRig, 
 			const FName& ControlName, 
 			const UMovieSceneControlRigParameterTrack* Track, 
@@ -404,14 +406,35 @@ namespace UE::ControlRigEditor
 				return SectionKeyedStatus;
 			}
 
-			for (const UMovieSceneSection* BaseSection : Track->GetAllSections())
+			// Calling GetAnimLayers without testing HasAnimLayers will create Anim Layers in current version, avoid this
+			TArray<UMovieSceneSection*> SectionsToKey;
+			UAnimLayers* AnimLayers = UAnimLayers::HasAnimLayers(&Sequencer) ?
+				UAnimLayers::GetAnimLayers(&Sequencer) :
+				nullptr;
+			if (AnimLayers)
 			{
-				const UMovieSceneControlRigParameterSection* Section = Cast<UMovieSceneControlRigParameterSection>(BaseSection);
+				SectionsToKey = AnimLayers->GetSelectedLayerSections();
+			}
+
+			// If there's no layer selected, evaluate for the base layer only
+			if (SectionsToKey.IsEmpty())
+			{
+				UMovieSceneSection* BaseSection = Track->GetSectionToKey(ControlElement->GetFName()) ?
+					Track->GetSectionToKey(ControlElement->GetFName()) :
+					Track->GetAllSections()[0];
+
+				SectionsToKey.AddUnique(BaseSection);
+			}
+
+			for (const UMovieSceneSection* Section : SectionsToKey)
+			{
+				const UMovieSceneControlRigParameterSection* ParameterSection = Cast<UMovieSceneControlRigParameterSection>(Section);
 				if (!Section)
 				{
 					continue;
 				}
-				EPropertyKeyedStatus NewSectionKeyedStatus = GetKeyedStatusInSection(ControlRig, ControlName, Section, Range, ChannelToKey);
+
+				const EPropertyKeyedStatus NewSectionKeyedStatus = GetKeyedStatusInSection(ControlRig, ControlName, ParameterSection, Range, ChannelToKey);
 				SectionKeyedStatus = FMath::Max(SectionKeyedStatus, NewSectionKeyedStatus);
 
 				// Maximum Status Reached no need to iterate further
@@ -539,6 +562,7 @@ namespace UE::ControlRigEditor
 		}
 
 		static EPropertyKeyedStatus GetKeyedStatusInTrack(
+			ISequencer& Sequencer,
 			const UMovieScenePropertyTrack* Track, 
 			const TRange<FFrameNumber>& Range, 
 			const EControlRigContextChannelToKey ChannelToKey, 
@@ -551,7 +575,7 @@ namespace UE::ControlRigEditor
 				{
 					continue;
 				}
-				EPropertyKeyedStatus NewSectionKeyedStatus = GetKeyedStatusInSection(BaseSection, Range, ChannelToKey, MaxNumIndices);
+				const EPropertyKeyedStatus NewSectionKeyedStatus = GetKeyedStatusInSection(BaseSection, Range, ChannelToKey, MaxNumIndices);
 				SectionKeyedStatus = FMath::Max(SectionKeyedStatus, NewSectionKeyedStatus);
 
 				// Maximum Status Reached no need to iterate further
@@ -739,7 +763,7 @@ EPropertyKeyedStatus UAnimDetailsProxyBase::GetPropertyKeyedStatus(const IProper
 	UMovieSceneControlRigParameterTrack* ControlRigParameterTrack = ControlRig ? FControlRigSequencerHelpers::FindControlRigTrack(Sequencer->GetFocusedMovieSceneSequence(), ControlRig) : nullptr;
 	if (ControlRig && ControlElementKey.IsValid() && ControlRigParameterTrack)
 	{
-		const EPropertyKeyedStatus NewKeyedStatus = KeyUtils::GetKeyedStatusInTrack(ControlRig, ControlElementKey.Name, ControlRigParameterTrack, FrameRange, ChannelToKey);
+		const EPropertyKeyedStatus NewKeyedStatus = KeyUtils::GetKeyedStatusInTrack(*Sequencer.Get(), ControlRig, ControlElementKey.Name, ControlRigParameterTrack, FrameRange, ChannelToKey);
 		KeyedStatus = FMath::Max(KeyedStatus, NewKeyedStatus);
 	}
 	else if (UMovieScenePropertyTrack* MovieScenePropertyTrack = Cast<UMovieScenePropertyTrack>(SequencerItem.GetMovieSceneTrack()))
@@ -754,7 +778,7 @@ EPropertyKeyedStatus UAnimDetailsProxyBase::GetPropertyKeyedStatus(const IProper
 			MaxNumIndices = 2;
 		}
 
-		const EPropertyKeyedStatus NewKeyedStatus = KeyUtils::GetKeyedStatusInTrack(MovieScenePropertyTrack, FrameRange, ChannelToKey, MaxNumIndices);
+		const EPropertyKeyedStatus NewKeyedStatus = KeyUtils::GetKeyedStatusInTrack(*Sequencer.Get(), MovieScenePropertyTrack, FrameRange, ChannelToKey, MaxNumIndices);
 		KeyedStatus = FMath::Max(KeyedStatus, NewKeyedStatus);
 	}
 
