'Bake To Control Rig' from inside an Animation Sequence

#rnx

#virtualized

diff --git a/Engine/Plugins/Animation/ControlRig/Source/ControlRigEditor/Private/ControlRigEditorModule.h b/Engine/Plugins/Animation/ControlRig/Source/ControlRigEditor/Private/ControlRigEditorModule.h
--- a/Engine/Plugins/Animation/ControlRig/Source/ControlRigEditor/Private/ControlRigEditorModule.h
+++ b/Engine/Plugins/Animation/ControlRig/Source/ControlRigEditor/Private/ControlRigEditorModule.h
@@ -47,7 +47,7 @@ public:
 	void ToggleIsDrivenByLevelSequence(UAnimSequence* AnimSequence)const;
 	bool IsDrivenByLevelSequence(UAnimSequence* AnimSequence)const;
 	void EditWithFKControlRig(UAnimSequence* AnimSequence, USkeletalMesh* SkelMesh, USkeleton* InSkeleton);
-	void BakeToControlRig(UClass* InClass,UAnimSequence* AnimSequence, USkeletalMesh* SkelMesh,USkeleton* InSkeleton);
+	CONTROLRIGEDITOR_API void BakeToControlRig(UClass* InClass,UAnimSequence* AnimSequence, USkeletalMesh* SkelMesh,USkeleton* InSkeleton);
 	static void OpenLevelSequence(UAnimSequence* AnimSequence);
 	static void UnLinkLevelSequence(UAnimSequence* AnimSequence);
 	void ExtendAnimSequenceMenu();
