Fix `UTransformableControlHandle::OnControlModified` to work under transactions.

diff --git a/Engine/Plugins/Animation/ControlRig/Source/ControlRig/Private/Constraints/ControlRigTransformableHandle.cpp b/Engine/Plugins/Animation/ControlRig/Source/ControlRig/Private/Constraints/ControlRigTransformableHandle.cpp
--- a/Engine/Plugins/Animation/ControlRig/Source/ControlRig/Private/Constraints/ControlRigTransformableHandle.cpp
+++ b/Engine/Plugins/Animation/ControlRig/Source/ControlRig/Private/Constraints/ControlRigTransformableHandle.cpp
@@ -8,6 +8,7 @@
 #include "ControlRigComponent.h"
 #include "IControlRigObjectBinding.h"
 #include "ControlRigObjectBinding.h"
+#include "Misc/TransactionallySafeCriticalSection.h"
 #include "ModularRig.h"
 #include "Transform/TransformableHandleUtils.h"
 #include "Rigs/RigHierarchyElements.h"
@@ -30,7 +31,7 @@ namespace UE::Private
 	}
 
 	// only allow notification from one thread at a time
-	FCriticalSection NotificationLock;
+	FTransactionallySafeCriticalSection NotificationLock;
 	
 	struct FControlPoseChangedNotifier
 	{
@@ -668,7 +669,7 @@ void UTransformableControlHandle::OnControlModified(
 
 		if (InControl->GetFName() == ControlName)
 		{
-			FScopeLock ScopeLock(&UE::Private::NotificationLock);
+			UE::TScopeLock ScopeLock(UE::Private::NotificationLock);
 			
 			// if that handle is wrapping InControl
 			if (InContext.bConstraintUpdate)
@@ -695,7 +696,7 @@ void UTransformableControlHandle::OnControlModified(
 			// the control being modified is not the one wrapped by this handle 
 			if (const FRigControlElement* Control = ControlRig->FindControl(ControlName))
 			{
-				FScopeLock ScopeLock(&UE::Private::NotificationLock);
+				UE::TScopeLock ScopeLock(UE::Private::NotificationLock);
 			
 				if (InContext.bConstraintUpdate)
 				{
