Anim In Engine:  Added a 'Tolerance' To the ControlRigPoseMirrorSettings Project settings that is used to match Controls also by position in rest pose. We still match by the Mirror table first, and fallback to position based matching. We also know cache the tables and reset them if a mirror table property is changed.

diff --git a/Engine/Plugins/Animation/ControlRig/Source/ControlRig/Private/Tools/ControlRigPose.cpp b/Engine/Plugins/Animation/ControlRig/Source/ControlRig/Private/Tools/ControlRigPose.cpp
--- a/Engine/Plugins/Animation/ControlRig/Source/ControlRig/Private/Tools/ControlRigPose.cpp
+++ b/Engine/Plugins/Animation/ControlRig/Source/ControlRig/Private/Tools/ControlRigPose.cpp
@@ -77,11 +77,7 @@ void FControlRigControlPose::PastePoseInternal(UControlRig* ControlRig, bool bDo
 	FRigControlModifiedContext Context;
 	Context.SetKey = bDoKey ? EControlRigSetKey::Always : EControlRigSetKey::DoNotCare;
 	FControlRigPoseMirrorTable MirrorTable;
-	if (bDoMirror)
-	{
-		MirrorTable.SetUpMirrorTable(ControlRig);
-	}
-
+	
 	TArray<FRigControlElement*> SortedControls;
 	ControlRig->GetControlsInOrder(SortedControls);
 	
@@ -95,7 +91,7 @@ void FControlRigControlPose::PastePoseInternal(UControlRig* ControlRig, bool bDo
 			continue;
 		}
 		
-		FRigControlCopy* CopyRigControl = MirrorTable.GetControl(*this, ControlElement->GetFName());
+		FRigControlCopy* CopyRigControl = MirrorTable.GetControl(ControlRig, *this, ControlElement->GetFName(), bDoMirror);
 		if (CopyRigControl)
 		{
 			//if not in default parent space we need to move it to default parent space first and then reset the global transforms
@@ -132,7 +128,7 @@ void FControlRigControlPose::PastePoseInternal(UControlRig* ControlRig, bool bDo
 				{
 					FTransform GlobalTransform;
 					FTransform LocalTransform;
-					bool bIsMatched = MirrorTable.IsMatched(CopyRigControl->Name);
+					bool bIsMatched = MirrorTable.IsMatched(ControlRig, CopyRigControl->Name);
 					MirrorTable.GetMirrorTransform(*CopyRigControl, bDoLocal,bIsMatched,GlobalTransform,LocalTransform);
 					SetControlMirrorTransform(bDoLocal,ControlRig, ControlElement->GetFName(), bIsMatched, GlobalTransform,LocalTransform,true,Context,bSetupUndo);
 				}				
@@ -214,10 +210,6 @@ void FControlRigControlPose::BlendWithInitialPoses(FControlRigControlPose& Initi
 	FRigControlModifiedContext Context;
 	Context.SetKey = bDoKey ? EControlRigSetKey::Always : EControlRigSetKey::DoNotCare;
 	FControlRigPoseMirrorTable MirrorTable;
-	if (bDoMirror)
-	{
-		MirrorTable.SetUpMirrorTable(ControlRig);
-	}
 
 	TArray<FRigControlElement*> SortedControls;
 	ControlRig->GetControlsInOrder(SortedControls);
@@ -230,7 +222,7 @@ void FControlRigControlPose::BlendWithInitialPoses(FControlRigControlPose& Initi
 		{
 			continue;
 		}
-		FRigControlCopy* CopyRigControl = MirrorTable.GetControl(*this, ControlElement->GetFName());
+		FRigControlCopy* CopyRigControl = MirrorTable.GetControl(ControlRig,*this, ControlElement->GetFName(),bDoMirror);
 		if (CopyRigControl)
 		{
 			FRigControlCopy* InitialFound = nullptr;
@@ -276,7 +268,7 @@ void FControlRigControlPose::BlendWithInitialPoses(FControlRigControlPose& Initi
 					{
 						FTransform GlobalTransform;
 						FTransform LocalTransform;
-						bool bIsMatched = MirrorTable.IsMatched(CopyRigControl->Name);
+						bool bIsMatched = MirrorTable.IsMatched(ControlRig, CopyRigControl->Name);
 						MirrorTable.GetMirrorTransform(*CopyRigControl, bDoLocal, bIsMatched, GlobalTransform, LocalTransform);
 						
 						BlendTransforms(InitialFound->GlobalTransform, GlobalTransform);
@@ -403,14 +395,13 @@ void UControlRigPoseAsset::SelectControls(UControlRig* InControlRig, bool bDoMir
 	FControlRigControlPose TempPose;
 	if (bDoMirror)
 	{
-		MirrorTable.SetUpMirrorTable(InControlRig);
 		TempPose.SavePose(InControlRig, true);
 	}
 	for (const FName& Name : Controls)
 	{
 		if (bDoMirror)
 		{
-			FRigControlCopy* CopyRigControl = MirrorTable.GetControl(TempPose, Name);
+			FRigControlCopy* CopyRigControl = MirrorTable.GetControl(InControlRig, TempPose, Name, bDoMirror);
 			if (CopyRigControl)
 			{
 				InControlRig->SelectControl(CopyRigControl->Name, true);
@@ -455,14 +446,9 @@ void UControlRigPoseAsset::ReplaceControlName(const FName& CurrentName, const FN
 	Pose.ReplaceControlName(CurrentName, NewName);
 }
 
-void UControlRigPoseAsset::SetUpMirrorMatchTable(UControlRig* InControlRig)
-{
-	MirrorMatchTable.SetUpMirrorTable(InControlRig);
-}
-
-bool UControlRigPoseAsset::DoesMirrorMatch(UControlRig* ControlRig, const FName& ControlName) const
+bool UControlRigPoseAsset::DoesMirrorMatch(UControlRig* ControlRig, const FName& ControlName) 
 {
-	return (MirrorMatchTable.IsMatched(ControlName));
+	return (MirrorMatchTable.IsMatched(ControlRig,ControlName));
 }
 
 
diff --git a/Engine/Plugins/Animation/ControlRig/Source/ControlRig/Private/Tools/ControlRigPoseMirrorSettings.cpp b/Engine/Plugins/Animation/ControlRig/Source/ControlRig/Private/Tools/ControlRigPoseMirrorSettings.cpp
--- a/Engine/Plugins/Animation/ControlRig/Source/ControlRig/Private/Tools/ControlRigPoseMirrorSettings.cpp
+++ b/Engine/Plugins/Animation/ControlRig/Source/ControlRig/Private/Tools/ControlRigPoseMirrorSettings.cpp
@@ -1,11 +1,23 @@
 // Copyright Epic Games, Inc. All Rights Reserved.
 #include "Tools/ControlRigPoseMirrorSettings.h"
 
+#include "Tools/ControlRigPoseMirrorTable.h"
+
 #include UE_INLINE_GENERATED_CPP_BY_NAME(ControlRigPoseMirrorSettings)
 
 UControlRigPoseMirrorSettings::UControlRigPoseMirrorSettings()
 {
 	MirrorAxis = EAxis::X;
 	AxisToFlip = EAxis::X;
+	MirrorMatchTolerance = .2;
 }
 
+#if WITH_EDITOR
+//property changed so reset the cached mirror tables
+void UControlRigPoseMirrorSettings::PostEditChangeProperty(struct FPropertyChangedEvent& PropertyChangedEvent)
+{
+	FControlRigPoseMirrorTable MirrorTable;
+	MirrorTable.Reset();
+	SaveConfig();
+}
+#endif
diff --git a/Engine/Plugins/Animation/ControlRig/Source/ControlRig/Private/Tools/ControlRigPoseMirrorTable.cpp b/Engine/Plugins/Animation/ControlRig/Source/ControlRig/Private/Tools/ControlRigPoseMirrorTable.cpp
--- a/Engine/Plugins/Animation/ControlRig/Source/ControlRig/Private/Tools/ControlRigPoseMirrorTable.cpp
+++ b/Engine/Plugins/Animation/ControlRig/Source/ControlRig/Private/Tools/ControlRigPoseMirrorTable.cpp
@@ -4,11 +4,60 @@
 #include "Tools/ControlRigPoseMirrorSettings.h"
 #include "ControlRig.h"
 #include "Tools/ControlRigPose.h"
+#include "Tools/MirrorCalculator.h"
 #include "RigVMFunctions/Math/RigVMMathLibrary.h"
 #include "Animation/AnimationSettings.h"
 #include "Animation/MirrorDataTable.h"
 #include "Internationalization/Regex.h"
 
+
+struct FControlRigToMirrorItems
+{
+	FControlRigToMirrorItems(const UControlRig* InControlRig);
+	FControlRigToMirrorItems() = delete;
+
+	const UControlRig* ControlRig;
+	TArray<FRigControlElement*> Elements;
+	TArray<UE::AIE::FMirrorItem> MirrorItems;
+	UE::AIE::FMirrorItemResults Results;
+	UE::AIE::FMirrorCalculator Calculator;
+};
+
+
+FControlRigToMirrorItems::FControlRigToMirrorItems(const UControlRig* InControlRig)
+{
+	using namespace UE::AIE;
+	ControlRig = InControlRig;
+	if (ControlRig)
+	{
+		if (URigHierarchy* Hierarchy = ControlRig->GetHierarchy())
+		{
+			if (const UControlRigPoseMirrorSettings* Settings = GetDefault<UControlRigPoseMirrorSettings>())
+			{
+				FTransform GlobalTransform;
+				TArray<FRigControlElement*> CurrentControls = ControlRig->AvailableControls();
+				for (FRigControlElement* ControlElement : CurrentControls)
+				{
+					if (Hierarchy->IsAnimatable(ControlElement))
+					{
+						GlobalTransform = Hierarchy->GetTransform(ControlElement, ERigTransformType::InitialGlobal);
+
+						FMirrorItem MirrorItem(GlobalTransform.GetLocation());
+						MirrorItems.Add(MirrorItem);
+						Elements.Add(ControlElement);
+
+					}
+				}
+
+				TEnumAsByte<EAxis::Type> MirrorAxis = Settings->MirrorAxis;
+				constexpr double AxisLocation = 0.0;
+				const double Tolerance = Settings->MirrorMatchTolerance;
+				Calculator.FindMirroredItems(MirrorItems, Results, MirrorAxis, AxisLocation, Tolerance);
+			}
+		}
+	}
+}
+
 namespace UE::ControlRigPoseMirrorTable
 {
 FString GetMirrorString(const FString& InNameString, const TArray<FMirrorFindReplaceExpression>& MirrorFindReplaceExpressions)
@@ -61,11 +110,21 @@ FString GetMirrorString(const FString& InNameString, const TArray<FMirrorFindRep
 	return FString();
 }
 }
+
+TMap<const UControlRig*, FControlRigPoseMirrorTable::FMatchedControls> FControlRigPoseMirrorTable::MatchedControls;
+
+void FControlRigPoseMirrorTable::Reset()
+{
+	FControlRigPoseMirrorTable::MatchedControls.Reset();
+}
+
 void FControlRigPoseMirrorTable::SetUpMirrorTable(const UControlRig* ControlRig)
 {
-	MatchedControls.Reset();
-	if (ControlRig)
+	FMatchedControls* CRMatchedControls = FControlRigPoseMirrorTable::MatchedControls.Find(ControlRig);
+	if (CRMatchedControls == nullptr || CRMatchedControls->Match.Num() == 0)
 	{
+		FMatchedControls CurrentMatchedControls;
+		//first try by name matching
 		TArray<FRigControlElement*> CurrentControls = ControlRig->AvailableControls();
 		if (UAnimationSettings* AnimSettings = UAnimationSettings::Get())
 		{
@@ -77,45 +136,68 @@ void FControlRigPoseMirrorTable::SetUpMirrorTable(const UControlRig* ControlRig)
 				{
 					FName CurrentName(*CurrentString);
 					FName NewName(*NewString);
-					MatchedControls.Add(NewName, CurrentName);
+					CurrentMatchedControls.Match.Add(NewName, CurrentName);
+				}
+			}
+		}
+		FControlRigToMirrorItems MirrorItems(ControlRig);
+		if (MirrorItems.Results.MirroredItems.Num() > 0)
+		{
+			for (TPair<int32, int32>& Pair : MirrorItems.Results.MirroredItems)
+			{
+				FName CurrentName(MirrorItems.Elements[Pair.Key]->GetFName());
+				FName NewName(MirrorItems.Elements[Pair.Value]->GetFName());
+				if (CurrentMatchedControls.Match.Contains(CurrentName) == false)
+				{
+					CurrentMatchedControls.Match.Add(CurrentName, NewName);
 				}
 			}
 		}
+		FControlRigPoseMirrorTable::MatchedControls.Add(ControlRig, CurrentMatchedControls);
 	}
 }
 
-FRigControlCopy* FControlRigPoseMirrorTable::GetControl(FControlRigControlPose& Pose, FName Name)
+FRigControlCopy* FControlRigPoseMirrorTable::GetControl(const UControlRig* ControlRig, FControlRigControlPose& Pose, FName Name, bool bDoMirror)
 {
-	TArray<FRigControlCopy> CopyOfControls;
-	
-	if (MatchedControls.Num() > 0) 
+	if (bDoMirror)
 	{
-		if (const FName* MatchedName = MatchedControls.Find(Name))
+		SetUpMirrorTable(ControlRig);
+
+		if (FMatchedControls* CRMatchedControls = FControlRigPoseMirrorTable::MatchedControls.Find(ControlRig))
 		{
-			int32* Index = Pose.CopyOfControlsNameToIndex.Find(*MatchedName);
-			if (Index != nullptr && (*Index) >= 0 && (*Index) < Pose.CopyOfControls.Num())
+			TArray<FRigControlCopy> CopyOfControls;
+
+			if (const FName* MatchedName = CRMatchedControls->Match.Find(Name))
 			{
-				return &(Pose.CopyOfControls[*Index]);
+				int32* Index = Pose.CopyOfControlsNameToIndex.Find(*MatchedName);
+				if (Index != nullptr && (*Index) >= 0 && (*Index) < Pose.CopyOfControls.Num())
+				{
+					return &(Pose.CopyOfControls[*Index]);
+				}
 			}
 		}
 	}
-	//okay not matched so just find it.
+	//okay not mirror or matched so just find it.
 	int32* Index = Pose.CopyOfControlsNameToIndex.Find(Name);
 	if (Index != nullptr && (*Index) >= 0 && (*Index) < Pose.CopyOfControls.Num())
 	{
 		return &(Pose.CopyOfControls[*Index]);
 	}
-
 	return nullptr;
 }
 
-bool FControlRigPoseMirrorTable::IsMatched(const FName& Name) const
+bool FControlRigPoseMirrorTable::IsMatched(const UControlRig* ControlRig,const FName& Name) 
 {
-	if (MatchedControls.Num() > 0)
+	SetUpMirrorTable(ControlRig);
+
+	if (FMatchedControls* CRMatchedControls = FControlRigPoseMirrorTable::MatchedControls.Find(ControlRig))
 	{
-		if (const FName* MatchedName = MatchedControls.Find(Name))
+		if (CRMatchedControls->Match.Num() > 0)
 		{
-			return true;
+			if (const FName* MatchedName = CRMatchedControls->Match.Find(Name))
+			{
+				return true;
+			}
 		}
 	}
 	return false;
@@ -148,3 +230,4 @@ void FControlRigPoseMirrorTable::GetMirrorTransform(const FRigControlCopy& Contr
 		}
 	}
 }
+
diff --git a/Engine/Plugins/Animation/ControlRig/Source/ControlRig/Private/Tools/MirrorCalculator.cpp b/Engine/Plugins/Animation/ControlRig/Source/ControlRig/Private/Tools/MirrorCalculator.cpp
new file mode 100644
--- /dev/null
+++ b/Engine/Plugins/Animation/ControlRig/Source/ControlRig/Private/Tools/MirrorCalculator.cpp
@@ -0,0 +1,78 @@
+// Copyright Epic Games, Inc. All Rights Reserved.
+
+
+#include "Tools/MirrorCalculator.h"
+
+using namespace UE::AIE;
+
+bool FMirrorCalculator::FindMirroredItems(const TArray<FMirrorItem>& Items, FMirrorItemResults& OutMirrorItemResults,
+	TEnumAsByte<EAxis::Type> MirrorAxis, double AxisLocation, double Tolerance)
+{
+
+	int32 MirrorIndex = 0;
+	int32 OtherIndex[2] = { 1,2 };
+	if (MirrorAxis == EAxis::Y)
+	{
+		MirrorIndex = 1;
+		OtherIndex[0] = 0;
+		OtherIndex[1] = 2;
+	}
+	else if (MirrorAxis == EAxis::Z)
+	{
+		MirrorIndex = 2;
+		OtherIndex[0] = 0;
+		OtherIndex[1] = 1;
+	}
+
+	//items found on first pass less than and greater than the Axis Location, those on the Axis will go straight into
+	TArray<int32> ItemsLess;
+	TArray<int32> ItemsMore;
+	
+	for (int32 Index = 0; Index < Items.Num(); ++Index)
+	{
+		const FMirrorItem& Item = Items[Index];
+		if (FMath::IsNearlyZero(Item.Location[MirrorIndex] - AxisLocation, Tolerance))
+		{
+			OutMirrorItemResults.ExactlyOnAxis.Add(Index);
+		}
+		else if (Item.Location[MirrorIndex] < AxisLocation)
+		{
+			ItemsLess.Add(Index);
+		}
+		else // > AxisLocation
+		{
+			ItemsMore.Add(Index);
+		}
+	}
+
+	for (const int32 LessIndex : ItemsLess)
+	{
+
+		const FMirrorItem& LessItem = Items[LessIndex];
+		bool bWasFound = false;
+		for (const int32 MoreIndex : ItemsMore)
+		{
+			const FMirrorItem& MoreItem = Items[MoreIndex];
+				
+			if (IsMirror(LessItem, MoreItem, AxisLocation, Tolerance, MirrorIndex, OtherIndex))
+			{
+				OutMirrorItemResults.MirroredItems.Add(LessIndex, MoreIndex);
+				OutMirrorItemResults.MirroredItems.Add(MoreIndex, LessIndex);
+				ItemsMore.Remove(MoreIndex);
+				bWasFound = true;
+				break;
+			}
+		}
+		if (bWasFound == false)
+		{
+			OutMirrorItemResults.NotOnAxis.Add(LessIndex);
+		}
+	}
+	//if any MoreItems left they weren't mirrored 
+	for (int32 MoreIndex : ItemsMore)
+	{
+		OutMirrorItemResults.NotOnAxis.Add(MoreIndex);
+	}
+	return (OutMirrorItemResults.MirroredItems.Num() > 0);
+	
+}
diff --git a/Engine/Plugins/Animation/ControlRig/Source/ControlRig/Public/Tools/ControlRigPose.h b/Engine/Plugins/Animation/ControlRig/Source/ControlRig/Public/Tools/ControlRigPose.h
--- a/Engine/Plugins/Animation/ControlRig/Source/ControlRig/Public/Tools/ControlRigPose.h
+++ b/Engine/Plugins/Animation/ControlRig/Source/ControlRig/Public/Tools/ControlRigPose.h
@@ -147,12 +147,13 @@ public:
 	UFUNCTION(BlueprintCallable, Category = "Pose")
 	UE_API void ReplaceControlName(const FName& CurrentName, const FName& NewName);
 
-	/* Set up an internal mirror table to be used by subsequent calls to DoesMirrorMatch,  need to call for each Control Rig*/
+	UE_DEPRECATED(5.7, "Function is no longer needed.")
 	UFUNCTION(BlueprintCallable, Category = "Pose")
-	UE_API void SetUpMirrorMatchTable(UControlRig* InControlRig);
+	UE_API void SetUpMirrorMatchTable(UControlRig* InControlRig){};
+
 
 	UFUNCTION(BlueprintPure, Category = "Pose")
-	UE_API bool DoesMirrorMatch(UControlRig* ControlRig, const FName& ControlName) const;
+	UE_API bool DoesMirrorMatch(UControlRig* ControlRig, const FName& ControlName);
 
 	UE_API void BlendWithInitialPoses(FControlRigControlPose& InitialPose, UControlRig* ControlRig, bool bDoKey, bool bdoMirror, float BlendValue, bool bDoAdditive = false);
 
diff --git a/Engine/Plugins/Animation/ControlRig/Source/ControlRig/Public/Tools/ControlRigPoseMirrorSettings.h b/Engine/Plugins/Animation/ControlRig/Source/ControlRig/Public/Tools/ControlRigPoseMirrorSettings.h
--- a/Engine/Plugins/Animation/ControlRig/Source/ControlRig/Public/Tools/ControlRigPoseMirrorSettings.h
+++ b/Engine/Plugins/Animation/ControlRig/Source/ControlRig/Public/Tools/ControlRigPoseMirrorSettings.h
@@ -37,6 +37,12 @@ public:
 	UPROPERTY(config, EditAnywhere, BlueprintReadWrite, Category = "Mirror Settings", meta = (ToolTip = "Axis to Flip for Rotations"))
 	TEnumAsByte<EAxis::Type> AxisToFlip;
 
+	// Tolerance to find the matching mirror control
+	UPROPERTY(config, EditAnywhere, BlueprintReadWrite, Category = "Mirror Settings", meta = (DisplayName = "Tolerance", ToolTip = "Tolerance to find the matching mirror control"))
+	double MirrorMatchTolerance;
+#if WITH_EDITOR
+	virtual void PostEditChangeProperty(struct FPropertyChangedEvent& PropertyChangedEvent) override;
+#endif
 };
 
 
diff --git a/Engine/Plugins/Animation/ControlRig/Source/ControlRig/Public/Tools/ControlRigPoseMirrorTable.h b/Engine/Plugins/Animation/ControlRig/Source/ControlRig/Public/Tools/ControlRigPoseMirrorTable.h
--- a/Engine/Plugins/Animation/ControlRig/Source/ControlRig/Public/Tools/ControlRigPoseMirrorTable.h
+++ b/Engine/Plugins/Animation/ControlRig/Source/ControlRig/Public/Tools/ControlRigPoseMirrorTable.h
@@ -9,32 +9,39 @@
 #include "CoreMinimal.h"
 #include "UObject/ObjectMacros.h"
 #include "UObject/Object.h"
+#include "Tools/MirrorCalculator.h"
 
 class UControlRig;
 struct FRigControl;
 struct FRigControlCopy;
 struct FControlRigControlPose;
 
+
 struct FControlRigPoseMirrorTable
 {
 public:
 	FControlRigPoseMirrorTable() {};
 	~FControlRigPoseMirrorTable() {};
 
-	/*Set up the Mirror Table*/
-	void SetUpMirrorTable(const UControlRig* ControlRig);
-
 	/*Get the matched control with the given name*/
-	FRigControlCopy* GetControl(FControlRigControlPose& Pose, FName ControlrigName);
+	FRigControlCopy* GetControl(const UControlRig* ControlRig, FControlRigControlPose& Pose, FName ControlrigName, bool bDoMiror);
 
 	/*Whether or not the Control with this name is matched*/
-	bool IsMatched(const FName& ControlName) const;
+	bool IsMatched(const UControlRig* ControlRig,const FName& ControlName);
 
 	/*Return the Mirrored Global(In Control Rig Space) Translation and Mirrored Local Rotation*/
 	void GetMirrorTransform(const FRigControlCopy& ControlCopy, bool bDoLocal, bool bIsMatched,FTransform& OutGlobalTransform, FTransform& OutLocalTransform) const;
 
+	void Reset();
 private:
+	/*Set up the Mirror Table*/
+	void SetUpMirrorTable(const UControlRig* ControlRig);
 
-	TMap<FName, FName>  MatchedControls;
+	struct FMatchedControls
+	{
+		TMap<FName, FName> Match;
+	};
+	//cache of matches, note raw pointer is fine for UControlRig, just used for lookup not accessed.
+	static TMap<const UControlRig*,FMatchedControls>  MatchedControls;
 
 };
\ No newline at end of file
diff --git a/Engine/Plugins/Animation/ControlRig/Source/ControlRig/Public/Tools/MirrorCalculator.h b/Engine/Plugins/Animation/ControlRig/Source/ControlRig/Public/Tools/MirrorCalculator.h
new file mode 100644
--- /dev/null
+++ b/Engine/Plugins/Animation/ControlRig/Source/ControlRig/Public/Tools/MirrorCalculator.h
@@ -0,0 +1,69 @@
+// Copyright Epic Games, Inc. All Rights Reserved.
+
+
+#pragma once
+
+#include "Math/Vector.h"
+#include "Math/Axis.h"
+#include "Misc/Guid.h"
+#include "Containers/Map.h"
+#include "Containers/Set.h"
+
+namespace UE
+{
+namespace AIE
+{
+
+/**
+* Item To Mirror, will get set by the Callee
+**/
+struct FMirrorItem
+{
+	FGuid Guid;
+	FVector Location;
+
+	FMirrorItem() = delete;
+	FMirrorItem(const FVector& InLocation) : Location(InLocation) { Guid = FGuid::NewGuid(); }
+	FMirrorItem(const FMirrorItem&) = default;
+};
+
+//indices from an array of FMirrorItems
+struct FMirrorItemResults
+{
+	//mirrored items will containt both A, B and B, A pairs.
+	TMap<int32, int32> MirroredItems;
+	//not mirrored and along the plane of mirroring
+	TSet<int32> ExactlyOnAxis;
+	//not mirrored and not along the plane of mirrorin.
+	TSet<int32> NotOnAxis;
+};
+
+/**
+	Calculates a set of mirrored item's based upon it's location, and the mirror plane defined by mirror axis and the mirror axis location.
+**/
+struct FMirrorCalculator
+{
+	/*
+		Calculate the mirror results given a set of FMirrorItems, the axis, and plane location.
+		Will return true if at least one item is mirrored
+	*/
+	bool FindMirroredItems(const TArray<FMirrorItem>& Items, FMirrorItemResults& OutMirrorItemResults,
+		TEnumAsByte<EAxis::Type> MirrorAxis = EAxis::X, double AxisLocation = 0.0, double Tolerance = 1e-3f);
+
+private:
+	/*
+	* Compare two items already know wone is is less than the AxisLocation, the other is greater than
+	*/
+	inline bool	IsMirror(const FMirrorItem& ItemLess, const FMirrorItem& ItemGreater, double AxisLocation, double Tolerance,
+		int32 MirrorIndex, int32 OtherIndex[2])
+	{
+		return (FMath::IsNearlyEqual(AxisLocation - ItemLess.Location[MirrorIndex], ItemGreater.Location[MirrorIndex] - AxisLocation, Tolerance) &&
+			FMath::IsNearlyEqual(ItemLess.Location[OtherIndex[0]], ItemGreater.Location[OtherIndex[0]], Tolerance) &&
+			FMath::IsNearlyEqual(ItemLess.Location[OtherIndex[1]], ItemGreater.Location[OtherIndex[1]], Tolerance));
+	}
+};
+
+} //AIE
+} //UE
+
+
diff --git a/Engine/Plugins/Animation/ControlRig/Source/ControlRigEditor/Private/EditMode/SControlRigControlViews.cpp b/Engine/Plugins/Animation/ControlRig/Source/ControlRigEditor/Private/EditMode/SControlRigControlViews.cpp
--- a/Engine/Plugins/Animation/ControlRig/Source/ControlRigEditor/Private/EditMode/SControlRigControlViews.cpp
+++ b/Engine/Plugins/Animation/ControlRig/Source/ControlRigEditor/Private/EditMode/SControlRigControlViews.cpp
@@ -840,7 +840,6 @@ void SControlRigPoseView::UpdateStatusBlocks()
 		uint32 MirrorMatching = 0;
 		for (UControlRig* ControlRig : ControlRigs)
 		{
-			PoseAsset->SetUpMirrorMatchTable(ControlRig);
 			TArray<FName> SelectedNames = ControlRig->CurrentControlSelection();
 			TotalSelected += SelectedNames.Num();
 			for (const FName& ControlName : ControlNames)
