RigVM: Adding IRigVMAssetInterface as parent of RigVMBlueprint. First step to separation of RigVM and Blueprint.

- RigVMBlueprint.* has been renamend to RigVMAsset.*, to not loose p4 history
- URigVMBlueprint lives now in RigVMBlueprintLegacy. It inherits from  UBlueprint, IRigVMAssetInterface and IRigVMClientHost. All IRigVMAssetInterface object also implement IRigVMClientHost, but this has to be specified in the implementing class rather than on the IRigVMAssetInterface itself. Otherwise, calls to ImplementsInterface will not work when called with IRigVMClientHost.
- Instead of passing around pointers to URigVMBlueprint, we are switching to passing TRigVMAssetInterfacePtr (which is just a TScriptInterface<IRigVMAssetInterface>). That way you have access to the Interface pointer and the UObject.
- Some pointers to URigVMBlueprint have not been replaced yet (for example RigVMEditorBlueprintLibrary, these UFUNCTIONs need to be deprecated and switched, or NameValidators that depend on Blueprints)
- Most common implementation from URigVMBlueprint implementation has moved to IRigVMAssetInterface. However, UPROPERTIES need to stay in the original class for backwards serialization purposes. Any implementing class will need to duplicate these properties.

diff --git a/Engine/Plugins/Animation/ControlRig/Source/ControlRig/Private/ControlRig.cpp b/Engine/Plugins/Animation/ControlRig/Source/ControlRig/Private/ControlRig.cpp
--- a/Engine/Plugins/Animation/ControlRig/Source/ControlRig/Private/ControlRig.cpp
+++ b/Engine/Plugins/Animation/ControlRig/Source/ControlRig/Private/ControlRig.cpp
@@ -30,7 +30,6 @@
 #include "Kismet2/BlueprintEditorUtils.h"
 #include "RigVMDeveloperTypeUtils.h"
 #include "Editor.h"
-#include "RigVMBlueprint.h"
 #include "RigVMCompiler/RigVMCompiler.h"
 #endif// WITH_EDITOR
 #include "ControlRigComponent.h"
diff --git a/Engine/Plugins/Animation/ControlRig/Source/ControlRigDeveloper/Private/ControlRigBlueprint.cpp b/Engine/Plugins/Animation/ControlRig/Source/ControlRigDeveloper/Private/ControlRigBlueprint.cpp
--- a/Engine/Plugins/Animation/ControlRig/Source/ControlRigDeveloper/Private/ControlRigBlueprint.cpp
+++ b/Engine/Plugins/Animation/ControlRig/Source/ControlRigDeveloper/Private/ControlRigBlueprint.cpp
@@ -1700,14 +1700,14 @@ void UControlRigBlueprint::GetTypeActions(FBlueprintActionDatabaseRegistrar& Act
 {
 	DECLARE_SCOPE_HIERARCHICAL_COUNTER_FUNC()
 
-	IControlRigEditorModule::Get().GetTypeActions((UControlRigBlueprint*)this, ActionRegistrar);
+	IControlRigEditorModule::Get().GetTypeActions(TRigVMAssetInterfacePtr(const_cast<UControlRigBlueprint*>(this)), ActionRegistrar);
 }
 
 void UControlRigBlueprint::GetInstanceActions(FBlueprintActionDatabaseRegistrar& ActionRegistrar) const
 {
 	DECLARE_SCOPE_HIERARCHICAL_COUNTER_FUNC()
 
-	IControlRigEditorModule::Get().GetInstanceActions((UControlRigBlueprint*)this, ActionRegistrar);
+	IControlRigEditorModule::Get().GetInstanceActions(TRigVMAssetInterfacePtr(const_cast<UControlRigBlueprint*>(this)), ActionRegistrar);
 }
 
 void UControlRigBlueprint::PostTransacted(const FTransactionObjectEvent& TransactionEvent)
@@ -2687,11 +2687,11 @@ void UControlRigBlueprint::OnModularDependencyVMCompiled(UObject* InBlueprint, U
 {
 	if(URigVMBlueprint* RigVMBlueprint = Cast<URigVMBlueprint>(InBlueprint))
 	{
-		OnModularDependencyChanged(RigVMBlueprint);
+		OnModularDependencyChanged(TRigVMAssetInterfacePtr(RigVMBlueprint));
 	}
 }
 
-void UControlRigBlueprint::OnModularDependencyChanged(URigVMBlueprint* InBlueprint)
+void UControlRigBlueprint::OnModularDependencyChanged(TRigVMAssetInterfacePtr InBlueprint)
 {
 	RefreshModuleVariables();
 	RefreshModuleConnectors();
diff --git a/Engine/Plugins/Animation/ControlRig/Source/ControlRigDeveloper/Private/Graph/ControlRigGraph.cpp b/Engine/Plugins/Animation/ControlRig/Source/ControlRigDeveloper/Private/Graph/ControlRigGraph.cpp
--- a/Engine/Plugins/Animation/ControlRig/Source/ControlRigDeveloper/Private/Graph/ControlRigGraph.cpp
+++ b/Engine/Plugins/Animation/ControlRig/Source/ControlRigDeveloper/Private/Graph/ControlRigGraph.cpp
@@ -18,13 +18,13 @@ UControlRigGraph::UControlRigGraph()
 	bIsFunctionDefinition = false;
 }
 
-void UControlRigGraph::InitializeFromBlueprint(URigVMBlueprint* InBlueprint)
+void UControlRigGraph::InitializeFromAsset(TRigVMAssetInterfacePtr InAsset)
 {
 	DECLARE_SCOPE_HIERARCHICAL_COUNTER_FUNC()
 
-	Super::InitializeFromBlueprint(InBlueprint);
+	Super::InitializeFromAsset(InAsset);
 
-	const UControlRigBlueprint* ControlRigBlueprint = CastChecked<UControlRigBlueprint>(InBlueprint);
+	const UControlRigBlueprint* ControlRigBlueprint = CastChecked<UControlRigBlueprint>(InAsset->GetObject());
 	URigHierarchy* Hierarchy = ControlRigBlueprint->Hierarchy;
 
 	if(UControlRig* ControlRig = Cast<UControlRig>(ControlRigBlueprint->GetObjectBeingDebugged()))
@@ -166,7 +166,7 @@ const TArray<TSharedPtr<FRigVMStringWithTag>>* UControlRigGraph::GetElementNameL
 	
 	if(!ElementNameLists.Contains(InElementType))
 	{
-		const UControlRigBlueprint* Blueprint = Cast<UControlRigBlueprint>(GetBlueprint());
+		const UControlRigBlueprint* Blueprint = Cast<UControlRigBlueprint>(GetAsset()->GetObject());
 		if(Blueprint == nullptr)
 		{
 			return &EmptyElementNameList;
@@ -216,7 +216,7 @@ const TArray<TSharedPtr<FRigVMStringWithTag>> UControlRigGraph::GetSelectedEleme
 		return OuterGraph->GetSelectedElementsNameList();
 	}
 	
-	const UControlRigBlueprint* Blueprint = CastChecked<UControlRigBlueprint>(GetBlueprint());
+	const UControlRigBlueprint* Blueprint = CastChecked<UControlRigBlueprint>(GetAsset()->GetObject());
 	if(Blueprint == nullptr)
 	{
 		return Result;
@@ -253,7 +253,7 @@ const TArray<TSharedPtr<FRigVMStringWithTag>>* UControlRigGraph::GetShapeNameLis
 
 FReply UControlRigGraph::HandleGetSelectedClicked(URigVMEdGraph* InEdGraph, URigVMPin* InPin, FString InDefaultValue)
 {
-	UControlRigBlueprint* RigVMBlueprint = CastChecked<UControlRigBlueprint>(InEdGraph->GetBlueprint());
+	UControlRigBlueprint* RigVMBlueprint = CastChecked<UControlRigBlueprint>(InEdGraph->GetAsset()->GetObject());
 	if(InPin->GetCustomWidgetName() == TEXT("ElementName"))
 	{
 		if (URigVMPin* ParentPin = InPin->GetParentPin())
@@ -308,7 +308,7 @@ FReply UControlRigGraph::HandleGetSelectedClicked(URigVMEdGraph* InEdGraph, URig
 
 FReply UControlRigGraph::HandleBrowseClicked(URigVMEdGraph* InEdGraph, URigVMPin* InPin, FString InDefaultValue)
 {
-	UControlRigBlueprint* RigVMBlueprint = CastChecked<UControlRigBlueprint>(InEdGraph->GetBlueprint());
+	UControlRigBlueprint* RigVMBlueprint = CastChecked<UControlRigBlueprint>(InEdGraph->GetAsset()->GetObject());
 	URigVMPin* KeyPin = InPin->GetParentPin();
 	if(KeyPin && KeyPin->GetCPPTypeObject() == FRigElementKey::StaticStruct())
 	{
diff --git a/Engine/Plugins/Animation/ControlRig/Source/ControlRigDeveloper/Public/ControlRigBlueprint.h b/Engine/Plugins/Animation/ControlRig/Source/ControlRigDeveloper/Public/ControlRigBlueprint.h
--- a/Engine/Plugins/Animation/ControlRig/Source/ControlRigDeveloper/Public/ControlRigBlueprint.h
+++ b/Engine/Plugins/Animation/ControlRig/Source/ControlRigDeveloper/Public/ControlRigBlueprint.h
@@ -16,7 +16,7 @@
 #include "RigVMCore/RigVMStatistics.h"
 #include "RigVMModel/RigVMClient.h"
 #include "ControlRigValidationPass.h"
-#include "RigVMBlueprint.h"
+#include "RigVMBlueprintLegacy.h"
 #include "Rigs/RigModuleDefines.h"
 #include "ModularRigModel.h"
 
@@ -50,7 +50,9 @@ public:
 	CONTROLRIGDEVELOPER_API virtual UClass* GetRigVMEditorSettingsClass() const override;
 
 	// URigVMBlueprint interface
-	virtual UClass* GetRigVMBlueprintGeneratedClassPrototype() const override { return UControlRigBlueprintGeneratedClass::StaticClass(); }
+	UE_DEPRECATED(5.7, "Please use GetRigVMGeneratedClassPrototype")
+	virtual UClass* GetRigVMBlueprintGeneratedClassPrototype() const { return UControlRigBlueprintGeneratedClass::StaticClass(); }
+	virtual UClass* GetRigVMGeneratedClassPrototype() const override { return UControlRigBlueprintGeneratedClass::StaticClass(); }
 	CONTROLRIGDEVELOPER_API virtual TArray<FString> GeneratePythonCommands(const FString InNewBlueprintName) override;
 	CONTROLRIGDEVELOPER_API virtual void GetPreloadDependencies(TArray<UObject*>& OutDeps) override;
 #if WITH_EDITOR
@@ -334,7 +336,7 @@ public:
 	CONTROLRIGDEVELOPER_API void PropagateModuleHierarchyFromBPToInstances() const;
 	CONTROLRIGDEVELOPER_API void UpdateModularDependencyDelegates();
 	CONTROLRIGDEVELOPER_API void OnModularDependencyVMCompiled(UObject* InBlueprint, URigVM* InVM, FRigVMExtendedExecuteContext& InExecuteContext);
-	CONTROLRIGDEVELOPER_API void OnModularDependencyChanged(URigVMBlueprint* InBlueprint);
+	CONTROLRIGDEVELOPER_API void OnModularDependencyChanged(TRigVMAssetInterfacePtr InBlueprint);
 	CONTROLRIGDEVELOPER_API void RequestConstructionOnAllModules();
 	CONTROLRIGDEVELOPER_API void RefreshModuleVariables();
 	CONTROLRIGDEVELOPER_API void RefreshModuleVariables(const FRigModuleReference* InModule);
diff --git a/Engine/Plugins/Animation/ControlRig/Source/ControlRigDeveloper/Public/Graph/ControlRigGraph.h b/Engine/Plugins/Animation/ControlRig/Source/ControlRigDeveloper/Public/Graph/ControlRigGraph.h
--- a/Engine/Plugins/Animation/ControlRig/Source/ControlRigDeveloper/Public/Graph/ControlRigGraph.h
+++ b/Engine/Plugins/Animation/ControlRig/Source/ControlRigDeveloper/Public/Graph/ControlRigGraph.h
@@ -29,7 +29,7 @@ public:
 	UE_API UControlRigGraph();
 
 	/** Set up this graph */
-	UE_API virtual void InitializeFromBlueprint(URigVMBlueprint* InBlueprint) override;
+	UE_API virtual void InitializeFromAsset(TRigVMAssetInterfacePtr InBlueprint) override;
 
 	UE_API virtual bool HandleModifiedEvent_Internal(ERigVMGraphNotifType InNotifType, URigVMGraph* InGraph, UObject* InSubject) override;
 
diff --git a/Engine/Plugins/Animation/ControlRig/Source/ControlRigEditor/Private/ControlRigBlueprintFactory.cpp b/Engine/Plugins/Animation/ControlRig/Source/ControlRigEditor/Private/ControlRigBlueprintFactory.cpp
--- a/Engine/Plugins/Animation/ControlRig/Source/ControlRigEditor/Private/ControlRigBlueprintFactory.cpp
+++ b/Engine/Plugins/Animation/ControlRig/Source/ControlRigEditor/Private/ControlRigBlueprintFactory.cpp
@@ -297,7 +297,7 @@ UObject* UControlRigBlueprintFactory::FactoryCreateNew(UClass* Class, UObject* I
 	else
 	{
 		UControlRigBlueprint* ControlRigBlueprint = CastChecked<UControlRigBlueprint>(FKismetEditorUtilities::CreateBlueprint(ParentClass, InParent, Name, BPTYPE_Normal, UControlRigBlueprint::StaticClass(), URigVMBlueprintGeneratedClass::StaticClass(), CallingContext));
-		FControlRigEditorModule::Get().CreateRootGraphIfRequired(ControlRigBlueprint);
+		FControlRigEditorModule::Get().CreateRootGraphIfRequired(TRigVMAssetInterfacePtr(ControlRigBlueprint));
 
 		// add the default module
 		if(ParentClass->IsChildOf(UModularRig::StaticClass()))
diff --git a/Engine/Plugins/Animation/ControlRig/Source/ControlRigEditor/Private/ControlRigEditorModule.cpp b/Engine/Plugins/Animation/ControlRig/Source/ControlRigEditor/Private/ControlRigEditorModule.cpp
--- a/Engine/Plugins/Animation/ControlRig/Source/ControlRigEditor/Private/ControlRigEditorModule.cpp
+++ b/Engine/Plugins/Animation/ControlRig/Source/ControlRigEditor/Private/ControlRigEditorModule.cpp
@@ -104,6 +104,7 @@
 #include "Editor/ControlRigLegacyEditor.h"
 #endif
 #include "Editor/ControlRigNewEditor.h"
+#include "Graph/ControlRigGraphSchema.h"
 
 #if WITH_RIGVMLEGACYEDITOR
 #include "SKismetInspector.h"
@@ -318,11 +319,21 @@ void FControlRigEditorModule::ShutdownModule()
 	UnregisterPropertySectionMappings();
 }
 
-UClass* FControlRigEditorModule::GetRigVMBlueprintClass() const
+UClass* FControlRigEditorModule::GetRigVMAssetClass() const
 {
 	return UControlRigBlueprint::StaticClass();
 }
 
+UClass* FControlRigEditorModule::GetRigVMEdGraphSchemaClass() const
+{
+	return UControlRigGraphSchema::StaticClass();
+}
+
+UScriptStruct* FControlRigEditorModule::GetRigVMExecuteContextStruct() const
+{
+	return FControlRigExecuteContext::StaticStruct();
+}
+
 TSharedRef<FPropertySection> FControlRigEditorModule::RegisterPropertySection(FPropertyEditorModule& PropertyModule, FName ClassName, FName SectionName, FText DisplayName)
 {
 	TSharedRef<FPropertySection> PropertySection = PropertyModule.FindOrCreateSection(ClassName, SectionName, DisplayName);
@@ -350,12 +361,12 @@ void FControlRigEditorModule::UnregisterPropertySectionMappings()
 	RegisteredPropertySections.Empty();
 }
 
-void FControlRigEditorModule::GetNodeContextMenuActions(IRigVMClientHost* RigVMClientHost,
-	const URigVMEdGraphNode* EdGraphNode, URigVMNode* ModelNode, UToolMenu* Menu) const
+void FControlRigEditorModule::GetNodeContextMenuActions(TScriptInterface<IRigVMClientHost> RigVMClientHost,
+                                                        const URigVMEdGraphNode* EdGraphNode, URigVMNode* ModelNode, UToolMenu* Menu) const
 {
 	FRigVMEditorModule::GetNodeContextMenuActions(RigVMClientHost, EdGraphNode, ModelNode, Menu);
 
-	UControlRigBlueprint* ControlRigBlueprint = Cast<UControlRigBlueprint>(RigVMClientHost);
+	UControlRigBlueprint* ControlRigBlueprint = Cast<UControlRigBlueprint>(RigVMClientHost.GetObject());
 	if(ControlRigBlueprint == nullptr)
 	{
 		return;
@@ -541,7 +552,7 @@ void FControlRigEditorModule::GetNodeContextMenuActions(IRigVMClientHost* RigVMC
 		}
 	}
 
-	GetDirectManipulationMenuActions(RigVMClientHost, ModelNode, nullptr, Menu);
+	GetDirectManipulationMenuActions(RigVMClientHost.GetInterface(), ModelNode, nullptr, Menu);
 
 	if (RigElementsToSelect.Num() > 0)
 	{
@@ -719,10 +730,10 @@ void FControlRigEditorModule::GetNodeContextMenuActions(IRigVMClientHost* RigVMC
 	}
 }
 
-void FControlRigEditorModule::GetPinContextMenuActions(IRigVMClientHost* RigVMClientHost, const UEdGraphPin* EdGraphPin, URigVMPin* ModelPin, UToolMenu* Menu) const
+void FControlRigEditorModule::GetPinContextMenuActions(TScriptInterface<IRigVMClientHost> RigVMClientHost, const UEdGraphPin* EdGraphPin, URigVMPin* ModelPin, UToolMenu* Menu) const
 {
 	FRigVMEditorModule::GetPinContextMenuActions(RigVMClientHost, EdGraphPin, ModelPin, Menu);
-	GetDirectManipulationMenuActions(RigVMClientHost, ModelPin->GetNode(), ModelPin, Menu);
+	GetDirectManipulationMenuActions(RigVMClientHost.GetInterface(), ModelPin->GetNode(), ModelPin, Menu);
 }
 
 bool FControlRigEditorModule::AssetsPublicFunctionsAllowed(const FAssetData& InAssetData) const
diff --git a/Engine/Plugins/Animation/ControlRig/Source/ControlRigEditor/Private/ControlRigEditorModule.h b/Engine/Plugins/Animation/ControlRig/Source/ControlRigEditor/Private/ControlRigEditorModule.h
--- a/Engine/Plugins/Animation/ControlRig/Source/ControlRigEditor/Private/ControlRigEditorModule.h
+++ b/Engine/Plugins/Animation/ControlRig/Source/ControlRigEditor/Private/ControlRigEditorModule.h
@@ -33,9 +33,11 @@ public:
 	virtual void ShutdownModule() override;
 
 	/** FRigVMEditorModule interface */
-	virtual UClass* GetRigVMBlueprintClass() const override;
-	virtual void GetNodeContextMenuActions(IRigVMClientHost* RigVMClientHost, const URigVMEdGraphNode* EdGraphNode, URigVMNode* ModelNode, UToolMenu* Menu) const override;
-	virtual void GetPinContextMenuActions(IRigVMClientHost* RigVMClientHost, const UEdGraphPin* EdGraphPin, URigVMPin* ModelPin, UToolMenu* Menu) const override;
+	virtual UClass* GetRigVMAssetClass() const override;
+	virtual UClass* GetRigVMEdGraphSchemaClass() const override;
+	virtual UScriptStruct* GetRigVMExecuteContextStruct() const override;
+	virtual void GetNodeContextMenuActions(TScriptInterface<IRigVMClientHost> RigVMClientHost, const URigVMEdGraphNode* EdGraphNode, URigVMNode* ModelNode, UToolMenu* Menu) const override;
+	virtual void GetPinContextMenuActions(TScriptInterface<IRigVMClientHost> RigVMClientHost, const UEdGraphPin* EdGraphPin, URigVMPin* ModelPin, UToolMenu* Menu) const override;
 	virtual bool AssetsPublicFunctionsAllowed(const FAssetData& InAssetData) const override;
 
 	/** IControlRigEditorModule interface */
diff --git a/Engine/Plugins/Animation/ControlRig/Source/ControlRigEditor/Private/Editor/ControlRigEditor.cpp b/Engine/Plugins/Animation/ControlRig/Source/ControlRigEditor/Private/Editor/ControlRigEditor.cpp
--- a/Engine/Plugins/Animation/ControlRig/Source/ControlRigEditor/Private/Editor/ControlRigEditor.cpp
+++ b/Engine/Plugins/Animation/ControlRig/Source/ControlRigEditor/Private/Editor/ControlRigEditor.cpp
@@ -187,7 +187,7 @@ UControlRigBlueprint* GetControlRigBlueprintFromContext(const UUnrealEdViewportT
 				const TSharedRef<FControlRigEditor>& Editor =
 					StaticCastSharedRef<FControlRigEditor>(AnimViewportClient->GetAssetEditorToolkit());
 
-				return Cast<UControlRigBlueprint>(Editor->GetRigVMBlueprint());
+				return Cast<UControlRigBlueprint>(Editor->GetRigVMAsset().GetObject());
 			}
 		}
 	}
@@ -369,7 +369,7 @@ void FControlRigBaseEditor::CreateEmptyGraphContentImpl(URigVMController* InCont
 
 UControlRigBlueprint* FControlRigBaseEditor::GetControlRigBlueprint() const
 {
-	return Cast<UControlRigBlueprint>(GetRigVMBlueprint());
+	return Cast<UControlRigBlueprint>(GetRigVMAsset().GetObject());
 }
 
 UControlRig* FControlRigBaseEditor::GetControlRig() const
@@ -390,13 +390,13 @@ URigHierarchy* FControlRigBaseEditor::GetHierarchyBeingDebugged() const
 	return nullptr;
 }
 
-void FControlRigBaseEditor::InitRigVMEditorImpl(const EToolkitMode::Type Mode, const TSharedPtr<class IToolkitHost>& InitToolkitHost, class URigVMBlueprint* InRigVMBlueprint)
+void FControlRigBaseEditor::InitRigVMEditorImpl(const EToolkitMode::Type Mode, const TSharedPtr<class IToolkitHost>& InitToolkitHost, TRigVMAssetInterfacePtr InRigVMBlueprint)
 {
 	DECLARE_SCOPE_HIERARCHICAL_COUNTER_FUNC()
 
 	InitRigVMEditorSuper(Mode, InitToolkitHost, InRigVMBlueprint);
 
-	UControlRigBlueprint* ControlRigBlueprint = CastChecked<UControlRigBlueprint>(InRigVMBlueprint);
+	UControlRigBlueprint* ControlRigBlueprint = CastChecked<UControlRigBlueprint>(InRigVMBlueprint->GetObject());
 
 	CreatePersonaToolKitIfRequired();
 	UControlRigBlueprint::sCurrentlyOpenedRigBlueprints.AddUnique(ControlRigBlueprint);
@@ -435,7 +435,7 @@ void FControlRigBaseEditor::InitRigVMEditorImpl(const EToolkitMode::Type Mode, c
 		if (ControlRigBlueprint->IsModularRig())
 		{
 			SchematicModel->SetEditor(SharedRef());
-			ControlRigBlueprint->OnSetObjectBeingDebugged().AddSP(SchematicModel.Get(), &FControlRigSchematicModel::OnSetObjectBeingDebugged);
+			ControlRigBlueprint->IRigVMAssetInterface::OnSetObjectBeingDebugged().AddSP(SchematicModel.Get(), &FControlRigSchematicModel::OnSetObjectBeingDebugged);
 			ControlRigBlueprint->GetModularRigController()->OnModified().AddSP(SchematicModel.Get(), &FControlRigSchematicModel::HandleModularRigModified);
 		}
 
@@ -1050,7 +1050,7 @@ void FControlRigBaseEditor::SetEventQueueImpl(TArray<FName> InEventQueue, bool b
 	}
 
 	TArray<FRigElementKey> PreviousSelection;
-	if (UControlRigBlueprint* RigBlueprint = Cast<UControlRigBlueprint>(GetRigVMBlueprint()))
+	if (UControlRigBlueprint* RigBlueprint = GetControlRigBlueprint())
 	{
 		if(bCompile)
 		{
@@ -1079,7 +1079,7 @@ void FControlRigBaseEditor::SetEventQueueImpl(TArray<FName> InEventQueue, bool b
 	{
 		if (InEventQueue.Num() > 0)
 		{
-			if (UControlRigBlueprint* RigBlueprint = Cast<UControlRigBlueprint>(GetRigVMBlueprint()))
+			if (UControlRigBlueprint* RigBlueprint = GetControlRigBlueprint())
 			{
 				RigBlueprint->Validator->SetControlRig(ControlRig);
 
@@ -1202,7 +1202,7 @@ void FControlRigBaseEditor::HandleSetObjectBeingDebuggedImpl(UObject* InObject)
 	HandleSetObjectBeingDebuggedSuper(InObject);
 	
 	UControlRig* DebuggedControlRig = Cast<UControlRig>(InObject);
-	if(UControlRig* PreviouslyDebuggedControlRig = Cast<UControlRig>(GetRigVMBlueprint()->GetObjectBeingDebugged()))
+	if(UControlRig* PreviouslyDebuggedControlRig = Cast<UControlRig>(GetRigVMAsset()->GetObjectBeingDebugged()))
 	{
 		if(!URigVMHost::IsGarbageOrDestroyed(PreviouslyDebuggedControlRig))
 		{
@@ -1215,7 +1215,7 @@ void FControlRigBaseEditor::HandleSetObjectBeingDebuggedImpl(UObject* InObject)
 		}
 	}
 
-	if (UControlRigBlueprint* RigBlueprint = Cast<UControlRigBlueprint>(GetRigVMBlueprint()))
+	if (UControlRigBlueprint* RigBlueprint = GetControlRigBlueprint())
 	{
 		RigBlueprint->Validator->SetControlRig(DebuggedControlRig);
 	}
@@ -1255,7 +1255,7 @@ void FControlRigBaseEditor::HandleSetObjectBeingDebuggedImpl(UObject* InObject)
 			if(bShouldExecute)
 			{
 				DebuggedControlRig->SetBoneInitialTransformsFromSkeletalMeshComponent(EditorSkelComp);
-				if(UControlRigBlueprint* RigBlueprint = Cast<UControlRigBlueprint>(GetRigVMBlueprint()))
+				if(UControlRigBlueprint* RigBlueprint = GetControlRigBlueprint())
 				{
 					// copy the initial transforms back to the blueprint
 					// no need to call modify here since this code only modifies the bp if the preview mesh changed
@@ -1343,7 +1343,7 @@ void FControlRigBaseEditor::SetDetailViewForRigElements(const TArray<FRigHierarc
 
 	ClearDetailObject();
 
-	UControlRigBlueprint* RigBlueprint = Cast<UControlRigBlueprint>(GetRigVMBlueprint());
+	UControlRigBlueprint* RigBlueprint = GetControlRigBlueprint();
 	URigHierarchy* HierarchyBeingDebugged = GetHierarchyBeingDebugged();
 	TArray<UObject*> Objects;
 
@@ -1357,7 +1357,7 @@ void FControlRigBaseEditor::SetDetailViewForRigElements(const TArray<FRigHierarc
 				continue;
 			}
 
-			URigVMDetailsViewWrapperObject* WrapperObject = URigVMDetailsViewWrapperObject::MakeInstance(GetDetailWrapperClass(), GetRigVMBlueprint(), Element->GetScriptStruct(), (uint8*)Element, HierarchyBeingDebugged);
+			URigVMDetailsViewWrapperObject* WrapperObject = URigVMDetailsViewWrapperObject::MakeInstance(GetDetailWrapperClass(), GetRigVMAsset()->GetObject(), Element->GetScriptStruct(), (uint8*)Element, HierarchyBeingDebugged);
 			WrapperObject->GetWrappedPropertyChangedChainEvent().AddSP(SharedRef(), &FControlRigBaseEditor::OnWrappedPropertyChangedChainEvent);
 			WrapperObject->AddToRoot();
 
@@ -1371,7 +1371,7 @@ void FControlRigBaseEditor::SetDetailViewForRigElements(const TArray<FRigHierarc
 				continue;
 			}
 
-			URigVMDetailsViewWrapperObject* WrapperObject = URigVMDetailsViewWrapperObject::MakeInstance(GetDetailWrapperClass(), GetRigVMBlueprint(), Component->GetScriptStruct(), (uint8*)Component, HierarchyBeingDebugged);
+			URigVMDetailsViewWrapperObject* WrapperObject = URigVMDetailsViewWrapperObject::MakeInstance(GetDetailWrapperClass(), GetRigVMAsset()->GetObject(), Component->GetScriptStruct(), (uint8*)Component, HierarchyBeingDebugged);
 			WrapperObject->GetWrappedPropertyChangedChainEvent().AddSP(SharedRef(), &FControlRigBaseEditor::OnWrappedPropertyChangedChainEvent);
 			WrapperObject->AddToRoot();
 
@@ -1511,7 +1511,7 @@ void FControlRigBaseEditor::SetDetailViewForRigModules(const TArray<FName> InMod
 
 	ClearDetailObject();
 
-	UControlRigBlueprint* RigBlueprint = Cast<UControlRigBlueprint>(GetRigVMBlueprint());
+	UControlRigBlueprint* RigBlueprint = GetControlRigBlueprint();
 	UModularRig* RigBeingDebugged = Cast<UModularRig>(RigBlueprint->GetDebuggedControlRig());
 
 	if (!RigBeingDebugged)
@@ -1992,7 +1992,7 @@ void FControlRigBaseEditor::SaveAsset_ExecuteImpl()
 	SaveAsset_ExecuteSuper();
 
 	// Save the new state of the hierarchy in the default object, so that it has the correct values on load
-	UControlRigBlueprint* RigBlueprint = Cast<UControlRigBlueprint>(GetRigVMBlueprint());
+	UControlRigBlueprint* RigBlueprint = GetControlRigBlueprint();
 	if(const UControlRig* ControlRig = GetControlRig())
 	{
 		UControlRig* CDO = ControlRig->GetClass()->GetDefaultObject<UControlRig>();
@@ -2010,7 +2010,7 @@ void FControlRigBaseEditor::SaveAssetAs_ExecuteImpl()
 	SaveAssetAs_ExecuteSuper();
 
 	// Save the new state of the hierarchy in the default object, so that it has the correct values on load
-	UControlRigBlueprint* RigBlueprint = Cast<UControlRigBlueprint>(GetRigVMBlueprint());
+	UControlRigBlueprint* RigBlueprint = GetControlRigBlueprint();
 	if(const UControlRig* ControlRig = GetControlRig())
 	{
 		UControlRig* CDO = ControlRig->GetClass()->GetDefaultObject<UControlRig>();
@@ -2025,7 +2025,7 @@ void FControlRigBaseEditor::SaveAssetAs_ExecuteImpl()
 
 bool FControlRigBaseEditor::IsModularRig() const
 {
-	if(UControlRigBlueprint* RigBlueprint = Cast<UControlRigBlueprint>(GetRigVMBlueprint()))
+	if(UControlRigBlueprint* RigBlueprint = GetControlRigBlueprint())
 	{
 		return RigBlueprint->IsModularRig();
 	}
@@ -2034,7 +2034,7 @@ bool FControlRigBaseEditor::IsModularRig() const
 
 bool FControlRigBaseEditor::IsRigModule() const
 {
-	if(UControlRigBlueprint* RigBlueprint = Cast<UControlRigBlueprint>(GetRigVMBlueprint()))
+	if(UControlRigBlueprint* RigBlueprint = GetControlRigBlueprint())
 	{
 		return RigBlueprint->IsControlRigModule();
 	}
@@ -2109,7 +2109,7 @@ void FControlRigBaseEditor::PostTransactionImpl(bool bSuccess, const FTransactio
 		return;
 	}
 
-	UControlRigBlueprint* RigBlueprint = Cast<UControlRigBlueprint>(GetRigVMBlueprint());
+	UControlRigBlueprint* RigBlueprint = GetControlRigBlueprint();
 	if(RigBlueprint == nullptr)
 	{
 		return;
@@ -2400,17 +2400,17 @@ void FControlRigBaseEditor::HandleViewportCreated(const TSharedRef<class IPerson
 	// TODO: this is duplicated code from FAnimBlueprintEditor, would be nice to consolidate. 
 	auto GetCompilationStateText = [this]()
 	{
-		if (UBlueprint* Blueprint = GetRigVMBlueprint())
+		if (TRigVMAssetInterfacePtr Blueprint = GetRigVMAsset())
 		{
-			switch (Blueprint->Status)
+			switch (Blueprint->GetAssetStatus())
 			{
-			case BS_UpToDate:
-			case BS_UpToDateWithWarnings:
+			case RVMA_UpToDate:
+			case RVMA_UpToDateWithWarnings:
 				// Fall thru and return empty string
 				break;
-			case BS_Dirty:
+			case RVMA_Dirty:
 				return LOCTEXT("ControlRigBP_Dirty", "Preview out of date");
-			case BS_Error:
+			case RVMA_Error:
 				return LOCTEXT("ControlRigBP_CompileError", "Compile Error");
 			default:
 				return LOCTEXT("ControlRigBP_UnknownStatus", "Unknown Status");
@@ -2449,7 +2449,7 @@ void FControlRigBaseEditor::HandleViewportCreated(const TSharedRef<class IPerson
 
 	auto CompileBlueprint = [this]()
 	{
-		if (UBlueprint* Blueprint = GetRigVMBlueprint())
+		if (TRigVMAssetInterfacePtr Blueprint = GetRigVMAsset())
 		{
 			if (!Blueprint->IsUpToDate())
 			{
@@ -2462,9 +2462,9 @@ void FControlRigBaseEditor::HandleViewportCreated(const TSharedRef<class IPerson
 
 	auto GetErrorSeverity = [this]()
 	{
-		if (UBlueprint* Blueprint = GetRigVMBlueprint())
+		if (TRigVMAssetInterfacePtr Blueprint = GetRigVMAsset())
 		{
-			return (Blueprint->Status == BS_Error) ? EMessageSeverity::Error : EMessageSeverity::Warning;
+			return (Blueprint->GetAssetStatus() == RVMA_Error) ? EMessageSeverity::Error : EMessageSeverity::Warning;
 		}
 
 		return EMessageSeverity::Warning;
@@ -2472,9 +2472,9 @@ void FControlRigBaseEditor::HandleViewportCreated(const TSharedRef<class IPerson
 
 	auto GetIcon = [this]()
 	{
-		if (UBlueprint* Blueprint = GetRigVMBlueprint())
+		if (TRigVMAssetInterfacePtr Blueprint = GetRigVMAsset())
 		{
-			return (Blueprint->Status == BS_Error) ? FEditorFontGlyphs::Exclamation_Triangle : FEditorFontGlyphs::Eye;
+			return (Blueprint->GetAssetStatus() == RVMA_Error) ? FEditorFontGlyphs::Exclamation_Triangle : FEditorFontGlyphs::Eye;
 		}
 
 		return FEditorFontGlyphs::Eye;
@@ -2937,7 +2937,7 @@ void FControlRigBaseEditor::HandleViewportCreated(const TSharedRef<class IPerson
 	{
 		AnimViewportClient->OnSetBoneSize.BindLambda([this](float InBoneSize)
 		{
-			if (UControlRigBlueprint* RigBlueprint = Cast<UControlRigBlueprint>(GetRigVMBlueprint()))
+			if (UControlRigBlueprint* RigBlueprint = GetControlRigBlueprint())
 			{
 				RigBlueprint->Modify();
 				RigBlueprint->DebugBoneRadius = InBoneSize;
@@ -2946,7 +2946,7 @@ void FControlRigBaseEditor::HandleViewportCreated(const TSharedRef<class IPerson
 		
 		AnimViewportClient->OnGetBoneSize.BindLambda([this]() -> float
 		{
-			if (UControlRigBlueprint* RigBlueprint = Cast<UControlRigBlueprint>(GetRigVMBlueprint()))
+			if (UControlRigBlueprint* RigBlueprint = GetControlRigBlueprint())
 			{
 				return RigBlueprint->DebugBoneRadius;
 			}
@@ -3165,7 +3165,7 @@ void FControlRigBaseEditor::HandlePreviewSceneCreated(const TSharedRef<IPersonaP
 	if (FWorldContext* WorldContext = GEngine->GetWorldContextFromWorld(InPersonaPreviewScene->GetWorld()))
 	{
 		static constexpr TCHAR Format[] = TEXT("ControlRigEditor (%s)");
-		WorldContext->CustomDescription = FString::Printf(Format, *GetRigVMBlueprint()->GetName());
+		WorldContext->CustomDescription = FString::Printf(Format, *GetRigVMAsset()->GetObject()->GetName());
 	}
 
 	AAnimationEditorPreviewActor* Actor = InPersonaPreviewScene->GetWorld()->SpawnActor<AAnimationEditorPreviewActor>(AAnimationEditorPreviewActor::StaticClass(), FTransform::Identity);
@@ -3208,7 +3208,7 @@ void FControlRigBaseEditor::UpdateRigVMHostImpl()
 
 	UpdateRigVMHostSuper();
 
-	UControlRigBlueprint* Blueprint = Cast<UControlRigBlueprint>(GetRigVMBlueprint());
+	UControlRigBlueprint* Blueprint = GetControlRigBlueprint();
 	if(UClass* Class = Blueprint->GeneratedClass)
 	{
 		UControlRigSkeletalMeshComponent* EditorSkelComp = Cast<UControlRigSkeletalMeshComponent>(GetPersonaToolkit()->GetPreviewScene()->GetPreviewMeshComponent());
@@ -3583,7 +3583,7 @@ void FControlRigBaseEditor::GenerateEventQueueMenuContentImpl(FMenuBuilder& Menu
     MenuBuilder.AddMenuEntry(FControlRigEditorCommands::Get().BackwardsAndForwardsSolveEvent, TEXT("BackwardsAndForwards"), TAttribute<FText>(), TAttribute<FText>(), GetEventQueueIconImpl(BackwardsAndForwardsSolveEventQueue));
     MenuBuilder.EndSection();
 
-    if (const UControlRigBlueprint* RigBlueprint = Cast<UControlRigBlueprint>(GetRigVMBlueprint()))
+    if (const UControlRigBlueprint* RigBlueprint = Cast<UControlRigBlueprint>(GetRigVMAsset()->GetObject()))
     {
     	URigVMEdGraphSchema* Schema = CastChecked<URigVMEdGraphSchema>(RigBlueprint->GetRigVMEdGraphSchemaClass()->GetDefaultObject());
     	
@@ -3674,7 +3674,7 @@ FTransform FControlRigBaseEditor::GetRigElementTransform(const FRigElementKey& I
 
 	if(bOnDebugInstance)
 	{
-		UControlRig* DebuggedControlRig = Cast<UControlRig>(GetRigVMBlueprint()->GetObjectBeingDebugged());
+		UControlRig* DebuggedControlRig = Cast<UControlRig>(GetRigVMAsset()->GetObjectBeingDebugged());
 		if (DebuggedControlRig == nullptr)
 		{
 			DebuggedControlRig = GetControlRig();
@@ -4077,7 +4077,7 @@ EVisibility FControlRigBaseEditor::GetDirectManipulationVisibility() const
 
 FText FControlRigBaseEditor::GetDirectionManipulationText() const
 {
-	if (UControlRig* DebuggedControlRig = Cast<UControlRig>(GetRigVMBlueprint()->GetObjectBeingDebugged()))
+	if (UControlRig* DebuggedControlRig = Cast<UControlRig>(GetRigVMAsset()->GetObjectBeingDebugged()))
 	{
 		if (URigHierarchy* Hierarchy = DebuggedControlRig->GetHierarchy())
 		{
@@ -4109,7 +4109,7 @@ void FControlRigBaseEditor::OnDirectManipulationChanged(TSharedPtr<FString> NewV
 		return;
 	}
 	
-	UControlRigBlueprint* ControlRigBlueprint = CastChecked<UControlRigBlueprint>(GetRigVMBlueprint());
+	UControlRigBlueprint* ControlRigBlueprint = CastChecked<UControlRigBlueprint>(GetRigVMAsset().GetObject());
 	if(ControlRigBlueprint == nullptr)
 	{
 		return;
@@ -4142,7 +4142,7 @@ const TArray<FRigDirectManipulationTarget> FControlRigBaseEditor::GetDirectManip
 {
 	if(DirectManipulationSubject.IsValid())
 	{
-		if (UControlRig* DebuggedControlRig = Cast<UControlRig>(GetRigVMBlueprint()->GetObjectBeingDebugged()))
+		if (UControlRig* DebuggedControlRig = Cast<UControlRig>(GetRigVMAsset()->GetObjectBeingDebugged()))
 		{
 			if(const URigVMUnitNode* Node = DirectManipulationSubject.Get())
 			{
@@ -4368,7 +4368,7 @@ void FControlRigBaseEditor::OnHierarchyChanged()
 
 void FControlRigBaseEditor::OnHierarchyModified(ERigHierarchyNotification InNotif, URigHierarchy* InHierarchy, const FRigNotificationSubject& InSubject)
 {
-	UControlRigBlueprint* RigBlueprint = Cast<UControlRigBlueprint>(GetRigVMBlueprint());
+	UControlRigBlueprint* RigBlueprint = GetControlRigBlueprint();
 	if(RigBlueprint == nullptr)
 	{
 		return;
@@ -4593,7 +4593,7 @@ void FControlRigBaseEditor::OnHierarchyModified_AnyThread(ERigHierarchyNotificat
 						else
 						{
 							// only clear the details if we are not looking at a transient control
-							if (UControlRig* DebuggedControlRig = Cast<UControlRig>(GetRigVMBlueprint()->GetObjectBeingDebugged()))
+							if (UControlRig* DebuggedControlRig = Cast<UControlRig>(GetRigVMAsset()->GetObjectBeingDebugged()))
 							{
 								if(DebuggedControlRig->RigUnitManipulationInfos.IsEmpty())
 								{
@@ -4753,7 +4753,7 @@ void FControlRigBaseEditor::HandleModularRigModified(EModularRigNotification InN
 	}
 }
 
-void FControlRigBaseEditor::HandlePostCompileModularRigs(URigVMBlueprint* InBlueprint)
+void FControlRigBaseEditor::HandlePostCompileModularRigs(TRigVMAssetInterfacePtr InBlueprint)
 {
 	RefreshDetailView();
 }
@@ -4761,7 +4761,7 @@ void FControlRigBaseEditor::HandlePostCompileModularRigs(URigVMBlueprint* InBlue
 void FControlRigBaseEditor::SwapModuleWithinAsset()
 {
 	const UControlRigBlueprint* Blueprint = GetControlRigBlueprint();
-	const FAssetData Asset = UE::RigVM::Editor::Tools::FindAssetFromAnyPath(GetRigVMBlueprint()->GetPathName(), true);
+	const FAssetData Asset = UE::RigVM::Editor::Tools::FindAssetFromAnyPath(GetRigVMAsset()->GetObject()->GetPathName(), true);
 	SRigVMSwapAssetReferencesWidget::FArguments WidgetArgs;
 
 	FRigVMAssetDataFilter FilterModules = FRigVMAssetDataFilter::CreateLambda([](const FAssetData& AssetData)
@@ -4844,7 +4844,7 @@ void FControlRigBaseEditor::SwapModuleWithinAsset()
 void FControlRigBaseEditor::SwapModuleAcrossProject()
 {
 	const UControlRigBlueprint* Blueprint = GetControlRigBlueprint();
-	const FAssetData Asset = UE::RigVM::Editor::Tools::FindAssetFromAnyPath(GetRigVMBlueprint()->GetPathName(), true);
+	const FAssetData Asset = UE::RigVM::Editor::Tools::FindAssetFromAnyPath(GetRigVMAsset()->GetObject()->GetPathName(), true);
 	SRigVMSwapAssetReferencesWidget::FArguments WidgetArgs;
 
 	FRigVMAssetDataFilter FilterModules = FRigVMAssetDataFilter::CreateLambda([](const FAssetData& AssetData)
@@ -4929,7 +4929,7 @@ void FControlRigBaseEditor::UpdateBoneModification(FName BoneName, const FTransf
 		}
 		
 		TMap<FName, FTransform>* TransformOverrideMap = &ControlRig->TransformOverrideForUserCreatedBones;
-		if (UControlRig* DebuggedControlRig = Cast<UControlRig>(GetRigVMBlueprint()->GetObjectBeingDebugged()))
+		if (UControlRig* DebuggedControlRig = Cast<UControlRig>(GetRigVMAsset()->GetObjectBeingDebugged()))
 		{
 			TransformOverrideMap = &DebuggedControlRig->TransformOverrideForUserCreatedBones;
 		}
@@ -4951,7 +4951,7 @@ void FControlRigBaseEditor::RemoveBoneModification(FName BoneName)
 		}
 
 		TMap<FName, FTransform>* TransformOverrideMap = &ControlRig->TransformOverrideForUserCreatedBones;
-		if (UControlRig* DebuggedControlRig = Cast<UControlRig>(GetRigVMBlueprint()->GetObjectBeingDebugged()))
+		if (UControlRig* DebuggedControlRig = Cast<UControlRig>(GetRigVMAsset()->GetObjectBeingDebugged()))
 		{
 			TransformOverrideMap = &DebuggedControlRig->TransformOverrideForUserCreatedBones;
 		}
@@ -4970,7 +4970,7 @@ void FControlRigBaseEditor::ResetAllBoneModification()
 		}
 
 		TMap<FName, FTransform>* TransformOverrideMap = &ControlRig->TransformOverrideForUserCreatedBones;
-		if (UControlRig* DebuggedControlRig = Cast<UControlRig>(GetRigVMBlueprint()->GetObjectBeingDebugged()))
+		if (UControlRig* DebuggedControlRig = Cast<UControlRig>(GetRigVMAsset()->GetObjectBeingDebugged()))
 		{
 			TransformOverrideMap = &DebuggedControlRig->TransformOverrideForUserCreatedBones;
 		}
@@ -6036,13 +6036,13 @@ void FControlRigBaseEditor::HandleMakeComponentGetterSetter(bool bIsGetter, TArr
 
 void FControlRigBaseEditor::HandleOnControlModified(UControlRig* Subject, FRigControlElement* ControlElement, const FRigControlModifiedContext& Context)
 {
-	UControlRig* DebuggedControlRig = Cast<UControlRig>(GetRigVMBlueprint()->GetObjectBeingDebugged());
+	UControlRig* DebuggedControlRig = Cast<UControlRig>(GetRigVMAsset()->GetObjectBeingDebugged());
 	if (Subject != DebuggedControlRig)
 	{
 		return;
 	}
 
-	UControlRigBlueprint* Blueprint = Cast<UControlRigBlueprint>(GetRigVMBlueprint());
+	UControlRigBlueprint* Blueprint = GetControlRigBlueprint();
 	if (Blueprint == nullptr)
 	{
 		return;
@@ -6190,7 +6190,7 @@ void FControlRigBaseEditor::HandleOnControlModified(UControlRig* Subject, FRigCo
 	}
 }
 
-void FControlRigBaseEditor::HandleRefreshEditorFromBlueprintImpl(URigVMBlueprint* InBlueprint)
+void FControlRigBaseEditor::HandleRefreshEditorFromBlueprintImpl(TRigVMAssetInterfacePtr InBlueprint)
 {
 	OnHierarchyChanged();
 	HandleRefreshEditorFromBlueprintSuper(InBlueprint);
diff --git a/Engine/Plugins/Animation/ControlRig/Source/ControlRigEditor/Private/Editor/ControlRigEditor.h b/Engine/Plugins/Animation/ControlRig/Source/ControlRigEditor/Private/Editor/ControlRigEditor.h
--- a/Engine/Plugins/Animation/ControlRig/Source/ControlRigEditor/Private/Editor/ControlRigEditor.h
+++ b/Engine/Plugins/Animation/ControlRig/Source/ControlRigEditor/Private/Editor/ControlRigEditor.h
@@ -118,10 +118,10 @@ public:
 	
 protected:
 
-	virtual void InitRigVMEditorSuper(const EToolkitMode::Type Mode, const TSharedPtr< class IToolkitHost >& InitToolkitHost, class URigVMBlueprint* InRigVMBlueprint) = 0;
+	virtual void InitRigVMEditorSuper(const EToolkitMode::Type Mode, const TSharedPtr< class IToolkitHost >& InitToolkitHost, TRigVMAssetInterfacePtr InRigVMBlueprint) = 0;
 
 	virtual bool IsControlRigLegacyEditor() const = 0;
-	virtual URigVMBlueprint* GetRigVMBlueprint() const = 0;
+	virtual TRigVMAssetInterfacePtr GetRigVMAsset() const = 0;
 	virtual URigVMHost* GetRigVMHost() const = 0;
 	virtual bool IsEditorInitialized() const = 0;
 	
@@ -173,7 +173,7 @@ protected:
 	virtual FMenuBuilder GenerateBulkEditMenuSuper() = 0;
 	virtual TWeakPtr<class SGraphEditor> GetFocusedGraphEd() const = 0;
 	virtual void OnGraphNodeDropToPerformSuper(TSharedPtr<FDragDropOperation> InDragDropOp, UEdGraph* InGraph, const FVector2f& InNodePosition, const FVector2f& InScreenPosition) = 0;
-	virtual void HandleRefreshEditorFromBlueprintSuper(URigVMBlueprint* InBlueprint) = 0;
+	virtual void HandleRefreshEditorFromBlueprintSuper(TRigVMAssetInterfacePtr InBlueprint) = 0;
 	virtual void OnGraphNodeDropToPerform(TSharedPtr<FDragDropOperation> InDragDropOp, UEdGraph* InGraph, const FVector2f& InNodePosition, const FVector2f& InScreenPosition) = 0;
 	
 	virtual void OnCloseSuper() = 0;
@@ -187,8 +187,11 @@ public:
 	static FControlRigBaseEditor* GetFromAssetEditorInstance(IAssetEditorInstance* Instance);
 	virtual TSharedRef<FControlRigBaseEditor> SharedRef() override { return StaticCastSharedRef<FControlRigBaseEditor>(SharedControlRigEditorRef());}
 
+	// Gets the Control Rig Blueprint being edited/viewed
+	virtual UControlRigBlueprint* GetControlRigBlueprint() const override;
+
 protected:
-	void InitRigVMEditorImpl(const EToolkitMode::Type Mode, const TSharedPtr< class IToolkitHost >& InitToolkitHost, class URigVMBlueprint* InRigVMBlueprint);
+	void InitRigVMEditorImpl(const EToolkitMode::Type Mode, const TSharedPtr< class IToolkitHost >& InitToolkitHost, TRigVMAssetInterfacePtr InRigVMBlueprint);
 
 	virtual const FName GetEditorAppNameImpl() const;
 	virtual const FName GetEditorModeNameImpl() const;
@@ -242,9 +245,6 @@ protected:
 	//  FTickableEditorObject Interface
 	virtual void TickImpl(float DeltaTime);
 
-	// Gets the Control Rig Blueprint being edited/viewed
-	virtual UControlRigBlueprint* GetControlRigBlueprint() const override;
-
 	// returns the hierarchy being debugged
 	virtual UControlRig* GetControlRig() const override;
 
@@ -287,7 +287,7 @@ protected:
 	void HandleRigTypeChanged(UControlRigBlueprint* InBlueprint);
 
 	void HandleModularRigModified(EModularRigNotification InNotification, const FRigModuleReference* InModule);
-	void HandlePostCompileModularRigs(URigVMBlueprint* InBlueprint);
+	void HandlePostCompileModularRigs(TRigVMAssetInterfacePtr InBlueprint);
 	void SwapModuleWithinAsset();
 	void SwapModuleAcrossProject();
 
@@ -415,7 +415,7 @@ protected:
 
 	void HandleOnControlModified(UControlRig* Subject, FRigControlElement* ControlElement, const FRigControlModifiedContext& Context);
 
-	virtual void HandleRefreshEditorFromBlueprintImpl(URigVMBlueprint* InBlueprint);
+	virtual void HandleRefreshEditorFromBlueprintImpl(TRigVMAssetInterfacePtr InBlueprint);
 
 	FText GetReplayAssetName() const;
 	FText GetReplayAssetTooltip() const;
diff --git a/Engine/Plugins/Animation/ControlRig/Source/ControlRigEditor/Private/Editor/ControlRigEditorMode.cpp b/Engine/Plugins/Animation/ControlRig/Source/ControlRigEditor/Private/Editor/ControlRigEditorMode.cpp
--- a/Engine/Plugins/Animation/ControlRig/Source/ControlRigEditor/Private/Editor/ControlRigEditorMode.cpp
+++ b/Engine/Plugins/Animation/ControlRig/Source/ControlRigEditor/Private/Editor/ControlRigEditorMode.cpp
@@ -27,7 +27,7 @@
 FControlRigLegacyEditorMode::FControlRigLegacyEditorMode(const TSharedRef<FControlRigLegacyEditor>& InControlRigEditor, bool bCreateDefaultLayout)
 	: FBlueprintEditorApplicationMode(InControlRigEditor.ToSharedPtr(), FControlRigEditorModes::ControlRigEditorMode, FControlRigEditorModes::GetLocalizedMode, false, false)
 {
-	ControlRigBlueprintPtr = CastChecked<UControlRigBlueprint>(InControlRigEditor->GetRigVMBlueprint());
+	ControlRigBlueprintPtr = InControlRigEditor->GetControlRigBlueprint(); 
 
 	TabFactories.RegisterFactory(MakeShared<FRigHierarchyTabSummoner>(InControlRigEditor));
 	TabFactories.RegisterFactory(MakeShared<FRigVMExecutionStackTabSummoner>(InControlRigEditor));
@@ -260,8 +260,8 @@ void FModularRigLegacyEditorMode::RegisterTabFactories(TSharedPtr<FTabManager> I
 FControlRigEditorMode::FControlRigEditorMode(const TSharedRef<FControlRigEditor>& InControlRigEditor, bool bCreateDefaultLayout)
 	: FRigVMNewEditorMode(InControlRigEditor)
 {
-	ControlRigBlueprintPtr = CastChecked<UControlRigBlueprint>(InControlRigEditor->GetRigVMBlueprint());
-
+	ControlRigBlueprintPtr = InControlRigEditor->GetControlRigBlueprint();
+	
 	TabFactories.RegisterFactory(MakeShared<FRigHierarchyTabSummoner>(InControlRigEditor));
 	TabFactories.RegisterFactory(MakeShared<FRigCurveContainerTabSummoner>(InControlRigEditor));
 	TabFactories.RegisterFactory(MakeShared<FRigValidationTabSummoner>(InControlRigEditor));
diff --git a/Engine/Plugins/Animation/ControlRig/Source/ControlRigEditor/Private/Editor/ControlRigLegacyEditor.cpp b/Engine/Plugins/Animation/ControlRig/Source/ControlRigEditor/Private/Editor/ControlRigLegacyEditor.cpp
--- a/Engine/Plugins/Animation/ControlRig/Source/ControlRigEditor/Private/Editor/ControlRigLegacyEditor.cpp
+++ b/Engine/Plugins/Animation/ControlRig/Source/ControlRigEditor/Private/Editor/ControlRigLegacyEditor.cpp
@@ -37,7 +37,7 @@ FControlRigLegacyEditor::~FControlRigLegacyEditor()
 			RigBlueprint->GetModularRigController()->OnModified().RemoveAll(this);
 			RigBlueprint->OnModularRigCompiled().RemoveAll(this);
 
-			RigBlueprint->OnSetObjectBeingDebugged().RemoveAll(&SchematicModel);
+			RigBlueprint->IRigVMAssetInterface::OnSetObjectBeingDebugged().RemoveAll(&SchematicModel);
 			RigBlueprint->OnHierarchyModified().RemoveAll(&SchematicModel);
 			RigBlueprint->GetModularRigController()->OnModified().RemoveAll(&SchematicModel);
 		}
diff --git a/Engine/Plugins/Animation/ControlRig/Source/ControlRigEditor/Private/Editor/ControlRigLegacyEditor.h b/Engine/Plugins/Animation/ControlRig/Source/ControlRigEditor/Private/Editor/ControlRigLegacyEditor.h
--- a/Engine/Plugins/Animation/ControlRig/Source/ControlRigEditor/Private/Editor/ControlRigLegacyEditor.h
+++ b/Engine/Plugins/Animation/ControlRig/Source/ControlRigEditor/Private/Editor/ControlRigLegacyEditor.h
@@ -11,8 +11,8 @@ class FControlRigLegacyEditor : public IControlRigLegacyEditor, public FControlR
 public:
 
 	
-	virtual void InitRigVMEditor(const EToolkitMode::Type Mode, const TSharedPtr< class IToolkitHost >& InitToolkitHost, class URigVMBlueprint* InRigVMBlueprint) override { return FControlRigBaseEditor::InitRigVMEditorImpl(Mode, InitToolkitHost, InRigVMBlueprint); }
-	virtual void InitRigVMEditorSuper(const EToolkitMode::Type Mode, const TSharedPtr< class IToolkitHost >& InitToolkitHost, class URigVMBlueprint* InRigVMBlueprint) override { return IControlRigLegacyEditor::InitRigVMEditor(Mode, InitToolkitHost, InRigVMBlueprint); }
+	virtual void InitRigVMEditor(const EToolkitMode::Type Mode, const TSharedPtr< class IToolkitHost >& InitToolkitHost, TRigVMAssetInterfacePtr InRigVMBlueprint) override { return FControlRigBaseEditor::InitRigVMEditorImpl(Mode, InitToolkitHost, InRigVMBlueprint); }
+	virtual void InitRigVMEditorSuper(const EToolkitMode::Type Mode, const TSharedPtr< class IToolkitHost >& InitToolkitHost, TRigVMAssetInterfacePtr InRigVMBlueprint) override { return IControlRigLegacyEditor::InitRigVMEditor(Mode, InitToolkitHost, InRigVMBlueprint); }
 
 	virtual const FName GetEditorAppName() const override { return FControlRigBaseEditor::GetEditorAppNameImpl(); }
 	virtual const FName GetEditorModeName() const override { return FControlRigBaseEditor::GetEditorModeNameImpl(); }
@@ -30,7 +30,7 @@ public:
 	virtual TSharedRef<const IRigVMEditor> SharedRigVMEditorRef() const override { return StaticCastSharedRef<const IRigVMEditor>(SharedThis(this)); }
 	
 	virtual bool IsControlRigLegacyEditor() const { return true; };
-	virtual URigVMBlueprint* GetRigVMBlueprint() const { return IControlRigLegacyEditor::GetRigVMBlueprint(); }
+	virtual TRigVMAssetInterfacePtr GetRigVMAsset() const override { return IControlRigLegacyEditor::GetRigVMAsset(); }
 	virtual URigVMHost* GetRigVMHost() const { return FRigVMEditorBase::GetRigVMHost(); }
 	virtual bool IsEditorInitialized() const override { return bRigVMEditorInitialized; }
 	virtual TSharedRef<FUICommandList> GetToolkitCommands() { return IControlRigLegacyEditor::GetToolkitCommands(); }
@@ -157,8 +157,8 @@ private:
 
 	virtual void GenerateEventQueueMenuContent(FMenuBuilder& MenuBuilder) override { return FControlRigBaseEditor::GenerateEventQueueMenuContentImpl(MenuBuilder); }
 
-	virtual void HandleRefreshEditorFromBlueprint(URigVMBlueprint* InBlueprint) override { return FControlRigBaseEditor::HandleRefreshEditorFromBlueprintImpl(InBlueprint); }
-	virtual void HandleRefreshEditorFromBlueprintSuper(URigVMBlueprint* InBlueprint) override { return IControlRigLegacyEditor::HandleRefreshEditorFromBlueprint(InBlueprint); }
+	virtual void HandleRefreshEditorFromBlueprint(TRigVMAssetInterfacePtr InBlueprint) override { return FControlRigBaseEditor::HandleRefreshEditorFromBlueprintImpl(InBlueprint); }
+	virtual void HandleRefreshEditorFromBlueprintSuper(TRigVMAssetInterfacePtr InBlueprint) override { return IControlRigLegacyEditor::HandleRefreshEditorFromBlueprint(InBlueprint); }
 
 	
 	/** delegate for changing property */
diff --git a/Engine/Plugins/Animation/ControlRig/Source/ControlRigEditor/Private/Editor/ControlRigNewEditor.cpp b/Engine/Plugins/Animation/ControlRig/Source/ControlRigEditor/Private/Editor/ControlRigNewEditor.cpp
--- a/Engine/Plugins/Animation/ControlRig/Source/ControlRigEditor/Private/Editor/ControlRigNewEditor.cpp
+++ b/Engine/Plugins/Animation/ControlRig/Source/ControlRigEditor/Private/Editor/ControlRigNewEditor.cpp
@@ -36,7 +36,7 @@ FControlRigEditor::~FControlRigEditor()
 			RigBlueprint->GetModularRigController()->OnModified().RemoveAll(this);
 			RigBlueprint->OnModularRigCompiled().RemoveAll(this);
 
-			RigBlueprint->OnSetObjectBeingDebugged().RemoveAll(&SchematicModel);
+			RigBlueprint->IRigVMAssetInterface::OnSetObjectBeingDebugged().RemoveAll(&SchematicModel);
 			RigBlueprint->OnHierarchyModified().RemoveAll(&SchematicModel);
 			RigBlueprint->GetModularRigController()->OnModified().RemoveAll(&SchematicModel);
 		}
diff --git a/Engine/Plugins/Animation/ControlRig/Source/ControlRigEditor/Private/Editor/ControlRigNewEditor.h b/Engine/Plugins/Animation/ControlRig/Source/ControlRigEditor/Private/Editor/ControlRigNewEditor.h
--- a/Engine/Plugins/Animation/ControlRig/Source/ControlRigEditor/Private/Editor/ControlRigNewEditor.h
+++ b/Engine/Plugins/Animation/ControlRig/Source/ControlRigEditor/Private/Editor/ControlRigNewEditor.h
@@ -10,8 +10,8 @@ class FControlRigEditor : public IControlRigNewEditor, public FControlRigBaseEdi
 {
 public:
 
-	virtual void InitRigVMEditor(const EToolkitMode::Type Mode, const TSharedPtr< class IToolkitHost >& InitToolkitHost, class URigVMBlueprint* InRigVMBlueprint) override { return FControlRigBaseEditor::InitRigVMEditorImpl(Mode, InitToolkitHost, InRigVMBlueprint); }
-	virtual void InitRigVMEditorSuper(const EToolkitMode::Type Mode, const TSharedPtr<class IToolkitHost>& InitToolkitHost, class URigVMBlueprint* InRigVMBlueprint) override { return IControlRigNewEditor::InitRigVMEditor(Mode, InitToolkitHost, InRigVMBlueprint); }
+	virtual void InitRigVMEditor(const EToolkitMode::Type Mode, const TSharedPtr< class IToolkitHost >& InitToolkitHost, TRigVMAssetInterfacePtr InRigVMBlueprint) override { return FControlRigBaseEditor::InitRigVMEditorImpl(Mode, InitToolkitHost, InRigVMBlueprint); }
+	virtual void InitRigVMEditorSuper(const EToolkitMode::Type Mode, const TSharedPtr<class IToolkitHost>& InitToolkitHost, TRigVMAssetInterfacePtr InRigVMBlueprint) override { return IControlRigNewEditor::InitRigVMEditor(Mode, InitToolkitHost, InRigVMBlueprint); }
 
 	virtual const FName GetEditorAppName() const override { return FControlRigBaseEditor::GetEditorAppNameImpl(); }
 	virtual const FName GetEditorModeName() const override { return FControlRigBaseEditor::GetEditorModeNameImpl(); }
@@ -30,7 +30,7 @@ public:
 	virtual TSharedRef<const IRigVMEditor> SharedRigVMEditorRef() const override { return StaticCastSharedRef<const IRigVMEditor>(SharedThis(this)); }
 	
 	virtual bool IsControlRigNewEditor() const { return true; };
-	virtual URigVMBlueprint* GetRigVMBlueprint() const { return IControlRigNewEditor::GetRigVMBlueprint(); }
+	virtual TRigVMAssetInterfacePtr GetRigVMAsset() const override { return IControlRigNewEditor::GetRigVMAsset(); }
 	virtual URigVMHost* GetRigVMHost() const { return FRigVMEditorBase::GetRigVMHost(); }
 	virtual bool IsEditorInitialized() const override { return bRigVMEditorInitialized; }
 	virtual TSharedRef<FUICommandList> GetToolkitCommands() { return IControlRigNewEditor::GetToolkitCommands(); }
@@ -155,8 +155,8 @@ private:
 
 	virtual void GenerateEventQueueMenuContent(FMenuBuilder& MenuBuilder) override { return FControlRigBaseEditor::GenerateEventQueueMenuContentImpl(MenuBuilder); }
 
-	virtual void HandleRefreshEditorFromBlueprint(URigVMBlueprint* InBlueprint) override { return FControlRigBaseEditor::HandleRefreshEditorFromBlueprintImpl(InBlueprint); }
-	virtual void HandleRefreshEditorFromBlueprintSuper(URigVMBlueprint* InBlueprint) override { return IControlRigNewEditor::HandleRefreshEditorFromBlueprint(InBlueprint); }
+	virtual void HandleRefreshEditorFromBlueprint(TRigVMAssetInterfacePtr InBlueprint) override { return FControlRigBaseEditor::HandleRefreshEditorFromBlueprintImpl(InBlueprint); }
+	virtual void HandleRefreshEditorFromBlueprintSuper(TRigVMAssetInterfacePtr InBlueprint) override { return IControlRigNewEditor::HandleRefreshEditorFromBlueprint(InBlueprint); }
 
 	
 	/** delegate for changing property */
diff --git a/Engine/Plugins/Animation/ControlRig/Source/ControlRigEditor/Private/Editor/RigValidationTabSummoner.cpp b/Engine/Plugins/Animation/ControlRig/Source/ControlRigEditor/Private/Editor/RigValidationTabSummoner.cpp
--- a/Engine/Plugins/Animation/ControlRig/Source/ControlRigEditor/Private/Editor/RigValidationTabSummoner.cpp
+++ b/Engine/Plugins/Animation/ControlRig/Source/ControlRigEditor/Private/Editor/RigValidationTabSummoner.cpp
@@ -32,7 +32,7 @@ TSharedRef<SWidget> FRigValidationTabSummoner::CreateTabBody(const FWorkflowTabS
 	
 	if (TSharedPtr<FRigVMEditorBase> Editor = EditorPtr)
 	{
-		UControlRigBlueprint* RigBlueprint = Cast<UControlRigBlueprint>(Editor->GetRigVMBlueprint());
+		UControlRigBlueprint* RigBlueprint = Cast<UControlRigBlueprint>(Editor->GetRigVMAsset()->GetObject());
 		check(RigBlueprint);
 	
 		UControlRigValidator* Validator = RigBlueprint->Validator;
diff --git a/Engine/Plugins/Animation/ControlRig/Source/ControlRigEditor/Private/Editor/SControlRigAnimAttributeView.cpp b/Engine/Plugins/Animation/ControlRig/Source/ControlRigEditor/Private/Editor/SControlRigAnimAttributeView.cpp
--- a/Engine/Plugins/Animation/ControlRig/Source/ControlRigEditor/Private/Editor/SControlRigAnimAttributeView.cpp
+++ b/Engine/Plugins/Animation/ControlRig/Source/ControlRigEditor/Private/Editor/SControlRigAnimAttributeView.cpp
@@ -18,7 +18,7 @@ SControlRigAnimAttributeView::~SControlRigAnimAttributeView()
 {
 	StopObservingCurrentControlRig();
 	
-	ControlRigBlueprint->OnSetObjectBeingDebugged().RemoveAll(this);
+	ControlRigBlueprint->UBlueprint::OnSetObjectBeingDebugged().RemoveAll(this);
 }
 
 void SControlRigAnimAttributeView::Construct(const FArguments& InArgs, TSharedRef<IControlRigBaseEditor> InControlRigEditor)
@@ -32,7 +32,7 @@ void SControlRigAnimAttributeView::Construct(const FArguments& InArgs, TSharedRe
 
 	ControlRigEditor = InControlRigEditor;
 	ControlRigBlueprint = ControlRigEditor.Pin()->GetControlRigBlueprint();
-	ControlRigBlueprint->OnSetObjectBeingDebugged().AddRaw(this, &SControlRigAnimAttributeView::HandleSetObjectBeingDebugged);
+	ControlRigBlueprint->UBlueprint::OnSetObjectBeingDebugged().AddRaw(this, &SControlRigAnimAttributeView::HandleSetObjectBeingDebugged);
 
 	if (UControlRig* ControlRig = Cast<UControlRig>(ControlRigBlueprint->GetObjectBeingDebugged()))
 	{
diff --git a/Engine/Plugins/Animation/ControlRig/Source/ControlRigEditor/Private/Editor/SModularRigModel.cpp b/Engine/Plugins/Animation/ControlRig/Source/ControlRigEditor/Private/Editor/SModularRigModel.cpp
--- a/Engine/Plugins/Animation/ControlRig/Source/ControlRigEditor/Private/Editor/SModularRigModel.cpp
+++ b/Engine/Plugins/Animation/ControlRig/Source/ControlRigEditor/Private/Editor/SModularRigModel.cpp
@@ -82,7 +82,7 @@ void SModularRigModel::Construct(const FArguments& InArgs, TSharedRef<IControlRi
 	ControlRigBlueprint = ControlRigEditor.Pin()->GetControlRigBlueprint();
 
 	ControlRigBlueprint->OnRefreshEditor().AddRaw(this, &SModularRigModel::HandleRefreshEditorFromBlueprint);
-	ControlRigBlueprint->OnSetObjectBeingDebugged().AddRaw(this, &SModularRigModel::HandleSetObjectBeingDebugged);
+	ControlRigBlueprint->IRigVMAssetInterface::OnSetObjectBeingDebugged().AddRaw(this, &SModularRigModel::HandleSetObjectBeingDebugged);
 	ControlRigBlueprint->OnModularRigPreCompiled().AddRaw(this, &SModularRigModel::HandlePreCompileModularRigs);
 	ControlRigBlueprint->OnModularRigCompiled().AddRaw(this, &SModularRigModel::HandlePostCompileModularRigs);
 
@@ -257,7 +257,7 @@ void SModularRigModel::OnEditorClose(IControlRigBaseEditor* InEditor, UControlRi
 	if (InBlueprint)
 	{
 		InBlueprint->OnRefreshEditor().RemoveAll(this);
-		InBlueprint->OnSetObjectBeingDebugged().RemoveAll(this);
+		InBlueprint->IRigVMAssetInterface::OnSetObjectBeingDebugged().RemoveAll(this);
 		InBlueprint->OnModularRigPreCompiled().RemoveAll(this);
 		InBlueprint->OnModularRigCompiled().RemoveAll(this);
 		if(UModularRigController* ModularRigController = InBlueprint->GetModularRigController())
@@ -407,11 +407,11 @@ TArray<FString> SModularRigModel::GetSelectedKeys() const
 }
 
 
-void SModularRigModel::HandlePreCompileModularRigs(URigVMBlueprint* InBlueprint)
+void SModularRigModel::HandlePreCompileModularRigs(TRigVMAssetInterfacePtr InBlueprint)
 {
 }
 
-void SModularRigModel::HandlePostCompileModularRigs(URigVMBlueprint* InBlueprint)
+void SModularRigModel::HandlePostCompileModularRigs(TRigVMAssetInterfacePtr InBlueprint)
 {
 	if(!bKeepCurrentEditedConnectors)
 	{
@@ -431,7 +431,7 @@ void SModularRigModel::HandlePostCompileModularRigs(URigVMBlueprint* InBlueprint
 	}
 }
 
-void SModularRigModel::HandleRefreshEditorFromBlueprint(URigVMBlueprint* InBlueprint)
+void SModularRigModel::HandleRefreshEditorFromBlueprint(TRigVMAssetInterfacePtr InBlueprint)
 {
 	RefreshTreeView();
 }
diff --git a/Engine/Plugins/Animation/ControlRig/Source/ControlRigEditor/Private/Editor/SModularRigModel.h b/Engine/Plugins/Animation/ControlRig/Source/ControlRigEditor/Private/Editor/SModularRigModel.h
--- a/Engine/Plugins/Animation/ControlRig/Source/ControlRigEditor/Private/Editor/SModularRigModel.h
+++ b/Engine/Plugins/Animation/ControlRig/Source/ControlRigEditor/Private/Editor/SModularRigModel.h
@@ -151,12 +151,12 @@ private:
 	UModularRig* GetDefaultModularRig() const;
 	const UModularRig* GetModularRigForTreeView() const { return GetModularRig(); }
 	void OnRequestDetailsInspection(const FName& InModuleName);
-	void HandlePreCompileModularRigs(URigVMBlueprint* InBlueprint);
-	void HandlePostCompileModularRigs(URigVMBlueprint* InBlueprint);
+	void HandlePreCompileModularRigs(TRigVMAssetInterfacePtr InBlueprint);
+	void HandlePostCompileModularRigs(TRigVMAssetInterfacePtr InBlueprint);
 	void OnModularRigModified(EModularRigNotification InNotif, const FRigModuleReference* InModule);
 	void OnHierarchyModified(ERigHierarchyNotification InNotif, URigHierarchy* InHierarchy, const FRigNotificationSubject& InSubject);
 
-	void HandleRefreshEditorFromBlueprint(URigVMBlueprint* InBlueprint);
+	void HandleRefreshEditorFromBlueprint(TRigVMAssetInterfacePtr InBlueprint);
 	void HandleSetObjectBeingDebugged(UObject* InObject);
 
 	TSharedRef<SWidget> OnGetOptionsMenu();
diff --git a/Engine/Plugins/Animation/ControlRig/Source/ControlRigEditor/Private/Editor/SRigCurveContainer.cpp b/Engine/Plugins/Animation/ControlRig/Source/ControlRigEditor/Private/Editor/SRigCurveContainer.cpp
--- a/Engine/Plugins/Animation/ControlRig/Source/ControlRigEditor/Private/Editor/SRigCurveContainer.cpp
+++ b/Engine/Plugins/Animation/ControlRig/Source/ControlRigEditor/Private/Editor/SRigCurveContainer.cpp
@@ -145,7 +145,7 @@ void SRigCurveContainer::Construct(const FArguments& InArgs, TSharedRef<IControl
 	bIsChangingRigHierarchy = false;
 
 	ControlRigBlueprint->Hierarchy->OnModified().AddRaw(this, &SRigCurveContainer::OnHierarchyModified_AnyThread);
-	ControlRigBlueprint->OnSetObjectBeingDebugged().AddRaw(this, &SRigCurveContainer::HandleSetObjectBeingDebugged);
+	ControlRigBlueprint->IRigVMAssetInterface::OnSetObjectBeingDebugged().AddRaw(this, &SRigCurveContainer::HandleSetObjectBeingDebugged);
 	ControlRigBlueprint->OnRefreshEditor().AddRaw(this, &SRigCurveContainer::HandleRefreshEditorFromBlueprint);
 
 	UEditorEngine* Editor = Cast<UEditorEngine>(GEngine);
@@ -327,7 +327,7 @@ void SRigCurveContainer::OnEditorClose(IControlRigBaseEditor* InEditor, UControl
 	{
 		InBlueprint->Hierarchy->OnModified().RemoveAll(this);
 		InBlueprint->OnRefreshEditor().RemoveAll(this);
-		InBlueprint->OnSetObjectBeingDebugged().RemoveAll(this);
+		InBlueprint->IRigVMAssetInterface::OnSetObjectBeingDebugged().RemoveAll(this);
 	}
 
 	if (InEditor)
@@ -757,7 +757,7 @@ void SRigCurveContainer::HandleSetObjectBeingDebugged(UObject* InObject)
 	RefreshCurveList_AnyThread();
 }
 
-void SRigCurveContainer::HandleRefreshEditorFromBlueprint(URigVMBlueprint* InBlueprint)
+void SRigCurveContainer::HandleRefreshEditorFromBlueprint(TRigVMAssetInterfacePtr InBlueprint)
 {
 	RefreshCurveList_AnyThread();
 }
diff --git a/Engine/Plugins/Animation/ControlRig/Source/ControlRigEditor/Private/Editor/SRigCurveContainer.h b/Engine/Plugins/Animation/ControlRig/Source/ControlRigEditor/Private/Editor/SRigCurveContainer.h
--- a/Engine/Plugins/Animation/ControlRig/Source/ControlRigEditor/Private/Editor/SRigCurveContainer.h
+++ b/Engine/Plugins/Animation/ControlRig/Source/ControlRigEditor/Private/Editor/SRigCurveContainer.h
@@ -248,7 +248,7 @@ private:
 	void OnHierarchyModified(ERigHierarchyNotification InNotif, URigHierarchy* InHierarchy, const FRigNotificationSubject& InSubject);
 	void OnPostConstruction(UControlRig* InControlRig, const FName& InName);
 	void HandleSetObjectBeingDebugged(UObject* InObject);
-	void HandleRefreshEditorFromBlueprint(URigVMBlueprint* InBlueprint);
+	void HandleRefreshEditorFromBlueprint(TRigVMAssetInterfacePtr InBlueprint);
 
 	// import curve part
 	void ImportCurve(const FAssetData& InAssetData);
diff --git a/Engine/Plugins/Animation/ControlRig/Source/ControlRigEditor/Private/Editor/SRigHierarchy.cpp b/Engine/Plugins/Animation/ControlRig/Source/ControlRigEditor/Private/Editor/SRigHierarchy.cpp
--- a/Engine/Plugins/Animation/ControlRig/Source/ControlRigEditor/Private/Editor/SRigHierarchy.cpp
+++ b/Engine/Plugins/Animation/ControlRig/Source/ControlRigEditor/Private/Editor/SRigHierarchy.cpp
@@ -85,7 +85,7 @@ void SRigHierarchy::Construct(const FArguments& InArgs, TSharedRef<IControlRigBa
 
 	ControlRigBlueprint->Hierarchy->OnModified().AddRaw(this, &SRigHierarchy::OnHierarchyModified);
 	ControlRigBlueprint->OnRefreshEditor().AddRaw(this, &SRigHierarchy::HandleRefreshEditorFromBlueprint);
-	ControlRigBlueprint->OnSetObjectBeingDebugged().AddRaw(this, &SRigHierarchy::HandleSetObjectBeingDebugged);
+	ControlRigBlueprint->UBlueprint::OnSetObjectBeingDebugged().AddRaw(this, &SRigHierarchy::HandleSetObjectBeingDebugged);
 
 	if(UModularRigController* ModularRigController = ControlRigBlueprint->GetModularRigController())
 	{
@@ -271,7 +271,7 @@ void SRigHierarchy::OnEditorClose(IControlRigBaseEditor* InEditor, UControlRigBl
 	{
 		BP->Hierarchy->OnModified().RemoveAll(this);
 		InBlueprint->OnRefreshEditor().RemoveAll(this);
-		InBlueprint->OnSetObjectBeingDebugged().RemoveAll(this);
+		InBlueprint->UBlueprint::OnSetObjectBeingDebugged().RemoveAll(this);
 
 		if(UModularRigController* ModularRigController = BP->GetModularRigController())
 		{
@@ -1031,7 +1031,7 @@ void SRigHierarchy::OnModularRigModified(EModularRigNotification InNotif, const
 	}
 }
 
-void SRigHierarchy::HandleRefreshEditorFromBlueprint(URigVMBlueprint* InBlueprint)
+void SRigHierarchy::HandleRefreshEditorFromBlueprint(TRigVMAssetInterfacePtr InBlueprint)
 {
 	if (bIsChangingRigHierarchy)
 	{
diff --git a/Engine/Plugins/Animation/ControlRig/Source/ControlRigEditor/Private/Editor/SRigHierarchy.h b/Engine/Plugins/Animation/ControlRig/Source/ControlRigEditor/Private/Editor/SRigHierarchy.h
--- a/Engine/Plugins/Animation/ControlRig/Source/ControlRigEditor/Private/Editor/SRigHierarchy.h
+++ b/Engine/Plugins/Animation/ControlRig/Source/ControlRigEditor/Private/Editor/SRigHierarchy.h
@@ -222,7 +222,7 @@ private:
 	void OnHierarchyModified(ERigHierarchyNotification InNotif, URigHierarchy* InHierarchy, const FRigNotificationSubject& InSubject);
 	void OnHierarchyModified_AnyThread(ERigHierarchyNotification InNotif, URigHierarchy* InHierarchy, const FRigNotificationSubject& InSubject);
 	void OnModularRigModified(EModularRigNotification InNotif, const FRigModuleReference* InModule);
-	void HandleRefreshEditorFromBlueprint(URigVMBlueprint* InBlueprint);
+	void HandleRefreshEditorFromBlueprint(TRigVMAssetInterfacePtr InBlueprint);
 	void HandleSetObjectBeingDebugged(UObject* InObject);
 	void OnPreConstruction_AnyThread(UControlRig* InRig, const FName& InEventName);
 	void OnPostConstruction_AnyThread(UControlRig* InRig, const FName& InEventName);
diff --git a/Engine/Plugins/Animation/ControlRig/Source/ControlRigEditor/Private/Editor/SRigModuleAssetBrowser.cpp b/Engine/Plugins/Animation/ControlRig/Source/ControlRigEditor/Private/Editor/SRigModuleAssetBrowser.cpp
--- a/Engine/Plugins/Animation/ControlRig/Source/ControlRigEditor/Private/Editor/SRigModuleAssetBrowser.cpp
+++ b/Engine/Plugins/Animation/ControlRig/Source/ControlRigEditor/Private/Editor/SRigModuleAssetBrowser.cpp
@@ -225,8 +225,8 @@ TSharedRef<SToolTip> SRigModuleAssetBrowser::CreateCustomAssetToolTip(FAssetData
 	// Add asset registry tags to a text list; except skeleton as that is implied in Persona
 	TSharedRef<SVerticalBox> DescriptionBox = SNew(SVerticalBox);
 
-	static const FName AssetVariantPropertyName = GET_MEMBER_NAME_CHECKED(URigVMBlueprint, AssetVariant);
-	const FProperty* AssetVariantProperty = CastField<FProperty>(URigVMBlueprint::StaticClass()->FindPropertyByName(AssetVariantPropertyName));
+	static const FName AssetVariantPropertyName = TEXT("AssetVariant");
+	const FProperty* AssetVariantProperty = CastField<FProperty>(AssetData.GetClass()->FindPropertyByName(AssetVariantPropertyName));
 	const FString VariantStr = AssetData.GetTagValueRef<FString>(AssetVariantPropertyName);
 	if(!VariantStr.IsEmpty())
 	{
@@ -510,8 +510,8 @@ TSharedRef<SToolTip> SRigModuleAssetBrowser::CreateCustomAssetToolTipNewStyle(FA
 	// Add asset registry tags to a text list; except skeleton as that is implied in Persona
 	TMap<FName, FText> TagsAndValuesToShow;
 
-	static const FName AssetVariantPropertyName = GET_MEMBER_NAME_CHECKED(URigVMBlueprint, AssetVariant);
-	const FProperty* AssetVariantProperty = CastField<FProperty>(URigVMBlueprint::StaticClass()->FindPropertyByName(AssetVariantPropertyName));
+	static const FName AssetVariantPropertyName = TEXT("AssetVariant");
+	const FProperty* AssetVariantProperty = CastField<FProperty>(AssetData.GetClass()->FindPropertyByName(AssetVariantPropertyName));
 	const FString VariantStr = AssetData.GetTagValueRef<FString>(AssetVariantPropertyName);
 	if(!VariantStr.IsEmpty())
 	{
diff --git a/Engine/Plugins/Experimental/UAF/UAF/Source/UAFEditor/Private/AnimNextEditorModule.cpp b/Engine/Plugins/Experimental/UAF/UAF/Source/UAFEditor/Private/AnimNextEditorModule.cpp
--- a/Engine/Plugins/Experimental/UAF/UAF/Source/UAFEditor/Private/AnimNextEditorModule.cpp
+++ b/Engine/Plugins/Experimental/UAF/UAF/Source/UAFEditor/Private/AnimNextEditorModule.cpp
@@ -1058,7 +1058,7 @@ void FAnimNextEditorModule::RegisterWorkspaceDocumentTypes(Workspace::IWorkspace
 
 			Controller->OpenUndoBracket(TEXT("Pasted Nodes."));
 
-			if (UE::RigVM::Editor::Tools::PasteNodes(InPasteLocation, TextToImport, Controller, RigVMEdGraph->GetModel(), RigVMClientHost->GetLocalFunctionLibrary(), RigVMClientHost->GetRigVMGraphFunctionHost(), true, true))
+			if (UE::RigVM::Editor::Tools::PasteNodes(InPasteLocation, TextToImport, Controller, RigVMEdGraph->GetModel(), RigVMClientHost->GetLocalFunctionLibrary(), RigVMClientHost->GetRigVMGraphFunctionHost().GetInterface(), true, true))
 			{
 				Controller->CloseUndoBracket();
 			}
diff --git a/Engine/Plugins/Experimental/UAF/UAF/Source/UAFUncookedOnly/Internal/AnimNextRigVMAssetEditorData.h b/Engine/Plugins/Experimental/UAF/UAF/Source/UAFUncookedOnly/Internal/AnimNextRigVMAssetEditorData.h
--- a/Engine/Plugins/Experimental/UAF/UAF/Source/UAFUncookedOnly/Internal/AnimNextRigVMAssetEditorData.h
+++ b/Engine/Plugins/Experimental/UAF/UAF/Source/UAFUncookedOnly/Internal/AnimNextRigVMAssetEditorData.h
@@ -7,7 +7,7 @@
 #include "IAnimNextRigVMExportInterface.h"
 #include "RigVMModel/RigVMGraph.h"
 #include "RigVMCore/RigVMGraphFunctionHost.h"
-#include "RigVMBlueprint.h"
+#include "EdGraph/RigVMEdGraph.h"
 #include "UncookedOnlyUtils.h"
 #include "Kismet/BlueprintFunctionLibrary.h"
 #include "Entries/AnimNextRigVMAssetEntry.h"
@@ -276,8 +276,8 @@ protected:
 	UE_API virtual UClass* GetRigVMEditorSettingsClass() const override;
 	UE_API virtual FRigVMClient* GetRigVMClient() override;
 	UE_API virtual const FRigVMClient* GetRigVMClient() const override;
-	UE_API virtual IRigVMGraphFunctionHost* GetRigVMGraphFunctionHost() override;
-	UE_API virtual const IRigVMGraphFunctionHost* GetRigVMGraphFunctionHost() const override;
+	UE_API virtual TScriptInterface<IRigVMGraphFunctionHost> GetRigVMGraphFunctionHost() override;
+	UE_API virtual TScriptInterface<const IRigVMGraphFunctionHost> GetRigVMGraphFunctionHost() const override;
 	UE_API virtual void HandleRigVMGraphAdded(const FRigVMClient* InClient, const FString& InNodePath) override;
 	UE_API virtual void HandleRigVMGraphRemoved(const FRigVMClient* InClient, const FString& InNodePath) override;
 	UE_API virtual void HandleRigVMGraphRenamed(const FRigVMClient* InClient, const FString& InOldNodePath, const FString& InNewNodePath) override;
diff --git a/Engine/Plugins/Experimental/UAF/UAF/Source/UAFUncookedOnly/Private/AnimNextRigVMAssetEditorData.cpp b/Engine/Plugins/Experimental/UAF/UAF/Source/UAFUncookedOnly/Private/AnimNextRigVMAssetEditorData.cpp
--- a/Engine/Plugins/Experimental/UAF/UAF/Source/UAFUncookedOnly/Private/AnimNextRigVMAssetEditorData.cpp
+++ b/Engine/Plugins/Experimental/UAF/UAF/Source/UAFUncookedOnly/Private/AnimNextRigVMAssetEditorData.cpp
@@ -565,12 +565,12 @@ const FRigVMClient* UAnimNextRigVMAssetEditorData::GetRigVMClient() const
 	return &RigVMClient;
 }
 
-IRigVMGraphFunctionHost* UAnimNextRigVMAssetEditorData::GetRigVMGraphFunctionHost()
+TScriptInterface<IRigVMGraphFunctionHost> UAnimNextRigVMAssetEditorData::GetRigVMGraphFunctionHost()
 {
 	return this;
 }
 
-const IRigVMGraphFunctionHost* UAnimNextRigVMAssetEditorData::GetRigVMGraphFunctionHost() const
+TScriptInterface<const IRigVMGraphFunctionHost> UAnimNextRigVMAssetEditorData::GetRigVMGraphFunctionHost() const
 {
 	return this;
 }
diff --git a/Engine/Plugins/Runtime/RigVM/Source/RigVM/Private/RigVMCore/RigVM.cpp b/Engine/Plugins/Runtime/RigVM/Source/RigVM/Private/RigVMCore/RigVM.cpp
--- a/Engine/Plugins/Runtime/RigVM/Source/RigVM/Private/RigVMCore/RigVM.cpp
+++ b/Engine/Plugins/Runtime/RigVM/Source/RigVM/Private/RigVMCore/RigVM.cpp
@@ -3241,6 +3241,44 @@ void URigVM::RefreshExternalPropertyPaths()
 	}
 }
 
+TArray<const UObject*> URigVM::GetUserDefinedDependencies(const TArray<const FRigVMMemoryStorageStruct*> InMemory)
+{
+	TArray<const UObject*> Dependencies;
+	for (const FRigVMMemoryStorageStruct* MemoryStorage : InMemory)
+	{
+		MemoryStorage->GetUserDefinedDependencies(Dependencies);
+	}
+
+	const TArray<const FRigVMFunction*>& Functions = GetFunctions();
+	for (const FRigVMFunction* Function : Functions)
+	{
+		FRigVMRegistry_NoLock* Registry = GetLocalizedRegistry();
+
+		// if the VM doesn't have a localized registry
+		// then access the global one and lock it for reading.
+		TUniquePtr<FRigVMRegistryReadLock> ReadLock;
+		if(Registry == nullptr)
+		{
+			FRigVMRegistry_RWLock& RegistryRef = FRigVMRegistry_RWLock::Get();
+			Registry = &RegistryRef;
+			ReadLock = MakeUnique<FRigVMRegistryReadLock>(RegistryRef);
+		}
+		
+		const TArray<TRigVMTypeIndex>& TypeIndices = Function->GetArgumentTypeIndices_NoLock(Registry->GetHandle_NoLock());
+		for (const TRigVMTypeIndex& TypeIndex : TypeIndices)
+		{
+			const FRigVMTemplateArgumentType& Type = Registry->GetType_NoLock(TypeIndex);
+			if (Cast<UUserDefinedStruct>(Type.CPPTypeObject) ||
+				Cast<UUserDefinedEnum>(Type.CPPTypeObject))
+			{
+				Dependencies.AddUnique(Type.CPPTypeObject);
+			}
+		}
+	}
+
+	return Dependencies;
+}
+
 void URigVM::SetupInstructionTracking(FRigVMExtendedExecuteContext& Context, int32 InInstructionCount)
 {
 #if WITH_EDITOR
diff --git a/Engine/Plugins/Runtime/RigVM/Source/RigVM/Private/RigVMHost.cpp b/Engine/Plugins/Runtime/RigVM/Source/RigVM/Private/RigVMHost.cpp
--- a/Engine/Plugins/Runtime/RigVM/Source/RigVM/Private/RigVMHost.cpp
+++ b/Engine/Plugins/Runtime/RigVM/Source/RigVM/Private/RigVMHost.cpp
@@ -132,7 +132,7 @@ void URigVMHost::Serialize(FArchive& Ar)
 	// to make sure they are loaded prior to the VM.
 	if (Ar.IsObjectReferenceCollector() && VM != nullptr)
 	{
-		const TArray<const UObject*> UserDefinedDependencies = GetUserDefinedDependencies({ GetDefaultMemoryByType(ERigVMMemoryType::Literal), GetDefaultMemoryByType(ERigVMMemoryType::Work) });
+		const TArray<const UObject*> UserDefinedDependencies = VM->GetUserDefinedDependencies({ GetDefaultMemoryByType(ERigVMMemoryType::Literal), GetDefaultMemoryByType(ERigVMMemoryType::Work) });
 		for (const UObject* UserDefinedDependency : UserDefinedDependencies)
 		{
 			if (Cast<UUserDefinedStruct>(UserDefinedDependency) ||
diff --git a/Engine/Plugins/Runtime/RigVM/Source/RigVM/Public/RigVMBlueprintGeneratedClass.h b/Engine/Plugins/Runtime/RigVM/Source/RigVM/Public/RigVMBlueprintGeneratedClass.h
--- a/Engine/Plugins/Runtime/RigVM/Source/RigVM/Public/RigVMBlueprintGeneratedClass.h
+++ b/Engine/Plugins/Runtime/RigVM/Source/RigVM/Public/RigVMBlueprintGeneratedClass.h
@@ -11,14 +11,6 @@
 
 #define UE_API RIGVM_API
 
-USTRUCT()
-struct FRigVMGraphFunctionHeaderArray
-{
-	GENERATED_BODY()
-	
-	UPROPERTY()
-	TArray<FRigVMGraphFunctionHeader> Headers;
-};
 
 UCLASS(MinimalAPI)
 class URigVMBlueprintGeneratedClass : public UBlueprintGeneratedClass, public IRigVMGraphFunctionHost
diff --git a/Engine/Plugins/Runtime/RigVM/Source/RigVM/Public/RigVMCore/RigVM.h b/Engine/Plugins/Runtime/RigVM/Source/RigVM/Public/RigVMCore/RigVM.h
--- a/Engine/Plugins/Runtime/RigVM/Source/RigVM/Public/RigVMCore/RigVM.h
+++ b/Engine/Plugins/Runtime/RigVM/Source/RigVM/Public/RigVMCore/RigVM.h
@@ -1031,6 +1031,7 @@ protected:
 public:
 	RIGVM_API bool ValidateBytecode();
 	RIGVM_API void RefreshArgumentNameCaches();
+	RIGVM_API TArray<const UObject*> GetUserDefinedDependencies(const TArray<const FRigVMMemoryStorageStruct*> InMemory);
 
 private:
 
diff --git a/Engine/Plugins/Runtime/RigVM/Source/RigVM/Public/RigVMCore/RigVMGraphFunctionDefinition.h b/Engine/Plugins/Runtime/RigVM/Source/RigVM/Public/RigVMCore/RigVMGraphFunctionDefinition.h
--- a/Engine/Plugins/Runtime/RigVM/Source/RigVM/Public/RigVMCore/RigVMGraphFunctionDefinition.h
+++ b/Engine/Plugins/Runtime/RigVM/Source/RigVM/Public/RigVMCore/RigVMGraphFunctionDefinition.h
@@ -504,6 +504,15 @@ protected:
 	friend struct FRigVMGraphFunctionData;
 };
 
+USTRUCT()
+struct FRigVMGraphFunctionHeaderArray
+{
+	GENERATED_BODY()
+	
+	UPROPERTY()
+	TArray<FRigVMGraphFunctionHeader> Headers;
+};
+
 USTRUCT(BlueprintType)
 struct FRigVMGraphFunctionData
 {
diff --git a/Engine/Plugins/Runtime/RigVM/Source/RigVM/Public/RigVMHost.h b/Engine/Plugins/Runtime/RigVM/Source/RigVM/Public/RigVMHost.h
--- a/Engine/Plugins/Runtime/RigVM/Source/RigVM/Public/RigVMHost.h
+++ b/Engine/Plugins/Runtime/RigVM/Source/RigVM/Public/RigVMHost.h
@@ -22,6 +22,19 @@
 #define UE_RIGVM_PROFILE_EXECUTE_UNITS_NUM 0
 #endif
 
+USTRUCT()
+struct FRigVMUserDefinedTypesInfo
+{
+	GENERATED_BODY()
+
+	UPROPERTY()
+	TMap<FString, FSoftObjectPath> StructGuidToPathName;
+	UPROPERTY()
+	TMap<FString, FSoftObjectPath> EnumToPathName;
+	UPROPERTY(transient)
+	TSet<TObjectPtr<UObject>> TypesInUse;
+};
+
 UCLASS(Abstract, editinlinenew, MinimalAPI)
 class URigVMHost : public UObject, public IInterface_AssetUserData
 {
@@ -549,6 +562,7 @@ protected:
 #endif
 
 	friend class URigVMBlueprint;
+	friend class IRigVMAssetInterface;
 	friend class URigVMBlueprintGeneratedClass;
 	friend class FRigVMCompileSettingsDetails;
 };
diff --git a/Engine/Plugins/Runtime/RigVM/Source/RigVMDeveloper/Private/EdGraph/RigVMEdGraph.cpp b/Engine/Plugins/Runtime/RigVM/Source/RigVMDeveloper/Private/EdGraph/RigVMEdGraph.cpp
--- a/Engine/Plugins/Runtime/RigVM/Source/RigVMDeveloper/Private/EdGraph/RigVMEdGraph.cpp
+++ b/Engine/Plugins/Runtime/RigVM/Source/RigVMDeveloper/Private/EdGraph/RigVMEdGraph.cpp
@@ -1,13 +1,14 @@
 // Copyright Epic Games, Inc. All Rights Reserved.
 
 #include "EdGraph/RigVMEdGraph.h"
-#include "RigVMBlueprint.h"
+#include "RigVMAsset.h"
 #include "RigVMBlueprintGeneratedClass.h"
 #include "Editor/EditorEngine.h"
 #include "EdGraph/RigVMEdGraphSchema.h"
 #include "RigVMModel/RigVMGraph.h"
 #include "RigVMObjectVersion.h"
 #include "EdGraphNode_Comment.h"
+#include "RigVMBlueprintLegacy.h"
 #include "RigVMModel/Nodes/RigVMLibraryNode.h"
 #include "RigVMModel/Nodes/RigVMFunctionEntryNode.h"
 #include "RigVMModel/Nodes/RigVMFunctionReturnNode.h"
@@ -105,10 +106,19 @@ void URigVMEdGraph::HandleRigVMGraphRenamed(const FString& InOldNodePath, const
 }
 
 const URigVMBlueprint* URigVMEdGraph::GetBlueprintDefaultObject() const
+{
+	return Cast<URigVMBlueprint>(GetAssetDefaultObject().GetObject());
+}
+
+const TRigVMAssetInterfacePtr URigVMEdGraph::GetAssetDefaultObject() const
 {
 	if(RigVMBlueprintClass)
 	{
-		return RigVMBlueprintClass->GetDefaultObject<URigVMBlueprint>();
+		UObject* Obj = RigVMBlueprintClass->GetDefaultObject(false);
+		if (Obj->Implements<URigVMAssetInterface>())
+		{
+			return Obj;
+		}
 	}
 	return nullptr;
 }
@@ -116,17 +126,17 @@ const URigVMBlueprint* URigVMEdGraph::GetBlueprintDefaultObject() const
 void URigVMEdGraph::SetBlueprintClass(const UClass* InClass)
 {
 	RigVMBlueprintClass = InClass;
-	if(const URigVMBlueprint* Blueprint = GetBlueprintDefaultObject())
+	if(const TRigVMAssetInterfacePtr Blueprint = GetAssetDefaultObject())
 	{
-		Schema = Blueprint->GetRigVMEdGraphSchemaClass();
+		Schema = Blueprint->GetRigVMClientHost()->GetRigVMEdGraphSchemaClass();
 	}
 }
 
-void URigVMEdGraph::InitializeFromBlueprint(URigVMBlueprint* InBlueprint)
+void URigVMEdGraph::InitializeFromAsset(TRigVMAssetInterfacePtr InBlueprint)
 {
 	DECLARE_SCOPE_HIERARCHICAL_COUNTER_FUNC()
 
-	SetBlueprintClass(InBlueprint->GetClass());
+	SetBlueprintClass(InBlueprint->GetObject()->GetClass());
 	InBlueprint->OnModified().RemoveAll(this);
 	InBlueprint->OnModified().AddUObject(this, &URigVMEdGraph::HandleModifiedEvent);
 	InBlueprint->OnVMCompiled().RemoveAll(this);
@@ -900,7 +910,7 @@ void URigVMEdGraph::CacheEntryNameList()
 	EntryNameList.Reset();
 	EntryNameList.Add(MakeShared<FRigVMStringWithTag>(FName(NAME_None).ToString()));
 
-	if(const URigVMBlueprint* Blueprint = CastChecked<URigVMBlueprint>(GetBlueprint()))
+	if(const TScriptInterface<IRigVMClientHost> Blueprint = GetAsset()->GetRigVMClientHost())
 	{
 		const TArray<FName> EntryNames = Blueprint->GetRigVMClient()->GetEntryNames();
 		for (const FName& EntryName : EntryNames)
@@ -965,6 +975,11 @@ UEdGraphNode* URigVMEdGraph::FindNodeForModelNodeName(const FName& InModelNodeNa
 	return nullptr;
 }
 
+URigVMBlueprint* URigVMEdGraph::GetBlueprint() const
+{
+	return Cast<URigVMBlueprint>(GetAsset().GetObject());
+}
+
 void URigVMEdGraph::AddLocalVariableSearchMetaDataInfo(const FName InVariableName, TArray<UBlueprintExtension::FSearchTagDataPair>& OutTaggedMetaData) const
 {
 	if (URigVMGraph* RigVMGraph = GetModel())
@@ -989,13 +1004,13 @@ void URigVMEdGraph::AddLocalVariableSearchMetaDataInfo(const FName InVariableNam
 	}
 }
 
-URigVMBlueprint* URigVMEdGraph::GetBlueprint() const
+TRigVMAssetInterfacePtr URigVMEdGraph::GetAsset() const
 {
 	if (URigVMEdGraph* OuterGraph = Cast<URigVMEdGraph>(GetOuter()))
 	{
-		return OuterGraph->GetBlueprint();
+		return OuterGraph->GetAsset();
 	}
-	return Cast<URigVMBlueprint>(GetOuter());
+	return GetOuter();
 }
 
 URigVMGraph* URigVMEdGraph::GetModel() const
diff --git a/Engine/Plugins/Runtime/RigVM/Source/RigVMDeveloper/Private/EdGraph/RigVMEdGraphNode.cpp b/Engine/Plugins/Runtime/RigVM/Source/RigVMDeveloper/Private/EdGraph/RigVMEdGraphNode.cpp
--- a/Engine/Plugins/Runtime/RigVM/Source/RigVMDeveloper/Private/EdGraph/RigVMEdGraphNode.cpp
+++ b/Engine/Plugins/Runtime/RigVM/Source/RigVMDeveloper/Private/EdGraph/RigVMEdGraphNode.cpp
@@ -18,8 +18,9 @@
 #endif
 #include "RigVMEditorModule.h"
 #include "Textures/SlateIcon.h"
-#include "RigVMBlueprint.h"
+#include "RigVMAsset.h"
 #include "PropertyPathHelpers.h"
+#include "RigVMBlueprintLegacy.h"
 #include "RigVMBlueprintUtils.h"
 #include "RigVMStringUtils.h"
 #include "RigVMCore/RigVMExecuteContext.h"
@@ -130,12 +131,12 @@ FText URigVMEdGraphNode::GetNodeTitle(ENodeTitleType::Type TitleType) const
 
 				if (SubTitle.Len() == 0)
 				{
-					if(UBlueprint* Blueprint = GetBlueprint())
+					if(TRigVMAssetInterfacePtr Blueprint = GetAsset())
 					{
 						const FName VariableName = VariableNode->GetVariableName();
-						for(const FBPVariableDescription& NewVariable : Blueprint->NewVariables)
+						for(const FRigVMGraphVariableDescription& NewVariable : Blueprint->GetAssetVariables())
 						{
-							if(NewVariable.VarName == VariableName)
+							if(NewVariable.Name == VariableName)
 							{
 								FString DefaultValue = NewVariable.DefaultValue;
 								if(DefaultValue.IsEmpty())
@@ -145,8 +146,9 @@ FText URigVMEdGraphNode::GetNodeTitle(ENodeTitleType::Type TitleType) const
 								}
 								else
 								{
+									const FEdGraphPinType PinType = NewVariable.ToPinType();
 									// Change the order of values in rotators so that they match the pin order
-									if (!NewVariable.VarType.IsContainer() && NewVariable.VarType.PinSubCategoryObject == TBaseStructure<FRotator>::Get())
+									if (!PinType.IsContainer() && PinType.PinSubCategoryObject == TBaseStructure<FRotator>::Get())
 									{
 										TArray<FString> Values;
 										DefaultValue.ParseIntoArray(Values, TEXT(","));
@@ -246,7 +248,7 @@ void URigVMEdGraphNode::ReconstructNode_Internal(bool bForce)
 	// Maintain watches up to date
 	if (const URigVMNode* Node = GetModelNode())
 	{
-		const UBlueprint* RigVMBlueprint = GetBlueprint();
+		TRigVMAssetInterfacePtr RigVMBlueprint = GetAsset();
 		for (UEdGraphPin* NewPin : Pins)
 		{
 			const FString PinName = NewPin->GetName();
@@ -256,7 +258,7 @@ void URigVMEdGraphNode::ReconstructNode_Internal(bool bForce)
 			{
 				if (ModelPin->RequiresWatch())
 				{
-					FKismetDebugUtilities::AddPinWatch(RigVMBlueprint, FBlueprintWatchedPin(NewPin));
+					RigVMBlueprint->AddPinWatch(NewPin);
 				}
 			}
 		}
@@ -369,11 +371,11 @@ void URigVMEdGraphNode::DestroyPinList(TArray<UEdGraphPin*>& InPins)
 {
 	DECLARE_SCOPE_HIERARCHICAL_COUNTER_FUNC()
 
-	const UBlueprint* Blueprint = GetBlueprint();
+	const TRigVMAssetInterfacePtr Blueprint = GetAsset();
 	bool bNotify = false;
 	if (Blueprint != nullptr)
 	{
-		bNotify = !Blueprint->bIsRegeneratingOnLoad;
+		bNotify = !Blueprint->IsRegeneratingOnLoad();
 	}
 
 	// Throw away the original pins
@@ -1615,25 +1617,12 @@ void URigVMEdGraphNode::AllocateDefaultPins()
 
 UClass* URigVMEdGraphNode::GetRigVMGeneratedClass() const
 {
-	if(const URigVMBlueprint* Blueprint = GetTypedOuter<URigVMBlueprint>())
+	if(const IRigVMAssetInterface* Blueprint = GetImplementingOuter<IRigVMAssetInterface>())
 	{
-		if (Blueprint->GeneratedClass)
+		if (const UClass* GeneratedClass = Blueprint->GetRigVMGeneratedClass())
 		{
-			check(Blueprint->GeneratedClass->IsChildOf(URigVMHost::StaticClass()));
-			return Blueprint->GeneratedClass;
-		}
-	}
-	return nullptr;
-}
-
-UClass* URigVMEdGraphNode::GetRigVMSkeletonGeneratedClass() const
-{
-	if(const URigVMBlueprint* Blueprint = GetTypedOuter<URigVMBlueprint>())
-	{
-		if (Blueprint->SkeletonGeneratedClass)
-		{
-			check(Blueprint->SkeletonGeneratedClass->IsChildOf(URigVMHost::StaticClass()));
-			return Blueprint->SkeletonGeneratedClass;
+			check(GeneratedClass->IsChildOf(URigVMHost::StaticClass()));
+			return const_cast<UClass*>(GeneratedClass);
 		}
 	}
 	return nullptr;
@@ -1814,7 +1803,7 @@ void URigVMEdGraphNode::GetNodeContextMenuActions(class UToolMenu* Menu, class U
 {
 #if WITH_EDITOR
 	const URigVMEdGraphSchema* Schema = Cast<URigVMEdGraphSchema>(GetSchema());
-	if(const URigVMBlueprint* Blueprint = GetBlueprint())
+	if(const TRigVMAssetInterfacePtr Blueprint = GetAsset())
 	{
 		return Blueprint->GetEditorModule()->GetContextMenuActions(Schema, Menu, Context);
 	}
@@ -1838,13 +1827,12 @@ void URigVMEdGraphNode::DestroyNode()
 	if(URigVMEdGraph* Graph = Cast<URigVMEdGraph>(GetOuter()))
 	{
 		BreakAllNodeLinks();
-		
-		URigVMBlueprint* RigVMBlueprint = Cast<URigVMBlueprint>(Graph->GetOuter());
-		if(RigVMBlueprint)
+
+		if(IRigVMAssetInterface* RigVMBlueprint = Graph->GetImplementingOuter<IRigVMAssetInterface>())
 		{
 			if(PropertyName_DEPRECATED.IsValid())
 			{
-				FRigVMBlueprintUtils::RemoveMemberVariableIfNotUsed(RigVMBlueprint, PropertyName_DEPRECATED);
+				FRigVMBlueprintUtils::RemoveMemberVariableIfNotUsed(RigVMBlueprint->GetObject(), PropertyName_DEPRECATED);
 			}
 		}
 	}
@@ -1908,7 +1896,16 @@ URigVMBlueprint* URigVMEdGraphNode::GetBlueprint() const
 {
 	if(const URigVMEdGraph* Graph = Cast<URigVMEdGraph>(GetOuter()))
 	{
-		return Graph->GetBlueprint();
+		return Cast<URigVMBlueprint>(Graph->GetAsset().GetObject());
+	}
+	return nullptr;
+}
+
+TRigVMAssetInterfacePtr URigVMEdGraphNode::GetAsset() const
+{
+	if(const URigVMEdGraph* Graph = Cast<URigVMEdGraph>(GetOuter()))
+	{
+		return Graph->GetAsset();
 	}
 	return nullptr;
 }
@@ -2277,39 +2274,39 @@ void URigVMEdGraphNode::UpdateVisualSettings()
 	
 	if(const URigVMEdGraph* RigGraph = Cast<URigVMEdGraph>(GetOuter()))
 	{
-		if(const URigVMBlueprint* RigBlueprint = RigGraph->GetBlueprint())
+		if(const TRigVMAssetInterfacePtr RigBlueprint = RigGraph->GetAsset())
 		{
-			bEnableProfiling = RigBlueprint->VMRuntimeSettings.bEnableProfiling;
+			bEnableProfiling = RigBlueprint->GetVMRuntimeSettings().bEnableProfiling;
 			if (bEnableProfiling.Get(false))
 			{
-				if(URigVMBlueprint* Blueprint = GetTypedOuter<URigVMBlueprint>())
+				if(IRigVMAssetInterface* Blueprint = GetImplementingOuter<IRigVMAssetInterface>())
 				{
 					if(URigVMHost* DebuggedHost = Cast<URigVMHost>(Blueprint->GetObjectBeingDebugged()))
 					{
 						if(const URigVMNode* ModelNode = GetModelNode())
 						{
 							const double CurrentMicroSeconds = ModelNode->GetInstructionMicroSeconds(DebuggedHost->GetRigVMExtendedExecuteContext(), DebuggedHost->GetVM(), FRigVMASTProxy());
-							MicroSeconds = Blueprint->RigGraphDisplaySettings.AggregateAverage(MicroSecondsFrames, MicroSeconds.Get(0), CurrentMicroSeconds);
+							MicroSeconds = Blueprint->GetRigGraphDisplaySettings().AggregateAverage(MicroSecondsFrames, MicroSeconds.Get(0), CurrentMicroSeconds);
 							if(MicroSeconds.GetValue() >= 0.0)
 							{
-								if(Blueprint->RigGraphDisplaySettings.bAutoDetermineRange)
+								if(Blueprint->GetRigGraphDisplaySettings().bAutoDetermineRange)
 								{
-									if(CurrentMicroSeconds < Blueprint->RigGraphDisplaySettings.MinMicroSeconds)
+									if(CurrentMicroSeconds < Blueprint->GetRigGraphDisplaySettings().MinMicroSeconds)
 									{
-										Blueprint->RigGraphDisplaySettings.MinMicroSeconds = CurrentMicroSeconds;
+										Blueprint->GetRigGraphDisplaySettings().MinMicroSeconds = CurrentMicroSeconds;
 									}
-									if(CurrentMicroSeconds > Blueprint->RigGraphDisplaySettings.MaxMicroSeconds)
+									if(CurrentMicroSeconds > Blueprint->GetRigGraphDisplaySettings().MaxMicroSeconds)
 									{
-										Blueprint->RigGraphDisplaySettings.MaxMicroSeconds = CurrentMicroSeconds;
+										Blueprint->GetRigGraphDisplaySettings().MaxMicroSeconds = CurrentMicroSeconds;
 									}
 								}
 							
-								const double MinMicroSeconds = Blueprint->RigGraphDisplaySettings.LastMinMicroSeconds;
-								const double MaxMicroSeconds = Blueprint->RigGraphDisplaySettings.LastMaxMicroSeconds;
+								const double MinMicroSeconds = Blueprint->GetRigGraphDisplaySettings().LastMinMicroSeconds;
+								const double MaxMicroSeconds = Blueprint->GetRigGraphDisplaySettings().LastMaxMicroSeconds;
 								if(MaxMicroSeconds > MinMicroSeconds)
 								{
-									const FLinearColor& MinColor = Blueprint->RigGraphDisplaySettings.MinDurationColor;
-									const FLinearColor& MaxColor = Blueprint->RigGraphDisplaySettings.MaxDurationColor;
+									const FLinearColor& MinColor = Blueprint->GetRigGraphDisplaySettings().MinDurationColor;
+									const FLinearColor& MaxColor = Blueprint->GetRigGraphDisplaySettings().MaxDurationColor;
 
 									const double T = (MicroSeconds.GetValue() - MinMicroSeconds) / (MaxMicroSeconds - MinMicroSeconds);
 									ProfilingColor = FMath::Lerp<FLinearColor>(MinColor, MaxColor, (float)T);
