Sequencer: Disable the track instead of disabling the sections

diff --git a/Engine/Plugins/Animation/ControlRig/Source/ControlRigEditor/Private/ControlRigSequencerEditorLibrary.cpp b/Engine/Plugins/Animation/ControlRig/Source/ControlRigEditor/Private/ControlRigSequencerEditorLibrary.cpp
--- a/Engine/Plugins/Animation/ControlRig/Source/ControlRigEditor/Private/ControlRigSequencerEditorLibrary.cpp
+++ b/Engine/Plugins/Animation/ControlRig/Source/ControlRigEditor/Private/ControlRigSequencerEditorLibrary.cpp
@@ -1386,14 +1386,17 @@ bool UControlRigSequencerEditorLibrary::BakeToControlRig(UWorld* World, ULevelSe
 					if (SkelTrack)
 					{
 						SkelTrack->Modify();
-						//can't just turn off the track so need to mute the sections
-						const TArray<UMovieSceneSection*>& Sections = SkelTrack->GetAllSections();
-						for (UMovieSceneSection* Section : Sections)
+
+						// Disable the entire track or each of the rows
+						if (SkelTrack->GetMaxRowIndex() == 0)
+						{
+							SkelTrack->SetEvalDisabled(true);
+						}
+						else
 						{
-							if (Section)
+							for (int32 RowIndex = 0; RowIndex <= SkelTrack->GetMaxRowIndex(); ++RowIndex)
 							{
-								Section->TryModify();
-								Section->SetIsActive(false);
+								SkelTrack->SetRowEvalDisabled(true, RowIndex);
 							}
 						}
 					}
diff --git a/Engine/Plugins/Animation/ControlRig/Source/ControlRigEditor/Private/Sequencer/ControlRigParameterTrackEditor.cpp b/Engine/Plugins/Animation/ControlRig/Source/ControlRigEditor/Private/Sequencer/ControlRigParameterTrackEditor.cpp
--- a/Engine/Plugins/Animation/ControlRig/Source/ControlRigEditor/Private/Sequencer/ControlRigParameterTrackEditor.cpp
+++ b/Engine/Plugins/Animation/ControlRig/Source/ControlRigEditor/Private/Sequencer/ControlRigParameterTrackEditor.cpp
@@ -1330,15 +1330,17 @@ void FControlRigParameterTrackEditor::BakeToControlRig(UClass* InClass, FGuid Ob
 					for (UMovieSceneSkeletalAnimationTrack* SkelTrack : SkelAnimationTracks)
 					{
 						SkelTrack->Modify();
-						
-						//can't just turn off the track so need to mute the sections
-						const TArray<UMovieSceneSection*>& Sections = SkelTrack->GetAllSections();
-						for (UMovieSceneSection* Section : Sections)
+
+						// Disable the entire track or each of the rows
+						if (SkelTrack->GetMaxRowIndex() == 0)
+						{
+							SkelTrack->SetEvalDisabled(true);
+						}
+						else
 						{
-							if (Section)
+							for (int32 RowIndex = 0; RowIndex <= SkelTrack->GetMaxRowIndex(); ++RowIndex)
 							{
-								Section->TryModify();
-								Section->SetIsActive(false);
+								SkelTrack->SetRowEvalDisabled(true, RowIndex);
 							}
 						}
 					}
