Space Switching & Constraints: make sure that constraints are evaluated when switching

#rnx

diff --git a/Engine/Plugins/Animation/ControlRig/Source/ControlRigEditor/Private/Sequencer/ControlRigSpaceChannelEditors.cpp b/Engine/Plugins/Animation/ControlRig/Source/ControlRigEditor/Private/Sequencer/ControlRigSpaceChannelEditors.cpp
--- a/Engine/Plugins/Animation/ControlRig/Source/ControlRigEditor/Private/Sequencer/ControlRigSpaceChannelEditors.cpp
+++ b/Engine/Plugins/Animation/ControlRig/Source/ControlRigEditor/Private/Sequencer/ControlRigSpaceChannelEditors.cpp
@@ -299,17 +299,27 @@ static TTuple<FRigControlElement*, FChannelMapInfo*, int32, int32> GetControlAnd
 *  4) For 1-3 we always compensate any transform keys to the new space
 */
 
-FKeyHandle FControlRigSpaceChannelHelpers::SequencerKeyControlRigSpaceChannel(UControlRig* ControlRig, ISequencer* Sequencer, FMovieSceneControlRigSpaceChannel* Channel, UMovieSceneSection* SectionToKey, FFrameNumber Time, URigHierarchy* RigHierarchy, const FRigElementKey& ControlKey, const FRigElementKey& SpaceKey)
+FKeyHandle FControlRigSpaceChannelHelpers::SequencerKeyControlRigSpaceChannel(
+	UControlRig* ControlRig,
+	ISequencer* Sequencer,
+	FMovieSceneControlRigSpaceChannel* Channel,
+	UMovieSceneSection* SectionToKey,
+	FFrameNumber Time,
+	URigHierarchy* RigHierarchy,
+	const FRigElementKey& ControlKey,
+	const FRigElementKey& SpaceKey)
 {
 	FKeyHandle Handle = FKeyHandle::Invalid();
-	if (bDoNotCompensate == true || ControlRig == nullptr || Sequencer == nullptr || Sequencer->GetFocusedMovieSceneSequence() == nullptr || Sequencer->GetFocusedMovieSceneSequence()->GetMovieScene() == nullptr)
+
+	UMovieSceneSequence* MovieSceneSequence = Sequencer ? Sequencer->GetFocusedMovieSceneSequence() : nullptr;
+	UMovieScene* MovieScene = MovieSceneSequence ? MovieSceneSequence->GetMovieScene() : nullptr;
+	if (bDoNotCompensate == true || ControlRig == nullptr || !MovieScene)
 	{
 		return Handle;
 	}
 
 	TGuardValue<bool> CompensateGuard(bDoNotCompensate, true);
 
-	URigHierarchy* Hierarchy = ControlRig->GetHierarchy();
 	FMovieSceneControlRigSpaceBaseKey ExistingValue, PreviousValue;
 	using namespace UE::MovieScene;
 	EvaluateChannel(Channel, Time, ExistingValue);
@@ -348,15 +358,25 @@ FKeyHandle FControlRigSpaceChannelHelpers::SequencerKeyControlRigSpaceChannel(UC
 	//we only key if the value is different.
 	if (Value != ExistingValue)
 	{
-		FFrameRate TickResolution = Sequencer->GetFocusedTickResolution();
-		FMovieSceneSequenceTransform RootToLocalTransform = Sequencer->GetFocusedMovieSceneSequenceTransform();
-
+		const FFrameRate TickResolution = Sequencer->GetFocusedTickResolution();
+		const FMovieSceneSequenceTransform RootToLocalTransform = Sequencer->GetFocusedMovieSceneSequenceTransform();
+		const FMovieSceneInverseSequenceTransform LocalToRootTransform = RootToLocalTransform.Inverse();
+		
+		const FConstraintsManagerController& Controller = FConstraintsManagerController::Get(ControlRig->GetWorld());
+		auto EvaluateSequencer = [Sequencer, TickResolution, &Controller](const FFrameTime& Time)
+		{
+			FMovieSceneContext SceneContext(FMovieSceneEvaluationRange(Time, TickResolution), Sequencer->GetPlaybackStatus());
+			SceneContext.SetHasJumped(true);
+			
+			Sequencer->GetEvaluationTemplate().EvaluateSynchronousBlocking(SceneContext);
+			Controller.EvaluateAllConstraints();
+		};
+		
 		//make sure to evaluate first frame
 		FFrameTime CurrentTime(Time);
-		CurrentTime = RootToLocalTransform.Inverse().TryTransformTime(CurrentTime).Get(CurrentTime);
+		CurrentTime = LocalToRootTransform.TryTransformTime(CurrentTime).Get(CurrentTime);
 
-		FMovieSceneContext SceneContext = FMovieSceneContext(FMovieSceneEvaluationRange(CurrentTime, TickResolution), Sequencer->GetPlaybackStatus()).SetHasJumped(true);
-		Sequencer->GetEvaluationTemplate().EvaluateSynchronousBlocking(SceneContext);
+		EvaluateSequencer(CurrentTime);
 		ControlRig->Evaluate_AnyThread();
 
 		TArray<FFrameNumber> Frames;
@@ -366,8 +386,7 @@ FKeyHandle FControlRigSpaceChannelHelpers::SequencerKeyControlRigSpaceChannel(UC
 		bool bSetPreviousKey = true;
 		//find all of the times in the space after this time we now need to compensate for
 		TSortedMap<FFrameNumber, FFrameNumber> ExtraFrames;
-		FControlRigSpaceChannelHelpers::GetFramesInThisSpaceAfterThisTime(ControlRig, ControlKey.Name, ExistingValue,
-			Channel, SectionToKey, Time, ExtraFrames);
+		GetFramesInThisSpaceAfterThisTime(ControlRig, ControlKey.Name, ExistingValue, Channel, SectionToKey, Time, ExtraFrames);
 		SectionToKey->Modify();
 		UMovieSceneControlRigParameterSection* ControlRigSection = Cast<UMovieSceneControlRigParameterSection>(SectionToKey);
 
@@ -397,7 +416,7 @@ FKeyHandle FControlRigSpaceChannelHelpers::SequencerKeyControlRigSpaceChannel(UC
 			}
 
 		}
-		const FFrameNumber StartFrame = Sequencer->GetFocusedMovieSceneSequence()->GetMovieScene()->GetPlaybackRange().GetLowerBoundValue();
+		const FFrameNumber StartFrame = MovieScene->GetPlaybackRange().GetLowerBoundValue();
 		if (StartFrame == Time)
 		{
 			bSetPreviousKey = false;
@@ -413,11 +432,8 @@ FKeyHandle FControlRigSpaceChannelHelpers::SequencerKeyControlRigSpaceChannel(UC
 		}
 
 		TArray<FTransform> ControlRigParentWorldTransforms;
-		ControlRigParentWorldTransforms.SetNum(Frames.Num());
-		for (FTransform& WorldTransform : ControlRigParentWorldTransforms)
-		{
-			WorldTransform  = FTransform::Identity;
-		}
+		ControlRigParentWorldTransforms.Init(FTransform::Identity, Frames.Num());
+		
 		TArray<FTransform> ControlWorldTransforms;
 		FControlRigSnapper Snapper;
 		Snapper.GetControlRigControlTransforms(Sequencer, ControlRig, ControlKey.Name, Frames, ControlRigParentWorldTransforms, ControlWorldTransforms);
@@ -474,30 +490,28 @@ FKeyHandle FControlRigSpaceChannelHelpers::SequencerKeyControlRigSpaceChannel(UC
 				{
 					ChannelInterface.DeleteKeys(OurKeyHandles);
 				}
-				//now delete any extra TimeOfDelete -1
+
+				// now delete any extra TimeOfDelete -1
 				// NOTE do we need to update the tangents here?
-				FControlRigSpaceChannelHelpers::DeleteTransformKeysAtThisTime(ControlRig, ControlRigSection, ControlKey.Name, Time - 1);
+				DeleteTransformKeysAtThisTime(ControlRig, ControlRigSection, ControlKey.Name, Time - 1);
 				bSetPreviousKey = false; //also don't set previous
 			}
 		}
 
-		FMovieSceneInverseSequenceTransform LocalToRootTransform = RootToLocalTransform.Inverse();
+		FRigControlModifiedContext Context(EControlRigSetKey::Always);
+		constexpr bool bNoUndo = false,  bNoPython = false, bFixEulerFlips = true;
 
 		//do any compensation or previous key adding
-		FRigControlModifiedContext Context;
-		Context.SetKey = EControlRigSetKey::Always;
-
 		// set previous key if we're going to switch space
 		if (bSetPreviousKey)
 		{
 			FFrameTime GlobalTime(Time - 1);
 			GlobalTime = LocalToRootTransform.TryTransformTime(GlobalTime).Get(GlobalTime);
 
-			SceneContext = FMovieSceneContext(FMovieSceneEvaluationRange(GlobalTime, TickResolution), Sequencer->GetPlaybackStatus()).SetHasJumped(true);
-			Sequencer->GetEvaluationTemplate().EvaluateSynchronousBlocking(SceneContext);
+			EvaluateSequencer(GlobalTime);
 
 			Context.LocalTime = TickResolution.AsSeconds(FFrameTime(Time - 1));
-			ControlRig->SetControlGlobalTransform(ControlKey.Name, ControlWorldTransforms[0], true, Context, false /*undo*/, false /*bPrintPython*/, true/* bFixEulerFlips*/);
+			ControlRig->SetControlGlobalTransform(ControlKey.Name, ControlWorldTransforms[0], true, Context, bNoUndo, bNoPython, bFixEulerFlips);
 			if (ControlRig->IsAdditive())
 			{
 				// We need to evaluate in order to trigger a notification
@@ -527,12 +541,13 @@ FKeyHandle FControlRigSpaceChannelHelpers::SequencerKeyControlRigSpaceChannel(UC
 			FFrameTime GlobalTime(Frame);
 			GlobalTime = LocalToRootTransform.TryTransformTime(GlobalTime).Get(GlobalTime);
 
-			SceneContext = FMovieSceneContext(FMovieSceneEvaluationRange(GlobalTime, TickResolution), Sequencer->GetPlaybackStatus()).SetHasJumped(true);
-			Sequencer->GetEvaluationTemplate().EvaluateSynchronousBlocking(SceneContext);
+			EvaluateSequencer(GlobalTime);
 
 			ControlRig->Evaluate_AnyThread();
+
 			Context.LocalTime = TickResolution.AsSeconds(FFrameTime(Frame));
-			ControlRig->SetControlGlobalTransform(ControlKey.Name, ControlWorldTransforms[FramesIndex], true, Context, false /*undo*/, false /*bPrintPython*/, true/* bFixEulerFlips*/);
+			ControlRig->SetControlGlobalTransform(ControlKey.Name, ControlWorldTransforms[FramesIndex], true, Context, bNoUndo, bNoPython, bFixEulerFlips);
+
 			if (ControlRig->IsAdditive())
 			{
 				// We need to evaluate in order to trigger a notification
