Anim in Engine: Anim Details now also displays properties that do not relate to a control rig and reside within a Shot Track

diff --git a/Engine/Plugins/Animation/ControlRig/Source/ControlRigEditor/Private/AnimDetails/AnimDetailsProxyManager.cpp b/Engine/Plugins/Animation/ControlRig/Source/ControlRigEditor/Private/AnimDetails/AnimDetailsProxyManager.cpp
--- a/Engine/Plugins/Animation/ControlRig/Source/ControlRigEditor/Private/AnimDetails/AnimDetailsProxyManager.cpp
+++ b/Engine/Plugins/Animation/ControlRig/Source/ControlRigEditor/Private/AnimDetails/AnimDetailsProxyManager.cpp
@@ -14,9 +14,11 @@
 #include "AnimDetails/Proxies/AnimDetailsProxyScale.h"
 #include "AnimDetails/Proxies/AnimDetailsProxyTransform.h"
 #include "AnimDetails/Proxies/AnimDetailsProxyVector2D.h"
+#include "Compilation/MovieSceneCompiledDataManager.h"
 #include "ControlRig.h"
 #include "EditMode/ControlRigEditMode.h"
 #include "EditorModeManager.h"
+#include "Evaluation/MovieSceneEvaluationTemplateInstance.h"
 #include "ISequencer.h"
 #include "LevelEditorViewport.h"
 #include "MovieScene.h"
@@ -282,14 +284,7 @@ void UAnimDetailsProxyManager::ForceUpdateProxies()
 		}
 		ObjectGuids = ObjectGuidsSet.Array();
 
-		UMovieSceneSequence* MovieSceneSequence = WeakSequencer.IsValid() ? WeakSequencer.Pin()->GetRootMovieSceneSequence() : nullptr;
-		UMovieScene* MovieScene = MovieSceneSequence ? MovieSceneSequence->GetMovieScene() : nullptr;
-		if (!MovieSceneSequence || !MovieScene)
-		{
-			return;
-		}
-
-		const TMap<FGuid, TArray<UMovieScenePropertyTrack*>> ObjectGuidToPropertyTracksMap = GetPropertyTracks(*MovieScene, Sequencer, ObjectGuids);
+		const TMap<FGuid, TArray<UMovieScenePropertyTrack*>> ObjectGuidToPropertyTracksMap = GetPropertyTracks(Sequencer, ObjectGuids);
 
 		for (const TTuple<FGuid, TArray<UMovieScenePropertyTrack*>>& ObjectGuidToPropertyTracksPair : ObjectGuidToPropertyTracksMap)
 		{
@@ -606,18 +601,38 @@ bool UAnimDetailsProxyManager::TryGetControlTypeFromTrackType(const UMovieSceneP
 	return false;
 }
 
-TMap<FGuid, TArray<UMovieScenePropertyTrack*>> UAnimDetailsProxyManager::GetPropertyTracks(const UMovieScene& MovieScene, const TSharedRef<ISequencer>& Sequencer, const TArray<FGuid>& ObjectGuids) const
+TMap<FGuid, TArray<UMovieScenePropertyTrack*>> UAnimDetailsProxyManager::GetPropertyTracks(const TSharedRef<ISequencer>& Sequencer, const TArray<FGuid>& ObjectGuids) const
 {
 	TMap<FGuid, TArray<UMovieScenePropertyTrack*>> ObjectGuidToPropertyTracksMap;
-	for (const FGuid& ObjectGuid : ObjectGuids)
+
+	// Find property tracks for all movie scenes in this sequencer
+	FMovieSceneRootEvaluationTemplateInstance& EvaluationTemplate = Sequencer->GetEvaluationTemplate();
+	UMovieSceneCompiledDataManager* CompiledDataManager = EvaluationTemplate.GetCompiledDataManager();
+	const FMovieSceneSequenceHierarchy* Hierarchy = CompiledDataManager ? CompiledDataManager->FindHierarchy(EvaluationTemplate.GetCompiledDataID()) : nullptr;
+	if (!Hierarchy)
 	{
-		const TArray<UMovieSceneTrack*> Tracks = MovieScene.FindTracks(UMovieScenePropertyTrack::StaticClass(), ObjectGuid);
+		return ObjectGuidToPropertyTracksMap;
+	}
 
-		for (UMovieSceneTrack* Track : Tracks)
+	for (const TTuple<FMovieSceneSequenceID, FMovieSceneSubSequenceData>& SequenceIDToSubsequenceDataPair : Hierarchy->AllSubSequenceData())
+	{
+		UMovieSceneSequence* Sequence = SequenceIDToSubsequenceDataPair.Value.GetSequence();
+		const UMovieScene* MovieScene = Sequence ? Sequence->GetMovieScene() : nullptr;
+		if (!Sequence || !MovieScene)
 		{
-			if (UMovieScenePropertyTrack* PropertyTrack = Cast<UMovieScenePropertyTrack>(Track))
+			continue;
+		}
+
+		for (const FGuid& ObjectGuid : ObjectGuids)
+		{
+			const TArray<UMovieSceneTrack*> Tracks = MovieScene->FindTracks(UMovieScenePropertyTrack::StaticClass(), ObjectGuid);
+
+			for (UMovieSceneTrack* Track : Tracks)
 			{
-				ObjectGuidToPropertyTracksMap.FindOrAdd(ObjectGuid).Add(PropertyTrack);
+				if (UMovieScenePropertyTrack* PropertyTrack = Cast<UMovieScenePropertyTrack>(Track))
+				{
+					ObjectGuidToPropertyTracksMap.FindOrAdd(ObjectGuid).Add(PropertyTrack);
+				}
 			}
 		}
 	}
diff --git a/Engine/Plugins/Animation/ControlRig/Source/ControlRigEditor/Private/AnimDetails/AnimDetailsProxyManager.h b/Engine/Plugins/Animation/ControlRig/Source/ControlRigEditor/Private/AnimDetails/AnimDetailsProxyManager.h
--- a/Engine/Plugins/Animation/ControlRig/Source/ControlRigEditor/Private/AnimDetails/AnimDetailsProxyManager.h
+++ b/Engine/Plugins/Animation/ControlRig/Source/ControlRigEditor/Private/AnimDetails/AnimDetailsProxyManager.h
@@ -121,7 +121,7 @@ private:
 	[[nodiscard]] bool TryGetControlTypeFromTrackType(const UMovieScenePropertyTrack& InPropertyTrack, ERigControlType& OutControlRigType) const;
 
 	/** Returns a map of object guids with their sequencer property tracks */
-	TMap<FGuid, TArray<UMovieScenePropertyTrack*>> GetPropertyTracks(const UMovieScene& MovieScene, const TSharedRef<ISequencer>& Sequencer, const TArray<FGuid>& ObjectGuids) const;
+	TMap<FGuid, TArray<UMovieScenePropertyTrack*>> GetPropertyTracks(const TSharedRef<ISequencer>& Sequencer, const TArray<FGuid>& ObjectGuids) const;
 
 	/** Returns a an array of objects bound in sequencer which match provided object id */
 	TArray<UObject*> GetBoundObjectsFromTrack(const TSharedRef<ISequencer>& Sequencer, const FGuid& ObjectGuid) const;
