Anim in Engine: Selecting an Anim Layer retains the controls selected in Anim Details where possible

diff --git a/Engine/Plugins/Animation/ControlRig/Source/ControlRigEditor/Private/AnimDetails/AnimDetailsProxyManager.cpp b/Engine/Plugins/Animation/ControlRig/Source/ControlRigEditor/Private/AnimDetails/AnimDetailsProxyManager.cpp
--- a/Engine/Plugins/Animation/ControlRig/Source/ControlRigEditor/Private/AnimDetails/AnimDetailsProxyManager.cpp
+++ b/Engine/Plugins/Animation/ControlRig/Source/ControlRigEditor/Private/AnimDetails/AnimDetailsProxyManager.cpp
@@ -170,7 +170,11 @@ void UAnimDetailsProxyManager::OnControlRigSelectionChanged(UControlRig* Control
 
 void UAnimDetailsProxyManager::OnSequencerSelectionChanged(TArray<FGuid> ObjectGuids)
 {
-	RequestUpdateProxies();
+	// Don't update proxies if anim layers is changing selection. Instead let the selection object override its selection.
+	if (AnimDetailsSelection && !AnimDetailsSelection->IsAnimLayersChangingSelection())
+	{
+		RequestUpdateProxies();
+	}
 }
 
 void UAnimDetailsProxyManager::OnObjectsReplaced(const TMap<UObject*, UObject*>& OldToNewInstanceMap)
diff --git a/Engine/Plugins/Animation/ControlRig/Source/ControlRigEditor/Private/AnimDetails/AnimDetailsSelection.cpp b/Engine/Plugins/Animation/ControlRig/Source/ControlRigEditor/Private/AnimDetails/AnimDetailsSelection.cpp
--- a/Engine/Plugins/Animation/ControlRig/Source/ControlRigEditor/Private/AnimDetails/AnimDetailsSelection.cpp
+++ b/Engine/Plugins/Animation/ControlRig/Source/ControlRigEditor/Private/AnimDetails/AnimDetailsSelection.cpp
@@ -289,6 +289,29 @@ void UAnimDetailsSelection::OnProxiesChanged()
 		return;
 	}
 
+	// Listen to anim layer selection changes
+	const TSharedPtr<ISequencer> Sequencer = ProxyManager->GetSequencer();
+	if (Sequencer.IsValid())
+	{
+		UAnimLayers* AnimLayers = UAnimLayers::HasAnimLayers(Sequencer.Get()) ? UAnimLayers::GetAnimLayers(Sequencer.Get()) : nullptr;
+		if (AnimLayers != WeakAnimLayers)
+		{
+			if (WeakAnimLayers.IsValid())
+			{
+				WeakAnimLayers->GetOnSelectionChanged().RemoveAll(this);
+			}
+
+			if (AnimLayers)
+			{
+				AnimLayers->GetOnSelectionChanged().AddUObject(this, &UAnimDetailsSelection::RequestPropagonateSelectionToCurveEditor);
+			}
+		}
+	}
+	else if (WeakAnimLayers.IsValid())
+	{
+		WeakAnimLayers->GetOnSelectionChanged().RemoveAll(this);
+	}
+
 	const TMap<FName, FAnimDetailsSelectionPropertyData> OldPropertyIDToPropertyDataMap = PropertyIDToPropertyDataMap;
 	PropertyIDToPropertyDataMap.Reset();
 
@@ -320,6 +343,21 @@ void UAnimDetailsSelection::OnProxiesChanged()
 	}
 }
 
+void UAnimDetailsSelection::OnAnimLayersSelected()
+{
+	bIsAnimLayersChangingSelection = true;
+
+	const TWeakObjectPtr<UAnimDetailsSelection> WeakThis = this;
+	GEditor->GetTimerManager()->SetTimerForNextTick(FTimerDelegate::CreateLambda([WeakThis, this]()
+		{
+			if (WeakThis.IsValid())
+			{
+				bIsAnimLayersChangingSelection = false;
+				PropagonateSelectionToCurveEditor();
+			}
+		}));
+}
+
 void UAnimDetailsSelection::OnFilterChanged()
 {
 	UAnimDetailsProxyManager* ProxyManager = GetTypedOuter<UAnimDetailsProxyManager>();
diff --git a/Engine/Plugins/Animation/ControlRig/Source/ControlRigEditor/Private/AnimDetails/AnimDetailsSelection.h b/Engine/Plugins/Animation/ControlRig/Source/ControlRigEditor/Private/AnimDetails/AnimDetailsSelection.h
--- a/Engine/Plugins/Animation/ControlRig/Source/ControlRigEditor/Private/AnimDetails/AnimDetailsSelection.h
+++ b/Engine/Plugins/Animation/ControlRig/Source/ControlRigEditor/Private/AnimDetails/AnimDetailsSelection.h
@@ -8,9 +8,11 @@
 #include "AnimDetailsSelection.generated.h"
 
 class IPropertyHandle;
+class UAnimLayers;
 class UControlRig;
 class UAnimDetailsProxyBase;
 struct FRigControlElement;
+namespace UE::Sequencer { class FOutlinerSelection; }
 
 namespace UE::ControlRigEditor
 {
@@ -111,10 +113,16 @@ public:
 	/** Returns the currently selected proxies */
 	TArray<UAnimDetailsProxyBase*> GetSelectedProxies() const;
 
+	/** Returns true while anim layers is changing selection. Useful to avoid updating anim details to override its selection */
+	bool IsAnimLayersChangingSelection() const { return bIsAnimLayersChangingSelection; }
+
 private:
 	/** Called when proxies changed in the proxy manager */
 	void OnProxiesChanged();
 
+	/** Called when anim layers were selected */
+	void OnAnimLayersSelected();
+
 	/** Called when the filter changed */
 	void OnFilterChanged();
 
@@ -140,6 +148,12 @@ private:
 	/** True while changing selection */
 	bool bIsChangingSelection = false;
 
+	/** True while anim layers is changing selection. Useful to override its selection */
+	bool bIsAnimLayersChangingSelection = false;
+
+	/** The current anim layers */
+	TWeakObjectPtr<UAnimLayers> WeakAnimLayers;
+
 	/** Timer handle to propagonate the selection to the curve editor */
 	FTimerHandle PropagonateSelectionToCurveEditorTimerHandle;
 };
diff --git a/Engine/Plugins/Animation/ControlRig/Source/ControlRigEditor/Private/Sequencer/AnimLayers/AnimLayers.cpp b/Engine/Plugins/Animation/ControlRig/Source/ControlRigEditor/Private/Sequencer/AnimLayers/AnimLayers.cpp
--- a/Engine/Plugins/Animation/ControlRig/Source/ControlRigEditor/Private/Sequencer/AnimLayers/AnimLayers.cpp
+++ b/Engine/Plugins/Animation/ControlRig/Source/ControlRigEditor/Private/Sequencer/AnimLayers/AnimLayers.cpp
@@ -701,7 +701,7 @@ void UAnimLayer::GetSections(TArray<UMovieSceneSection*>& OutSections) const
 	}
 }
 
-void UAnimLayer::SetSelectedInList(bool bInValue) 
+void UAnimLayer::SetSelectedInList(const FAnimLayersScopedSelection& ScopedSelection, bool bInValue)
 {
 	if (bIsSelectedInList != bInValue)
 	{
diff --git a/Engine/Plugins/Animation/ControlRig/Source/ControlRigEditor/Private/Sequencer/AnimLayers/AnimLayers.h b/Engine/Plugins/Animation/ControlRig/Source/ControlRigEditor/Private/Sequencer/AnimLayers/AnimLayers.h
--- a/Engine/Plugins/Animation/ControlRig/Source/ControlRigEditor/Private/Sequencer/AnimLayers/AnimLayers.h
+++ b/Engine/Plugins/Animation/ControlRig/Source/ControlRigEditor/Private/Sequencer/AnimLayers/AnimLayers.h
@@ -23,7 +23,7 @@ class UMovieSceneControlRigParameterTrack;
 class UMovieSceneControlRigParameterSection;
 struct FMovieSceneFloatChannel;
 class IPropertyHandle;
-
+namespace UE::Sequencer { class FOutlinerSelection; }
 
 USTRUCT(BlueprintType)
 struct FMergeAnimLayerSettings
@@ -228,6 +228,7 @@ public:
 //for friend access to State variables directly
 class UAnimLayers;
 struct FAnimLayerSourceUIEntry;
+struct FAnimLayersScopedSelection;
 
 UCLASS(BlueprintType, MinimalAPI)
 class UAnimLayer : public UObject
@@ -280,7 +281,7 @@ public:
 	bool RemoveSelectedInSequencer();
 
 	bool RemoveAnimLayerItem(UObject* InObject);
-	void SetSelectedInList(bool bInValue); 
+	void SetSelectedInList(const FAnimLayersScopedSelection& ScopedSelection, bool bInValue);
 	bool GetSelectedInList() const { return bIsSelectedInList; }
 	ECheckBoxState GetSelected(TSet<UObject*>& OutSelectedObjects, TMap<UControlRig*, TArray<FName>>& OutSelectedControls) const;
 
@@ -352,16 +353,24 @@ public:
 	static bool HasAnimLayers(ISequencer* InSequencer);
 
 	static 	TSharedPtr<ISequencer> GetSequencerFromAsset();
+
 public:
 	DECLARE_EVENT_OneParam(UAnimLayers, FAnimLayerListChanged, UAnimLayers*);
 
 	FAnimLayerListChanged& AnimLayerListChanged() { return OnAnimLayerListChanged; };
 
+	/** Returns a delegate broadcast when anim outliner changed selection, both for anim layers self as well as for sequencer. */
+	FSimpleMulticastDelegate& GetOnSelectionChanged() { return OnSelectionChanged; }
+
 private:
 	void AnimLayerListChangedBroadcast();
 	FAnimLayerListChanged OnAnimLayerListChanged;
 	void AddBaseLayer();
 	bool SetKeyValueOrPassthrough(ISequencer* InSequencer, int32 Index, bool bJustValue);
+
+	/** A delegate broadcast when the selected layer sections changed */
+	FSimpleMulticastDelegate OnSelectionChanged;
+
 public:
 	void SetUpBaseLayerSections();
 
@@ -370,6 +379,23 @@ public:
 
 };
 
+/** A scope in which selection can occur in anim layers. Raises OnSelectionChanged when destructed. */
+struct FAnimLayersScopedSelection
+{
+	FAnimLayersScopedSelection(UAnimLayers & AnimLayers)
+		: WeakAnimLayers(&AnimLayers)
+	{
+	}
 
+	~FAnimLayersScopedSelection()
+	{
+		if (WeakAnimLayers.IsValid())
+		{
+			WeakAnimLayers->GetOnSelectionChanged().Broadcast();
+		}
+	}
 
+private:
+	TWeakObjectPtr<UAnimLayers> WeakAnimLayers;
+};
 
diff --git a/Engine/Plugins/Animation/ControlRig/Source/ControlRigEditor/Private/Sequencer/AnimLayers/SAnimLayers.cpp b/Engine/Plugins/Animation/ControlRig/Source/ControlRigEditor/Private/Sequencer/AnimLayers/SAnimLayers.cpp
--- a/Engine/Plugins/Animation/ControlRig/Source/ControlRigEditor/Private/Sequencer/AnimLayers/SAnimLayers.cpp
+++ b/Engine/Plugins/Animation/ControlRig/Source/ControlRigEditor/Private/Sequencer/AnimLayers/SAnimLayers.cpp
@@ -61,7 +61,7 @@ public:
 	FAnimLayerSourceUIEntry(UAnimLayer* InAnimLayer): AnimLayer(InAnimLayer)
 	{}
 
-	void SetSelectedInList(bool bInValue);
+	void SetSelectedInList(const FAnimLayersScopedSelection& ScopedSelection, bool bInValue);
 	bool GetSelectedInList() const;
 
 	void SelectObjects() const;
@@ -302,11 +302,11 @@ void FAnimLayerSourceUIEntry::SetSelected(bool bInSelected)
 	}
 }
 
-void FAnimLayerSourceUIEntry::SetSelectedInList(bool bInValue)
+void FAnimLayerSourceUIEntry::SetSelectedInList(const FAnimLayersScopedSelection& ScopedSelection, bool bInValue)
 { 
 	if (AnimLayer)
 	{
-		AnimLayer->SetSelectedInList(bInValue);
+		AnimLayer->SetSelectedInList(ScopedSelection, bInValue);
 		if (bInValue)
 		{
 			AnimLayer->SetKeyed();//selection also sets keyed
@@ -521,7 +521,7 @@ public:
 	void RefreshSelectionData();
 	void HandleOnAnimLayerListChanged(UAnimLayers*) { RebuildSourceList(); }
 	void FocusRenameLayer(int32 Index) { if (SourcesView.IsValid()) SourcesView->FocusRenameLayer(Index); }
-	void SelectItem(int32 Index, bool bClear = true);
+	void SelectItem(const FAnimLayersScopedSelection& ScopedSelection, int32 Index, bool bClear = true);
 	void ToggleSelectionFilterActive();
 	bool IsSelectionFilterActive() const;
 public:
@@ -1503,18 +1503,21 @@ void FAnimLayerSourcesView::OnSourceListSelectionChanged(FAnimLayerSourceUIEntry
 	{
 		if (UAnimLayers* AnimLayers = UAnimLayers::GetAnimLayers(SequencerPtr.Get()))
 		{
+			const FAnimLayersScopedSelection ScopedSelection(*AnimLayers);
+
 			TArray<FAnimLayerSourceUIEntryPtr> Selected;
 			SourcesListView->GetSelectedItems(Selected);
 			FSelectionEventSuppressor SuppressSelectionEvents = SequencerPtr->GetViewModel()->GetSelection()->SuppressEvents();
 
 			for (UAnimLayer* AnimLayer : AnimLayers->AnimLayers)
 			{
-				AnimLayer->SetSelectedInList(false);
+				AnimLayer->SetSelectedInList(ScopedSelection, false);
 			}
 			for (const FAnimLayerSourceUIEntryPtr& Item : Selected)
 			{
-				Item->SetSelectedInList(true);
+				Item->SetSelectedInList(ScopedSelection, true);
 			}
+
 			//refresh tree to rerun the filter
 			SequencerPtr->NotifyMovieSceneDataChanged(EMovieSceneDataChangeType::RefreshTree);
 		}
@@ -1595,7 +1598,7 @@ void FAnimLayerController::RefreshTimeDependantData()
 		}
 	}
 }
-void FAnimLayerController::SelectItem(int32 Index, bool bClear)
+void FAnimLayerController::SelectItem(const FAnimLayersScopedSelection& ScopedSelection, int32 Index, bool bClear)
 {
 	if (SourcesView.IsValid() && SourcesView->SourcesListView.IsValid())
 	{
@@ -1609,12 +1612,12 @@ void FAnimLayerController::SelectItem(int32 Index, bool bClear)
 		{
 			if (Count == Index)
 			{
-				Item->SetSelectedInList(true);
+				Item->SetSelectedInList(ScopedSelection, true);
 				Selected.Add(Item);
 			}
 			else if (bClear && Item->GetSelectedInList())
 			{
-				Item->SetSelectedInList(false);
+				Item->SetSelectedInList(ScopedSelection, false);
 			}
 			++Count;
 		}
@@ -1966,13 +1969,15 @@ FReply SAnimLayers::OnAddClicked()
 	{
 		if (UAnimLayers* AnimLayersPtr = UAnimLayers::GetAnimLayers(SequencerPtr))
 		{
+			const FAnimLayersScopedSelection ScopedSelection(*AnimLayersPtr);
+
 			int32 Index = AnimLayersPtr->AddAnimLayerFromSelection(SequencerPtr);
 			if (Index != INDEX_NONE)
 			{
 				if (AnimLayerController.IsValid())
 				{
 					AnimLayerController->FocusRenameLayer(Index);
-					AnimLayerController->SelectItem(Index, true);
+					AnimLayerController->SelectItem(ScopedSelection, Index, true);
 				}
 			}
 		}
