Anim in Engine: Improve const correctness in UAnimDetailsProxyManager::GetPropertyTracks

#rnx

diff --git a/Engine/Plugins/Animation/ControlRig/Source/ControlRigEditor/Private/AnimDetails/AnimDetailsProxyManager.cpp b/Engine/Plugins/Animation/ControlRig/Source/ControlRigEditor/Private/AnimDetails/AnimDetailsProxyManager.cpp
--- a/Engine/Plugins/Animation/ControlRig/Source/ControlRigEditor/Private/AnimDetails/AnimDetailsProxyManager.cpp
+++ b/Engine/Plugins/Animation/ControlRig/Source/ControlRigEditor/Private/AnimDetails/AnimDetailsProxyManager.cpp
@@ -610,8 +610,8 @@ TMap<FGuid, TArray<UMovieScenePropertyTrack*>> UAnimDetailsProxyManager::GetProp
 	TMap<FGuid, TArray<UMovieScenePropertyTrack*>> ObjectGuidToPropertyTracksMap;
 
 	// Find property tracks for all movie scenes in this sequencer
-	FMovieSceneRootEvaluationTemplateInstance& EvaluationTemplate = Sequencer->GetEvaluationTemplate();
-	UMovieSceneCompiledDataManager* CompiledDataManager = EvaluationTemplate.GetCompiledDataManager();
+	const FMovieSceneRootEvaluationTemplateInstance& EvaluationTemplate = Sequencer->GetEvaluationTemplate();
+	const UMovieSceneCompiledDataManager* CompiledDataManager = EvaluationTemplate.GetCompiledDataManager();
 	const FMovieSceneSequenceHierarchy* Hierarchy = CompiledDataManager ? CompiledDataManager->FindHierarchy(EvaluationTemplate.GetCompiledDataID()) : nullptr;
 	if (!Hierarchy)
 	{
@@ -620,7 +620,7 @@ TMap<FGuid, TArray<UMovieScenePropertyTrack*>> UAnimDetailsProxyManager::GetProp
 
 	for (const TTuple<FMovieSceneSequenceID, FMovieSceneSubSequenceData>& SequenceIDToSubsequenceDataPair : Hierarchy->AllSubSequenceData())
 	{
-		UMovieSceneSequence* Sequence = SequenceIDToSubsequenceDataPair.Value.GetSequence();
+		const UMovieSceneSequence* Sequence = SequenceIDToSubsequenceDataPair.Value.GetSequence();
 		const UMovieScene* MovieScene = Sequence ? Sequence->GetMovieScene() : nullptr;
 		if (!Sequence || !MovieScene)
 		{
