Anim in Engine: Anim Details no longer updates its data if its tab is closed

#rnx

diff --git a/Engine/Plugins/Animation/ControlRig/Source/ControlRigEditor/Private/AnimDetails/AnimDetailsProxyManager.cpp b/Engine/Plugins/Animation/ControlRig/Source/ControlRigEditor/Private/AnimDetails/AnimDetailsProxyManager.cpp
--- a/Engine/Plugins/Animation/ControlRig/Source/ControlRigEditor/Private/AnimDetails/AnimDetailsProxyManager.cpp
+++ b/Engine/Plugins/Animation/ControlRig/Source/ControlRigEditor/Private/AnimDetails/AnimDetailsProxyManager.cpp
@@ -42,9 +42,24 @@ UAnimDetailsProxyManager::UAnimDetailsProxyManager()
 	AnimDetailsSelection = NewObject<UAnimDetailsSelection>(this, "AnimDetailsSelection", RF_Transactional);
 }
 
-void UAnimDetailsProxyManager::NotifySequencerChanged()
+void UAnimDetailsProxyManager::SetSuspended(const bool bInSuspended)
 {
-	RequestUpdateProxies();
+	bSuspended = bInSuspended;
+
+	if (bSuspended)
+	{
+		ClearSequencer();
+
+		// Clear out selection so it won't update itself
+		AnimDetailsSelection = nullptr;
+	}
+	else
+	{
+		const bool bSuccess = UpdateSequencer();
+
+		// Reinstall selection
+		AnimDetailsSelection = NewObject<UAnimDetailsSelection>(this, "AnimDetailsSelection", RF_Transactional);
+	}
 }
 
 TSharedPtr<ISequencer> UAnimDetailsProxyManager::GetSequencer() const
@@ -52,8 +67,27 @@ TSharedPtr<ISequencer> UAnimDetailsProxyManager::GetSequencer() const
 	return WeakSequencer.IsValid() ? WeakSequencer.Pin() : nullptr;
 }
 
+void UAnimDetailsProxyManager::RequestUpdateProxies()
+{
+	if (bSuspended)
+	{
+		return;
+	}
+
+	if (ensureMsgf(IsInGameThread(), TEXT("Anim Details proxies can only be updated in game thread. Ignoring call")) &&
+		!RequestUpdateProxiesTimerHandle.IsValid())
+	{
+		RequestUpdateProxiesTimerHandle = GEditor->GetTimerManager()->SetTimerForNextTick(FTimerDelegate::CreateUObject(this, &UAnimDetailsProxyManager::ForceUpdateProxies));
+	}
+}
+
 void UAnimDetailsProxyManager::RequestUpdateProxyValues()
 {
+	if (bSuspended)
+	{
+		return;
+	}
+
 	if (ensureMsgf(IsInGameThread(), TEXT("Anim Details proxy values can only be updated in game thread. Ignoring call")) &&
 		!RequestUpdateProxyValuesTimerHandle.IsValid())
 	{
@@ -75,7 +109,7 @@ const UAnimDetailsSelection* UAnimDetailsProxyManager::GetAnimDetailsSelection()
 void UAnimDetailsProxyManager::PostUndo(bool bSuccess)
 {
 	// Update proxies to reflect any external changes
-	RequestUpdateProxies(); 
+	RequestUpdateProxies();
 }
 
 void UAnimDetailsProxyManager::PostRedo(bool bSuccess) 
@@ -84,8 +118,22 @@ void UAnimDetailsProxyManager::PostRedo(bool bSuccess)
 	PostUndo(bSuccess);
 }
 
+void UAnimDetailsProxyManager::ClearSequencer()
+{
+	WeakSequencer.Reset();
+
+	ReleaseBindings();
+}
+
 bool UAnimDetailsProxyManager::UpdateSequencer()
 {
+	if (bSuspended)
+	{
+		WeakSequencer.Reset();
+
+		return false;
+	}
+
 	const TWeakPtr<ISequencer> NewWeakSequencer = UE::AnimationEditMode::GetSequencer();
 	
 	if (WeakSequencer.IsValid() && WeakSequencer == NewWeakSequencer)
@@ -93,15 +141,12 @@ bool UAnimDetailsProxyManager::UpdateSequencer()
 		// Nothing changed
 		return true;
 	}
-	else if (WeakSequencer.IsValid())
-	{
-		// Sequencer changed, unbind from the previous sequencer events if it is still valid
-		WeakSequencer.Pin()->GetSelectionChangedObjectGuids().RemoveAll(this);
-	}
 
 	if (!NewWeakSequencer.IsValid())
 	{
 		// No new sequencer
+		ReleaseBindings();
+
 		return false;
 	}
 
@@ -109,13 +154,21 @@ bool UAnimDetailsProxyManager::UpdateSequencer()
 	Proxies.Reset();
 	ExternalSelection.Reset();
 
-	// Set the new sequencer and listen to it
-	const TSharedRef<ISequencer> Sequencer = NewWeakSequencer.Pin().ToSharedRef();
-	WeakSequencer = Sequencer;
+	// Set the new sequencer and refresh from that
+	WeakSequencer = NewWeakSequencer;
+
+	SetupBingings();
+	RequestUpdateProxies();
+
+	return true;
+}
 
-	if (!Sequencer->GetSelectionChangedObjectGuids().IsBoundToObject(this))
+void UAnimDetailsProxyManager::SetupBingings()
+{
+	if (WeakSequencer.IsValid() &&
+		!WeakSequencer.Pin()->GetSelectionChangedObjectGuids().IsBoundToObject(this))
 	{
-		Sequencer->GetSelectionChangedObjectGuids().AddUObject(this, &UAnimDetailsProxyManager::OnSequencerSelectionChanged);
+		WeakSequencer.Pin()->GetSelectionChangedObjectGuids().AddUObject(this, &UAnimDetailsProxyManager::OnSequencerSelectionChanged);
 	}
 
 	// Listen to the control rig edit mode
@@ -135,12 +188,27 @@ bool UAnimDetailsProxyManager::UpdateSequencer()
 	{
 		FCoreUObjectDelegates::OnObjectsReplaced.AddUObject(this, &UAnimDetailsProxyManager::OnObjectsReplaced);
 	}
+}
 
-	RequestUpdateProxies();
+void UAnimDetailsProxyManager::ReleaseBindings()
+{
+	if (WeakSequencer.IsValid())
+	{
+		WeakSequencer.Pin()->GetSelectionChangedObjectGuids().RemoveAll(this);
+	}
 
-	return true;
+	FControlRigEditMode* EditMode = static_cast<FControlRigEditMode*>(GLevelEditorModeTools().GetActiveMode(FControlRigEditMode::ModeName));
+	if (EditMode)
+	{
+		EditMode->OnControlRigAddedOrRemoved().RemoveAll(this);
+		EditMode->OnControlRigShapeActorsRecreated().RemoveAll(this);
+		EditMode->OnControlRigSelected().RemoveAll(this);
+	}
+
+	FCoreUObjectDelegates::OnObjectsReplaced.RemoveAll(this);
 }
 
+
 void UAnimDetailsProxyManager::OnControlRigShapeActorsRecreated()
 {
 	FControlRigEditMode* EditMode = static_cast<FControlRigEditMode*>(GLevelEditorModeTools().GetActiveMode(FControlRigEditMode::ModeName));
@@ -182,15 +250,6 @@ void UAnimDetailsProxyManager::OnObjectsReplaced(const TMap<UObject*, UObject*>&
 	RequestUpdateProxies();
 }
 
-void UAnimDetailsProxyManager::RequestUpdateProxies()
-{
-	if (ensureMsgf(IsInGameThread(), TEXT("Anim Details proxies can only be updated in game thread. Ignoring call")) &&
-		!RequestUpdateProxiesTimerHandle.IsValid())
-	{
-		RequestUpdateProxiesTimerHandle = GEditor->GetTimerManager()->SetTimerForNextTick(FTimerDelegate::CreateUObject(this, &UAnimDetailsProxyManager::ForceUpdateProxies));
-	}
-}
-
 void UAnimDetailsProxyManager::ForceUpdateProxies()
 {
 	if (!ensureMsgf(IsInGameThread(), TEXT("Anim Details proxies can only be updated in game thread. Ignoring call")))
@@ -198,12 +257,12 @@ void UAnimDetailsProxyManager::ForceUpdateProxies()
 		return;
 	}
 
+	// FTimerHandle has to be invalidated on the game thread
 	RequestUpdateProxiesTimerHandle.Invalidate();
 
 	const FControlRigEditMode* EditMode = static_cast<FControlRigEditMode*>(GLevelEditorModeTools().GetActiveMode(FControlRigEditMode::ModeName));
 	if (!EditMode || !UpdateSequencer())
 	{
-		WeakSequencer.Reset();
 		Proxies.Reset();
 		return;
 	}
@@ -357,8 +416,19 @@ void UAnimDetailsProxyManager::ForceUpdateProxies()
 
 void UAnimDetailsProxyManager::ForceUpdateProxyValues()
 {
+	if (!ensureMsgf(IsInGameThread(), TEXT("Anim Details proxies can only be updated in game thread. Ignoring call")))
+	{
+		return;
+	}
+
+	// FTimerHandle has to be invalidated on the game thread
 	RequestUpdateProxyValuesTimerHandle.Invalidate();
 
+	if (bSuspended)
+	{
+		return;
+	}
+
 	if (ExternalSelection.IsEmpty())
 	{
 		return;
diff --git a/Engine/Plugins/Animation/ControlRig/Source/ControlRigEditor/Private/AnimDetails/AnimDetailsProxyManager.h b/Engine/Plugins/Animation/ControlRig/Source/ControlRigEditor/Private/AnimDetails/AnimDetailsProxyManager.h
--- a/Engine/Plugins/Animation/ControlRig/Source/ControlRigEditor/Private/AnimDetails/AnimDetailsProxyManager.h
+++ b/Engine/Plugins/Animation/ControlRig/Source/ControlRigEditor/Private/AnimDetails/AnimDetailsProxyManager.h
@@ -17,6 +17,7 @@ class UAnimDetailsProxyBase;
 class UAnimDetailsSelection;
 class UMovieScene;
 class UMovieScenePropertyTrack;
+class UMovieSceneSection;
 enum class ERigControlType : uint8;
 struct FRigControlElement;
 struct FRigElementKey;
@@ -36,13 +37,19 @@ class UAnimDetailsProxyManager
 public:
 	UAnimDetailsProxyManager();
 
-	/** Lets this proxy manager know the sequencer changed. */
-	void NotifySequencerChanged();
+	/** Suspends or resumes this manager. While the manager is suspended, its proxies are no longer updated. */
+	void SetSuspended(const bool bEnabled);
 
 	/** Returns the sequencer this proxy manager handles or nullptr if not currently assigned to a sequencer */
 	TSharedPtr<ISequencer> GetSequencer() const;
 
-	/** Updates the proxy values on the next tick. Useful to let proxies know that values changed externally. */
+	/** Updates proxies on the next tick. Ignored if the manager is suspended. */
+	void RequestUpdateProxies();
+
+	/** 
+	 * Updates the proxy values on the next tick. Useful to let proxies know that values changed externally.
+	 * Ignored if the manager is suspended. 
+	 */
 	void RequestUpdateProxyValues();
 
 	/** Delegate broadcase when a displayed proxies changed */
@@ -75,9 +82,18 @@ private:
 	virtual void PostRedo(bool bSuccess) override;
 	//~ End FSelfRegisteringEditorUndoClient interface
 
+	/** Clears the sequencer. This is useful to suspend this proxy manager. */
+	void ClearSequencer();
+
 	/** Updates the sequencer this manager uses. Returns true if a sequencer was set successfully */
 	[[nodiscard]] bool UpdateSequencer();
 
+	/** Binds to other objects' events to update on their changes */
+	void SetupBingings();
+
+	/** Releases bindings to other objects' events */
+	void ReleaseBindings();
+
 	/** Called when control rig shape actors were recreated  */
 	void OnControlRigShapeActorsRecreated();
 
@@ -93,9 +109,6 @@ private:
 	/** Called from the editor when a blueprint object replacement has occurred. Useful to restore the control rig for proxies when it was replaced. */
 	void OnObjectsReplaced(const TMap<UObject*, UObject*>& OldToNewInstanceMap);
 
-	/** Updates proxies on the next tick */
-	void RequestUpdateProxies();
-
 	/** Updates all proxies. */
 	void ForceUpdateProxies();
 
@@ -129,6 +142,9 @@ private:
 	/** Convenience to add a proxy to ExternalSelection, only if the proxy is valid */
 	void AddProxyToExternalSelectionIfValid(UAnimDetailsProxyBase* Proxy);
 
+	/** While true, the manager is not updated */
+	bool bSuspended = true;
+
 	/** The current proxies */
 	UPROPERTY()
 	TArray<TObjectPtr<UAnimDetailsProxyBase>> Proxies; 
diff --git a/Engine/Plugins/Animation/ControlRig/Source/ControlRigEditor/Private/EditMode/ControlRigEditMode.cpp b/Engine/Plugins/Animation/ControlRig/Source/ControlRigEditor/Private/EditMode/ControlRigEditMode.cpp
--- a/Engine/Plugins/Animation/ControlRig/Source/ControlRigEditor/Private/EditMode/ControlRigEditMode.cpp
+++ b/Engine/Plugins/Animation/ControlRig/Source/ControlRigEditor/Private/EditMode/ControlRigEditMode.cpp
@@ -293,7 +293,7 @@ bool FControlRigEditMode::SetSequencer(TWeakPtr<ISequencer> InSequencer)
 			
 			SetSequencerDelegates(WeakSequencer);
 			
-			AnimDetailsProxyManager->NotifySequencerChanged();
+			AnimDetailsProxyManager->RequestUpdateProxies();
 
 			{
 				TSharedRef<SSequencer> SequencerWidget = StaticCastSharedRef<SSequencer>(Sequencer->GetSequencerWidget());
diff --git a/Engine/Plugins/Animation/ControlRig/Source/ControlRigEditor/Private/EditMode/ControlRigEditModeToolkit.cpp b/Engine/Plugins/Animation/ControlRig/Source/ControlRigEditor/Private/EditMode/ControlRigEditModeToolkit.cpp
--- a/Engine/Plugins/Animation/ControlRig/Source/ControlRigEditor/Private/EditMode/ControlRigEditModeToolkit.cpp
+++ b/Engine/Plugins/Animation/ControlRig/Source/ControlRigEditor/Private/EditMode/ControlRigEditModeToolkit.cpp
@@ -6,6 +6,7 @@
 #include "EditMode/ControlRigEditModeToolkit.h"
 
 #include "ControlRigEditModeCommands.h"
+#include "AnimDetails/AnimDetailsProxyManager.h"
 #include "AnimDetails/Views/SAnimDetailsView.h"
 #include "Widgets/DeclarativeSyntaxSupport.h"
 #include "EditorModes.h"
@@ -28,6 +29,7 @@
 #include "Sequencer/AnimLayers/SAnimLayers.h"
 #include "Sequencer/AnimLayers/AnimLayers.h"
 #include "Tools/MotionTrailOptions.h"
+#include "UObject/WeakObjectPtr.h"
 
 class SAnimDetailsView;
 namespace UE::ControlRigEditor
@@ -233,10 +235,33 @@ TSharedRef<SDockTab> SpawnSpacePickerTab(const FSpawnTabArgs& Args, FControlRigE
 
 TSharedRef<SDockTab> SpawnDetailsTab(const FSpawnTabArgs& Args, FControlRigEditMode* InEditorMode)
 {
-	return SNew(SDockTab)
-		[
-			SAssignNew(FControlRigEditModeToolkit::Details, UE::ControlRigEditor::SAnimDetailsView)
-		];
+	UAnimDetailsProxyManager* ProxyManager = InEditorMode ? InEditorMode->GetAnimDetailsProxyManager() : nullptr;
+
+	if (ProxyManager)
+	{
+		ProxyManager->SetSuspended(false);
+
+		const TWeakObjectPtr<UAnimDetailsProxyManager> WeakProxyManager = ProxyManager;
+		return SNew(SDockTab)
+			.OnTabClosed_Lambda([WeakProxyManager](TSharedRef<SDockTab> InTab)
+				{
+					if (WeakProxyManager.IsValid())
+					{
+						// For performance reasons, suspend the anim details proxy manager while its view is not spawned
+						WeakProxyManager->SetSuspended(true);
+					}
+				})
+			[
+				SAssignNew(FControlRigEditModeToolkit::Details, UE::ControlRigEditor::SAnimDetailsView)
+			];
+	}
+	else
+	{
+		return SNew(SDockTab)
+			[
+				SNullWidget::NullWidget
+			];
+	}
 }
 
 void FControlRigEditModeToolkit::RequestModeUITabs()
