Fixed ControlRig serialization not storing the previous hierarchy name map

diff --git a/Engine/Plugins/Animation/ControlRig/Source/ControlRig/Private/Rigs/RigHierarchy.cpp b/Engine/Plugins/Animation/ControlRig/Source/ControlRig/Private/Rigs/RigHierarchy.cpp
--- a/Engine/Plugins/Animation/ControlRig/Source/ControlRig/Private/Rigs/RigHierarchy.cpp
+++ b/Engine/Plugins/Animation/ControlRig/Source/ControlRig/Private/Rigs/RigHierarchy.cpp
@@ -274,7 +274,7 @@ void URigHierarchy::Save(FArchive& Ar)
 	}
 
 	*ArchiveForElements << PreviousHierarchyNameMap;
-	*ArchiveForElements << PreviousHierarchyNameMap;
+	*ArchiveForElements << PreviousHierarchyParentMap;
 
 	{
 		TMap<FRigElementKey, FMetadataStorage> ElementMetadataToSave;
