Anim in Engine: Fix a regression where Anim Details no longer showed properties of the root sequence

#rnx

diff --git a/Engine/Plugins/Animation/ControlRig/Source/ControlRigEditor/Private/AnimDetails/AnimDetailsProxyManager.cpp b/Engine/Plugins/Animation/ControlRig/Source/ControlRigEditor/Private/AnimDetails/AnimDetailsProxyManager.cpp
--- a/Engine/Plugins/Animation/ControlRig/Source/ControlRigEditor/Private/AnimDetails/AnimDetailsProxyManager.cpp
+++ b/Engine/Plugins/Animation/ControlRig/Source/ControlRigEditor/Private/AnimDetails/AnimDetailsProxyManager.cpp
@@ -679,7 +679,31 @@ TMap<FGuid, TArray<UMovieScenePropertyTrack*>> UAnimDetailsProxyManager::GetProp
 {
 	TMap<FGuid, TArray<UMovieScenePropertyTrack*>> ObjectGuidToPropertyTracksMap;
 
-	// Find property tracks for all movie scenes in this sequencer
+	const auto FindPropertyTracks = [&ObjectGuidToPropertyTracksMap, &ObjectGuids](const UMovieScene* MovieScene)
+		{
+			for (const FGuid& ObjectGuid : ObjectGuids)
+			{
+				const TArray<UMovieSceneTrack*> Tracks = MovieScene->FindTracks(UMovieScenePropertyTrack::StaticClass(), ObjectGuid);
+
+				for (UMovieSceneTrack* Track : Tracks)
+				{
+					if (UMovieScenePropertyTrack* PropertyTrack = Cast<UMovieScenePropertyTrack>(Track))
+					{
+						ObjectGuidToPropertyTracksMap.FindOrAdd(ObjectGuid).Add(PropertyTrack);
+					}
+				}
+			}
+		};
+
+	// Find property tracks of the movie scene in the root sequence
+	const UMovieSceneSequence* RootSequence = Sequencer->GetRootMovieSceneSequence();
+	const UMovieScene* RootMovieScene = RootSequence ? RootSequence->GetMovieScene() : nullptr;
+	if (RootMovieScene)
+	{
+		FindPropertyTracks(RootMovieScene);
+	}
+
+	// Find property tracks of movie scenes in sub sequences
 	const FMovieSceneRootEvaluationTemplateInstance& EvaluationTemplate = Sequencer->GetEvaluationTemplate();
 	const UMovieSceneCompiledDataManager* CompiledDataManager = EvaluationTemplate.GetCompiledDataManager();
 	const FMovieSceneSequenceHierarchy* Hierarchy = CompiledDataManager ? CompiledDataManager->FindHierarchy(EvaluationTemplate.GetCompiledDataID()) : nullptr;
@@ -690,25 +714,14 @@ TMap<FGuid, TArray<UMovieScenePropertyTrack*>> UAnimDetailsProxyManager::GetProp
 
 	for (const TTuple<FMovieSceneSequenceID, FMovieSceneSubSequenceData>& SequenceIDToSubsequenceDataPair : Hierarchy->AllSubSequenceData())
 	{
-		const UMovieSceneSequence* Sequence = SequenceIDToSubsequenceDataPair.Value.GetSequence();
-		const UMovieScene* MovieScene = Sequence ? Sequence->GetMovieScene() : nullptr;
-		if (!Sequence || !MovieScene)
+		const UMovieSceneSequence* SubSequence = SequenceIDToSubsequenceDataPair.Value.GetSequence();
+		const UMovieScene* InnerMovieScene = SubSequence ? SubSequence->GetMovieScene() : nullptr;
+		if (!SubSequence || !InnerMovieScene)
 		{
 			continue;
 		}
 
-		for (const FGuid& ObjectGuid : ObjectGuids)
-		{
-			const TArray<UMovieSceneTrack*> Tracks = MovieScene->FindTracks(UMovieScenePropertyTrack::StaticClass(), ObjectGuid);
-
-			for (UMovieSceneTrack* Track : Tracks)
-			{
-				if (UMovieScenePropertyTrack* PropertyTrack = Cast<UMovieScenePropertyTrack>(Track))
-				{
-					ObjectGuidToPropertyTracksMap.FindOrAdd(ObjectGuid).Add(PropertyTrack);
-				}
-			}
-		}
+		FindPropertyTracks(InnerMovieScene);
 	}
 
 	return ObjectGuidToPropertyTracksMap;
