Fix Control Rig editor reseting curve operations made on the curve container (delete, rename and import)

diff --git a/Engine/Plugins/Animation/ControlRig/Source/ControlRigEditor/Private/Editor/SRigCurveContainer.cpp b/Engine/Plugins/Animation/ControlRig/Source/ControlRigEditor/Private/Editor/SRigCurveContainer.cpp
--- a/Engine/Plugins/Animation/ControlRig/Source/ControlRigEditor/Private/Editor/SRigCurveContainer.cpp
+++ b/Engine/Plugins/Animation/ControlRig/Source/ControlRigEditor/Private/Editor/SRigCurveContainer.cpp
@@ -478,7 +478,7 @@ void SRigCurveContainer::RefreshCurveList()
 
 void SRigCurveContainer::OnNameCommitted(const FText& InNewName, ETextCommit::Type CommitType, FDisplayedRigCurveInfoPtr Item)
 {
-	URigHierarchy* Hierarchy = GetHierarchy();
+	URigHierarchy* Hierarchy = GetDefaultHierarchy();
 	if (Hierarchy)
 	{
 		if (CommitType == ETextCommit::OnEnter)
@@ -511,7 +511,7 @@ void SRigCurveContainer::PostRedo(bool bSuccess)
 
 void SRigCurveContainer::OnDeleteNameClicked()
 {
-	URigHierarchy* Hierarchy = GetHierarchy();
+	URigHierarchy* Hierarchy = GetDefaultHierarchy();
 	if (Hierarchy)
 	{
 		TGuardValue<bool> SuspendBlueprintNotifs(ControlRigBlueprint->bSuspendAllNotifications, true);
@@ -558,7 +558,7 @@ void SRigCurveContainer::ChangeCurveName(const FName& OldName, const FName& NewN
 {
 	TGuardValue<bool> GuardReentry(bIsChangingRigHierarchy, true);
 
-	URigHierarchy* Hierarchy = GetHierarchy();
+	URigHierarchy* Hierarchy = GetDefaultHierarchy();
 	if (Hierarchy)
 	{
 		if(URigHierarchyController* Controller = Hierarchy->GetController())
@@ -798,7 +798,7 @@ bool SRigCurveContainer::ShouldFilterOnImport(const FAssetData& AssetData) const
 
 void SRigCurveContainer::ImportCurve(const FAssetData& InAssetData)
 {
-	URigHierarchy* Hierarchy = GetHierarchy();
+	URigHierarchy* Hierarchy = GetDefaultHierarchy();
 	if (Hierarchy)
 	{
 		TGuardValue<bool> SuspendBlueprintNotifs(ControlRigBlueprint->bSuspendAllNotifications, true);
