Control Rig: Represent instruction as part of the dependency map

The dependency map now contains elements as well as instructions.

#rnx

diff --git a/Engine/Plugins/Animation/ControlRig/Source/ControlRig/Private/ControlRig.cpp b/Engine/Plugins/Animation/ControlRig/Source/ControlRig/Private/ControlRig.cpp
--- a/Engine/Plugins/Animation/ControlRig/Source/ControlRig/Private/ControlRig.cpp
+++ b/Engine/Plugins/Animation/ControlRig/Source/ControlRig/Private/ControlRig.cpp
@@ -115,19 +115,8 @@ void UControlRig::ResetRecordedTransforms(const FName& InEventName)
 	{
 		if(Hierarchy->bRecordInstructionsAtRuntime)
 		{
-			bool bResetRecordedTransforms = SupportsEvent(InEventName);
-			if(InEventName == FRigUnit_PostBeginExecution::EventName ||
-				InEventName == FRigUnit_PreBeginExecution::EventName)
-			{
-				bResetRecordedTransforms = false;
-			}
-			
-			if(bResetRecordedTransforms)
-			{
-				FControlRigExecuteContext& PublicData = GetRigVMExtendedExecuteContext().GetPublicData<FControlRigExecuteContext>();
-				PublicData.ReadRecords.Reset();
-				PublicData.WrittenRecords.Reset();
-			}
+			FControlRigExecuteContext& PublicData = GetRigVMExtendedExecuteContext().GetPublicData<FControlRigExecuteContext>();
+			PublicData.ResetInstructionRecords();
 		}
 	}
 }
@@ -1126,7 +1115,10 @@ bool UControlRig::Execute(const FName& InEventName)
 		}
 		
 		RecordTransformsPerInstructionGuard = MakeShared<TGuardValue<bool>>(Hierarchy->bRecordInstructionsAtRuntime, bRecordInstructionsAtRuntime);
-		ResetRecordedTransforms(InEventName);
+		if (!bIsEventInQueue || bIsEventFirstInQueue)
+		{
+			ResetRecordedTransforms(InEventName);
+		}
 #endif
 	}
 
diff --git a/Engine/Plugins/Animation/ControlRig/Source/ControlRig/Private/ModularRig.cpp b/Engine/Plugins/Animation/ControlRig/Source/ControlRig/Private/ModularRig.cpp
--- a/Engine/Plugins/Animation/ControlRig/Source/ControlRig/Private/ModularRig.cpp
+++ b/Engine/Plugins/Animation/ControlRig/Source/ControlRig/Private/ModularRig.cpp
@@ -644,6 +644,9 @@ void UModularRig::ExecuteQueue()
 				RigPublicContext.RuntimeSettings = PublicContext.RuntimeSettings;
 
 #if WITH_EDITOR
+				// redirect the records to be recorded to the top level context
+				const TGuardValue<FInstructionRecordContainer*> RecordsGuard(RigPublicContext.Records, PublicContext.Records);
+				
 				if (!FirstModuleEvent.Contains(ModuleInstance))
 				{
 					//ModuleRig->InstructionVisitInfo.FirstEntryEventInQueue = ExecutionElement.EventName;
diff --git a/Engine/Plugins/Animation/ControlRig/Source/ControlRig/Private/Rigs/RigDependency.cpp b/Engine/Plugins/Animation/ControlRig/Source/ControlRig/Private/Rigs/RigDependency.cpp
--- a/Engine/Plugins/Animation/ControlRig/Source/ControlRig/Private/Rigs/RigDependency.cpp
+++ b/Engine/Plugins/Animation/ControlRig/Source/ControlRig/Private/Rigs/RigDependency.cpp
@@ -10,7 +10,13 @@ const TRigHierarchyDependencyMap& FEmptyRigDependenciesProvider::GetRigHierarchy
 	return Empty;
 }
 
-FRigDependenciesProviderForVM::FRigDependenciesProviderForVM(UControlRig* InControlRig, FName InEventName)
+const TRigHierarchyDependencyMap& FEmptyRigDependenciesProvider::GetReverseRigHierarchyDependencies() const
+{
+	static const TRigHierarchyDependencyMap Empty;
+	return Empty;
+}
+
+FRigDependenciesProviderForVM::FRigDependenciesProviderForVM(UControlRig* InControlRig, const FName& InEventName)
 	: WeakControlRig(InControlRig)
 	, EventName(InEventName)
 {}
@@ -40,78 +46,97 @@ const TRigHierarchyDependencyMap& FRigDependenciesProviderForVM::GetRigHierarchy
 		return Empty;
 	}
 
-	uint32 NewHash = 0; 
-	auto LazyGetTopologyHash = [Hierarchy, &NewHash]()
-	{
-		if (NewHash == 0)
-		{
-			// NOTE: GetTopologyVersion() is used to track simple topology changes (GetTopologyHash() is probably more complete but also slower)
-			NewHash = Hierarchy->GetTopologyVersion();
-		}
-		return NewHash;
-	};
-
 	const FRigVMExtendedExecuteContext& ExtendedExecuteContext = ControlRig->GetRigVMExtendedExecuteContext();
 	const FControlRigExecuteContext& Context = ExtendedExecuteContext.GetPublicData<FControlRigExecuteContext>();
-	
-	if (!CachedDependencies.IsEmpty() && TopologyHash == LazyGetTopologyHash())
+	const uint32 CurrentRecordsHash = Context.GetInstructionRecordsHash(); 
+	if (!CachedDependencies.IsEmpty() && RecordsHash == CurrentRecordsHash)
 	{
 		return CachedDependencies;
 	}
 	
-	FName EventNameToUse = EventName;
-	if(EventNameToUse.IsNone())
-	{
-		EventNameToUse = FRigUnit_BeginExecution().GetEventName();
-	}
-
 	CachedDependencies.Reset();
+
+	ComputeDependencies(Context, VM, NAME_None, CachedDependencies);
+
+	RecordsHash = CurrentRecordsHash;
+	ReverseRecordsHash = UINT32_MAX;
 	
-	const FRigVMInstructionArray Instructions = VM->GetByteCode().GetInstructions();
+#endif
 
-	// if the VM doesn't implement the given event
-	if(!VM->ContainsEntry(EventNameToUse))
+	return CachedDependencies;
+}
+
+const TRigHierarchyDependencyMap& FRigDependenciesProviderForVM::GetReverseRigHierarchyDependencies() const
+{
+	const TRigHierarchyDependencyMap& Dependencies = GetRigHierarchyDependencies();
+	if (ReverseRecordsHash == RecordsHash)
 	{
-		return Empty;
+		return CachedReverseDependencies;
 	}
 
-	// only represent instruction for a given event
-	const int32 EntryIndex = VM->GetByteCode().FindEntryIndex(EventNameToUse);
-	const int32 EntryInstructionIndex = VM->GetByteCode().GetEntry(EntryIndex).InstructionIndex;
-
-	TMap<FRigVMOperand, TArray<int32>> OperandToInstructions;
-
-	for(int32 InstructionIndex = EntryInstructionIndex; InstructionIndex < Instructions.Num(); InstructionIndex++)
+	CachedReverseDependencies.Reset();
+	CachedReverseDependencies.Reserve(Dependencies.Num());
+	for (const TPair<FRigHierarchyRecord, TSet<FRigHierarchyRecord>>& Pair : Dependencies)
 	{
-		// early exit since following instructions belong to another event
-		if(Instructions[InstructionIndex].OpCode == ERigVMOpCode::Exit)
+		for (const FRigHierarchyRecord& DependentRecord : Pair.Value)
 		{
-			break;
+			CachedReverseDependencies.FindOrAdd(DependentRecord).Add(Pair.Key);
 		}
+	}
 
-		const FRigVMOperandArray InputOperands = VM->GetByteCode().GetInputOperands(InstructionIndex);
-		for(const FRigVMOperand InputOperand : InputOperands)
-		{
-			const FRigVMOperand InputOperandNoRegisterOffset(InputOperand.GetMemoryType(), InputOperand.GetRegisterIndex());
-			OperandToInstructions.FindOrAdd(InputOperandNoRegisterOffset).Add(InstructionIndex);
-		}
+	ReverseRecordsHash = RecordsHash;
+	return CachedReverseDependencies;
+}
+
+void FRigDependenciesProviderForVM::InvalidateCache()
+{
+	CachedDependencies.Reset();
+	CachedReverseDependencies.Reset();
+	RecordsHash = ReverseRecordsHash = UINT32_MAX;
+}
+
+UControlRig* FRigDependenciesProviderForVM::GetControlRig() const
+{
+	if (WeakControlRig.IsValid())
+	{
+		return WeakControlRig.Get();
+	}
+	return nullptr;
+}
+
+void FRigDependenciesProviderForVM::ComputeDependencies(const FControlRigExecuteContext& InContext, const URigVM* InVM, const FName& InVMName, TRigHierarchyDependencyMap& OutDependencies) const
+{
+#if WITH_EDITOR
+	
+	FName EventNameToUse = EventName;
+	if(EventNameToUse.IsNone())
+	{
+		EventNameToUse = FRigUnit_BeginExecution().GetEventName();
+	}
+
+	// if the VM doesn't implement the given event
+	if(!InVM->ContainsEntry(EventNameToUse))
+	{
+		return;
 	}
 
+	const FRigVMInstructionArray Instructions = InVM->GetByteCode().GetInstructions();
+	const FInstructionRecordContainer& InstructionRecords = InContext.GetInstructionRecords();
 	TArray<TArray<FInstructionRecord>> ReadRecordsPerInstruction, WrittenRecordsPerInstruction;
 	
 	// Find the max instruction index
 	{
-		/** TODO: UE-207320
-		 * This function needs to be rewritten to account for modular rigs. Rig Hierarchies will have reads/writes coming from
-		 * different VMs, so we cannot depend on the algorithm used in this function to find dependencies. */
-		
-		int32 MaxInstructionIndex = Instructions.Num();
+		int32 MaxInstructionIndex = 0;
 		for(int32 RecordType = 0; RecordType < 2; RecordType++)
 		{
-			const TArray<FInstructionRecord>& Records = RecordType == 0 ? Context.ReadRecords : Context.WrittenRecords; 
+			const TArray<FInstructionRecord>& Records = RecordType == 0 ? InstructionRecords.ReadRecords : InstructionRecords.WrittenRecords; 
 			for(int32 RecordIndex = 0; RecordIndex < Records.Num(); RecordIndex++)
 			{
 				const FInstructionRecord& Record = Records[RecordIndex];
+				if (Record.VM != InVMName)
+				{
+					continue;
+				}
 				MaxInstructionIndex = FMath::Max(MaxInstructionIndex, Record.InstructionIndex);
 			}
 		}
@@ -122,38 +147,52 @@ const TRigHierarchyDependencyMap& FRigDependenciesProviderForVM::GetRigHierarchy
 	// fill lookup tables per instruction / element
 	for(int32 RecordType = 0; RecordType < 2; RecordType++)
 	{
-		const TArray<FInstructionRecord>& Records = RecordType == 0 ? Context.ReadRecords : Context.WrittenRecords; 
+		const TArray<FInstructionRecord>& Records = RecordType == 0 ? InstructionRecords.ReadRecords : InstructionRecords.WrittenRecords; 
 		TArray<TArray<FInstructionRecord>>& PerInstruction = RecordType == 0 ? ReadRecordsPerInstruction : WrittenRecordsPerInstruction;
 
 		for(int32 RecordIndex = 0; RecordIndex < Records.Num(); RecordIndex++)
 		{
 			const FInstructionRecord& Record = Records[RecordIndex];
+			if (Record.VM != InVMName)
+			{
+				continue;
+			}
 			PerInstruction[Record.InstructionIndex].Add(Record);
 		}
 	}
 
-	// for each read transform on an instruction
-	// follow the operands to the next instruction affected by it
-	TArray<FInstructionRecord> FilteredRecords;
+	// for each written transform / metadata on an instruction
 	TArray<int32> InstructionsToVisit;
+	TArray<FInstructionRecord> FilteredReadRecords;
+	TArray<TTuple<FRigHierarchyRecord,FRigHierarchyRecord>> ReadRecordsToAddAtEnd;
+	TMap<FRigVMOperand, TArray<int32>> OutputOperandToInstructions;
+
+	for(int32 InstructionIndex = 0; InstructionIndex < Instructions.Num(); InstructionIndex++)
+	{
+		FRigVMOperandArray Operands = InVM->GetByteCode().GetOutputOperands(InstructionIndex);
+		for (const FRigVMOperand& Operand : Operands)
+		{
+			const FRigVMOperand OperandNoRegisterOffset(Operand.GetMemoryType(), Operand.GetRegisterIndex());
+			OutputOperandToInstructions.FindOrAdd(OperandNoRegisterOffset).Add(InstructionIndex);
+		}
+	}
 
-	for(int32 InstructionIndex = EntryInstructionIndex; InstructionIndex < Instructions.Num(); InstructionIndex++)
+	for(int32 InstructionIndex = 0; InstructionIndex < Instructions.Num(); InstructionIndex++)
 	{
-		// early exit since following instructions belong to another event
-		if(Instructions[InstructionIndex].OpCode == ERigVMOpCode::Exit)
+		if (!WrittenRecordsPerInstruction.IsValidIndex(InstructionIndex))
 		{
-			break;
+			continue;
 		}
 
-		const TArray<FInstructionRecord>& ReadRecords = ReadRecordsPerInstruction[InstructionIndex];
-		if(ReadRecords.IsEmpty())
+		const TArray<FInstructionRecord>& WrittenRecords = WrittenRecordsPerInstruction[InstructionIndex];
+		if(WrittenRecords.IsEmpty())
 		{
 			continue;
 		}
 
-		FilteredRecords.Reset();
+		FilteredReadRecords.Reset();
 
-		for(int32 ReadRecordIndex = 0; ReadRecordIndex < ReadRecords.Num(); ReadRecordIndex++)
+		for(int32 WrittenRecordIndex = 0; WrittenRecordIndex < WrittenRecords.Num(); WrittenRecordIndex++)
 		{
 			InstructionsToVisit.Reset();
 			InstructionsToVisit.Add(InstructionIndex);
@@ -162,95 +201,116 @@ const TRigHierarchyDependencyMap& FRigDependenciesProviderForVM::GetRigHierarchy
 			{
 				const int32 InstructionToVisit = InstructionsToVisit[InstructionToVisitIndex];
 
-				const TArray<FInstructionRecord>& WrittenRecords = WrittenRecordsPerInstruction[InstructionToVisit];
-				for(int32 WrittenRecordIndex = 0; WrittenRecordIndex < WrittenRecords.Num(); WrittenRecordIndex++)
+				if (ReadRecordsPerInstruction.IsValidIndex(InstructionToVisit))
 				{
-					const FInstructionRecord& WrittenRecord = WrittenRecords[WrittenRecordIndex];
-
-					// for the first instruction in this pass let's only care about
-					// written transforms which have not been read before
-					if(InstructionToVisitIndex == 0)
+					const TArray<FInstructionRecord>& ReadRecords = ReadRecordsPerInstruction[InstructionToVisit];
+					for(int32 ReadRecordIndex = 0; ReadRecordIndex < ReadRecords.Num(); ReadRecordIndex++)
 					{
-						if(ReadRecords.Contains(WrittenRecord))
-						{
-							continue;
-						}
+						const FInstructionRecord& ReadRecord = ReadRecords[ReadRecordIndex];
+						FilteredReadRecords.AddUnique(ReadRecord);
 					}
-
-					// store the filtered records for any instruction here
-					// by clearing the instruction index field.
-					FInstructionRecord SimplifiedRecord = WrittenRecord;
-					SimplifiedRecord.InstructionIndex = INDEX_NONE;
-					FilteredRecords.AddUnique(SimplifiedRecord);
 				}
 
-				FRigVMOperandArray OutputOperands = VM->GetByteCode().GetOutputOperands(InstructionToVisit);
-				for(const FRigVMOperand OutputOperand : OutputOperands)
+				FRigVMOperandArray InputOperands = InVM->GetByteCode().GetInputOperands(InstructionToVisit);
+				for(const FRigVMOperand InputOperand : InputOperands)
 				{
-					const FRigVMOperand OutputOperandNoRegisterOffset(OutputOperand.GetMemoryType(), OutputOperand.GetRegisterIndex());
+					const FRigVMOperand InputOperandNoRegisterOffset(InputOperand.GetMemoryType(), InputOperand.GetRegisterIndex());
 
-					if(const TArray<int32>* InstructionsWithInputOperand = OperandToInstructions.Find(OutputOperandNoRegisterOffset))
+					if(const TArray<int32>* InstructionsWithOutputOperand = OutputOperandToInstructions.Find(InputOperandNoRegisterOffset))
 					{
-						for(int32 InstructionWithInputOperand : *InstructionsWithInputOperand)
+						for(int32 InstructionWithOutputOperand : *InstructionsWithOutputOperand)
 						{
-							InstructionsToVisit.AddUnique(InstructionWithInputOperand);
+							InstructionsToVisit.AddUnique(InstructionWithOutputOperand);
 						}
 					}
 				}
 			}
 		}
-		
-		for(const FInstructionRecord& ReadRecord : ReadRecords)
+
+		for(const FInstructionRecord& WrittenRecord : WrittenRecords)
+		{
+			const FRigHierarchyRecord TargetRecord(
+				WrittenRecord.MetadataName.IsNone() ? FRigHierarchyRecord::EType_RigElement : FRigHierarchyRecord::EType_Metadata,
+				WrittenRecord.ElementIndex,
+				WrittenRecord.MetadataName);
+
+			const FRigHierarchyRecord InstructionRecord(
+				FRigHierarchyRecord::EType_Instruction,
+				WrittenRecord.InstructionIndex,
+				WrittenRecord.VM);
+
+			OutDependencies.FindOrAdd(TargetRecord).Add(InstructionRecord);
+
+			for(const FInstructionRecord& FilteredReadRecord : FilteredReadRecords)
+			{
+				const FRigHierarchyRecord SourceRecord(
+					FilteredReadRecord.MetadataName.IsNone() ? FRigHierarchyRecord::EType_RigElement : FRigHierarchyRecord::EType_Metadata,
+					FilteredReadRecord.ElementIndex,
+					FilteredReadRecord.MetadataName);
+
+				ReadRecordsToAddAtEnd.Emplace(InstructionRecord, SourceRecord);
+			}
+		}
+	}
+
+	for (const TTuple<FRigHierarchyRecord,FRigHierarchyRecord>& Pair : ReadRecordsToAddAtEnd)
+	{
+		const FRigHierarchyRecord& Target = Pair.Get<0>();
+		const FRigHierarchyRecord& Source = Pair.Get<1>();
+
+		if (Target.IsInstruction())
 		{
-			for(const FInstructionRecord& FilteredRecord : FilteredRecords)
+			if (const TSet<FRigHierarchyRecord>* SourceToTarget = OutDependencies.Find(Source))
 			{
-				// only create dependencies for reads and writes that are on the same slice
-				if(ReadRecord != FilteredRecord && ReadRecord.SliceIndex == FilteredRecord.SliceIndex)
+				if (SourceToTarget->Contains(Target))
 				{
-					const FRigHierarchyRecord SourceHierarchyRecord(ReadRecord.ElementIndex, ReadRecord.MetadataName);
-					const FRigHierarchyRecord TargetHierarchyRecord(FilteredRecord.ElementIndex, FilteredRecord.MetadataName);
-					CachedDependencies.FindOrAdd(TargetHierarchyRecord).AddUnique(SourceHierarchyRecord);
+					continue;
 				}
 			}
 		}
+		
+		OutDependencies.FindOrAdd(Pair.Get<0>()).Add(Pair.Get<1>());
 	}
 
-	TopologyHash = NewHash;
-	
 #endif
-
-	return CachedDependencies;
-}
-
-UControlRig* FRigDependenciesProviderForVM::GetControlRig() const
-{
-	if (WeakControlRig.IsValid())
-	{
-		return WeakControlRig.Get();
-	}
-	return nullptr;
 }
 
-FRigDependenciesProviderForControlRig::FRigDependenciesProviderForControlRig(UControlRig* InControlRig, FName InEventName)
+FRigDependenciesProviderForControlRig::FRigDependenciesProviderForControlRig(UControlRig* InControlRig, const FName& InEventName)
 	: FRigDependenciesProviderForVM(InControlRig, InEventName)
 {
-	if (const UControlRig* ControlRig = GetControlRig())
-	{
-		if (ControlRig->IsA<UModularRig>())
-		{
-			// todo
-		}
-	}
 }
 
 const TRigHierarchyDependencyMap& FRigDependenciesProviderForControlRig::GetRigHierarchyDependencies() const
 {
-	if (const UControlRig* ControlRig = GetControlRig())
+	if (const UModularRig* ModularRig = Cast<UModularRig>(GetControlRig()))
 	{
-		if (ControlRig->IsA<UModularRig>())
+		const FRigVMExtendedExecuteContext& ExtendedExecuteContext = ModularRig->GetRigVMExtendedExecuteContext();
+		const FControlRigExecuteContext& Context = ExtendedExecuteContext.GetPublicData<FControlRigExecuteContext>();
+		const uint32 CurrentRecordsHash = Context.GetInstructionRecordsHash(); 
+	
+		if (!CachedDependencies.IsEmpty() && RecordsHash == CurrentRecordsHash)
 		{
-			// todo
+			return CachedDependencies;
 		}
+
+		CachedDependencies.Reset();
+		CachedReverseDependencies.Reset();
+		
+		ModularRig->ForEachModule([this, Context](const FRigModuleInstance* InModule) -> bool
+		{
+			if (UControlRig* ModularRig = InModule->GetRig())
+			{
+				if (const URigVM* VM = ModularRig->GetVM())
+				{
+					ComputeDependencies(Context, VM, InModule->Name, CachedDependencies);
+				}
+			}
+			return true;
+		});
+
+		RecordsHash = CurrentRecordsHash;
+		ReverseRecordsHash = UINT32_MAX;
+		return CachedDependencies;
 	}
 	return FRigDependenciesProviderForVM::GetRigHierarchyDependencies();
 }
diff --git a/Engine/Plugins/Animation/ControlRig/Source/ControlRig/Private/Rigs/RigHierarchy.cpp b/Engine/Plugins/Animation/ControlRig/Source/ControlRig/Private/Rigs/RigHierarchy.cpp
--- a/Engine/Plugins/Animation/ControlRig/Source/ControlRig/Private/Rigs/RigHierarchy.cpp
+++ b/Engine/Plugins/Animation/ControlRig/Source/ControlRig/Private/Rigs/RigHierarchy.cpp
@@ -4965,13 +4965,13 @@ FTransform URigHierarchy::GetTransform(FRigTransformElement* InTransformElement,
 			{
 				const FControlRigExecuteContext& PublicData = ExecuteContext->GetPublicData<FControlRigExecuteContext>();
 				const FName VMName = PublicData.GetRigModuleInstance() ? PublicData.GetRigModuleInstance()->Name : NAME_None;
-				PublicData.ReadRecords.Emplace(
-					PublicData.GetInstructionIndex(),
+
+				const FInstructionRecord Record(PublicData.GetInstructionIndex(),
 					ExecuteContext->GetSlice().GetIndex(),
 					InTransformElement->GetIndex(),
-					InTransformType,
-					VMName
-				);
+					VMName);
+
+				PublicData.AddReadInstructionRecord(Record);
 			}			
 		}
 	}
@@ -5129,13 +5129,8 @@ void URigHierarchy::SetTransform(FRigTransformElement* InTransformElement, const
 				const FControlRigExecuteContext& PublicData = ExecuteContext->GetPublicData<FControlRigExecuteContext>();
 				const FRigVMSlice& Slice = ExecuteContext->GetSlice();
 				const FName VMName = PublicData.GetRigModuleInstance() ? PublicData.GetRigModuleInstance()->Name : NAME_None;
-				PublicData.WrittenRecords.Emplace(
-					PublicData.GetInstructionIndex(),
-					Slice.GetIndex(),
-					InTransformElement->GetIndex(),
-					InTransformType,
-					VMName
-				);
+				const FInstructionRecord Record(PublicData.GetInstructionIndex(), Slice.GetIndex(), InTransformElement->GetIndex(), VMName);
+				PublicData.AddWrittenInstructionRecord(Record);
 
 				// if we are setting a control / null parent after a child here - let's let the user know
 				if(InTransformElement->IsA<FRigControlElement>() || InTransformElement->IsA<FRigNullElement>())
@@ -5147,7 +5142,7 @@ void URigHierarchy::SetTransform(FRigTransformElement* InTransformElement, const
 						{
 							for (FRigBaseElement* Child : GetChildren(InTransformElement))
 							{
-								const bool bChildFound = PublicData.WrittenRecords.ContainsByPredicate([Child](const FInstructionRecord& Record) -> bool
+								const bool bChildFound = PublicData.GetInstructionRecords().WrittenRecords.ContainsByPredicate([Child](const FInstructionRecord& Record) -> bool
 								{
 									return Record.ElementIndex == Child->GetIndex();
 								});
@@ -5314,13 +5309,12 @@ FTransform URigHierarchy::GetControlOffsetTransform(FRigControlElement* InContro
 				const FControlRigExecuteContext& PublicData = ExecuteContext->GetPublicData<FControlRigExecuteContext>();
 				const FName VMName = PublicData.GetRigModuleInstance() ? PublicData.GetRigModuleInstance()->Name : NAME_None;
 
-				PublicData.ReadRecords.Emplace(
-					PublicData.GetInstructionIndex(),
+				const FInstructionRecord Record(PublicData.GetInstructionIndex(),
 					ExecuteContext->GetSlice().GetIndex(),
 					InControlElement->GetIndex(),
-					InTransformType,
-					VMName
-				);
+					VMName);
+
+				PublicData.AddReadInstructionRecord(Record);
 			}
 		}
 	}
@@ -5410,13 +5404,8 @@ void URigHierarchy::SetControlOffsetTransform(FRigControlElement* InControlEleme
 			{
 				const FControlRigExecuteContext& PublicData = ExecuteContext->GetPublicData<FControlRigExecuteContext>();
 				const FName VMName = PublicData.GetRigModuleInstance() ? PublicData.GetRigModuleInstance()->Name : NAME_None;
-				PublicData.WrittenRecords.Emplace(
-					PublicData.GetInstructionIndex(),
-					ExecuteContext->GetSlice().GetIndex(),
-					InControlElement->GetIndex(),
-					InTransformType,
-					VMName
-				);
+				const FInstructionRecord Record(PublicData.GetInstructionIndex(), ExecuteContext->GetSlice().GetIndex(), InControlElement->GetIndex(), VMName);
+				PublicData.AddWrittenInstructionRecord(Record);
 			}
 		}
 	}
@@ -6390,29 +6379,59 @@ bool URigHierarchy::IsDependentOn(FRigBaseElement* InDependent, FRigBaseElement*
 		}
 	}
 
-	// check the optional dependency map
-	const TRigHierarchyDependencyMap& DependencyMap = InDependencyProvider.GetRigHierarchyDependencies();
-	if(const TArray<FRigHierarchyRecord>* DependentRecordsPtr = DependencyMap.Find(InDependent->GetIndex()))
+	struct Local
 	{
-		const TArray<FRigHierarchyRecord>& DependentRecords = *DependentRecordsPtr;
-		for(const FRigHierarchyRecord& DependentRecord : DependentRecords)
+		static bool FoundDependencyInProvider(
+			const URigHierarchy* InHierarchy,
+			FRigBaseElement* InDependency,
+			const IRigDependenciesProvider& InDependencyProvider,
+			bool bIsOnActualTopology,
+			const TTuple<int32,int32>& CacheKey,
+			const FRigHierarchyRecord& InDependentRecord)
 		{
-			if (!DependentRecord.IsElement())
-			{
-				continue;
-			}
+			const TRigHierarchyDependencyMap& DependencyMap = InDependencyProvider.GetRigHierarchyDependencies();
 			
-			ensure(Elements.IsValidIndex(DependentRecord.Index));
-			if(IsDependentOn(Elements[DependentRecord.Index], InDependency, InDependencyProvider, false))
+			if(const TSet<FRigHierarchyRecord>* DependentRecordsPtr = DependencyMap.Find(InDependentRecord))
 			{
-				// we'll only update the caches if we are following edges on the actual topology
-				if(bIsOnActualTopology)
+				const TSet<FRigHierarchyRecord>& DependentRecords = *DependentRecordsPtr;
+				for(const FRigHierarchyRecord& DependentRecord : DependentRecords)
 				{
-					ElementDependencyCache.Get().FindOrAdd(CacheKey, true);
+					if (DependentRecord.IsElement())
+					{
+						ensure(InHierarchy->Elements.IsValidIndex(DependentRecord.Index));
+						if(InHierarchy->IsDependentOn(InHierarchy->Elements[DependentRecord.Index], InDependency, InDependencyProvider, false))
+						{
+							// we'll only update the caches if we are following edges on the actual topology
+							if(bIsOnActualTopology)
+							{
+								InHierarchy->ElementDependencyCache.Get().FindOrAdd(CacheKey, true);
+							}
+							return true;
+						}
+					}
+					else if (DependentRecord.IsInstruction())
+					{
+						if (FoundDependencyInProvider(InHierarchy, InDependency, InDependencyProvider, bIsOnActualTopology, CacheKey, DependentRecord))
+						{
+							return true;
+						}
+					}
 				}
-				return true;
 			}
+			return false;
+		}
+	};
+
+	// check the optional dependency map
+	const FRigHierarchyRecord DependentRecord(FRigHierarchyRecord::EType_RigElement, InDependent->GetIndex(), NAME_None);
+	if (Local::FoundDependencyInProvider(this, InDependency, InDependencyProvider, bIsOnActualTopology, CacheKey, DependentRecord))
+	{
+		// we'll only update the caches if we are following edges on the actual topology
+		if(bIsOnActualTopology)
+		{
+			ElementDependencyCache.Get().FindOrAdd(CacheKey, true);
 		}
+		return true;		
 	}
 
 	// we'll only update the caches if we are following edges on the actual topology
diff --git a/Engine/Plugins/Animation/ControlRig/Source/ControlRig/Public/Rigs/RigDependency.h b/Engine/Plugins/Animation/ControlRig/Source/ControlRig/Public/Rigs/RigDependency.h
--- a/Engine/Plugins/Animation/ControlRig/Source/ControlRig/Public/Rigs/RigDependency.h
+++ b/Engine/Plugins/Animation/ControlRig/Source/ControlRig/Public/Rigs/RigDependency.h
@@ -6,6 +6,7 @@
 #include "UObject/WeakObjectPtr.h"
 #include "UObject/WeakObjectPtrTemplates.h"
 #include "RigDependencyRecords.h"
+#include "Units/RigUnitContext.h"
 
 #define UE_API CONTROLRIG_API
 
@@ -21,6 +22,8 @@ struct IRigDependenciesProvider
 {
 	virtual ~IRigDependenciesProvider() = default;
 	virtual const TRigHierarchyDependencyMap& GetRigHierarchyDependencies() const = 0;
+	virtual const TRigHierarchyDependencyMap& GetReverseRigHierarchyDependencies() const = 0;
+	virtual void InvalidateCache() = 0;
 };
 
 /**
@@ -31,6 +34,8 @@ struct FEmptyRigDependenciesProvider : public IRigDependenciesProvider
 {
 	virtual ~FEmptyRigDependenciesProvider() override = default;
 	UE_API virtual const TRigHierarchyDependencyMap& GetRigHierarchyDependencies() const override;
+	UE_API virtual const TRigHierarchyDependencyMap& GetReverseRigHierarchyDependencies() const;
+	UE_API virtual void InvalidateCache() {}
 };
 
 /**
@@ -40,20 +45,26 @@ struct FEmptyRigDependenciesProvider : public IRigDependenciesProvider
 struct FRigDependenciesProviderForVM : public IRigDependenciesProvider
 {
 	FRigDependenciesProviderForVM() = default;
-	UE_API FRigDependenciesProviderForVM(UControlRig* InControlRig, FName InEventName = NAME_None);
+	UE_API FRigDependenciesProviderForVM(UControlRig* InControlRig, const FName& InEventName = NAME_None);
 	virtual ~FRigDependenciesProviderForVM() override = default;
 	
 	UE_API virtual const TRigHierarchyDependencyMap& GetRigHierarchyDependencies() const override;
+	UE_API virtual const TRigHierarchyDependencyMap& GetReverseRigHierarchyDependencies() const override;
+	UE_API virtual void InvalidateCache() override;
 
 	UE_API UControlRig* GetControlRig() const;
 
 protected:
 
+	void ComputeDependencies(const FControlRigExecuteContext& InContext, const URigVM* InVM, const FName& InVMName, TRigHierarchyDependencyMap& OutDependencies) const;
+
 	TWeakObjectPtr<UControlRig> WeakControlRig = nullptr;
 	FName EventName = NAME_None;
 	
 	mutable TRigHierarchyDependencyMap CachedDependencies;
-	mutable uint32 TopologyHash = 0;
+	mutable TRigHierarchyDependencyMap CachedReverseDependencies;
+	mutable uint32 RecordsHash = 0;
+	mutable uint32 ReverseRecordsHash = 0;
 };
 
 /**
@@ -64,14 +75,10 @@ protected:
 struct FRigDependenciesProviderForControlRig : public FRigDependenciesProviderForVM
 {
 	FRigDependenciesProviderForControlRig() = default;
-	UE_API FRigDependenciesProviderForControlRig(UControlRig* InControlRig, FName InEventName = NAME_None);
+	UE_API FRigDependenciesProviderForControlRig(UControlRig* InControlRig, const FName& InEventName = NAME_None);
 	virtual ~FRigDependenciesProviderForControlRig() override = default;
 	
 	UE_API virtual const TRigHierarchyDependencyMap& GetRigHierarchyDependencies() const override;
-
-protected:
-
-	TMap<FName, TSharedPtr<FRigDependenciesProviderForVM>> ModuleDependencyProviders;
 };
 
 #undef UE_API
diff --git a/Engine/Plugins/Animation/ControlRig/Source/ControlRig/Public/Rigs/RigDependencyRecords.h b/Engine/Plugins/Animation/ControlRig/Source/ControlRig/Public/Rigs/RigDependencyRecords.h
--- a/Engine/Plugins/Animation/ControlRig/Source/ControlRig/Public/Rigs/RigDependencyRecords.h
+++ b/Engine/Plugins/Animation/ControlRig/Source/ControlRig/Public/Rigs/RigDependencyRecords.h
@@ -10,40 +10,63 @@
 
 struct FRigHierarchyRecord
 {
+	enum EType
+	{
+		EType_RigElement,
+		EType_Metadata,
+		EType_Instruction,
+		EType_Invalid
+	};
+
 	FRigHierarchyRecord()
-		: Index(INDEX_NONE)
-		, MetadataName(NAME_None)
+		: Type(EType_Invalid)
+		, Index(INDEX_NONE)
+		, Name(NAME_None)
 	{
 	}
 
-	FRigHierarchyRecord(int32 InIndex, const FName& InMetadataName = NAME_None)
-		: Index(InIndex)
-		, MetadataName(InMetadataName)
+	FRigHierarchyRecord(EType InType, int32 InIndex, const FName& InName)
+		: Type(InType)
+		, Index(InIndex)
+		, Name(InName)
 	{
 	}
 	
 	bool IsValid() const
 	{
-		return Index != INDEX_NONE;
+		return Type != EType_Invalid;
 	}
 
 	bool IsElement() const
 	{
-		return IsValid() && MetadataName.IsNone();
+		return Type == EType_RigElement;
 	}
 
 	bool IsMetadata() const
 	{
-		return IsValid() && !MetadataName.IsNone();
+		return Type == EType_Metadata;
+	}
+
+	bool IsInstruction() const
+	{
+		return Type == EType_Instruction;
 	}
 
 	bool operator==(const FRigHierarchyRecord& InRecord) const
 	{
-		return Index == InRecord.Index && MetadataName == InRecord.MetadataName;
+		return Type == InRecord.Type && Index == InRecord.Index && Name == InRecord.Name;
 	}
 
 	bool operator <(const FRigHierarchyRecord& InRecord) const
 	{
+		if (Type < InRecord.Type)
+		{
+			return true;
+		}
+		if (Type > InRecord.Type)
+		{
+			return false;
+		}
 		if (Index < InRecord.Index)
 		{
 			return true;
@@ -52,11 +75,19 @@ struct FRigHierarchyRecord
 		{
 			return false;
 		}
-		return MetadataName.Compare(InRecord.MetadataName) < 0;
+		return Name.Compare(InRecord.Name) < 0;
 	}
 
 	bool operator >(const FRigHierarchyRecord& InRecord) const
 	{
+		if (Type > InRecord.Type)
+		{
+			return true;
+		}
+		if (Type < InRecord.Type)
+		{
+			return false;
+		}
 		if (Index > InRecord.Index)
 		{
 			return true;
@@ -65,19 +96,21 @@ struct FRigHierarchyRecord
 		{
 			return false;
 		}
-		return MetadataName.Compare(InRecord.MetadataName) > 0;
+		return Name.Compare(InRecord.Name) > 0;
 	}
 
+	EType Type;
 	int32 Index;
-	FName MetadataName;
+	FName Name;
 };
 
-typedef TMap<FRigHierarchyRecord, TArray<FRigHierarchyRecord>> TRigHierarchyDependencyMap;
-typedef TPair<FRigHierarchyRecord, TArray<FRigHierarchyRecord>> TRigHierarchyDependencyMapPair;
+typedef TMap<FRigHierarchyRecord, TSet<FRigHierarchyRecord>> TRigHierarchyDependencyMap;
+typedef TPair<FRigHierarchyRecord, TSet<FRigHierarchyRecord>> TRigHierarchyDependencyMapPair;
 
 inline uint32 GetTypeHash(const FRigHierarchyRecord& InRecord)
 {
-	return HashCombine(GetTypeHash(InRecord.Index), GetTypeHash(InRecord.MetadataName));
+	return HashCombine(GetTypeHash(static_cast<int32>(InRecord.Type)),
+		HashCombine(GetTypeHash(InRecord.Index), GetTypeHash(InRecord.Name)));
 }
 
 struct FInstructionRecord
@@ -86,17 +119,15 @@ struct FInstructionRecord
 		: InstructionIndex(INDEX_NONE)
 		, SliceIndex(INDEX_NONE)
 		, ElementIndex(INDEX_NONE)
-		, TransformType(ERigTransformType::CurrentLocal)
 		, VM(NAME_None)
 		, MetadataName(NAME_None)
 	{
 	}
 
-	FInstructionRecord(int32 InInstructionIndex, int32 InSliceIndex, int32 InElementIndex, ERigTransformType::Type InTransformType, const FName& InVM = NAME_None, const FName& InMetadataName = NAME_None)
+	FInstructionRecord(int32 InInstructionIndex, int32 InSliceIndex, int32 InElementIndex, const FName& InVM = NAME_None, const FName& InMetadataName = NAME_None)
 		:InstructionIndex(InInstructionIndex)
 		, SliceIndex(InSliceIndex)
 		, ElementIndex(InElementIndex)
-		, TransformType(InTransformType)
 		, VM(InVM)
 		, MetadataName(InMetadataName)
 	{
@@ -107,7 +138,6 @@ struct FInstructionRecord
 		return InstructionIndex == InRecord.InstructionIndex &&
 			SliceIndex == InRecord.SliceIndex &&
 			ElementIndex == InRecord.ElementIndex &&
-			TransformType == InRecord.TransformType &&
 			VM == InRecord.VM &&
 			MetadataName == InRecord.MetadataName;
 	}
@@ -138,14 +168,6 @@ struct FInstructionRecord
 		{
 			return false;
 		}
-		if (TransformType < InRecord.TransformType)
-		{
-			return true;
-		}
-		if (TransformType > InRecord.TransformType)
-		{
-			return false;
-		}
 		const int32 VMComp = VM.Compare(InRecord.VM);
 		if (VMComp < 0)
 		{
@@ -184,14 +206,6 @@ struct FInstructionRecord
 		{
 			return false;
 		}
-		if (TransformType > InRecord.TransformType)
-		{
-			return true;
-		}
-		if (TransformType < InRecord.TransformType)
-		{
-			return false;
-		}
 		const int32 VMComp = VM.Compare(InRecord.VM);
 		if (VMComp > 0)
 		{
@@ -207,7 +221,6 @@ struct FInstructionRecord
 	int32 InstructionIndex;
 	int32 SliceIndex;
 	int32 ElementIndex;
-	ERigTransformType::Type TransformType;
 	FName VM;
 	FName MetadataName;
 };
@@ -216,7 +229,6 @@ inline uint32 GetTypeHash(const FInstructionRecord& InRecord)
 {
 	uint32 Hash = HashCombine(GetTypeHash(InRecord.InstructionIndex), GetTypeHash(InRecord.SliceIndex));
 	Hash = HashCombine(Hash, GetTypeHash(InRecord.ElementIndex));
-	Hash = HashCombine(Hash, GetTypeHash(static_cast<int32>(InRecord.TransformType)));
 	Hash = HashCombine(Hash, GetTypeHash(InRecord.VM));
 	Hash = HashCombine(Hash, GetTypeHash(InRecord.MetadataName));
 	return Hash;
@@ -225,4 +237,18 @@ inline uint32 GetTypeHash(const FInstructionRecord& InRecord)
 typedef TMap<FInstructionRecord, TArray<FInstructionRecord>> TInstructionRecordMap;
 typedef TPair<FInstructionRecord, TArray<FInstructionRecord>> TInstructionRecordMapPair;
 
+struct FInstructionRecordContainer
+{
+	FInstructionRecordContainer()
+		: ReadRecordsHash(UINT32_MAX)
+		, WrittenRecordsHash(UINT32_MAX)
+	{
+	}
+	
+	mutable TArray<FInstructionRecord> ReadRecords;
+	mutable TArray<FInstructionRecord> WrittenRecords;
+	mutable uint32 ReadRecordsHash;
+	mutable uint32 WrittenRecordsHash;
+};
+
 #undef UE_API
diff --git a/Engine/Plugins/Animation/ControlRig/Source/ControlRig/Public/Units/RigUnitContext.h b/Engine/Plugins/Animation/ControlRig/Source/ControlRig/Public/Units/RigUnitContext.h
--- a/Engine/Plugins/Animation/ControlRig/Source/ControlRig/Public/Units/RigUnitContext.h
+++ b/Engine/Plugins/Animation/ControlRig/Source/ControlRig/Public/Units/RigUnitContext.h
@@ -133,6 +133,8 @@ public:
 		, ControlRig(nullptr)
 		, RigModulePrefix()
 		, RigModulePrefixHash(0)
+		, RigModuleInstance(nullptr)
+		, Records(&RecordsStorage)
 	{
 	}
 
@@ -242,8 +244,40 @@ public:
 	URigHierarchy* Hierarchy;
 	UControlRig* ControlRig;
 	
-	mutable TArray<FInstructionRecord> ReadRecords;
-	mutable TArray<FInstructionRecord> WrittenRecords;
+	void AddReadInstructionRecord(const FInstructionRecord& InRecord) const
+	{
+		check(Records);
+		Records->ReadRecords.Add(InRecord);
+		Records->ReadRecordsHash = HashCombine(Records->ReadRecordsHash, GetTypeHash(InRecord));
+	}
+
+	void AddWrittenInstructionRecord(const FInstructionRecord& InRecord) const
+	{
+		check(Records);
+		Records->WrittenRecords.Add(InRecord);
+		Records->WrittenRecordsHash = HashCombine(Records->WrittenRecordsHash, GetTypeHash(InRecord));
+	}
+
+	void ResetInstructionRecords() const
+	{
+		// only reset the memory we own.
+		// (ignore the Records pointer since it may point to unowned memory)
+		RecordsStorage.ReadRecords.Reset();
+		RecordsStorage.WrittenRecords.Reset();
+		RecordsStorage.ReadRecordsHash = RecordsStorage.WrittenRecordsHash = 0;
+	}
+
+	uint32 GetInstructionRecordsHash() const
+	{
+		check(Records);
+		return HashCombine(Records->ReadRecordsHash, Records->WrittenRecordsHash);
+	}
+
+	const FInstructionRecordContainer& GetInstructionRecords() const
+	{
+		check(Records);
+		return *Records;
+	}
 
 #if WITH_EDITOR
 	virtual void Report(const FRigVMLogSettings& InLogSettings, const FName& InFunctionName, int32 InInstructionIndex, const FString& InMessage) const override
@@ -268,6 +302,9 @@ private:
 	uint32 RigModulePrefixHash;
 	const FRigModuleInstance* RigModuleInstance;
 
+	mutable FInstructionRecordContainer RecordsStorage;
+	mutable FInstructionRecordContainer* Records;
+
 	friend class FControlRigExecuteContextRigModuleGuard;
 	friend class UModularRig;
 };
diff --git a/Engine/Plugins/Runtime/RigVM/Source/RigVM/Private/RigVMCore/RigVMByteCode.cpp b/Engine/Plugins/Runtime/RigVM/Source/RigVM/Private/RigVMCore/RigVMByteCode.cpp
--- a/Engine/Plugins/Runtime/RigVM/Source/RigVM/Private/RigVMCore/RigVMByteCode.cpp
+++ b/Engine/Plugins/Runtime/RigVM/Source/RigVM/Private/RigVMCore/RigVMByteCode.cpp
@@ -1677,6 +1677,42 @@ TArray<int32> FRigVMByteCode::GetInstructionsForOperand(const FRigVMOperand& InO
 	return InstructionIndices;
 }
 
+#if WITH_EDITOR
+
+TArray<int32> FRigVMByteCode::GetInstructionsForInputOperand(const FRigVMOperand& InOperand) const
+{
+	TArray<int32> InstructionIndices;
+	
+	const FRigVMInstructionArray Instructions = GetInstructions();
+	for(int32 InstructionIndex = 0; InstructionIndex < GetNumInstructions(); InstructionIndex++)
+	{
+		if(InputOperandsPerInstruction[InstructionIndex].Contains(InOperand))
+		{
+			InstructionIndices.Add(InstructionIndex);
+		}
+
+	}
+	return InstructionIndices;
+}
+
+TArray<int32> FRigVMByteCode::GetInstructionsForOutputOperand(const FRigVMOperand& InOperand) const
+{
+	TArray<int32> InstructionIndices;
+	
+	const FRigVMInstructionArray Instructions = GetInstructions();
+	for(int32 InstructionIndex = 0; InstructionIndex < GetNumInstructions(); InstructionIndex++)
+	{
+		if(OutputOperandsPerInstruction[InstructionIndex].Contains(InOperand))
+		{
+			InstructionIndices.Add(InstructionIndex);
+		}
+
+	}
+	return InstructionIndices;
+}
+
+#endif
+
 int32 FRigVMByteCode::GetOpAlignment(ERigVMOpCode InOpCode)
 {
 	switch (InOpCode)
diff --git a/Engine/Plugins/Runtime/RigVM/Source/RigVM/Public/RigVMCore/RigVMByteCode.h b/Engine/Plugins/Runtime/RigVM/Source/RigVM/Public/RigVMCore/RigVMByteCode.h
--- a/Engine/Plugins/Runtime/RigVM/Source/RigVM/Public/RigVMCore/RigVMByteCode.h
+++ b/Engine/Plugins/Runtime/RigVM/Source/RigVM/Public/RigVMCore/RigVMByteCode.h
@@ -1316,9 +1316,17 @@ public:
 	// returns the byte index of the first operand for this instructions
 	UE_API int32 GetFirstOperandByteIndex(const FRigVMInstruction& InInstruction) const;
 
-	// returns all of the operands for a given instruction
+	// returns all instructions for a given operand
 	UE_API TArray<int32> GetInstructionsForOperand(const FRigVMOperand& InOperand) const;
 
+#if WITH_EDITOR
+	// returns all instructions for a given input operand
+	UE_API TArray<int32> GetInstructionsForInputOperand(const FRigVMOperand& InOperand) const;
+
+	// returns all instructions for a given output operand 
+	UE_API TArray<int32> GetInstructionsForOutputOperand(const FRigVMOperand& InOperand) const;
+#endif
+
 	// returns true if the operator in question is used by multiple instructions
 	bool IsOperandShared(const FRigVMOperand& InOperand) const { return GetInstructionsForOperand(InOperand).Num() > 1; }
 
diff --git a/Engine/Plugins/Runtime/RigVM/Source/RigVMDeveloper/Private/RigVMModel/RigVMController.cpp b/Engine/Plugins/Runtime/RigVM/Source/RigVMDeveloper/Private/RigVMModel/RigVMController.cpp
--- a/Engine/Plugins/Runtime/RigVM/Source/RigVMDeveloper/Private/RigVMModel/RigVMController.cpp
+++ b/Engine/Plugins/Runtime/RigVM/Source/RigVMDeveloper/Private/RigVMModel/RigVMController.cpp
@@ -20133,6 +20133,14 @@ void URigVMController::SendUserFacingNotification(const FString& InMessage, floa
 #endif
 }
 
+void URigVMController::RequestJumpToHyperLink(const URigVMNode* InSubject)
+{
+	if(RequestJumpToHyperlinkDelegate.IsBound())
+	{
+		RequestJumpToHyperlinkDelegate.Execute(InSubject);
+	}
+}
+
 void URigVMController::CreateDefaultValueForStructIfRequired(UScriptStruct* InStruct, FString& InOutDefaultValue)
 {
 	if (InStruct != nullptr)
diff --git a/Engine/Plugins/Runtime/RigVM/Source/RigVMDeveloper/Public/RigVMModel/RigVMController.h b/Engine/Plugins/Runtime/RigVM/Source/RigVMDeveloper/Public/RigVMModel/RigVMController.h
--- a/Engine/Plugins/Runtime/RigVM/Source/RigVMDeveloper/Public/RigVMModel/RigVMController.h
+++ b/Engine/Plugins/Runtime/RigVM/Source/RigVMDeveloper/Public/RigVMModel/RigVMController.h
@@ -1262,6 +1262,7 @@ public:
 	UE_API void ReportAndNotifyError(const FString& InMessage) const;
 	UE_API void ReportPinTypeChange(URigVMPin* InPin, const FString& InNewCPPType);
 	UE_API void SendUserFacingNotification(const FString& InMessage, float InDuration = 0.f, const UObject* InSubject = nullptr, const FName& InBrushName = TEXT("MessageLog.Warning")) const;
+	UE_API void RequestJumpToHyperLink(const URigVMNode* InSubject);
 
 	template <typename... Types>
 	void ReportInfof(UE::Core::TCheckedFormatString<FString::FmtCharType, Types...> Fmt, Types... Args) const
