PR #13645: Make generic vector channel data display more robust for position/rotation/scale control types. This also fixes the euler curve filter not recognizing rotation keys coming from Rotator control types.

diff --git a/Engine/Plugins/Animation/ControlRig/Source/ControlRig/Private/Sequencer/MovieSceneControlRigParameterSection.cpp b/Engine/Plugins/Animation/ControlRig/Source/ControlRig/Private/Sequencer/MovieSceneControlRigParameterSection.cpp
--- a/Engine/Plugins/Animation/ControlRig/Source/ControlRig/Private/Sequencer/MovieSceneControlRigParameterSection.cpp
+++ b/Engine/Plugins/Animation/ControlRig/Source/ControlRig/Private/Sequencer/MovieSceneControlRigParameterSection.cpp
@@ -193,74 +193,113 @@ struct FParameterFloatChannelEditorData
 //Set up with all 4 Channels so it can be used by all vector types.
 struct FParameterVectorChannelEditorData
 {
-	FParameterVectorChannelEditorData(UControlRig *InControlRig, const FName& InName, bool bEnabledOverride, const FText& GroupName, int SortStartIndex, int32 NumChannels)
+	FParameterVectorChannelEditorData(UControlRig *InControlRig, const FName& InName, bool bEnabledOverride, const FText& GroupName, int SortStartIndex, int32 NumChannels, ERigControlType ControlType)
 	{
 		ControlRig = InControlRig;
 		ParameterName = InName;
-		FString NameAsString = InName.ToString();
-		FString TotalName = NameAsString;
-		const EAxisList::Type XAxis = EAxisList::X;
-		const EAxisList::Type YAxis = EAxisList::Y;
-		const EAxisList::Type ZAxis = EAxisList::Z;
 
+		static constexpr EAxisList::Type Axes[] = { EAxisList::Forward, EAxisList::Left, EAxisList::Up, EAxisList::None };
+		static_assert(UE_ARRAY_COUNT(Axes) == UE_ARRAY_COUNT(MetaData));
+		
+		static constexpr FStringView NameSuffixes[] = { TEXT(".X"), TEXT(".Y"), TEXT(".Z"), TEXT(".W") };
+		static_assert(UE_ARRAY_COUNT(NameSuffixes) == UE_ARRAY_COUNT(MetaData));
+		
+		static constexpr FStringView PositionNameSuffixes[] = { TEXT(".Positon.X"), TEXT(".Position.Y"), TEXT(".Position.Z"), nullptr };
+		static_assert(UE_ARRAY_COUNT(PositionNameSuffixes) == UE_ARRAY_COUNT(MetaData));
+		
+		static constexpr FStringView RotationNameSuffixes[] = { TEXT(".Rotation.Roll"), TEXT(".Rotation.Pitch"), TEXT(".Rotation.Yaw"), nullptr };
+		static_assert(UE_ARRAY_COUNT(RotationNameSuffixes) == UE_ARRAY_COUNT(MetaData));
+		
+		static constexpr FStringView ScaleNameSuffixes[] = { TEXT(".Scale.X"), TEXT(".Scale.Y"), TEXT(".Scale.Z"), nullptr };
+		static_assert(UE_ARRAY_COUNT(ScaleNameSuffixes) == UE_ARRAY_COUNT(MetaData));
+		
+		static const FText DisplayTexts[] = { FCommonChannelData::ChannelX, FCommonChannelData::ChannelY, FCommonChannelData::ChannelZ, FCommonChannelData::ChannelW };
+		static_assert(UE_ARRAY_COUNT(DisplayTexts) == UE_ARRAY_COUNT(MetaData));
+		
+		static const FText PositionDisplayTexts[] =
 		{
-			TotalName += ".X";
-			MetaData[0].SetIdentifiers(FName(*TotalName), FCommonChannelData::ChannelX);
-			TotalName = NameAsString;
-			MetaData[0].IntentName = FCommonChannelData::ChannelX;
-			MetaData[0].Color = AxisDisplayInfo::GetAxisColor(XAxis);
-			MetaData[0].Group = GroupName;
-			MetaData[0].bEnabled = bEnabledOverride;
-			MetaData[0].SortOrder = SortStartIndex++;
-			MetaData[0].bCanCollapseToTrack = true;
-		}
+			NSLOCTEXT("MovieSceneControlParameterRigSection", "Location.X", "Location.X"), 
+			NSLOCTEXT("MovieSceneControlParameterRigSection", "Location.Y", "Location.Y"), 
+			NSLOCTEXT("MovieSceneControlParameterRigSection", "Location.Z", "Location.Z"),
+			FText()
+		};
+		static_assert(UE_ARRAY_COUNT(PositionDisplayTexts) == UE_ARRAY_COUNT(MetaData));
+		
+		static const FText RotationDisplayTexts[] =
 		{
-			TotalName += ".Y";
-			MetaData[1].SetIdentifiers(FName(*TotalName), FCommonChannelData::ChannelY);
-			TotalName = NameAsString;
-			MetaData[1].IntentName = FCommonChannelData::ChannelY;
-			MetaData[1].Color = AxisDisplayInfo::GetAxisColor(YAxis);
-			MetaData[1].Group = GroupName;
-			MetaData[1].bEnabled = bEnabledOverride;
-			MetaData[1].SortOrder = SortStartIndex++;
-			MetaData[1].bCanCollapseToTrack = true;
-		}
+			NSLOCTEXT("MovieSceneControlParameterRigSection", "Rotation.X", "Rotation.Roll"),
+			NSLOCTEXT("MovieSceneControlParameterRigSection", "Rotation.Y", "Rotation.Pitch"),
+			NSLOCTEXT("MovieSceneControlParameterRigSection", "Rotation.Z", "Rotation.Yaw"),
+			FText()
+		};
+		static_assert(UE_ARRAY_COUNT(RotationDisplayTexts) == UE_ARRAY_COUNT(MetaData));
+		
+		static const FText ScaleDisplayTexts[] =
 		{
-			TotalName += ".Z";
-			MetaData[2].SetIdentifiers(FName(*TotalName), FCommonChannelData::ChannelZ);
-			TotalName = NameAsString;
-			MetaData[2].IntentName = FCommonChannelData::ChannelZ;
-			MetaData[2].Color = AxisDisplayInfo::GetAxisColor(ZAxis);
-			MetaData[2].Group = GroupName;
-			MetaData[2].bEnabled = bEnabledOverride;
-			MetaData[2].SortOrder = SortStartIndex++;
-			MetaData[2].bCanCollapseToTrack = true;
-		}
+			NSLOCTEXT("MovieSceneControlParameterRigSection", "Scale.X", "Scale.X"),
+			NSLOCTEXT("MovieSceneControlParameterRigSection", "Scale.Y", "Scale.Y"),
+			NSLOCTEXT("MovieSceneControlParameterRigSection", "Scale.Z", "Scale.Z"),
+			FText()
+		};
+		static_assert(UE_ARRAY_COUNT(ScaleDisplayTexts) == UE_ARRAY_COUNT(MetaData));
+		
+		static const TFunction<TOptional<float>(UObject&, FTrackInstancePropertyBindings*)> GetExternalValueFuncs[] =
 		{
-			TotalName += ".W";
-			MetaData[3].SetIdentifiers(FName(*TotalName), FCommonChannelData::ChannelW);
-			TotalName = NameAsString;
-			MetaData[3].IntentName = FCommonChannelData::ChannelW;
-			MetaData[3].Group = GroupName;
-			MetaData[3].bEnabled = bEnabledOverride;
-			MetaData[3].SortOrder = SortStartIndex++;
-			MetaData[3].bCanCollapseToTrack = true;
-		}
-
-		ExternalValues[0].OnGetExternalValue = [InControlRig, InName,NumChannels](UObject& InObject, FTrackInstancePropertyBindings* Bindings) { return ExtractChannelX(InObject, InControlRig, InName, NumChannels); };
-		ExternalValues[1].OnGetExternalValue = [InControlRig, InName,NumChannels](UObject& InObject, FTrackInstancePropertyBindings* Bindings) { return ExtractChannelY(InObject, InControlRig, InName, NumChannels); };
-		ExternalValues[2].OnGetExternalValue = [InControlRig, InName,NumChannels](UObject& InObject, FTrackInstancePropertyBindings* Bindings) { return ExtractChannelZ(InObject, InControlRig, InName, NumChannels); };
-		ExternalValues[3].OnGetExternalValue = [InControlRig, InName,NumChannels](UObject& InObject, FTrackInstancePropertyBindings* Bindings) { return ExtractChannelW(InObject, InControlRig, InName, NumChannels); };
+			[InControlRig, InName, NumChannels](UObject& InObject, FTrackInstancePropertyBindings*) { return ExtractChannelX(InObject, InControlRig, InName, NumChannels); },
+			[InControlRig, InName, NumChannels](UObject& InObject, FTrackInstancePropertyBindings*) { return ExtractChannelY(InObject, InControlRig, InName, NumChannels); },
+			[InControlRig, InName, NumChannels](UObject& InObject, FTrackInstancePropertyBindings*) { return ExtractChannelZ(InObject, InControlRig, InName, NumChannels); },
+			[InControlRig, InName, NumChannels](UObject& InObject, FTrackInstancePropertyBindings*) { return ExtractChannelW(InObject, InControlRig, InName, NumChannels); },
+		};
+		static_assert(UE_ARRAY_COUNT(GetExternalValueFuncs) == UE_ARRAY_COUNT(ExternalValues));
 
-		ExternalValues[0].OnGetCurrentValueAndWeight = [InName, NumChannels](UObject* Object, UMovieSceneSection*  SectionToKey, FFrameNumber KeyTime, FFrameRate TickResolution, FMovieSceneRootEvaluationTemplateInstance& RootTemplate,
-			float& OutValue, float& OutWeight) { GetChannelValueAndWeight(InName, NumChannels, 0, Object, SectionToKey, KeyTime, TickResolution, RootTemplate, OutValue, OutWeight); };
-		ExternalValues[1].OnGetCurrentValueAndWeight = [InName, NumChannels](UObject* Object, UMovieSceneSection*  SectionToKey, FFrameNumber KeyTime, FFrameRate TickResolution, FMovieSceneRootEvaluationTemplateInstance& RootTemplate,
-			float& OutValue, float& OutWeight) { GetChannelValueAndWeight(InName, NumChannels, 1, Object, SectionToKey, KeyTime, TickResolution, RootTemplate, OutValue, OutWeight); };
-		ExternalValues[2].OnGetCurrentValueAndWeight = [InName, NumChannels](UObject* Object, UMovieSceneSection*  SectionToKey, FFrameNumber KeyTime, FFrameRate TickResolution, FMovieSceneRootEvaluationTemplateInstance& RootTemplate,
-			float& OutValue, float& OutWeight) { GetChannelValueAndWeight(InName, NumChannels, 2, Object, SectionToKey, KeyTime, TickResolution, RootTemplate, OutValue, OutWeight); };
-		ExternalValues[3].OnGetCurrentValueAndWeight = [InName, NumChannels](UObject* Object, UMovieSceneSection*  SectionToKey, FFrameNumber KeyTime, FFrameRate TickResolution, FMovieSceneRootEvaluationTemplateInstance& RootTemplate,
-			float& OutValue, float& OutWeight) { GetChannelValueAndWeight(InName, NumChannels, 3, Object, SectionToKey, KeyTime, TickResolution, RootTemplate, OutValue, OutWeight); };
+		const FString NameAsString = InName.ToString();
+		
+		check(NumChannels <= UE_ARRAY_COUNT(MetaData));
+		check(NumChannels <= UE_ARRAY_COUNT(ExternalValues));
+		for (int ChannelIdx = 0; ChannelIdx < NumChannels; ChannelIdx++)
+		{
+			const EAxisList::Type* Axis = &Axes[ChannelIdx];
+			const FStringView* NameSuffix;
+			const FText* DisplayText;
+			switch (ControlType)
+			{
+			case ERigControlType::Position:
+				NameSuffix = &PositionNameSuffixes[ChannelIdx];
+				DisplayText = &PositionDisplayTexts[ChannelIdx];
+				break;
+			case ERigControlType::Rotator:
+				NameSuffix = &RotationNameSuffixes[ChannelIdx];
+				DisplayText = &RotationDisplayTexts[ChannelIdx];
+				break;
+			case ERigControlType::Scale:
+				NameSuffix = &ScaleNameSuffixes[ChannelIdx];
+				DisplayText = &ScaleDisplayTexts[ChannelIdx];
+				break;
+			default:
+				NameSuffix = &NameSuffixes[ChannelIdx];
+				DisplayText = &DisplayTexts[ChannelIdx];
+				Axis = nullptr;
+				break;
+			}
 
+			MetaData[ChannelIdx].SetIdentifiers(FName(NameAsString + *NameSuffix), *DisplayText);
+			MetaData[ChannelIdx].IntentName = *DisplayText;
+			if (Axis != nullptr && *Axis != EAxisList::None)
+			{
+				MetaData[ChannelIdx].Color = AxisDisplayInfo::GetAxisColor(*Axis);
+			}
+			MetaData[ChannelIdx].Group = GroupName;
+			MetaData[ChannelIdx].bEnabled = bEnabledOverride;
+			MetaData[ChannelIdx].SortOrder = SortStartIndex++;
+			MetaData[ChannelIdx].bCanCollapseToTrack = true;
+			
+			ExternalValues[ChannelIdx].OnGetExternalValue = GetExternalValueFuncs[ChannelIdx];
+			ExternalValues[ChannelIdx].OnGetCurrentValueAndWeight = [InName, NumChannels, ChannelIdx](UObject* Object, UMovieSceneSection* SectionToKey, FFrameNumber KeyTime,
+				FFrameRate TickResolution, FMovieSceneRootEvaluationTemplateInstance& RootTemplate, float& OutValue, float& OutWeight)
+			{
+				GetChannelValueAndWeight(InName, NumChannels, ChannelIdx, Object, SectionToKey, KeyTime, TickResolution, RootTemplate, OutValue, OutWeight);
+			};
+		}
 	}
 
 	static FVector4 GetPropertyValue(UControlRig* ControlRig, FName ParameterName, UObject& InObject,int32 NumChannels)
@@ -2308,7 +2347,7 @@ EMovieSceneChannelProxyType UMovieSceneControlRigParameterSection::CacheChannelP
 								const int32 ParentControlIndex = pChannelIndex ? pChannelIndex->ControlIndex : INDEX_NONE;
 								ControlChannelMap.Add(Vector2D.ParameterName, FChannelMapInfo(ControlIndex, SortOrder, FloatChannelIndex, ParentControlIndex, NAME_None, MaskIndex, CategoryIndex));
 							}
-							FParameterVectorChannelEditorData EditorData(ControlRig, Vector2D.ParameterName, bEnabled, Group, SortOrder, 2);
+							FParameterVectorChannelEditorData EditorData(ControlRig, Vector2D.ParameterName, bEnabled, Group, SortOrder, 2, ControlElement->Settings.ControlType);
 							MaybeApplyChannelMask(EditorData.MetaData[0], ControlElement, ERigControlTransformChannel::TranslationX);
 							MaybeApplyChannelMask(EditorData.MetaData[1], ControlElement, ERigControlTransformChannel::TranslationY);
 							Channels.Add(Vector2D.XCurve, EditorData.MetaData[0], EditorData.ExternalValues[0]);
@@ -2367,7 +2406,7 @@ EMovieSceneChannelProxyType UMovieSceneControlRigParameterSection::CacheChannelP
 							}
 
 
-							FParameterVectorChannelEditorData EditorData(ControlRig, Vector.ParameterName, bEnabled, Group, SortOrder, 3);
+							FParameterVectorChannelEditorData EditorData(ControlRig, Vector.ParameterName, bEnabled, Group, SortOrder, 3, ControlElement->Settings.ControlType);
 
 							if (ControlElement->Settings.ControlType == ERigControlType::Position)
 							{
