ControlRig: Fix undo/redo crash when hierarchy does not exist anymore
#rnx

diff --git a/Engine/Plugins/Animation/ControlRig/Source/ControlRigDeveloper/Private/ControlRigBlueprint.cpp b/Engine/Plugins/Animation/ControlRig/Source/ControlRigDeveloper/Private/ControlRigBlueprint.cpp
--- a/Engine/Plugins/Animation/ControlRig/Source/ControlRigDeveloper/Private/ControlRigBlueprint.cpp
+++ b/Engine/Plugins/Animation/ControlRig/Source/ControlRigDeveloper/Private/ControlRigBlueprint.cpp
@@ -1991,10 +1991,13 @@ void UControlRigBlueprint::ClearTransientControls()
 			UControlRig* InstancedControlRig = Cast<UControlRig>(ArchetypeInstance);
 			if (InstancedControlRig)
 			{
-				if(!InstancedControlRig->GetHierarchy()->GetTransientControls().IsEmpty())
+				if (URigHierarchy* InstanceHierarchy = InstancedControlRig->GetHierarchy())
 				{
-					bHasAnyTransientControls = true;
-					break;
+					if(!InstanceHierarchy->GetTransientControls().IsEmpty())
+					{
+						bHasAnyTransientControls = true;
+						break;
+					}
 				}
 			}
 		}
