ControlRig - RigHierarchy record metadata changes

#rnx

diff --git a/Engine/Plugins/Animation/ControlRig/Source/ControlRig/Private/Rigs/RigDependency.cpp b/Engine/Plugins/Animation/ControlRig/Source/ControlRig/Private/Rigs/RigDependency.cpp
--- a/Engine/Plugins/Animation/ControlRig/Source/ControlRig/Private/Rigs/RigDependency.cpp
+++ b/Engine/Plugins/Animation/ControlRig/Source/ControlRig/Private/Rigs/RigDependency.cpp
@@ -162,15 +162,22 @@ void FRigDependenciesProviderForVM::ComputeDependencies(const FControlRigExecute
 	}
 
 	// for each written transform / metadata on an instruction
-	TArray<int32> InstructionsToVisit;
-	TArray<FInstructionRecord> FilteredReadRecords;
-	TArray<TTuple<FRigHierarchyRecord,FRigHierarchyRecord>> ReadRecordsToAddAtEnd;
+	TArray<TTuple<int32,FRigHierarchyRecord>> InstructionsToVisit;
+	TArray<FRigHierarchyRecord> ReadVariables, WrittenVariables;
+	TMap<FRigVMOperand, TArray<int32>> InputOperandToInstructions;
 	TMap<FRigVMOperand, TArray<int32>> OutputOperandToInstructions;
 
 	for(int32 InstructionIndex = 0; InstructionIndex < Instructions.Num(); InstructionIndex++)
 	{
-		FRigVMOperandArray Operands = InVM->GetByteCode().GetOutputOperands(InstructionIndex);
-		for (const FRigVMOperand& Operand : Operands)
+		FRigVMOperandArray InputOperands = InVM->GetByteCode().GetInputOperands(InstructionIndex);
+		for (const FRigVMOperand& Operand : InputOperands)
+		{
+			const FRigVMOperand OperandNoRegisterOffset(Operand.GetMemoryType(), Operand.GetRegisterIndex());
+			InputOperandToInstructions.FindOrAdd(OperandNoRegisterOffset).Add(InstructionIndex);
+		}
+
+		FRigVMOperandArray OutputOperands = InVM->GetByteCode().GetOutputOperands(InstructionIndex);
+		for (const FRigVMOperand& Operand : OutputOperands)
 		{
 			const FRigVMOperand OperandNoRegisterOffset(Operand.GetMemoryType(), Operand.GetRegisterIndex());
 			OutputOperandToInstructions.FindOrAdd(OperandNoRegisterOffset).Add(InstructionIndex);
@@ -179,99 +186,134 @@ void FRigDependenciesProviderForVM::ComputeDependencies(const FControlRigExecute
 
 	for(int32 InstructionIndex = 0; InstructionIndex < Instructions.Num(); InstructionIndex++)
 	{
-		if (!WrittenRecordsPerInstruction.IsValidIndex(InstructionIndex))
-		{
-			continue;
-		}
+		ReadVariables.Reset();
+		WrittenVariables.Reset();
 
-		const TArray<FInstructionRecord>& WrittenRecords = WrittenRecordsPerInstruction[InstructionIndex];
-		if(WrittenRecords.IsEmpty())
+		if (ReadRecordsPerInstruction.IsValidIndex(InstructionIndex))
 		{
-			continue;
-		}
-
-		FilteredReadRecords.Reset();
+			const TArray<FInstructionRecord>& ReadRecords = ReadRecordsPerInstruction[InstructionIndex];
+			for(int32 ReadRecordIndex = 0; ReadRecordIndex < ReadRecords.Num(); ReadRecordIndex++)
+			{
+				const FInstructionRecord& ReadRecord = ReadRecords[ReadRecordIndex];
 
-		for(int32 WrittenRecordIndex = 0; WrittenRecordIndex < WrittenRecords.Num(); WrittenRecordIndex++)
-		{
-			InstructionsToVisit.Reset();
-			InstructionsToVisit.Add(InstructionIndex);
+				const FRigHierarchyRecord ReadElementRecord(
+					ReadRecord.MetadataName.IsNone() ? FRigHierarchyRecord::EType_RigElement : FRigHierarchyRecord::EType_Metadata,
+					ReadRecord.ElementIndex,
+					ReadRecord.MetadataName);
 
-			for(int32 InstructionToVisitIndex = 0; InstructionToVisitIndex < InstructionsToVisit.Num(); InstructionToVisitIndex++)
-			{
-				const int32 InstructionToVisit = InstructionsToVisit[InstructionToVisitIndex];
+				InstructionsToVisit.Reset();
+				InstructionsToVisit.Emplace(InstructionIndex, ReadElementRecord);
 
-				if (ReadRecordsPerInstruction.IsValidIndex(InstructionToVisit))
+				// follow the output operands through the bytecode to determine variable writes of this instruction
+				for(int32 InstructionToVisitIndex = 0; InstructionToVisitIndex < InstructionsToVisit.Num(); InstructionToVisitIndex++)
 				{
-					const TArray<FInstructionRecord>& ReadRecords = ReadRecordsPerInstruction[InstructionToVisit];
-					for(int32 ReadRecordIndex = 0; ReadRecordIndex < ReadRecords.Num(); ReadRecordIndex++)
+					const int32 InstructionToVisit = InstructionsToVisit[InstructionToVisitIndex].Get<0>();
+
+					FRigVMOperandArray OutputOperands = InVM->GetByteCode().GetOutputOperands(InstructionToVisit);
+					for(const FRigVMOperand OutputOperand : OutputOperands)
 					{
-						const FInstructionRecord& ReadRecord = ReadRecords[ReadRecordIndex];
-						FilteredReadRecords.AddUnique(ReadRecord);
-					}
-				}
+						FRigHierarchyRecord LastRecord = InstructionsToVisit[InstructionToVisitIndex].Get<1>();
 
-				FRigVMOperandArray InputOperands = InVM->GetByteCode().GetInputOperands(InstructionToVisit);
-				for(const FRigVMOperand InputOperand : InputOperands)
-				{
-					const FRigVMOperand InputOperandNoRegisterOffset(InputOperand.GetMemoryType(), InputOperand.GetRegisterIndex());
+						if (OutputOperand.GetMemoryType() == ERigVMMemoryType::External)
+						{
+							const FRigHierarchyRecord VariableRecord(FRigHierarchyRecord::EType_Variable, OutputOperand.GetRegisterIndex(), ReadRecord.VM);
+							WrittenVariables.AddUnique(VariableRecord);
+							LastRecord = VariableRecord;
+						}
+						
+						const FRigVMOperand OutputOperandNoRegisterOffset(OutputOperand.GetMemoryType(), OutputOperand.GetRegisterIndex());
 
-					if(const TArray<int32>* InstructionsWithOutputOperand = OutputOperandToInstructions.Find(InputOperandNoRegisterOffset))
-					{
-						for(int32 InstructionWithOutputOperand : *InstructionsWithOutputOperand)
+						if(const TArray<int32>* InstructionsWithOutputOperand = InputOperandToInstructions.Find(OutputOperandNoRegisterOffset))
 						{
-							InstructionsToVisit.AddUnique(InstructionWithOutputOperand);
+							for(int32 InstructionWithOutputOperand : *InstructionsWithOutputOperand)
+							{
+								const TTuple<int32,FRigHierarchyRecord> Tuple(InstructionWithOutputOperand, LastRecord);
+								InstructionsToVisit.AddUnique(Tuple);
+							}
 						}
 					}
 				}
+
+				const FRigHierarchyRecord ElementRecord(
+					ReadRecord.MetadataName.IsNone() ? FRigHierarchyRecord::EType_RigElement : FRigHierarchyRecord::EType_Metadata,
+					ReadRecord.ElementIndex,
+					ReadRecord.MetadataName);
+
+				for (const FRigHierarchyRecord& WrittenVariableRecord : WrittenVariables)
+				{
+					OutDependencies.FindOrAdd(WrittenVariableRecord).Add(ElementRecord);
+				}
 			}
 		}
 
-		for(const FInstructionRecord& WrittenRecord : WrittenRecords)
+		if (WrittenRecordsPerInstruction.IsValidIndex(InstructionIndex))
 		{
-			const FRigHierarchyRecord TargetRecord(
-				WrittenRecord.MetadataName.IsNone() ? FRigHierarchyRecord::EType_RigElement : FRigHierarchyRecord::EType_Metadata,
-				WrittenRecord.ElementIndex,
-				WrittenRecord.MetadataName);
-
-			const FRigHierarchyRecord InstructionRecord(
-				FRigHierarchyRecord::EType_Instruction,
-				WrittenRecord.InstructionIndex,
-				WrittenRecord.VM);
+			const TArray<FInstructionRecord>& WrittenRecords = WrittenRecordsPerInstruction[InstructionIndex];
+			for(int32 WrittenRecordIndex = 0; WrittenRecordIndex < WrittenRecords.Num(); WrittenRecordIndex++)
+			{
+				const FInstructionRecord& WrittenRecord = WrittenRecords[WrittenRecordIndex];
 
-			OutDependencies.FindOrAdd(TargetRecord).Add(InstructionRecord);
+				const FRigHierarchyRecord InstructionRecord(
+					FRigHierarchyRecord::EType_Instruction,
+					WrittenRecord.InstructionIndex,
+					WrittenRecord.VM);
 
-			for(const FInstructionRecord& FilteredReadRecord : FilteredReadRecords)
-			{
-				const FRigHierarchyRecord SourceRecord(
-					FilteredReadRecord.MetadataName.IsNone() ? FRigHierarchyRecord::EType_RigElement : FRigHierarchyRecord::EType_Metadata,
-					FilteredReadRecord.ElementIndex,
-					FilteredReadRecord.MetadataName);
+				const FRigHierarchyRecord TargetRecord(
+					WrittenRecord.MetadataName.IsNone() ? FRigHierarchyRecord::EType_RigElement : FRigHierarchyRecord::EType_Metadata,
+					WrittenRecord.ElementIndex,
+					WrittenRecord.MetadataName);
 
-				ReadRecordsToAddAtEnd.Emplace(InstructionRecord, SourceRecord);
-			}
-		}
-	}
+				OutDependencies.FindOrAdd(TargetRecord).Add(InstructionRecord);
 
-	for (const TTuple<FRigHierarchyRecord,FRigHierarchyRecord>& Pair : ReadRecordsToAddAtEnd)
-	{
-		const FRigHierarchyRecord& Target = Pair.Get<0>();
-		const FRigHierarchyRecord& Source = Pair.Get<1>();
+				InstructionsToVisit.Reset();
+				InstructionsToVisit.Emplace(InstructionIndex, InstructionRecord);
 
-		if (Target.IsInstruction())
-		{
-			if (const TSet<FRigHierarchyRecord>* SourceToTarget = OutDependencies.Find(Source))
-			{
-				if (SourceToTarget->Contains(Target))
+				// follow the input operands through the bytecode to determine inputs of this instruction
+				for(int32 InstructionToVisitIndex = 0; InstructionToVisitIndex < InstructionsToVisit.Num(); InstructionToVisitIndex++)
 				{
-					continue;
+					const int32 InstructionToVisit = InstructionsToVisit[InstructionToVisitIndex].Get<0>();
+
+					if (ReadRecordsPerInstruction.IsValidIndex(InstructionToVisit))
+					{
+						const TArray<FInstructionRecord>& ReadRecords = ReadRecordsPerInstruction[InstructionToVisit];
+						for(int32 ReadRecordIndex = 0; ReadRecordIndex < ReadRecords.Num(); ReadRecordIndex++)
+						{
+							const FInstructionRecord& ReadRecord = ReadRecords[ReadRecordIndex];
+							const FRigHierarchyRecord SourceRecord(
+								ReadRecord.MetadataName.IsNone() ? FRigHierarchyRecord::EType_RigElement : FRigHierarchyRecord::EType_Metadata,
+								ReadRecord.ElementIndex,
+								ReadRecord.MetadataName);
+							OutDependencies.FindOrAdd(InstructionsToVisit[InstructionToVisitIndex].Get<1>()).Add(SourceRecord);
+						}
+					}
+
+					FRigVMOperandArray InputOperands = InVM->GetByteCode().GetInputOperands(InstructionToVisit);
+					for(const FRigVMOperand InputOperand : InputOperands)
+					{
+						FRigHierarchyRecord LastRecord = InstructionsToVisit[InstructionToVisitIndex].Get<1>();
+
+						if (InputOperand.GetMemoryType() == ERigVMMemoryType::External)
+						{
+							const FRigHierarchyRecord VariableRecord(FRigHierarchyRecord::EType_Variable, InputOperand.GetRegisterIndex(), WrittenRecord.VM);
+							OutDependencies.FindOrAdd(InstructionsToVisit[InstructionToVisitIndex].Get<1>()).Add(VariableRecord);
+							LastRecord = VariableRecord;
+						}
+						
+						const FRigVMOperand InputOperandNoRegisterOffset(InputOperand.GetMemoryType(), InputOperand.GetRegisterIndex());
+
+						if(const TArray<int32>* InstructionsWithOutputOperand = OutputOperandToInstructions.Find(InputOperandNoRegisterOffset))
+						{
+							for(int32 InstructionWithOutputOperand : *InstructionsWithOutputOperand)
+							{
+								const TTuple<int32,FRigHierarchyRecord> Tuple(InstructionWithOutputOperand, LastRecord);
+								InstructionsToVisit.AddUnique(Tuple);
+							}
+						}
+					}
 				}
 			}
 		}
-		
-		OutDependencies.FindOrAdd(Pair.Get<0>()).Add(Pair.Get<1>());
 	}
-
 #endif
 }
 
diff --git a/Engine/Plugins/Animation/ControlRig/Source/ControlRig/Private/Rigs/RigHierarchy.cpp b/Engine/Plugins/Animation/ControlRig/Source/ControlRig/Private/Rigs/RigHierarchy.cpp
--- a/Engine/Plugins/Animation/ControlRig/Source/ControlRig/Private/Rigs/RigHierarchy.cpp
+++ b/Engine/Plugins/Animation/ControlRig/Source/ControlRig/Private/Rigs/RigHierarchy.cpp
@@ -6409,7 +6409,7 @@ bool URigHierarchy::IsDependentOn(FRigBaseElement* InDependent, FRigBaseElement*
 							return true;
 						}
 					}
-					else if (DependentRecord.IsInstruction())
+					else if (DependentRecord.IsInstruction() || DependentRecord.IsVariable())
 					{
 						if (FoundDependencyInProvider(InHierarchy, InDependency, InDependencyProvider, bIsOnActualTopology, CacheKey, DependentRecord))
 						{
diff --git a/Engine/Plugins/Animation/ControlRig/Source/ControlRig/Public/Rigs/RigDependencyRecords.h b/Engine/Plugins/Animation/ControlRig/Source/ControlRig/Public/Rigs/RigDependencyRecords.h
--- a/Engine/Plugins/Animation/ControlRig/Source/ControlRig/Public/Rigs/RigDependencyRecords.h
+++ b/Engine/Plugins/Animation/ControlRig/Source/ControlRig/Public/Rigs/RigDependencyRecords.h
@@ -15,6 +15,7 @@ struct FRigHierarchyRecord
 		EType_RigElement,
 		EType_Metadata,
 		EType_Instruction,
+		EType_Variable,
 		EType_Invalid
 	};
 
@@ -52,6 +53,11 @@ struct FRigHierarchyRecord
 		return Type == EType_Instruction;
 	}
 
+	bool IsVariable() const
+	{
+		return Type == EType_Variable;
+	}
+
 	bool operator==(const FRigHierarchyRecord& InRecord) const
 	{
 		return Type == InRecord.Type && Index == InRecord.Index && Name == InRecord.Name;
diff --git a/Engine/Plugins/Animation/ControlRig/Source/ControlRig/Public/Units/Hierarchy/RigUnit_Metadata.h b/Engine/Plugins/Animation/ControlRig/Source/ControlRig/Public/Units/Hierarchy/RigUnit_Metadata.h
--- a/Engine/Plugins/Animation/ControlRig/Source/ControlRig/Public/Units/Hierarchy/RigUnit_Metadata.h
+++ b/Engine/Plugins/Animation/ControlRig/Source/ControlRig/Public/Units/Hierarchy/RigUnit_Metadata.h
@@ -4,6 +4,8 @@
 
 #include "Units/RigDispatchFactory.h"
 #include "Units/RigUnit.h"
+#include "ModularRig.h"
+#include "Rigs/RigDependencyRecords.h"
 #include "RigUnit_Metadata.generated.h"
 
 #define UE_API CONTROLRIG_API
@@ -235,6 +237,16 @@ protected:
 		{
 			Value = Md->GetValue();
 			Found = true;
+
+#if WITH_EDITOR
+			const FControlRigExecuteContext& ExecuteContext = InContext.GetPublicData<FControlRigExecuteContext>();
+			if (ExecuteContext.Hierarchy->bRecordInstructionsAtRuntime)
+			{
+				const FName VMName = ExecuteContext.GetRigModuleInstance() ? ExecuteContext.GetRigModuleInstance()->Name : NAME_None;
+				const FInstructionRecord Record(ExecuteContext.GetInstructionIndex(), InContext.GetSlice().GetIndex(), Cache.GetIndex(), VMName, Md->GetName());
+				ExecuteContext.AddReadInstructionRecord(Record);
+			}
+#endif			
 		}
 		else
 		{
@@ -298,6 +310,15 @@ protected:
 		// extract the metadata
 		if (MetadataType* Md = Cast<MetadataType>(FindOrAddMetadata(ExecuteContext, Item, Name, EnumValue, NameSpace, Cache)))
 		{
+#if WITH_EDITOR
+			if (ExecuteContext.Hierarchy->bRecordInstructionsAtRuntime)
+			{
+				const FName VMName = ExecuteContext.GetRigModuleInstance() ? ExecuteContext.GetRigModuleInstance()->Name : NAME_None;
+				const FInstructionRecord Record(ExecuteContext.GetInstructionIndex(), InContext.GetSlice().GetIndex(), Cache.GetIndex(), VMName, Md->GetName());
+				ExecuteContext.AddWrittenInstructionRecord(Record);
+			}
+#endif
+
 			if(!Equals<ValueType>(Md->GetValue(), Value))
 			{
 				Md->GetValue() = Value;
