Control Rig: Dependency Viewer

#rnx
#nojira

diff --git a/Engine/Plugins/Animation/ControlRig/Source/ControlRigEditor/Private/ControlRigEditorModule.cpp b/Engine/Plugins/Animation/ControlRig/Source/ControlRigEditor/Private/ControlRigEditorModule.cpp
--- a/Engine/Plugins/Animation/ControlRig/Source/ControlRigEditor/Private/ControlRigEditorModule.cpp
+++ b/Engine/Plugins/Animation/ControlRig/Source/ControlRigEditor/Private/ControlRigEditorModule.cpp
@@ -42,6 +42,7 @@
 #include "ControlRigHierarchyCommands.h"
 #include "ControlRigModularRigCommands.h"
 #include "Editor/ModularRigEventQueueCommands.h"
+#include "RigDependencyGraph/RigDependencyGraphCommands.h"
 #include "Animation/AnimSequence.h"
 #include "Editor/ControlRigEditorEditMode.h"
 #include "ControlRigElementDetails.h"
@@ -125,6 +126,7 @@ void FControlRigEditorModule::StartupModule()
 	FControlRigHierarchyCommands::Register();
 	FControlRigModularRigCommands::Register();
 	FModularRigEventQueueCommands::Register();
+	FRigDependencyGraphCommands::Register();
 	FControlRigEditorStyle::Get();
 
 	EdGraphPanelPinFactory = MakeShared<FControlRigGraphPanelPinFactory>();
@@ -255,6 +257,9 @@ void FControlRigEditorModule::StartupModule()
 	FRigVMUserWorkflowProvider Provider;
 	Provider.BindUFunction(UControlRigTransformWorkflowOptions::StaticClass()->GetDefaultObject(), TEXT("ProvideWorkflows"));
 	WorkflowHandles.Add(WorkflowRegistry->RegisterProvider(nullptr, Provider));
+
+	RigDependencyGraphPanelNodeFactory = MakeShareable(new FRigDependencyGraphPanelNodeFactory());
+	FEdGraphUtilities::RegisterVisualNodeFactory(RigDependencyGraphPanelNodeFactory);
 }
 
 void FControlRigEditorModule::ShutdownModule()
@@ -317,6 +322,12 @@ void FControlRigEditorModule::ShutdownModule()
 	WorkflowHandles.Reset();
 
 	UnregisterPropertySectionMappings();
+
+	if (RigDependencyGraphPanelNodeFactory.IsValid())
+	{
+		FEdGraphUtilities::UnregisterVisualNodeFactory(RigDependencyGraphPanelNodeFactory);
+		RigDependencyGraphPanelNodeFactory.Reset();
+	}
 }
 
 UClass* FControlRigEditorModule::GetRigVMAssetClass() const
diff --git a/Engine/Plugins/Animation/ControlRig/Source/ControlRigEditor/Private/ControlRigEditorModule.h b/Engine/Plugins/Animation/ControlRig/Source/ControlRigEditor/Private/ControlRigEditorModule.h
--- a/Engine/Plugins/Animation/ControlRig/Source/ControlRigEditor/Private/ControlRigEditorModule.h
+++ b/Engine/Plugins/Animation/ControlRig/Source/ControlRigEditor/Private/ControlRigEditorModule.h
@@ -8,6 +8,7 @@
 #include "IControlRigModule.h"
 #include "IAnimationEditor.h"
 #include "IAnimationEditorModule.h"
+#include "RigDependencyGraph/RigDependencyGraphPanelNodeFactory.h"
 
 class IControlRigBaseEditor;
 class UBlueprint;
@@ -77,4 +78,6 @@ private:
 
 	/** Handles for all registered workflows */
 	TArray<int32> WorkflowHandles;
+	
+	TSharedPtr<FRigDependencyGraphPanelNodeFactory> RigDependencyGraphPanelNodeFactory;
 };
\ No newline at end of file
diff --git a/Engine/Plugins/Animation/ControlRig/Source/ControlRigEditor/Private/ControlRigEditorStyle.h b/Engine/Plugins/Animation/ControlRig/Source/ControlRigEditor/Private/ControlRigEditorStyle.h
--- a/Engine/Plugins/Animation/ControlRig/Source/ControlRigEditor/Private/ControlRigEditorStyle.h
+++ b/Engine/Plugins/Animation/ControlRig/Source/ControlRigEditor/Private/ControlRigEditorStyle.h
@@ -197,6 +197,34 @@ public:
 			
 			Set("ConstraintManager.ComboButton", ConstraintComboButton);
 		}
+
+		// Dependency Graph
+		{
+ 			const FTextBlockStyle NormalText = FAppStyle::Get().GetWidgetStyle<FTextBlockStyle>("NormalText");
+			Set("ControlRig.DependencyGraph.TextStyle", FTextBlockStyle(NormalText)
+				.SetColorAndOpacity(FLinearColor(0.0f, 0.0f, 0.0f, 1.0f))
+				.SetFont(DEFAULT_FONT("Regular", 8)));
+
+			Set( "ControlRig.DependencyGraph.NodeBody", new BOX_BRUSH( "Slate/DependencyGraphNode_Body", FMargin(16.f/64.f, 25.f/64.f, 16.f/64.f, 16.f/64.f) ) );
+
+			Set("ControlRig.DependencyGraph.Colors.Instruction", FLinearColor(FColor::FromHex(TEXT("#00AAAA"))));
+			Set("ControlRig.DependencyGraph.Colors.Bone", FLinearColor(FColor::FromHex(TEXT("#E166B6"))));
+			Set("ControlRig.DependencyGraph.Colors.Control", FLinearColor(FColor::FromHex(TEXT("#45C761"))));
+			Set("ControlRig.DependencyGraph.Colors.AnimationChannel", FLinearColor(FColor::FromHex(TEXT("#FE8539"))));
+			Set("ControlRig.DependencyGraph.Colors.Null", FLinearColor(FColor::FromHex(TEXT("#C7C8FE"))));
+			Set("ControlRig.DependencyGraph.Colors.Socket", FLinearColor(FColor::FromHex(TEXT("#F9E264"))));
+			Set("ControlRig.DependencyGraph.Colors.Connector", FLinearColor(FColor::FromHex(TEXT("#FE8989"))));
+			Set("ControlRig.DependencyGraph.Colors.Metadata", FLinearColor(FColor::FromHex(TEXT("#BB6BF0"))));
+			Set("ControlRig.DependencyGraph.Colors.Variable", FLinearColor(FColor::FromHex(TEXT("#CCCC00"))));
+
+			Set("ControlRig.DependencyGraph.SyncContent", new CORE_IMAGE_BRUSH_SVG("Starship/Common/Undo", Icon16x16));
+			{
+				FContentRootBracket Bracket(this, EngineEditorSlateDir);
+ 				Set( "ControlRig.DependencyGraph.HideUnrelated", new IMAGE_BRUSH_SVG( "Starship/GraphEditors/HideUnrelated", Icon16x16 ) );
+ 				Set( "ControlRig.DependencyGraph.IsolateSelection", new IMAGE_BRUSH_SVG( "Sequencer/TrackIsolate", Icon16x16 ) );
+			}
+
+		}
 		
 		FSlateStyleRegistry::RegisterSlateStyle(*this);
 	}
diff --git a/Engine/Plugins/Animation/ControlRig/Source/ControlRigEditor/Private/Editor/ControlRigEditorMode.cpp b/Engine/Plugins/Animation/ControlRig/Source/ControlRigEditor/Private/Editor/ControlRigEditorMode.cpp
--- a/Engine/Plugins/Animation/ControlRig/Source/ControlRigEditor/Private/Editor/ControlRigEditorMode.cpp
+++ b/Engine/Plugins/Animation/ControlRig/Source/ControlRigEditor/Private/Editor/ControlRigEditorMode.cpp
@@ -22,6 +22,7 @@
 #include "Editor/RigVMDetailsInspectorTabSummoner.h"
 #include "Editor/RigVMFindReferencesTabSummoner.h"
 #include "Editor/RigOverrideListTabSummoner.h"
+#include "RigDependencyGraph/RigDependencyGraphSummoner.h"
 
 #if WITH_RIGVMLEGACYEDITOR
 FControlRigLegacyEditorMode::FControlRigLegacyEditorMode(const TSharedRef<FControlRigLegacyEditor>& InControlRigEditor, bool bCreateDefaultLayout)
@@ -35,6 +36,7 @@ FControlRigLegacyEditorMode::FControlRigLegacyEditorMode(const TSharedRef<FContr
 	TabFactories.RegisterFactory(MakeShared<FRigCurveContainerTabSummoner>(InControlRigEditor));
 	TabFactories.RegisterFactory(MakeShared<FRigValidationTabSummoner>(InControlRigEditor));
 	TabFactories.RegisterFactory(MakeShared<FRigAnimAttributeTabSummoner>(InControlRigEditor));
+	TabFactories.RegisterFactory(MakeShared<FRigDependencyGraphSummoner>(InControlRigEditor));
 
 	FPersonaModule& PersonaModule = FModuleManager::LoadModuleChecked<FPersonaModule>("Persona");
 
@@ -266,6 +268,7 @@ FControlRigEditorMode::FControlRigEditorMode(const TSharedRef<FControlRigEditor>
 	TabFactories.RegisterFactory(MakeShared<FRigCurveContainerTabSummoner>(InControlRigEditor));
 	TabFactories.RegisterFactory(MakeShared<FRigValidationTabSummoner>(InControlRigEditor));
 	TabFactories.RegisterFactory(MakeShared<FRigAnimAttributeTabSummoner>(InControlRigEditor));
+	TabFactories.RegisterFactory(MakeShared<FRigDependencyGraphSummoner>(InControlRigEditor));
 
 	FPersonaModule& PersonaModule = FModuleManager::LoadModuleChecked<FPersonaModule>("Persona");
 
@@ -378,6 +381,7 @@ FModularRigEditorMode::FModularRigEditorMode(const TSharedRef<FControlRigEditor>
 
 	TabFactories.RegisterFactory(MakeShared<FModularRigModelTabSummoner>(InControlRigEditor));
 	TabFactories.RegisterFactory(MakeShared<FRigModuleAssetBrowserTabSummoner>(InControlRigEditor));
+	TabFactories.RegisterFactory(MakeShared<FRigDependencyGraphSummoner>(InControlRigEditor));
 
 	TabLayout = FTabManager::NewLayout("Standalone_ModularRigEditMode_Layout_v1.3")
 		->AddArea
diff --git a/Engine/Plugins/Animation/ControlRig/Source/ControlRigEditor/Private/RigDependencyGraph/RigDependencyConnectionDrawingPolicy.cpp b/Engine/Plugins/Animation/ControlRig/Source/ControlRigEditor/Private/RigDependencyGraph/RigDependencyConnectionDrawingPolicy.cpp
new file mode 100644
--- /dev/null
+++ b/Engine/Plugins/Animation/ControlRig/Source/ControlRigEditor/Private/RigDependencyGraph/RigDependencyConnectionDrawingPolicy.cpp
@@ -0,0 +1,25 @@
+// Copyright Epic Games, Inc. All Rights Reserved.
+
+#include "RigDependencyGraph/RigDependencyConnectionDrawingPolicy.h"
+#include "RigDependencyGraph/RigDependencyGraphNode.h"
+
+void FRigDependencyConnectionDrawingPolicy::DetermineWiringStyle(UEdGraphPin* OutputPin, UEdGraphPin* InputPin, /*inout*/ FConnectionParams& Params)
+{
+	const URigDependencyGraphNode* OutputNode = Cast<URigDependencyGraphNode>(OutputPin->GetOwningNode());
+	const URigDependencyGraphNode* InputNode = Cast<URigDependencyGraphNode>(InputPin->GetOwningNode());
+
+	if (OutputNode && InputNode)
+	{
+		if (OutputNode->GetNodeId().IsElement() && InputNode->GetNodeId().IsElement())
+		{
+			static const FLinearColor ParentChildWireColor = FLinearColor(FColor::FromHex(TEXT("#AAAAAA")));
+			Params.WireColor = ParentChildWireColor;
+			Params.WireColor.A = 0.75f;
+		}
+
+		if (OutputNode->IsFadedOut() || InputNode->IsFadedOut())
+		{
+			Params.WireColor.A *= 0.2f;
+		}
+	}
+}
\ No newline at end of file
diff --git a/Engine/Plugins/Animation/ControlRig/Source/ControlRigEditor/Private/RigDependencyGraph/RigDependencyConnectionDrawingPolicy.h b/Engine/Plugins/Animation/ControlRig/Source/ControlRigEditor/Private/RigDependencyGraph/RigDependencyConnectionDrawingPolicy.h
new file mode 100644
--- /dev/null
+++ b/Engine/Plugins/Animation/ControlRig/Source/ControlRigEditor/Private/RigDependencyGraph/RigDependencyConnectionDrawingPolicy.h
@@ -0,0 +1,25 @@
+// Copyright Epic Games, Inc. All Rights Reserved.
+
+#pragma once
+
+#include "ConnectionDrawingPolicy.h"
+
+class FRigDependencyConnectionDrawingPolicy : public FConnectionDrawingPolicy
+{
+public:
+	FRigDependencyConnectionDrawingPolicy(int32 InBackLayerID, int32 InFrontLayerID, float InZoomFactor, const FSlateRect& InClippingRect, FSlateWindowElementList& InDrawElements)
+		: FConnectionDrawingPolicy(InBackLayerID, InFrontLayerID, InZoomFactor, InClippingRect, InDrawElements)
+	{
+		// We don't draw arrows here
+		ArrowImage = nullptr;
+		ArrowRadius = FVector2f::ZeroVector;
+	}
+
+	virtual FVector2f ComputeSplineTangent(const FVector2f& Start, const FVector2f& End) const override
+	{
+		const int32 Tension = FMath::Abs<int32>(Start.X - End.X);
+		return Tension * FVector2f(1.0f, 0.0f);
+	}
+
+	virtual void DetermineWiringStyle(UEdGraphPin* OutputPin, UEdGraphPin* InputPin, /*inout*/ FConnectionParams& Params) override;
+};
diff --git a/Engine/Plugins/Animation/ControlRig/Source/ControlRigEditor/Private/RigDependencyGraph/RigDependencyGraph.cpp b/Engine/Plugins/Animation/ControlRig/Source/ControlRigEditor/Private/RigDependencyGraph/RigDependencyGraph.cpp
new file mode 100644
--- /dev/null
+++ b/Engine/Plugins/Animation/ControlRig/Source/ControlRigEditor/Private/RigDependencyGraph/RigDependencyGraph.cpp
@@ -0,0 +1,1267 @@
+// Copyright Epic Games, Inc. All Rights Reserved.
+
+#include "RigDependencyGraph/RigDependencyGraph.h"
+
+#include "ITransportControl.h"
+#include "ModularRig.h"
+#include "Editor/ControlRigEditor.h"
+#include "RigVMModel/RigVMClient.h"
+#include "RigDependencyGraph/RigDependencyGraphNode.h"
+#include "RigDependencyGraph/RigDependencyGraphSchema.h"
+
+#include UE_INLINE_GENERATED_CPP_BY_NAME(RigDependencyGraph)
+
+URigDependencyGraph::URigDependencyGraph(const FObjectInitializer& ObjectInitializer)
+	: Super(ObjectInitializer)
+	, DefaultLayoutIterationsPerNode(50)
+	, LayoutIterationsLeft(0)
+	, LayoutIterationsTotal(0)
+	, LayoutIterationsMax(DefaultLayoutIterationsPerNode * 1000)
+	, bFollowParentRelationShips(true)
+	, bFollowVMRelationShips(true)
+	, bLockContent(false)
+	, BlockSelectionCounter(2)
+{
+}
+
+void URigDependencyGraph::Initialize(const TSharedRef<IControlRigBaseEditor>& InControlRigEditor, const TWeakObjectPtr<UControlRig>& InControlRig)
+{
+	const bool bIsFirstTime = !WeakControlRigEditor.IsValid();
+	
+	if (WeakControlRigEditor.IsValid())
+	{
+		if (TSharedPtr<IControlRigBaseEditor> PreviousControlRigEditor = WeakControlRigEditor.Pin())
+		{
+			Cast<IRigVMAssetInterface>(PreviousControlRigEditor->GetControlRigBlueprint())->OnSetObjectBeingDebugged().Remove(OnSetObjectBeingDebuggedHandle);
+			PreviousControlRigEditor->GetControlRigBlueprint()->OnVMCompiled().Remove(OnVMCompiledHandle);
+			PreviousControlRigEditor->GetControlRigBlueprint()->OnModified().Remove(OnRigVMGraphModifiedHandle);
+		}
+	}
+	
+	if (UControlRig* PreviousControlRig = GetControlRig())
+	{
+		if (IsValid(PreviousControlRig))
+		{
+			if (URigHierarchy* Hierarchy = PreviousControlRig->GetHierarchy())
+			{
+				Hierarchy->OnModified().Remove(OnHierarchyModifiedHandle);
+				Hierarchy->OnMetadataChanged().Remove(OnMetadataChangedHandle);
+			}
+		}
+	}
+	
+	WeakControlRigEditor = InControlRigEditor;
+	WeakControlRig = InControlRig;
+
+	if(WeakControlRig.IsValid() && WeakControlRig.Get())
+	{
+		RigDependencyProvider = MakeShareable(new FRigDependenciesProviderForControlRig(WeakControlRig.Get()));
+	}
+	else
+	{
+		RigDependencyProvider = MakeShareable(new FEmptyRigDependenciesProvider());
+	}
+
+	OnSetObjectBeingDebuggedHandle = Cast<IRigVMAssetInterface>(InControlRigEditor->GetControlRigBlueprint())->OnSetObjectBeingDebugged().AddUObject(this, &URigDependencyGraph::OnSetObjectBeingDebugged);
+	OnVMCompiledHandle = InControlRigEditor->GetControlRigBlueprint()->OnVMCompiled().AddUObject(this, &URigDependencyGraph::OnVMCompiled);
+	OnRigVMGraphModifiedHandle = InControlRigEditor->GetControlRigBlueprint()->OnModified().AddUObject(this, &URigDependencyGraph::OnRigVMGraphModified);
+
+	if (UControlRig* ControlRig = GetControlRig())
+	{
+		if (URigHierarchy* Hierarchy = ControlRig->GetHierarchy())
+		{
+			OnHierarchyModifiedHandle = Hierarchy->OnModified().AddUObject(this, &URigDependencyGraph::OnHierarchyModified);
+			OnMetadataChangedHandle = Hierarchy->OnMetadataChanged().AddUObject(this, &URigDependencyGraph::OnMetadataChanged);
+		}
+	}
+
+	if (bIsFirstTime && WeakControlRigEditor.IsValid())
+	{
+		RebuildGraph(true);
+	}
+	else if (RigDependencyProvider.IsValid())
+	{
+		RigDependencyProvider->InvalidateCache();
+	}
+}
+
+UControlRig* URigDependencyGraph::GetControlRig() const
+{
+	if (WeakControlRig.IsValid())
+	{
+		return WeakControlRig.Get();
+	}
+	return nullptr;
+}
+
+URigHierarchy* URigDependencyGraph::GetRigHierarchy() const
+{
+	if (UControlRig* ControlRig = GetControlRig())
+	{
+		return ControlRig->GetHierarchy();
+	}
+	return nullptr;
+}
+
+const FRigVMClient* URigDependencyGraph::GetRigVMClient() const
+{
+	TSharedRef<IControlRigBaseEditor> ControlRigEditor = GetControlRigEditor();
+	return ControlRigEditor->GetControlRigBlueprint()->GetRigVMClient();
+}
+
+void URigDependencyGraph::RequestRefreshLayout()
+{
+	UpdateNodeIndices();
+
+	int32 NumSimulatedNodes = 0;
+	for (const TObjectPtr<URigDependencyGraphNode>& Node : DependencyGraphNodes)
+	{
+		if (Node->FollowLayout.Get(true))
+		{
+			NumSimulatedNodes++;
+		}
+	}
+	
+	LayoutIterationsLeft = LayoutIterationsTotal =
+		FMath::Min(LayoutIterationsMax, DefaultLayoutIterationsPerNode * NumSimulatedNodes);
+
+	LayoutEdges.Reset();
+
+	UpdateFadeOutStates();
+}
+
+bool URigDependencyGraph::NeedsRefreshLayout() const
+{
+	return LayoutIterationsLeft > 0;
+}
+
+TSharedRef<IRigDependenciesProvider> URigDependencyGraph::GetRigDependencyProvider() const
+{
+	return RigDependencyProvider.ToSharedRef();
+}
+
+const URigDependencyGraphSchema* URigDependencyGraph::GetRigDependencyGraphSchema() const
+{
+	return CastChecked<const URigDependencyGraphSchema>(GetSchema());
+}
+
+const URigDependencyGraphNode* URigDependencyGraph::FindNode(const FNodeId& InNodeId) const
+{
+	if (URigDependencyGraphNode* const* NodePtr = NodeIdLookup.Find(InNodeId))
+	{
+		return *NodePtr;
+	}
+	return nullptr;
+}
+
+URigDependencyGraphNode* URigDependencyGraph::FindNode(const FNodeId& InNodeId)
+{
+	const URigDependencyGraph* ConstThis = this;
+	return const_cast<URigDependencyGraphNode*>(ConstThis->FindNode(InNodeId));
+}
+
+bool URigDependencyGraph::Contains(const FNodeId& InNodeId) const
+{
+	return FindNode(InNodeId) != nullptr;
+}
+
+void URigDependencyGraph::RebuildGraph(bool bExpand)
+{
+	const URigHierarchy* Hierarchy = GetRigHierarchy();
+	if (Hierarchy == nullptr)
+	{
+		return;
+	}
+
+	bool bNeedToRebuildGraph = false;
+	TArray<FNodeId> PreviousNodes;
+	PreviousNodes.Reserve(DependencyGraphNodes.Num());
+	TMap<FNodeId, FVector2f> PreviousPositions;
+	PreviousPositions.Reserve(DependencyGraphNodes.Num());
+	
+	for (const TObjectPtr<URigDependencyGraphNode>& PreviousNode : DependencyGraphNodes)
+	{
+		const FNodeId& PreviousNodeId = PreviousNode->NodeId;
+
+		switch (PreviousNodeId.Type)
+		{
+			case FNodeId::EType_RigElement:
+			case FNodeId::EType_Metadata:
+			{
+				const FRigBaseElement* Element = PreviousNode->GetRigElement();
+				if (!Element)
+				{
+					bNeedToRebuildGraph = true;
+					continue;
+				}
+				if (!PreviousNodeId.IsMetadata())
+				{
+					break;
+				}
+				if (!Hierarchy->FindMetadataForElement(Element, PreviousNodeId.Name, ERigMetadataType::Invalid))
+				{
+					bNeedToRebuildGraph = true;
+					continue;
+				}
+				break;
+			}
+			case FNodeId::EType_Instruction:
+			{
+				if (!PreviousNode->GetRigVMNodeForInstruction())
+				{
+					bNeedToRebuildGraph = true;
+					continue;
+				}
+				break;
+			}
+			case FNodeId::EType_Variable:
+			{
+				if (!PreviousNode->GetExternalVariable()->IsValid())
+				{
+					bNeedToRebuildGraph = true;
+					continue;
+				}
+				break;
+			}
+			case FNodeId::EType_Invalid:
+			default:
+			{
+				bNeedToRebuildGraph = true;
+				continue;
+			}
+		}
+		
+		PreviousNodes.Add(PreviousNodeId);
+		PreviousPositions.Add(PreviousNodeId, FVector2f(PreviousNode->NodePosX, PreviousNode->NodePosY));
+	}
+
+	if (!bNeedToRebuildGraph)
+	{
+		RebuildLinks();
+		return;
+	}
+	
+	RemoveAllNodes();
+	ConstructNodes(PreviousNodes, bExpand);
+
+	for (const TObjectPtr<URigDependencyGraphNode>& CurrentNode : DependencyGraphNodes)
+	{
+		if (const FVector2f* PreviousPosition = PreviousPositions.Find(CurrentNode->NodeId))
+		{
+			CurrentNode->NodePosX = PreviousPosition->X;
+			CurrentNode->NodePosY = PreviousPosition->Y;
+			CurrentNode->FollowLayout = false;
+		}
+	}
+}
+
+void URigDependencyGraph::RebuildLinks()
+{
+	for (const TObjectPtr<URigDependencyGraphNode>& Node : DependencyGraphNodes)
+	{
+		Node->InputPin->BreakAllPinLinks();
+		Node->OutputPin->BreakAllPinLinks();
+	}
+
+	auto MakeLink = [](URigDependencyGraphNode* InSourceNode, URigDependencyGraphNode* InTargetNode)
+	{
+		if (InTargetNode->NodeId == InSourceNode->NodeId)
+		{
+			return;
+		}
+		if (InTargetNode->GetNodeId().IsInstruction() &&
+			InTargetNode->OutputPin->LinkedTo.Contains(InSourceNode->InputPin))
+		{
+			return;
+		}
+		if (!InSourceNode->OutputPin->LinkedTo.Contains(InTargetNode->InputPin))
+		{
+			InSourceNode->OutputPin->MakeLinkTo(InTargetNode->InputPin);
+		}
+			if (InSourceNode->GetNodeId().IsInstruction() &&
+				InTargetNode->OutputPin->LinkedTo.Contains(InSourceNode->InputPin))
+			{
+				InTargetNode->OutputPin->BreakLinkTo(InSourceNode->InputPin);
+			}
+	};
+
+	if (bFollowParentRelationShips)
+	{
+		if (const URigHierarchy* Hierarchy = GetRigHierarchy())
+		{
+			for (TObjectPtr<URigDependencyGraphNode>& ParentElementNode : DependencyGraphNodes)
+			{
+				if (!ParentElementNode->NodeId.IsElement())
+				{
+					continue;
+				}
+				if (Hierarchy->IsValidIndex(ParentElementNode->NodeId.Index))
+				{
+					const TArray<int32> Children = Hierarchy->GetChildren(ParentElementNode->NodeId.Index);
+					for (const int32& ChildIndex : Children)
+					{
+						const FNodeId ChildId(FNodeId::EType_RigElement, ChildIndex, NAME_None);
+						if (URigDependencyGraphNode* ChildElementNode = FindNode(ChildId))
+						{
+							MakeLink(ParentElementNode, ChildElementNode);
+						}
+					}
+				}
+			}
+
+			for (TObjectPtr<URigDependencyGraphNode>& ChildElementNode : DependencyGraphNodes)
+			{
+				if (!ChildElementNode->NodeId.IsElement())
+				{
+					continue;
+				}
+				if (Hierarchy->IsValidIndex(ChildElementNode->NodeId.Index))
+				{
+					const TArray<int32> Parents = Hierarchy->GetParents(ChildElementNode->NodeId.Index);
+					for (const int32& ParentIndex : Parents)
+					{
+						const FNodeId ParentId(FNodeId::EType_RigElement, ParentIndex, NAME_None);
+						if (URigDependencyGraphNode* ParentElementNode = FindNode(ParentId))
+						{
+							MakeLink(ParentElementNode, ChildElementNode);
+						}
+					}
+				}
+			}
+		}
+	}
+
+	if (bFollowVMRelationShips)
+	{
+		const TRigHierarchyDependencyMap& Dependencies = RigDependencyProvider->GetRigHierarchyDependencies();
+		for (TObjectPtr<URigDependencyGraphNode>& TargetNode : DependencyGraphNodes)
+		{
+			if (const TSet<FRigHierarchyRecord>* Records = Dependencies.Find(TargetNode->NodeId))
+			{
+				for (const FRigHierarchyRecord& DependencyRecord : *Records)
+				{
+					if (URigDependencyGraphNode* SourceNode = FindNode(DependencyRecord))
+					{
+						MakeLink(SourceNode, TargetNode);
+					}
+				}
+			}
+		}
+	}
+
+	UpdatePinVisibility();
+	UpdateFadeOutStates();
+	RebuildIslands();
+}
+
+void URigDependencyGraph::RebuildIslands()
+{
+	NodeIslands.Reset();
+	TMap<FGuid, int32> GuidToIsland;
+
+	for (const TObjectPtr<URigDependencyGraphNode>& Node : DependencyGraphNodes)
+	{
+		Node->IslandGuid.Reset();
+	}
+
+	for (const TObjectPtr<URigDependencyGraphNode>& Node : DependencyGraphNodes)
+	{
+		const FGuid& Guid = Node->GetIslandGuid();
+		check(Guid.IsValid());
+
+		if (const int32* ExistingIslandIndex = GuidToIsland.Find(Guid))
+		{
+			NodeIslands[*ExistingIslandIndex].NodeIds.Add(Node->GetNodeId());
+		}
+		else
+		{
+			GuidToIsland.Add(Guid, NodeIslands.AddDefaulted(1));
+			NodeIslands.Last().Guid = Guid;
+			NodeIslands.Last().NodeIds.Add(Node->GetNodeId());
+		}
+	}
+}
+
+void URigDependencyGraph::UpdatePinVisibility()
+{
+	const URigHierarchy* Hierarchy = GetRigHierarchy();
+	
+	for (const TObjectPtr<URigDependencyGraphNode>& Node : DependencyGraphNodes)
+	{
+		const FNodeId& NodeId = Node->GetNodeId();
+		const uint32 OldHash = HashCombine(GetTypeHash(Node->InputPin->bHidden), GetTypeHash(Node->OutputPin->bHidden));
+		Node->InputPin->bHidden = false;
+		Node->OutputPin->bHidden = false;
+
+		if (Node->InputPin->LinkedTo.IsEmpty())
+		{
+			Node->InputPin->bHidden = true;
+			
+			if (Node->InputPin->bHidden && NodeId.IsElement() && Hierarchy && bFollowParentRelationShips)
+			{
+				if (Hierarchy->GetNumberOfParents(NodeId.Index) > 0)
+				{
+					Node->InputPin->bHidden = false;
+				}
+			}
+
+			if (Node->InputPin->bHidden && RigDependencyProvider.IsValid() && bFollowVMRelationShips)
+			{
+				if (const TSet<FNodeId>* Dependencies = RigDependencyProvider->GetRigHierarchyDependencies().Find(NodeId))
+				{
+					Node->InputPin->bHidden = false;
+				}
+			}
+		}
+
+		if (Node->OutputPin->LinkedTo.IsEmpty())
+		{
+			Node->OutputPin->bHidden = true;
+			
+			if (Node->OutputPin->bHidden && NodeId.IsElement() && Hierarchy && bFollowParentRelationShips)
+			{
+				if (Hierarchy->IsValidIndex(NodeId.Index))
+				{
+					if (!Hierarchy->GetChildren(NodeId.Index).IsEmpty())
+					{
+						Node->OutputPin->bHidden = false;
+					}
+				}
+			}
+
+			if (Node->OutputPin->bHidden && RigDependencyProvider.IsValid() && bFollowVMRelationShips)
+			{
+				if (RigDependencyProvider->GetReverseRigHierarchyDependencies().Contains(NodeId))
+				{
+					Node->OutputPin->bHidden = false;
+				}
+			}
+		}
+
+		const uint32 NewHash = HashCombine(GetTypeHash(Node->InputPin->bHidden), GetTypeHash(Node->OutputPin->bHidden));
+		if (NewHash != OldHash)
+		{
+			NotifyNodeChanged(Node);
+		}
+	}
+}
+
+void URigDependencyGraph::UpdateFadeOutStates()
+{
+	if (SelectedNodes.IsEmpty())
+	{
+		for (const TObjectPtr<URigDependencyGraphNode>& Node : DependencyGraphNodes)
+		{
+			Node->bIsFadedOut.Reset();
+		}
+		return;
+	}
+
+	for (const TObjectPtr<URigDependencyGraphNode>& Node : DependencyGraphNodes)
+	{
+		Node->bIsFadedOut = true;
+	}
+
+	TArray<TPair<const URigDependencyGraphNode*, int32>> NodesOfInterest;
+	NodesOfInterest.Reserve(DependencyGraphNodes.Num());
+	for (const FNodeId& NodeId : SelectedNodes)
+	{
+		if (const URigDependencyGraphNode* Node = FindNode(NodeId))
+		{
+			NodesOfInterest.Emplace(Node, false);
+		}
+	}
+
+	for (int32 Index = 0; Index < NodesOfInterest.Num(); Index++)
+	{
+		const URigDependencyGraphNode* Node = NodesOfInterest[Index].Key;
+		if (Node->bIsFadedOut.IsSet() && Node->bIsFadedOut.GetValue() == false)
+		{
+			continue;
+		}
+		
+		const int32 Direction = NodesOfInterest[Index].Value;
+		Node->bIsFadedOut = false;
+		
+		if (Direction <= 0)
+		{
+			for (const UEdGraphPin* LinkedPin : Node->InputPin->LinkedTo)
+			{
+				if (const URigDependencyGraphNode* LinkedNode = Cast<URigDependencyGraphNode>(LinkedPin->GetOwningNode()))
+				{
+					NodesOfInterest.Emplace(LinkedNode, -1);
+				}
+			}
+		}
+		if (Direction >= 0)
+		{
+			for (const UEdGraphPin* LinkedPin : Node->OutputPin->LinkedTo)
+			{
+				if (const URigDependencyGraphNode* LinkedNode = Cast<URigDependencyGraphNode>(LinkedPin->GetOwningNode()))
+				{
+					NodesOfInterest.Emplace(LinkedNode, 1);
+				}
+			}
+		}
+	}
+}
+
+void URigDependencyGraph::SetFollowParentRelationShips(bool InFollowParentRelationShips)
+{
+	bFollowParentRelationShips = InFollowParentRelationShips;
+	RebuildGraph(false);
+}
+
+void URigDependencyGraph::SetFollowVMRelationShips(bool InFollowVMRelationShips)
+{
+	bFollowVMRelationShips = InFollowVMRelationShips;
+	RebuildGraph(false);
+}
+
+void URigDependencyGraph::SetContentLocked(bool InLockContent)
+{
+	bLockContent = InLockContent
+;}
+
+TArray<URigDependencyGraph::FNodeId> URigDependencyGraph::GetNodesToExpand(const TArray<FNodeId>& InNodeIds) const
+{
+	TArray<FNodeId> NodesToFollow = InNodeIds;
+	TArray<FNodeId> NodesToExpand;
+	
+	// expand the nodes to construct one level using the dependency provider
+	if (bFollowVMRelationShips && RigDependencyProvider.IsValid())
+	{
+		const TRigHierarchyDependencyMap& Dependencies = RigDependencyProvider->GetRigHierarchyDependencies();
+		const TRigHierarchyDependencyMap& ReverseDependencies = RigDependencyProvider->GetReverseRigHierarchyDependencies();
+
+		for (int32 NodeIndex = 0; NodeIndex < NodesToFollow.Num(); ++NodeIndex)
+		{
+			const FRigHierarchyRecord Record = NodesToFollow[NodeIndex];
+			if (const TSet<FRigHierarchyRecord>* Records = Dependencies.Find(Record))
+			{
+				for (const FRigHierarchyRecord& DependencyRecord : *Records)
+				{
+					NodesToExpand.AddUnique(DependencyRecord);
+					
+					if(DependencyRecord.IsInstruction() || DependencyRecord.IsVariable() || DependencyRecord.IsMetadata())
+					{
+						NodesToFollow.Add(DependencyRecord);
+					}
+				}
+			}
+			if (const TSet<FRigHierarchyRecord>* ReverseRecords = ReverseDependencies.Find(Record))
+			{
+				for (const FRigHierarchyRecord& DependentRecord : *ReverseRecords)
+				{
+					NodesToExpand.AddUnique(DependentRecord);
+				}
+			}
+		}
+	}
+
+	if (bFollowParentRelationShips)
+	{
+		if (const URigHierarchy* Hierarchy = GetRigHierarchy())
+		{
+			// also expand the nodes to construct parent and child nodes
+			for (int32 NodeIndex = 0; NodeIndex < InNodeIds.Num(); ++NodeIndex)
+			{
+				if (InNodeIds[NodeIndex].Type != FNodeId::EType_RigElement)
+				{
+					continue;
+				}
+
+				if (Hierarchy->IsValidIndex(InNodeIds[NodeIndex].Index))
+				{
+					TArray<int32> ParentsAndChildren = Hierarchy->GetParents(InNodeIds[NodeIndex].Index);
+					ParentsAndChildren.Append(Hierarchy->GetChildren(InNodeIds[NodeIndex].Index));
+
+					for (const int32& ParentOrChildIndex : ParentsAndChildren)
+					{
+						const FNodeId ParentOrChildId(FNodeId::EType_RigElement, ParentOrChildIndex, NAME_None);
+						NodesToExpand.AddUnique(ParentOrChildId);
+					}
+				}
+			}
+		}
+	}
+
+	// only consider nodes to expand that we don't have yet
+	NodesToExpand.RemoveAll([this](const FNodeId& InNodeId) -> bool
+	{
+		return FindNode(InNodeId) != nullptr;
+	});
+
+	return NodesToExpand;
+}
+
+const URigVMNode* URigDependencyGraph::GetRigVMNodeForInstruction(const FNodeId& InNodeId) const
+{
+	if (!InNodeId.IsInstruction())
+	{
+		return nullptr;
+	}
+
+	UControlRig* ControlRig = GetControlRig();
+	if (!ControlRig)
+	{
+		return nullptr;
+	}
+
+	UModularRig* ModularRig = Cast<UModularRig>(ControlRig);
+
+	// find the node path in question based on the VM
+	const URigVM* VM = nullptr;
+								
+	if (InNodeId.Name.IsNone() || !ModularRig)
+	{
+		VM = ControlRig->GetVM();
+	}
+	else
+	{
+		const FRigModuleInstance* RigModule = ModularRig->FindModule(InNodeId.Name);
+		if (!RigModule)
+		{
+			return nullptr;
+		}
+		UControlRig* ModuleRig = RigModule->GetRig();
+		if (!ModuleRig)
+		{
+			return nullptr;
+		}
+		VM = ModuleRig->GetVM();
+	}
+
+	check(VM);
+
+	if (!VM->GetByteCode().GetInstructions().IsValidIndex(InNodeId.Index))
+	{
+		return nullptr;
+	}
+	
+	const FRigVMInstruction Instruction = VM->GetByteCode().GetInstructions()[InNodeId.Index];
+	if (Instruction.OpCode != ERigVMOpCode::Execute)
+	{
+		return nullptr;
+	}
+
+	return Cast<URigVMNode>(VM->GetByteCode().GetSubjectForInstruction(InNodeId.Index));
+}
+
+FRigHierarchyRecord URigDependencyGraph::GetRecordForRigVMNode(const URigVMNode* InNode) const
+{
+	if (!InNode)
+	{
+		return FRigHierarchyRecord();
+	}
+
+	UControlRig* ControlRig = GetControlRig();
+	if (!ControlRig)
+	{
+		return FRigHierarchyRecord();
+	}
+
+	UModularRig* ModularRig = Cast<UModularRig>(ControlRig);
+	
+	TArray<TTuple<const URigVM*, FName>> VMs;;
+								
+	if (!ModularRig)
+	{
+		VMs.Emplace(ControlRig->GetVM(), NAME_None);
+	}
+	else
+	{
+		ModularRig->ForEachModule([&VMs](FRigModuleInstance* RigModuleInstance)
+		{
+			if (UControlRig* ModuleRig = RigModuleInstance->GetRig())
+			{
+				VMs.Emplace(ModuleRig->GetVM(), RigModuleInstance->Name);
+			}
+			return true;
+		});
+	}
+
+	TSharedRef<IRigDependenciesProvider> Provider = GetRigDependencyProvider();
+	for (const TTuple<const URigVM*, FName>& Pair : VMs)
+	{
+		const URigVM* VM = Pair.Get<0>();
+		const TArray<int32>& InstructionIndices = VM->GetByteCode().GetAllInstructionIndicesForSubject(const_cast<URigVMNode*>(InNode));
+		for (const int32& InstructionIndex : InstructionIndices)
+		{
+			const FRigVMInstruction Instruction = VM->GetByteCode().GetInstructions()[InstructionIndex];
+			if (Instruction.OpCode != ERigVMOpCode::Execute)
+			{
+				continue;
+			}
+			
+			const FRigHierarchyRecord Record(FRigHierarchyRecord::EType_Instruction, InstructionIndex, Pair.Get<1>());
+			if (Provider->GetRigHierarchyDependencies().Contains(Record) ||
+				Provider->GetReverseRigHierarchyDependencies().Contains(Record))
+			{
+				return Record;
+			}
+		}
+	}
+
+	return FRigHierarchyRecord();
+}
+
+const FRigVMExternalVariableDef* URigDependencyGraph::GetExternalVariable(const FNodeId& InNodeId) const
+{
+	if (!InNodeId.IsVariable())
+	{
+		return nullptr;
+	}
+
+	UControlRig* ControlRig = GetControlRig();
+	if (!ControlRig)
+	{
+		return nullptr;
+	}
+
+	UModularRig* ModularRig = Cast<UModularRig>(ControlRig);
+
+	// find the node path in question based on the VM
+	const URigVM* VM = nullptr;
+								
+	if (InNodeId.Name.IsNone() || !ModularRig)
+	{
+		VM = ControlRig->GetVM();
+	}
+	else
+	{
+		const FRigModuleInstance* RigModule = ModularRig->FindModule(InNodeId.Name);
+		if (!RigModule)
+		{
+			return nullptr;
+		}
+		UControlRig* ModuleRig = RigModule->GetRig();
+		if (!ModuleRig)
+		{
+			return nullptr;
+		}
+		VM = ModuleRig->GetVM();
+	}
+
+	check(VM);
+
+	const TArray<FRigVMExternalVariableDef>& ExternalVariableDefs = VM->GetExternalVariableDefs();
+
+	if (ExternalVariableDefs.IsValidIndex(InNodeId.Index))
+	{
+		return &ExternalVariableDefs[InNodeId.Index];
+	}
+
+	return nullptr;
+}
+
+void URigDependencyGraph::Tick(float DeltaTime)
+{
+	if (--BlockSelectionCounter == 0)
+	{
+		UE::TScopeLock Lock(NodesToSelectCriticalSection);
+		if (NodesToSelectDuringTick.IsSet())
+		{
+			const TArray<FNodeId> NodeIds = NodesToSelectDuringTick.GetValue();
+			bool bRemoveNodes = RemoveAllNodesDuringTick.Get(true);
+
+			SelectedNodes.Reset();
+			NodesToSelectDuringTick.Reset();
+			RemoveAllNodesDuringTick.Reset();
+			ZoomAndFitDuringLayout = bRemoveNodes;
+
+			if (bRemoveNodes)
+			{
+				RemoveAllNodes();
+			}
+			ConstructNodes(NodeIds, bRemoveNodes);
+			SelectNodes(NodeIds);
+		}
+	}
+
+	if (LastRecordsHash.IsSet() && !bLockContent)
+	{
+		if (const UControlRig* ControlRig = GetControlRig())
+		{
+			const FRigVMExtendedExecuteContext& ExtendedExecuteContext = ControlRig->GetRigVMExtendedExecuteContext();
+			const FControlRigExecuteContext& Context = ExtendedExecuteContext.GetPublicData<FControlRigExecuteContext>();
+			if (LastRecordsHash.GetValue() != Context.GetInstructionRecordsHash())
+			{
+				RebuildGraph(false);
+				LastRecordsHash = Context.GetInstructionRecordsHash();
+			}
+		}
+	}
+}
+
+void URigDependencyGraph::SelectAllNodes()
+{
+	TArray<FNodeId> NodeIds;
+	NodeIds.Reserve(DependencyGraphNodes.Num());
+	for (const TObjectPtr<URigDependencyGraphNode>& Node : DependencyGraphNodes)
+	{
+		NodeIds.Add(Node->GetNodeId());
+	}
+	SelectNodes(NodeIds);
+}
+
+void URigDependencyGraph::ConstructNodes(const TArray<FNodeId>& InNodesToConstruct, bool bExpand)
+{
+	if (!RigDependencyProvider.IsValid())
+	{
+		return;
+	}
+
+	const URigHierarchy* Hierarchy = GetRigHierarchy();
+	if (!Hierarchy)
+	{
+		return;
+	}
+
+	TArray<FNodeId> NodesToConstruct = InNodesToConstruct;
+
+	if (bExpand)
+	{
+		if (!NodesToConstruct.IsEmpty())
+		{
+			const TArray<FNodeId> NodesToExpand = GetNodesToExpand(InNodesToConstruct);
+			for (const FNodeId& NodeId : NodesToExpand)
+			{
+				NodesToConstruct.AddUnique(NodeId);
+			}
+		}
+	}
+	
+	const TArray<FNodeId> PreviousSelection = SelectedNodes;
+	SelectedNodes.Reset();
+	bool bCreatedAnyNodes = false;
+
+	TSet<const UEdGraphNode*> NodesToSelect;
+	for (const FNodeId& NodeId : NodesToConstruct)
+	{
+		URigDependencyGraphNode* Node = FindNode(NodeId);
+		if (!Node)
+		{
+			Node = GetRigDependencyGraphSchema()->CreateGraphNode(this, NodeId);
+			bCreatedAnyNodes = true;
+		}
+		
+		if(PreviousSelection.Contains(NodeId))
+		{
+			SelectedNodes.Add(NodeId);
+			NodesToSelect.Add(Node);
+		}
+	}
+
+	SelectNodeSet(NodesToSelect);
+
+	RebuildLinks();
+
+	if (bCreatedAnyNodes)
+	{
+		RequestRefreshLayout();
+		NotifyGraphChanged();
+	}
+	
+	if (const UControlRig* ControlRig = GetControlRig())
+	{
+		const FRigVMExtendedExecuteContext& ExtendedExecuteContext = ControlRig->GetRigVMExtendedExecuteContext();
+		const FControlRigExecuteContext& Context = ExtendedExecuteContext.GetPublicData<FControlRigExecuteContext>();
+		LastRecordsHash = Context.GetInstructionRecordsHash();
+	}
+}
+
+void URigDependencyGraph::RemoveAllNodes()
+{
+	if (Nodes.IsEmpty())
+	{
+		return;
+	}
+
+	TArray<URigDependencyGraphNode*> NodesToRemove = DependencyGraphNodes;
+	for (int32 NodeIndex = 0; NodeIndex < NodesToRemove.Num(); ++NodeIndex)
+	{
+		RemoveDependencyGraphNode(NodesToRemove[NodeIndex]);
+	}
+	
+	DependencyGraphNodes.Reset();
+	LastRecordsHash.Reset();
+
+	RequestRefreshLayout();
+	NotifyGraphChanged();
+}
+
+void URigDependencyGraph::RemoveSelectedNodes()
+{
+	if (DependencyGraphNodes.IsEmpty() || SelectedNodes.IsEmpty())
+	{
+		return;
+	}
+
+	TArray<TObjectPtr<URigDependencyGraphNode>> NodesToRemove = DependencyGraphNodes;
+	for (int32 NodeIndex = 0; NodeIndex < NodesToRemove.Num(); ++NodeIndex)
+	{
+		if (!SelectedNodes.Contains(NodesToRemove[NodeIndex]->NodeId))
+		{
+			continue;
+		}
+		RemoveDependencyGraphNode(NodesToRemove[NodeIndex]);
+	}
+
+	SelectedNodes.Reset();
+	RequestRefreshLayout();
+	NotifyGraphChanged();
+}
+
+void URigDependencyGraph::RemoveUnselectedNodes()
+{
+	if (DependencyGraphNodes.IsEmpty() || SelectedNodes.IsEmpty())
+	{
+		return;
+	}
+
+	TArray<TObjectPtr<URigDependencyGraphNode>> NodesToRemove = DependencyGraphNodes;
+	for (int32 NodeIndex = 0; NodeIndex < NodesToRemove.Num(); ++NodeIndex)
+	{
+		if (SelectedNodes.Contains(NodesToRemove[NodeIndex]->NodeId))
+		{
+			continue;
+		}
+		RemoveDependencyGraphNode(NodesToRemove[NodeIndex]);
+	}
+
+	RequestRefreshLayout();
+	NotifyGraphChanged();
+}
+
+void URigDependencyGraph::RemoveUnrelatedNodes()
+{
+	if (DependencyGraphNodes.IsEmpty())
+	{
+		return;
+	}
+
+	bool bRemovedAnyNodes = false;
+	TArray<TObjectPtr<URigDependencyGraphNode>> PreviousNodes = DependencyGraphNodes;
+	for (int32 NodeIndex = 0; NodeIndex < PreviousNodes.Num(); ++NodeIndex)
+	{
+		if (PreviousNodes[NodeIndex]->IsUnrelated())
+		{
+			DependencyGraphNodes.Remove(PreviousNodes[NodeIndex]);
+			SelectedNodes.Remove(PreviousNodes[NodeIndex]->NodeId);
+			RemoveDependencyGraphNode(PreviousNodes[NodeIndex]);
+			bRemovedAnyNodes = true;
+		}
+	}
+
+	if (bRemovedAnyNodes)
+	{
+		RequestRefreshLayout();
+		NotifyGraphChanged();
+	}
+}
+
+void URigDependencyGraph::UpdateNodeIndices()
+{
+	for (int32 NodeIndex = 0; NodeIndex < DependencyGraphNodes.Num(); ++NodeIndex)
+	{
+		DependencyGraphNodes[NodeIndex]->Index = NodeIndex;
+	}
+}
+
+void URigDependencyGraph::RemoveDependencyGraphNode(URigDependencyGraphNode* InNode)
+{
+	check(InNode);
+	SelectedNodes.Remove(InNode->GetNodeId());
+	DependencyGraphNodes.Remove(InNode);
+	NodeIdLookup.Remove(InNode->GetNodeId());
+	RemoveNode(InNode);
+	UpdateFadeOutStates();
+}
+
+void URigDependencyGraph::SelectLinkedNodes(const TArray<FNodeId>& InNodeIds, bool bInputNodes, bool bClearSelection, bool bRecursive)
+{
+	TArray<FNodeId> NodesToSelect = InNodeIds;
+	for (int32 Index = 0; Index < NodesToSelect.Num(); ++Index)
+	{
+		if (!bRecursive)
+		{
+			if (!InNodeIds.Contains(NodesToSelect[Index]))
+			{
+				break;
+			}
+		}
+		if (const URigDependencyGraphNode* Node = FindNode(NodesToSelect[Index]))
+		{
+			const TArray<UEdGraphPin*>& LinkedPins = bInputNodes ? Node->InputPin->LinkedTo : Node->OutputPin->LinkedTo; 
+			for (const UEdGraphPin* LinkedPin : LinkedPins)
+			{
+				if (const URigDependencyGraphNode* LinkedNode = Cast<URigDependencyGraphNode>(LinkedPin->GetOwningNode()))
+				{
+					NodesToSelect.AddUnique(LinkedNode->GetNodeId());
+				}
+			}
+		}
+	}
+
+	if (!bClearSelection)
+	{
+		for (const FNodeId& SelectedNodeId : SelectedNodes)
+		{
+			NodesToSelect.AddUnique(SelectedNodeId);
+		}
+	}
+
+	SelectNodes(NodesToSelect);
+}
+
+void URigDependencyGraph::SelectNodeIsland(const TArray<FNodeId>& InNodeIds, bool bClearSelection)
+{
+	TArray<FGuid> Guids;
+	for (const FNodeId& NodeId : InNodeIds)
+	{
+		if (const URigDependencyGraphNode* Node = FindNode(NodeId))
+		{
+			if (Node->IslandGuid.IsSet())
+			{
+				Guids.AddUnique(Node->IslandGuid.GetValue());
+			}
+		}
+	}
+	
+	TArray<FNodeId> NodesToSelect;
+	for (const FNodeIsland& Island : NodeIslands)
+	{
+		if (Guids.Contains(Island.Guid))
+		{
+			for (const FNodeId& NodeId : Island.NodeIds)
+			{
+				NodesToSelect.AddUnique(NodeId);
+			}
+		}
+	}
+
+	SelectNodes(NodesToSelect);
+}
+
+FBox2D URigDependencyGraph::GetAllNodesBounds() const
+{
+	if (DependencyGraphNodes.IsEmpty())
+	{
+		return FBox2D();
+	}
+
+	FBox2D Bounds(EForceInit::ForceInit);
+	for (const TObjectPtr<URigDependencyGraphNode>& Node : DependencyGraphNodes)
+	{
+		const FVector2D TopLeft(Node->NodePosX, Node->NodePosY);
+		Bounds += TopLeft;
+		Bounds += TopLeft + Node->Dimensions;
+	}
+	return Bounds;
+}
+
+FBox2D URigDependencyGraph::GetSelectedNodesBounds() const
+{
+	if (DependencyGraphNodes.IsEmpty() || SelectedNodes.IsEmpty())
+	{
+		return FBox2D();
+	}
+
+	return GetNodesBounds(SelectedNodes);
+}
+
+void URigDependencyGraph::CancelLayout()
+{
+	LayoutIterationsLeft = 0;
+}
+
+void URigDependencyGraph::OnHierarchyModified(ERigHierarchyNotification InNotification, URigHierarchy* InHierarchy, const FRigNotificationSubject& InSubject)
+{
+	if (InNotification == ERigHierarchyNotification::ElementSelected || 
+		InNotification == ERigHierarchyNotification::ElementDeselected)
+	{
+		UE::TScopeLock Lock(NodesToSelectCriticalSection);
+
+		if (!NodesToSelectDuringTick.IsSet())
+		{
+			if (BlockSelectionCounter > 0)
+			{
+				return;
+			}
+		}
+		
+		URigHierarchy* Hierarchy = GetRigHierarchy();
+		if (Hierarchy == nullptr)
+		{
+			return;
+		}
+		
+		const TArray<FRigElementKey> SelectedElements = Hierarchy->GetSelectedKeys();
+
+		bool bRequiresReconstruction = false;
+		TArray<FNodeId> NodesToSelect;
+		for (const FRigElementKey& SelectedKey : SelectedElements)
+		{
+			const FNodeId NodeId(FNodeId::EType_RigElement, Hierarchy->GetIndex(SelectedKey), NAME_None);
+			if (!bRequiresReconstruction && FindNode(NodeId) == nullptr)
+			{
+				if (bLockContent)
+				{
+					continue;
+				}
+				bRequiresReconstruction = true;
+			}
+			NodesToSelect.Add(NodeId);
+		}
+
+		// maintain the selection of non-hierarchy based nodes
+		for (const FNodeId& SelectedNode : SelectedNodes)
+		{
+			if (!SelectedNode.IsElement() && !SelectedNode.IsMetadata())
+			{
+				NodesToSelect.Add(SelectedNode);
+			}
+		}
+
+		NodesToSelectDuringTick = NodesToSelect;
+		RemoveAllNodesDuringTick = bRequiresReconstruction;
+		BlockSelectionCounter = 2;
+	}
+}
+
+void URigDependencyGraph::OnMetadataChanged(const FRigElementKey& InKey, const FName& InMetadataName)
+{
+	// todo
+}
+
+void URigDependencyGraph::OnSetObjectBeingDebugged(UObject* InObject)
+{
+	if (!WeakControlRigEditor.IsValid())
+	{
+		return;
+	}
+	
+	UControlRig* ControlRig = Cast<UControlRig>(InObject);
+	if (!ControlRig || !IsValid(ControlRig))
+	{
+		ControlRig = nullptr;
+	}
+	Initialize(GetControlRigEditor(), ControlRig);
+}
+
+void URigDependencyGraph::OnVMCompiled(UObject* InAsset, URigVM* InVM, FRigVMExtendedExecuteContext& InExtendedContext)
+{
+	for (const TObjectPtr<URigDependencyGraphNode>& DependencyGraphNode : DependencyGraphNodes)
+	{
+		DependencyGraphNode->InvalidateCache();
+	}
+
+	if (RigDependencyProvider.IsValid())
+	{
+		RigDependencyProvider->InvalidateCache();
+	}
+}
+
+void URigDependencyGraph::OnRigVMGraphModified(ERigVMGraphNotifType InNotification, URigVMGraph* InRigVMGraph, UObject* InSubject)
+{
+	if (!WeakControlRigEditor.IsValid())
+	{
+		return;
+	}
+
+	if (InNotification == ERigVMGraphNotifType::NodeSelectionChanged)
+	{
+		UE::TScopeLock Lock(NodesToSelectCriticalSection);
+
+		if (!NodesToSelectDuringTick.IsSet())
+		{
+			if (BlockSelectionCounter > 0)
+			{
+				return;
+			}
+		}
+
+		check(InRigVMGraph);
+		
+		const TArray<FName> SelectNodeNames = InRigVMGraph->GetSelectNodes();
+		
+		TArray<FNodeId> NodesToSelect;
+		for (const FName& SelectedNodeName : SelectNodeNames)
+		{
+			if (const URigVMNode* SelectedNode = InRigVMGraph->FindNodeByName(SelectedNodeName))
+			{
+				for (const TObjectPtr<URigDependencyGraphNode>& DependencyGraphNode : DependencyGraphNodes)
+				{
+					if (DependencyGraphNode->GetRigVMNodeForInstruction() == SelectedNode)
+					{
+						NodesToSelect.Add(DependencyGraphNode->GetNodeId());
+					}
+				}
+			}
+		}
+
+		RemoveAllNodesDuringTick = false;
+		
+		if (NodesToSelect.IsEmpty() && !bLockContent)
+		{
+			// we didn't find the nodes selected on our view just yet,
+			// let's sync to this new view. first - we need to find the nodes in question.
+			for (const FName& SelectedNodeName : SelectNodeNames)
+			{
+				if (const URigVMNode* SelectedNode = InRigVMGraph->FindNodeByName(SelectedNodeName))
+				{
+					const FRigHierarchyRecord& Record = GetRecordForRigVMNode(SelectedNode);
+					if (Record.IsValid())
+					{
+						NodesToSelect.Add(Record);
+						RemoveAllNodesDuringTick = true;
+					}
+				}
+			}
+		}
+
+		// maintain the selection of non-instruction based nodes
+		for (const FNodeId& SelectedNode : SelectedNodes)
+		{
+			if (!SelectedNode.IsInstruction() || !SelectedNode.IsVariable())
+			{
+				NodesToSelect.Add(SelectedNode);
+			}
+		}
+
+		NodesToSelectDuringTick = NodesToSelect;
+		BlockSelectionCounter = 2;
+	}
+}
+
+void URigDependencyGraph::SelectNodes(const TArray<FNodeId>& InNodeIds)
+{
+	BlockSelectionCounter = 2;
+
+	SelectedNodes.Reset();
+	
+	TSet<const UEdGraphNode*> NodesToSelect;
+	for (const FNodeId& NodeId : InNodeIds)
+	{
+		if (const URigDependencyGraphNode* Node = FindNode(NodeId))
+		{
+			SelectedNodes.Add(NodeId);
+			NodesToSelect.Add(Node);
+		}
+	}
+
+	UpdateFadeOutStates();
+	
+	if (NodesToSelect.IsEmpty())
+	{
+		ClearNodeSelection();
+		return;
+	}
+	SelectNodeSet(NodesToSelect);
+}
+
+void URigDependencyGraph::ClearNodeSelection()
+{
+	(void)OnClearSelection.ExecuteIfBound();
+}
diff --git a/Engine/Plugins/Animation/ControlRig/Source/ControlRigEditor/Private/RigDependencyGraph/RigDependencyGraph.h b/Engine/Plugins/Animation/ControlRig/Source/ControlRigEditor/Private/RigDependencyGraph/RigDependencyGraph.h
new file mode 100644
--- /dev/null
+++ b/Engine/Plugins/Animation/ControlRig/Source/ControlRigEditor/Private/RigDependencyGraph/RigDependencyGraph.h
@@ -0,0 +1,260 @@
+// Copyright Epic Games, Inc. All Rights Reserved.
+
+#pragma once
+
+#include "RigDependencyGraphNode.h"
+#include "EdGraph/EdGraph.h"
+#include "Rigs/RigDependency.h"
+#include "RigVMCore/RigVMExecuteContext.h"
+#include "RigVMModel/RigVMNotifications.h"
+#include "RigDependencyGraph.generated.h"
+
+class URigDependencyGraphSchema;
+class IControlRigBaseEditor;
+class UControlRig;
+struct FRigVMClient;
+
+UCLASS()
+class URigDependencyGraph : public UEdGraph
+{
+	GENERATED_UCLASS_BODY()
+
+public:
+	using FNodeId = URigDependencyGraphNode::FNodeId;
+	
+	void Initialize(const TSharedRef<IControlRigBaseEditor>& InControlRigEditor, const TWeakObjectPtr<UControlRig>& InControlRig);
+
+	/** Rebuilds graph from selection */
+	void RebuildGraph(bool bExpand = false);
+
+	/** Rebuilds the links */
+	void RebuildLinks();
+
+	/* Recomputes the islands based on the connectivity */ 
+	void RebuildIslands();
+
+	/** Update the input and output pins */
+	void UpdatePinVisibility();
+
+	/** Update the nodes' fade out states */
+	void UpdateFadeOutStates();
+
+	/** Select the nodes */
+	void SelectNodes(const TArray<FNodeId>& InNodeIds);
+
+	/** Removes all nodes from the selection */
+	void ClearNodeSelection();
+
+	/** Get the dependency graph nodes we are displaying */
+	const TArray<URigDependencyGraphNode*>& GetNodes() const { return DependencyGraphNodes; }
+
+	/** Returns the node given a Node Id */
+	const URigDependencyGraphNode* FindNode(const FNodeId& InNodeId) const;
+	URigDependencyGraphNode* FindNode(const FNodeId& InNodeId);
+
+	/** Returns true if the node given a Node Id exists */
+	bool Contains(const FNodeId& InNodeId) const;
+
+	/** Get the control rig asset editor we are embedded in */
+	TSharedRef<IControlRigBaseEditor> GetControlRigEditor() const { return WeakControlRigEditor.Pin().ToSharedRef(); }
+
+	/** Get the control rig we are displaying */
+	UControlRig* GetControlRig() const;
+
+	/** Get the hierarchy we are displaying */
+	URigHierarchy* GetRigHierarchy() const;
+
+	/** Get the client to resolve a node for */
+	const FRigVMClient* GetRigVMClient() const;
+
+	/** Requests a layout refresh */
+	void RequestRefreshLayout();
+
+	/** Get whether a layout refresh was requested */
+	bool NeedsRefreshLayout() const;
+
+	/** Get the dependency provider we are editing */
+	TSharedRef<IRigDependenciesProvider> GetRigDependencyProvider() const;
+
+	/** Get the dependency graph schema */
+	const URigDependencyGraphSchema* GetRigDependencyGraphSchema() const;
+
+	bool DoesFollowParentRelationShips() const { return bFollowParentRelationShips; } 
+	bool DoesFollowVMRelationShips() const { return bFollowVMRelationShips; } 
+	bool IsContentLocked() const { return bLockContent; } 
+
+	void SetFollowParentRelationShips(bool InFollowParentRelationShips);
+	void SetFollowVMRelationShips(bool InFollowVMRelationShips);
+	void SetContentLocked(bool InLockContent);
+
+	TArray<FNodeId> GetNodesToExpand(const TArray<FNodeId>& InNodeIds) const;
+
+	const URigVMNode* GetRigVMNodeForInstruction(const FNodeId& InNodeId) const;
+	FRigHierarchyRecord GetRecordForRigVMNode(const URigVMNode* InNode) const;
+
+	const FRigVMExternalVariableDef* GetExternalVariable(const FNodeId& InNodeId) const;
+
+	void Tick(float DeltaTime);
+
+	void SelectAllNodes();
+	void RemoveAllNodes();
+	void RemoveSelectedNodes();
+	void RemoveUnselectedNodes();
+	void RemoveUnrelatedNodes();
+	void SelectLinkedNodes(const TArray<FNodeId>& InNodeIds, bool bInputNodes, bool bClearSelection, bool bRecursive);
+	void SelectNodeIsland(const TArray<FNodeId>& InNodeIds, bool bClearSelection);
+
+private:
+	void ConstructNodes(const TArray<FNodeId>& InNodesToConstruct, bool bExpand);
+	void UpdateNodeIndices();
+	void RemoveDependencyGraphNode(URigDependencyGraphNode* InNode);
+
+	FBox2D GetAllNodesBounds() const;
+	FBox2D GetSelectedNodesBounds() const;
+
+	template<typename RangeType>
+	FBox2D GetNodesBounds(const RangeType& InNodeIds) const
+	{
+		FBox2D Bounds(EForceInit::ForceInit);
+		for (const FNodeId& NodeId : InNodeIds)
+		{
+			if (const URigDependencyGraphNode* Node = FindNode(NodeId))
+			{
+				const FVector2D TopLeft(Node->NodePosX, Node->NodePosY);
+				Bounds += TopLeft + Node->Dimensions;
+			}
+		}
+		return Bounds;
+	}
+
+	void CancelLayout();
+
+	void OnHierarchyModified(ERigHierarchyNotification InNotification, URigHierarchy* InHierarchy, const FRigNotificationSubject& InSubject);
+	FDelegateHandle OnHierarchyModifiedHandle;
+
+	void OnMetadataChanged(const FRigElementKey& InKey, const FName& InMetadataName);
+	FDelegateHandle OnMetadataChangedHandle;;
+	
+	void OnSetObjectBeingDebugged(UObject* InObject);
+	FDelegateHandle OnSetObjectBeingDebuggedHandle;
+
+	void OnVMCompiled(UObject* InAsset, URigVM* InVM, FRigVMExtendedExecuteContext& InExtendedContext);
+	FDelegateHandle OnVMCompiledHandle;
+
+	void OnRigVMGraphModified(ERigVMGraphNotifType InNotification, URigVMGraph* InRigVMGraph, UObject* InSubject);
+	FDelegateHandle OnRigVMGraphModifiedHandle;
+	
+private:
+	
+	TArray<FNodeId> SelectedNodes;
+
+	UPROPERTY()
+	TArray<TObjectPtr<URigDependencyGraphNode>> DependencyGraphNodes;
+
+	TMap<FNodeId, URigDependencyGraphNode*> NodeIdLookup;
+	
+	TWeakPtr<IControlRigBaseEditor> WeakControlRigEditor;
+	TWeakObjectPtr<UControlRig> WeakControlRig;
+	TSharedPtr<IRigDependenciesProvider> RigDependencyProvider;
+
+	mutable FTransactionallySafeCriticalSection NodesToSelectCriticalSection;
+	mutable TOptional<TArray<FNodeId>> NodesToSelectDuringTick;
+	mutable TOptional<bool> RemoveAllNodesDuringTick;
+	mutable TOptional<bool> ZoomAndFitDuringLayout;
+	TOptional<uint32> LastRecordsHash;
+
+	int32 DefaultLayoutIterationsPerNode;
+	int32 LayoutIterationsLeft;
+	int32 LayoutIterationsTotal;
+	int32 LayoutIterationsMax;
+	bool bFollowParentRelationShips;
+	bool bFollowVMRelationShips;
+	bool bLockContent;
+
+	/** As selection broadcasting is deferred in the graph, we need to block re-broadcasting a for a few frames on refresh */
+	int32 BlockSelectionCounter;
+
+	struct FLayoutEdge
+	{
+		FLayoutEdge()
+			: NodeA(INDEX_NONE)
+			, NodeB(INDEX_NONE)
+			, AnchorA(FVector2D::ZeroVector)
+			, AnchorB(FVector2D::ZeroVector)
+			, Strength(0.f)
+		{
+		}
+		
+		FLayoutEdge(int32 InNodeA, int32 InNodeB, const FVector2D& InAnchorA, const FVector2D& InAnchorB, float InStrength)
+			: NodeA(InNodeA)
+			, NodeB(InNodeB)
+			, AnchorA(InAnchorA)
+			, AnchorB(InAnchorB)
+			, Strength(InStrength)
+		{
+			if (NodeA > NodeB)
+			{
+				Swap(NodeA, NodeB);
+				Swap(AnchorA, AnchorB);
+			}
+		}
+
+		friend uint32 GetTypeHash(const FLayoutEdge& InEdge)
+		{
+			return HashCombine(GetTypeHash(InEdge.NodeA), GetTypeHash(InEdge.NodeB));
+		}
+
+		bool operator==(const FLayoutEdge& InOther) const
+		{
+			return NodeA == InOther.NodeA && NodeB == InOther.NodeB;
+		}
+
+		bool operator>(const FLayoutEdge& InOther) const
+		{
+			if (NodeA > InOther.NodeA)
+			{
+				return true;
+			}
+			if (NodeA < InOther.NodeA)
+			{
+				return false;
+			}
+			return NodeB > InOther.NodeB;
+		}
+
+		bool operator<(const FLayoutEdge& InOther) const
+		{
+			if (NodeA < InOther.NodeA)
+			{
+				return true;
+			}
+			if (NodeA > InOther.NodeA)
+			{
+				return false;
+			}
+			return NodeB < InOther.NodeB;
+		}
+
+		int32 NodeA;
+		int32 NodeB;
+		FVector2D AnchorA;
+		FVector2D AnchorB;
+		float Strength;
+	};
+
+	struct FNodeIsland
+	{
+		FGuid Guid;
+		TSet<FNodeId> NodeIds;
+		FBox2D Bounds;
+	};
+
+	TSet<FLayoutEdge> LayoutEdges;
+	TArray<FNodeIsland> NodeIslands;
+
+	FSimpleDelegate OnClearSelection;
+
+	friend class SRigDependencyGraph;
+	friend class URigDependencyGraphSchema;
+};
+
diff --git a/Engine/Plugins/Animation/ControlRig/Source/ControlRigEditor/Private/RigDependencyGraph/RigDependencyGraphCommands.cpp b/Engine/Plugins/Animation/ControlRig/Source/ControlRigEditor/Private/RigDependencyGraph/RigDependencyGraphCommands.cpp
new file mode 100644
--- /dev/null
+++ b/Engine/Plugins/Animation/ControlRig/Source/ControlRigEditor/Private/RigDependencyGraph/RigDependencyGraphCommands.cpp
@@ -0,0 +1,26 @@
+// Copyright Epic Games, Inc. All Rights Reserved.
+
+#include "RigDependencyGraphCommands.h"
+
+#define LOCTEXT_NAMESPACE "RigDependencyGraphCommands"
+
+void FRigDependencyGraphCommands::RegisterCommands()
+{
+	UI_COMMAND(FrameSelection, "Frame Selection", "Expands and frames the selection in the tree", EUserInterfaceActionType::Button, FInputChord(EKeys::F));
+	UI_COMMAND(ShowParentChildRelationships, "Show Parent/Child", "Whether to show or hide parent and child relationships", EUserInterfaceActionType::ToggleButton, FInputChord(EKeys::P));
+	UI_COMMAND(ShowInstructionRelationships, "Show VM Nodes", "Whether to show or hide VM based relationships", EUserInterfaceActionType::ToggleButton, FInputChord(EKeys::O));
+	UI_COMMAND(SyncContent, "Sync Content", "Whether or not to sync the content with selection in other views", EUserInterfaceActionType::ToggleButton, FInputChord(EKeys::L));
+	UI_COMMAND(SelectAllNodes, "Select All", "Selects all nodes in the graph", EUserInterfaceActionType::Button, FInputChord(EModifierKey::Control, EKeys::A));
+	UI_COMMAND(SelectInputNodes, "Select Inputs", "Selects the linked inputs", EUserInterfaceActionType::Button, FInputChord(EKeys::MiddleMouseButton));
+	UI_COMMAND(SelectOutputNodes, "Select Outputs", "Selects the linked outputs", EUserInterfaceActionType::Button, FInputChord(EModifierKey::Control, EKeys::MiddleMouseButton));
+	UI_COMMAND(SelectNodeIsland, "Select Island", "Selects the linked island", EUserInterfaceActionType::Button, FInputChord(EModifierKey::Alt, EKeys::MiddleMouseButton));
+	UI_COMMAND(RemoveAllNodes, "Clear View", "Hides all nodes", EUserInterfaceActionType::Button, FInputChord(EKeys::X));
+	UI_COMMAND(RemoveSelected, "Hide Selection", "Hides the selected nodes", EUserInterfaceActionType::Button, FInputChord(EKeys::Delete));
+	UI_COMMAND(RemoveUnrelated, "Hide Unrelated", "Hides nodes without any dependencies", EUserInterfaceActionType::Button, FInputChord());
+	UI_COMMAND(IsolateSelected, "Isolate Selection", "Hides all non-selected nodes", EUserInterfaceActionType::Button, FInputChord(EKeys::I));
+	UI_COMMAND(SyncWithHierarchy, "Sync Hierarchy", "Syncs the content of the graph with the hierarchy selection", EUserInterfaceActionType::Button, FInputChord(EKeys::F5));
+	UI_COMMAND(CancelLayoutSimulation, "Cancel Layout", "Cancels the currently running layout simulation", EUserInterfaceActionType::Button, FInputChord(EKeys::Escape));
+	UI_COMMAND(FindNodesByName, "Find Nodes By Name", "Finds nodes by name", EUserInterfaceActionType::Button, FInputChord(EModifierKey::Control, EKeys::F));
+}
+
+#undef LOCTEXT_NAMESPACE
diff --git a/Engine/Plugins/Animation/ControlRig/Source/ControlRigEditor/Private/RigDependencyGraph/RigDependencyGraphCommands.h b/Engine/Plugins/Animation/ControlRig/Source/ControlRigEditor/Private/RigDependencyGraph/RigDependencyGraphCommands.h
new file mode 100644
--- /dev/null
+++ b/Engine/Plugins/Animation/ControlRig/Source/ControlRigEditor/Private/RigDependencyGraph/RigDependencyGraphCommands.h
@@ -0,0 +1,70 @@
+// Copyright Epic Games, Inc. All Rights Reserved.
+
+#pragma once
+
+#include "CoreMinimal.h"
+#include "Framework/Commands/Commands.h"
+#include "ControlRigEditorStyle.h"
+
+class FRigDependencyGraphCommands : public TCommands<FRigDependencyGraphCommands>
+{
+public:
+	FRigDependencyGraphCommands() : TCommands<FRigDependencyGraphCommands>
+	(
+		"ControlRigDependencyGraph",
+		NSLOCTEXT("Contexts", "RigDependencyGraph", "Dependencies"),
+		NAME_None, // "MainFrame" // @todo Fix this crash
+		FControlRigEditorStyle::Get().GetStyleSetName() // Icon Style Set
+	)
+	{}
+	
+	/* frames the selection in the graph */
+	TSharedPtr<FUICommandInfo> FrameSelection;
+
+	/* Show parent-child relationships */
+	TSharedPtr< FUICommandInfo > ShowParentChildRelationships;
+
+	/* Show instruction based relationships */
+	TSharedPtr< FUICommandInfo > ShowInstructionRelationships;
+
+	/* Syncs the content with other views */
+	TSharedPtr< FUICommandInfo > SyncContent;
+
+	/* select all nodes */
+	TSharedPtr<FUICommandInfo> SelectAllNodes;
+
+	/* select the input nodes of the current selection */
+	TSharedPtr<FUICommandInfo> SelectInputNodes;
+
+	/* select the output nodes of the current selection */
+	TSharedPtr<FUICommandInfo> SelectOutputNodes;
+
+	/* select the complete node island */
+	TSharedPtr<FUICommandInfo> SelectNodeIsland;
+
+	/* remove all nodes */
+	TSharedPtr<FUICommandInfo> RemoveAllNodes;
+
+	/* remove selected nodes */
+	TSharedPtr<FUICommandInfo> RemoveSelected;
+
+	/* isolate selected nodes */
+	TSharedPtr<FUICommandInfo> IsolateSelected;
+
+	/* remove unrealted nodes */
+	TSharedPtr<FUICommandInfo> RemoveUnrelated;
+
+	/* sync the view with the hierarchy */
+	TSharedPtr<FUICommandInfo> SyncWithHierarchy;
+
+	/* cancel an ongoing layout simulation */
+	TSharedPtr<FUICommandInfo> CancelLayoutSimulation;
+
+	/* search nodes by name */
+	TSharedPtr<FUICommandInfo> FindNodesByName;
+
+	/**
+	 * Initialize commands
+	 */
+	virtual void RegisterCommands() override;
+};
diff --git a/Engine/Plugins/Animation/ControlRig/Source/ControlRigEditor/Private/RigDependencyGraph/RigDependencyGraphNode.cpp b/Engine/Plugins/Animation/ControlRig/Source/ControlRigEditor/Private/RigDependencyGraph/RigDependencyGraphNode.cpp
new file mode 100644
--- /dev/null
+++ b/Engine/Plugins/Animation/ControlRig/Source/ControlRigEditor/Private/RigDependencyGraph/RigDependencyGraphNode.cpp
@@ -0,0 +1,475 @@
+// Copyright Epic Games, Inc. All Rights Reserved.
+
+#include "RigDependencyGraph/RigDependencyGraphNode.h"
+
+#include "ControlRig.h"
+#include "ControlRigEditorStyle.h"
+#include "MaterialPropertyHelpers.h"
+#include "EdGraph/EdGraphPin.h"
+#include "RigDependencyGraph/RigDependencyGraph.h"
+#include "Templates/Casts.h"
+#include "UObject/UnrealNames.h"
+#include "Editor/SRigHierarchyTreeView.h"
+#include "RigVMModel/RigVMClient.h"
+#include "RigVMModel/Nodes/RigVMLibraryNode.h"
+#include "RigVMModel/RigVMFunctionLibrary.h"
+
+#include UE_INLINE_GENERATED_CPP_BY_NAME(RigDependencyGraphNode)
+
+class UObject;
+
+URigDependencyGraphNode::URigDependencyGraphNode(const FObjectInitializer& ObjectInitializer)
+	: Super(ObjectInitializer)
+{
+	Index = INDEX_NONE;
+	InputPin = nullptr;
+	OutputPin = nullptr;
+	NodeBodyColor = FLinearColor::White;
+
+	LayoutPosition = FVector2D::ZeroVector;
+	LayoutVelocity = FVector2D::ZeroVector;
+	LayoutForce = FVector2D::ZeroVector;
+	Dimensions = FVector2D::ZeroVector;
+}
+
+void URigDependencyGraphNode::SetupRigDependencyNode(const FNodeId& InNodeId)
+{
+	NodeId = InNodeId;
+
+	const URigHierarchy* Hierarchy = GetRigHierarchy();
+	if (Hierarchy == nullptr)
+	{
+		return;
+	}
+
+	switch (NodeId.Type)
+	{
+		case FNodeId::EType_RigElement:
+		{
+			const FRigElementKey ElementKey = Hierarchy->GetKey(NodeId.Index);
+			NodeTitle = FText::FromName(ElementKey.Name);
+			NodeTooltip = FText::FromString(ElementKey.ToString());
+
+			switch (ElementKey.Type)
+			{
+				case ERigElementType::Bone:
+				{
+					NodeBodyColor = FControlRigEditorStyle::Get().GetColor(TEXT("ControlRig.DependencyGraph.Colors.Bone"));
+					break;
+				}
+				case ERigElementType::Control:
+				{
+					if(const FRigControlElement* ControlElement = Hierarchy->Find<FRigControlElement>(ElementKey))
+					{
+						if (ControlElement->IsAnimationChannel())
+						{
+							NodeBodyColor = FControlRigEditorStyle::Get().GetColor(TEXT("ControlRig.DependencyGraph.Colors.AnimationChannel"));
+							break;
+						}
+					}
+					NodeBodyColor = FControlRigEditorStyle::Get().GetColor(TEXT("ControlRig.DependencyGraph.Colors.Control"));
+					break;
+				}
+				case ERigElementType::Null:
+				{
+					NodeBodyColor = FControlRigEditorStyle::Get().GetColor(TEXT("ControlRig.DependencyGraph.Colors.Null"));
+					break;
+				}
+				case ERigElementType::Socket:
+				{
+					NodeBodyColor = FControlRigEditorStyle::Get().GetColor(TEXT("ControlRig.DependencyGraph.Colors.Socket"));
+					break;
+				}
+				case ERigElementType::Connector:
+				{
+					NodeBodyColor = FControlRigEditorStyle::Get().GetColor(TEXT("ControlRig.DependencyGraph.Colors.Connector"));
+					break;
+				}
+				default:
+				{
+					break;
+				}
+			}
+			break;
+		}
+		case FNodeId::EType_Metadata:
+		{
+			const FRigElementKey ElementKey = Hierarchy->GetKey(NodeId.Index);
+			NodeTitle = FText::FromString(
+				FString::Printf(TEXT("%s\n%s"),
+					*NodeId.Name.ToString(),
+					*ElementKey.ToString()
+				));
+			NodeBodyColor = FControlRigEditorStyle::Get().GetColor(TEXT("ControlRig.DependencyGraph.Colors.Metadata"));
+			NodeTooltip = NodeTitle;
+			break;
+		}
+		case FNodeId::EType_Variable:
+		{
+			if (const FRigVMExternalVariableDef* Variable = GetExternalVariable())
+			{
+				NodeTitle = FText::FromString(Variable->Name.ToString());
+			}
+			NodeBodyColor = FControlRigEditorStyle::Get().GetColor(TEXT("ControlRig.DependencyGraph.Colors.Variable"));
+			NodeTooltip = NodeTitle;
+			break;
+		}
+		case FNodeId::EType_Instruction:
+		{
+			if (const URigVMNode* Node = GetRigVMNodeForInstruction())
+			{
+				NodeTitle = FText::FromString(Node->GetNodeTitle());
+
+				const URigVMFunctionLibrary* FunctionLibrary = Node->GetTypedOuter<URigVMFunctionLibrary>();
+				if (FunctionLibrary)
+				{
+					FText FunctionName;
+					const UObject* Parent = Node->GetOuter();
+					while (Parent && Parent != FunctionLibrary)
+					{
+						const UObject* GrandParent = Parent->GetOuter();
+						if (GrandParent == FunctionLibrary)
+						{
+							if (const URigVMLibraryNode* LibraryNode = Cast<URigVMLibraryNode>(Parent))
+							{
+								FunctionName = FText::FromName(LibraryNode->GetFunctionIdentifier().GetFunctionFName());
+							}
+							break;
+						}
+						Parent = GrandParent;
+					}
+
+					FText Asset, AssetSuffix;
+					if (FunctionLibrary->GetOutermost() != GetRigVMClient()->GetOuter()->GetOutermost())
+					{
+						FString AssetName = FunctionLibrary->GetOutermost()->GetName();
+						AssetName.RightChopInline(1 + AssetName.Find(TEXT("/"), ESearchCase::IgnoreCase, ESearchDir::FromEnd));
+						Asset =  FText::FromString(AssetName);
+						AssetSuffix = NSLOCTEXT("RigDependencyGraphNode", "AssetSuffix", " - "); 
+					}
+					
+					const FText NodeTitleFormat = NSLOCTEXT("RigDependencyGraphNode", "NodeTitleFormat", "{0}\n{1}{2}{3}");
+					NodeTitle = FText::Format(NodeTitleFormat, NodeTitle, Asset, AssetSuffix, FunctionName);
+					NodeTooltip = Node->GetToolTipText();
+				}
+				
+				NodeBodyColor = FControlRigEditorStyle::Get().GetColor(TEXT("ControlRig.DependencyGraph.Colors.Instruction"));
+			}
+			break;
+		}
+		default:
+		{
+			NodeTitle = NSLOCTEXT("RigDependencyGraphNode", "Node", "Node");
+			NodeTooltip = NodeTitle;
+		}
+	}
+
+	NodeBodyColor *= FLinearColor(0.2f, 0.2f, 0.2f, 1.f);
+}
+
+UObject* URigDependencyGraphNode::GetDetailsObject()
+{
+	return nullptr;
+}
+
+URigDependencyGraph* URigDependencyGraphNode::GetRigDependencyGraph() const
+{
+	return CastChecked<URigDependencyGraph>(GetOuter());
+}
+
+FText URigDependencyGraphNode::GetNodeTitle(ENodeTitleType::Type TitleType) const
+{
+	return NodeTitle;
+}
+
+FLinearColor URigDependencyGraphNode::GetNodeBodyTintColor() const
+{
+	return NodeBodyColor;
+}
+
+void URigDependencyGraphNode::AllocateDefaultPins()
+{
+	InputPin = CreatePin(EEdGraphPinDirection::EGPD_Input, NAME_None, NAME_None, NAME_None);
+	//InputPin->bHidden = true;
+	OutputPin = CreatePin(EEdGraphPinDirection::EGPD_Output, NAME_None, NAME_None, NAME_None);
+	//OutputPin->bHidden = true;
+}
+
+FSlateIcon URigDependencyGraphNode::GetIconAndTint(FLinearColor& OutColor) const
+{
+	OutColor = FLinearColor::White;
+
+	switch (NodeId.Type)
+	{
+		case FNodeId::EType_RigElement:
+		{
+			static const FSlateIcon ControlIcon = FSlateIcon(FControlRigEditorStyle::Get().GetStyleSetName(), TEXT("ControlRig.Tree.Control"));
+			static const FSlateIcon NullIcon = FSlateIcon(FControlRigEditorStyle::Get().GetStyleSetName(), TEXT("ControlRig.Tree.Null"));
+			static const FSlateIcon BoneIcon = FSlateIcon(FControlRigEditorStyle::Get().GetStyleSetName(), TEXT("ControlRig.Tree.BoneImported"));
+			static const FSlateIcon SocketIcon = FSlateIcon(FControlRigEditorStyle::Get().GetStyleSetName(), TEXT("ControlRig.Tree.Socket_Open"));
+			static const FSlateIcon ConnectorIcon = FSlateIcon(FControlRigEditorStyle::Get().GetStyleSetName(), TEXT("ControlRig.ConnectorPrimary"));
+			static const FSlateIcon AnimationChannelIcon = FSlateIcon(FAppStyle::GetAppStyleSetName(), TEXT("Kismet.VariableList.TypeIcon"));
+
+			if (const URigHierarchy* Hierarchy = GetRigHierarchy())
+			{
+				const FRigElementKey ElementKey = Hierarchy->GetKey(NodeId.Index);
+
+				switch (ElementKey.Type)
+				{
+					case ERigElementType::Bone:
+					{
+						return BoneIcon;
+					}
+					case ERigElementType::Control:
+					{
+						if (const UControlRig* ControlRig = GetControlRig())
+						{
+							if (const FRigControlElement* Control = ControlRig->FindControl(ElementKey.Name))
+							{
+								if (Control->IsAnimationChannel())
+								{
+									return AnimationChannelIcon;
+								}
+							}
+						}
+						return ControlIcon;
+					}
+					case ERigElementType::Socket:
+					{
+						return SocketIcon;
+					}
+					case ERigElementType::Connector:
+					{
+						return ConnectorIcon;
+					}
+					default:
+					{
+						return NullIcon;
+					}
+				}
+			}
+			break;
+		}
+		case FNodeId::EType_Metadata:
+		case FNodeId::EType_Variable:
+		{
+			static FSlateIcon TypeIcon(FAppStyle::GetAppStyleSetName(), "Kismet.VariableList.TypeIcon");
+			return TypeIcon;
+		}
+		case FNodeId::EType_Instruction:
+		{
+			static const FSlateIcon NodeIcon(TEXT("RigVMEditorStyle"), TEXT("RigVM.Unit"));
+			return NodeIcon;
+		}
+		default:
+		{
+			break;
+		}
+	}
+	
+	return Super::GetIconAndTint(OutColor);
+}
+
+UEdGraphPin& URigDependencyGraphNode::GetInputPin() const
+{
+	return *InputPin;
+}
+
+UEdGraphPin& URigDependencyGraphNode::GetOutputPin() const
+{
+	return *OutputPin;
+}
+
+void URigDependencyGraphNode::SetDimensions(const FVector2D& InDimensions)
+{
+	Dimensions = InDimensions;
+}
+
+int32 URigDependencyGraphNode::GetDependencyDepth() const
+{
+	if (DependencyDepth.IsSet())
+	{
+		return DependencyDepth.GetValue();
+	}
+	
+	TArray<bool> Visited;
+	return GetDependencyDepth_Impl(Visited);
+}
+
+UControlRig* URigDependencyGraphNode::GetControlRig() const
+{
+	if (const URigDependencyGraph* Graph = GetRigDependencyGraph())
+	{
+		return Graph->GetControlRig();
+	}
+	return nullptr;
+}
+
+URigHierarchy* URigDependencyGraphNode::GetRigHierarchy() const
+{
+	if (const URigDependencyGraph* Graph = GetRigDependencyGraph())
+	{
+		return Graph->GetRigHierarchy();
+	}
+	return nullptr;
+}
+
+const FRigVMClient* URigDependencyGraphNode::GetRigVMClient() const
+{
+	if (const URigDependencyGraph* Graph = GetRigDependencyGraph())
+	{
+		return Graph->GetRigVMClient();
+	}
+	return nullptr;
+}
+
+FRigElementKey URigDependencyGraphNode::GetRigElementKey() const
+{
+	if (NodeId.IsElement() || NodeId.IsMetadata())
+	{
+		if (const URigHierarchy* Hierarchy = GetRigHierarchy())
+		{
+			return Hierarchy->GetKey(NodeId.Index);
+		}
+	}
+	return FRigElementKey();
+}
+
+const FRigBaseElement* URigDependencyGraphNode::GetRigElement() const
+{
+	if (NodeId.IsElement() || NodeId.IsMetadata())
+	{
+		if (const URigHierarchy* Hierarchy = GetRigHierarchy())
+		{
+			return Hierarchy->Get(NodeId.Index);
+		}
+	}
+	return nullptr;
+}
+
+const URigVMNode* URigDependencyGraphNode::GetRigVMNodeForInstruction() const
+{
+	if (CachedRigVMNode.IsValid())
+	{
+		return CachedRigVMNode.Get();
+	}
+
+	const URigVMNode* SubjectNode = GetRigDependencyGraph()->GetRigVMNodeForInstruction(NodeId);
+	if (SubjectNode)
+	{
+		CachedRigVMNode = SubjectNode;
+	}
+	return SubjectNode;
+}
+
+const FRigVMExternalVariableDef* URigDependencyGraphNode::GetExternalVariable() const
+{
+	if (!CachedRigVMExternalVariableDef.IsValid())
+	{
+		if (const FRigVMExternalVariableDef* ExternalVariable = GetRigDependencyGraph()->GetExternalVariable(NodeId))
+		{
+			CachedRigVMExternalVariableDef = *ExternalVariable;
+		}
+	}
+	return &CachedRigVMExternalVariableDef;
+}
+
+bool URigDependencyGraphNode::IsUnrelated() const
+{
+	return InputPin->bHidden && OutputPin->bHidden;
+}
+
+bool URigDependencyGraphNode::IsFadedOut() const
+{
+	return GetFadedOutState() < 1.f - SMALL_NUMBER;
+}
+
+float URigDependencyGraphNode::GetFadedOutState() const
+{
+	if (FadedOutOverride.IsSet())
+	{
+		return FadedOutOverride.GetValue();
+	}
+	return bIsFadedOut.Get(false) ? 0.1f : 1.f;
+}
+
+void URigDependencyGraphNode::OverrideFadeOutState(TOptional<float> InFadedOutState)
+{
+	FadedOutOverride = InFadedOutState;
+}
+
+const FGuid& URigDependencyGraphNode::GetIslandGuid() const
+{
+	if (IslandGuid.IsSet())
+	{
+		return IslandGuid.GetValue();
+	}
+
+	// make as invalid
+	IslandGuid = FGuid();
+
+	for (int32 Phase = 0; Phase < 2; Phase++)
+	{
+		for (const UEdGraphPin* LinkedPin : (Phase == 0 ? InputPin : OutputPin)->LinkedTo)
+		{
+			if (const URigDependencyGraphNode* LinkedNode = Cast<URigDependencyGraphNode>(LinkedPin->GetOwningNode()))
+			{
+				const FGuid& LinkedGuid = LinkedNode->GetIslandGuid();
+				if (LinkedGuid.IsValid())
+				{
+					IslandGuid = LinkedGuid;
+					return LinkedGuid;
+				}
+			}
+		}
+	}
+
+	IslandGuid = FGuid::NewGuid();
+	return IslandGuid.GetValue();
+}
+
+void URigDependencyGraphNode::InvalidateCache()
+{
+	CachedRigVMExternalVariableDef = FRigVMExternalVariableDef();
+	CachedRigVMNode.Reset();
+}
+
+int32 URigDependencyGraphNode::GetDependencyDepth_Impl(TArray<bool>& InOutVisited) const
+{
+	URigDependencyGraph* DependencyGraph = GetRigDependencyGraph();
+	if (DependencyGraph == nullptr)
+	{
+		return INDEX_NONE;
+	}
+
+	if (InOutVisited.IsEmpty())
+	{
+		InOutVisited.AddZeroed(DependencyGraph->GetNodes().Num());
+	}
+
+	if (InOutVisited[Index])
+	{
+		if (!DependencyDepth.IsSet())
+		{
+			DependencyDepth = 0;
+		}
+		return DependencyDepth.GetValue();
+	}
+
+	InOutVisited[Index] = true;
+
+	if (InputPin->LinkedTo.IsEmpty())
+	{
+		DependencyDepth = 0;
+		return DependencyDepth.GetValue();
+	}
+
+	for (const UEdGraphPin* SourcePin : InputPin->LinkedTo)
+	{
+		const URigDependencyGraphNode* SourceNode = CastChecked<URigDependencyGraphNode>(SourcePin->GetOwningNode());
+		DependencyDepth = FMath::Max<int32>(SourceNode->GetDependencyDepth_Impl(InOutVisited) + 1, DependencyDepth.Get(INDEX_NONE));
+	}
+
+	check(DependencyDepth.IsSet());
+	return DependencyDepth.GetValue();
+}
diff --git a/Engine/Plugins/Animation/ControlRig/Source/ControlRigEditor/Private/RigDependencyGraph/RigDependencyGraphNode.h b/Engine/Plugins/Animation/ControlRig/Source/ControlRigEditor/Private/RigDependencyGraph/RigDependencyGraphNode.h
new file mode 100644
--- /dev/null
+++ b/Engine/Plugins/Animation/ControlRig/Source/ControlRigEditor/Private/RigDependencyGraph/RigDependencyGraphNode.h
@@ -0,0 +1,115 @@
+// Copyright Epic Games, Inc. All Rights Reserved.
+
+#pragma once
+
+#include "EdGraph/EdGraphNode.h"
+#include "Internationalization/Text.h"
+#include "Math/Vector2D.h"
+#include "Rigs/RigHierarchyDefines.h"
+#include "Rigs/RigHierarchyElements.h"
+#include "Rigs/RigDependencyRecords.h"
+#include "RigVMCore/RigVMExternalVariable.h"
+
+#include "RigDependencyGraphNode.generated.h"
+
+class UEdGraphPin;
+class UObject;
+class UControlRig;
+class URigVMNode;
+struct FRigVMClient;
+
+UCLASS()
+class URigDependencyGraphNode : public UEdGraphNode
+{
+	GENERATED_UCLASS_BODY()
+
+	using FNodeId = FRigHierarchyRecord;
+
+	virtual void SetupRigDependencyNode(const FNodeId& InNodeId);
+	virtual UObject* GetDetailsObject();
+
+	class URigDependencyGraph* GetRigDependencyGraph() const;
+
+	// UEdGraphNode implementation
+	virtual FText GetNodeTitle(ENodeTitleType::Type TitleType) const override;
+	virtual FLinearColor GetNodeBodyTintColor() const override;
+	virtual FText GetTooltipText() const override { return NodeTooltip; }
+	virtual void AllocateDefaultPins() override;
+	virtual bool ShowPaletteIconOnNode() const override { return true; }
+	virtual FSlateIcon GetIconAndTint(FLinearColor& OutColor) const;
+
+	int32 GetIndex() const { return Index; }
+	const FNodeId& GetNodeId() const { return NodeId; }
+	
+	UEdGraphPin& GetInputPin() const;
+
+	UEdGraphPin& GetOutputPin() const;
+
+	void SetDimensions(const FVector2D& InDimensions);
+	const FVector2D& GetDimensions() const { return Dimensions; }
+
+	int32 GetDependencyDepth() const;
+
+	/** Get the control rig we are displaying */
+	UControlRig* GetControlRig() const;
+
+	/** Get the hierarchy we are displaying */
+	URigHierarchy* GetRigHierarchy() const;
+
+	/** Get the client to resolve a node for */
+	const FRigVMClient* GetRigVMClient() const;
+
+	FRigElementKey GetRigElementKey() const;
+	const FRigBaseElement* GetRigElement() const;
+	const URigVMNode* GetRigVMNodeForInstruction() const;
+	const FRigVMExternalVariableDef* GetExternalVariable() const;
+
+	bool IsUnrelated() const;
+	bool IsFadedOut() const;
+	float GetFadedOutState() const;
+	void OverrideFadeOutState(TOptional<float> InFadedOutState);
+	void ResetFadedOutState() { OverrideFadeOutState({}); }
+
+	const FGuid& GetIslandGuid() const;
+
+	void InvalidateCache();
+
+protected:
+
+	int32 GetDependencyDepth_Impl(TArray<bool>& InOutVisited) const;
+
+	int32 Index;
+	
+	FNodeId NodeId;
+	
+	/** Cached title for the node */
+	FText NodeTitle;
+	FText NodeTooltip;
+	FLinearColor NodeBodyColor;
+
+	/** Our one input pin */
+	UEdGraphPin* InputPin;
+
+	/** Our one output pin */
+	UEdGraphPin* OutputPin;
+
+	/** Cached dimensions of this node (used for layout) */
+	TOptional<bool> FollowLayout;
+	FVector2D LayoutPosition;
+	FVector2D LayoutVelocity;
+	FVector2D LayoutForce;
+	FVector2D Dimensions;
+
+	mutable TOptional<int32> DependencyDepth;
+	mutable TOptional<FGuid> IslandGuid;
+	mutable TOptional<bool> bIsFadedOut;
+	mutable TOptional<float> FadedOutOverride;
+
+	mutable TWeakObjectPtr<const URigVMNode> CachedRigVMNode;
+	mutable FRigVMExternalVariableDef CachedRigVMExternalVariableDef;
+
+	friend class URigDependencyGraph;
+	friend class URigDependencyGraphSchema;
+};
+
+
diff --git a/Engine/Plugins/Animation/ControlRig/Source/ControlRigEditor/Private/RigDependencyGraph/RigDependencyGraphPanelNodeFactory.cpp b/Engine/Plugins/Animation/ControlRig/Source/ControlRigEditor/Private/RigDependencyGraph/RigDependencyGraphPanelNodeFactory.cpp
new file mode 100644
--- /dev/null
+++ b/Engine/Plugins/Animation/ControlRig/Source/ControlRigEditor/Private/RigDependencyGraph/RigDependencyGraphPanelNodeFactory.cpp
@@ -0,0 +1,23 @@
+// Copyright Epic Games, Inc. All Rights Reserved.
+
+#include "RigDependencyGraph/RigDependencyGraphPanelNodeFactory.h"
+
+#include "EdGraph/EdGraphNode.h"
+#include "RigDependencyGraph/RigDependencyGraph.h"
+#include "RigDependencyGraph/RigDependencyGraphNode.h"
+#include "RigDependencyGraph/SRigDependencyGraphNode.h"
+#include "SGraphNode.h"
+#include "Templates/Casts.h"
+#include "Widgets/DeclarativeSyntaxSupport.h"
+
+TSharedPtr<SGraphNode> FRigDependencyGraphPanelNodeFactory::CreateNode(UEdGraphNode* Node) const
+{
+	if (URigDependencyGraphNode* RigDependencyGraphNode = Cast<URigDependencyGraphNode>(Node))
+	{
+		TSharedRef<SGraphNode> GraphNode = SNew(SRigDependencyGraphNode, RigDependencyGraphNode);
+		GraphNode->SlatePrepass();
+		RigDependencyGraphNode->SetDimensions(GraphNode->GetDesiredSize());
+		return GraphNode;
+	}
+	return nullptr;
+}
diff --git a/Engine/Plugins/Animation/ControlRig/Source/ControlRigEditor/Private/RigDependencyGraph/RigDependencyGraphPanelNodeFactory.h b/Engine/Plugins/Animation/ControlRig/Source/ControlRigEditor/Private/RigDependencyGraph/RigDependencyGraphPanelNodeFactory.h
new file mode 100644
--- /dev/null
+++ b/Engine/Plugins/Animation/ControlRig/Source/ControlRigEditor/Private/RigDependencyGraph/RigDependencyGraphPanelNodeFactory.h
@@ -0,0 +1,13 @@
+// Copyright Epic Games, Inc. All Rights Reserved.
+
+#pragma once
+
+#include "EdGraphUtilities.h"
+#include "Templates/SharedPointer.h"
+
+class UEdGraphNode;
+
+class FRigDependencyGraphPanelNodeFactory : public FGraphPanelNodeFactory
+{
+	virtual TSharedPtr<class SGraphNode> CreateNode(UEdGraphNode* Node) const override;
+};
diff --git a/Engine/Plugins/Animation/ControlRig/Source/ControlRigEditor/Private/RigDependencyGraph/RigDependencyGraphSchema.cpp b/Engine/Plugins/Animation/ControlRig/Source/ControlRigEditor/Private/RigDependencyGraph/RigDependencyGraphSchema.cpp
new file mode 100644
--- /dev/null
+++ b/Engine/Plugins/Animation/ControlRig/Source/ControlRigEditor/Private/RigDependencyGraph/RigDependencyGraphSchema.cpp
@@ -0,0 +1,492 @@
+// Copyright Epic Games, Inc. All Rights Reserved.
+
+#include "RigDependencyGraph/RigDependencyGraphSchema.h"
+
+#include "RigDependencyGraphNode.h"
+#include "RigDependencyGraph/RigDependencyGraph.h"
+#include "RigDependencyGraph/RigDependencyGraphNode.h"
+#include "RigDependencyGraph/RigDependencyConnectionDrawingPolicy.h"
+#include "ToolMenus.h"
+#include "Toolkits/GlobalEditorCommonCommands.h"
+
+#include UE_INLINE_GENERATED_CPP_BY_NAME(RigDependencyGraphSchema)
+
+#define LOCTEXT_NAMESPACE "RigDependencyGraphSchema"
+
+URigDependencyGraphSchema::URigDependencyGraphSchema(const FObjectInitializer& ObjectInitializer)
+	: Super(ObjectInitializer)
+{
+}
+
+FLinearColor URigDependencyGraphSchema::GetPinTypeColor(const FEdGraphPinType& PinType) const
+{
+	return FLinearColor::White;
+}
+
+void URigDependencyGraphSchema::BreakPinLinks(UEdGraphPin& TargetPin, bool bSendsNodeNotifcation) const
+{
+	// Don't allow breaking any links
+}
+
+void URigDependencyGraphSchema::BreakSinglePinLink(UEdGraphPin* SourcePin, UEdGraphPin* TargetPin) const
+{
+	// Don't allow breaking any links
+}
+
+FPinConnectionResponse URigDependencyGraphSchema::MovePinLinks(UEdGraphPin& MoveFromPin, UEdGraphPin& MoveToPin, bool bIsIntermediateMove, bool bNotifyLinkedNodes) const
+{
+	// Don't allow moving any links
+	return FPinConnectionResponse(CONNECT_RESPONSE_DISALLOW, TEXT(""));
+}
+
+FPinConnectionResponse URigDependencyGraphSchema::CopyPinLinks(UEdGraphPin& CopyFromPin, UEdGraphPin& CopyToPin, bool bIsIntermediateCopy) const
+{
+	// Don't allow copying any links
+	return FPinConnectionResponse(CONNECT_RESPONSE_DISALLOW, TEXT(""));
+}
+
+FConnectionDrawingPolicy* URigDependencyGraphSchema::CreateConnectionDrawingPolicy(int32 InBackLayerID, int32 InFrontLayerID, float InZoomFactor, const FSlateRect& InClippingRect, class FSlateWindowElementList& InDrawElements, class UEdGraph* InGraphObj) const
+{
+	return new FRigDependencyConnectionDrawingPolicy(InBackLayerID, InFrontLayerID, InZoomFactor, InClippingRect, InDrawElements);
+}
+
+URigDependencyGraphNode* URigDependencyGraphSchema::CreateGraphNode(URigDependencyGraph* InGraph, const FNodeId& InNodeId) const
+{
+	static constexpr bool bSelectNewNode = false;
+	
+	FGraphNodeCreator<URigDependencyGraphNode> GraphNodeCreator(*InGraph);
+	URigDependencyGraphNode* RigDependencyGraphNode = GraphNodeCreator.CreateNode(bSelectNewNode);
+	GraphNodeCreator.Finalize();
+
+	RigDependencyGraphNode->SetupRigDependencyNode(InNodeId);
+
+	RigDependencyGraphNode->Index = InGraph->DependencyGraphNodes.Add(RigDependencyGraphNode);
+	InGraph->NodeIdLookup.Add(InNodeId, RigDependencyGraphNode);
+
+	return RigDependencyGraphNode;
+}
+
+float URigDependencyGraphSchema_Random_Helper(int32& Seed)
+{
+	Seed = (Seed * 196314165) + 907633515;
+	union { float f; int32 i; } Result;
+	union { float f; int32 i; } Temp;
+	const float SRandTemp = 1.0f;
+	Temp.f = SRandTemp;
+	Result.i = (Temp.i & 0xff800000) | (Seed & 0x007fffff);
+	return FPlatformMath::Fractional(Result.f);
+}
+
+void URigDependencyGraphSchema::LayoutNodes(URigDependencyGraph* InGraph, int32 InNumIterations) const
+{
+	const TArray<URigDependencyGraphNode*>& Nodes = InGraph->GetNodes();
+	const int32 NumNodes = Nodes.Num();
+	if (NumNodes == 0)
+	{
+		InGraph->LayoutIterationsLeft = 0;
+		return;
+	}
+	
+	const bool bFirstTime = InGraph->LayoutIterationsTotal == InGraph->LayoutIterationsLeft;
+	bool bRandomize = bFirstTime;
+
+	if (InGraph->LayoutIterationsLeft <= 0)
+	{
+		return;
+	}
+
+	static constexpr int32 NodeMarginX = 150;
+	static constexpr int32 NodeMarginY = 90;
+	static constexpr float DefaultDisplacement = (NodeMarginX + NodeMarginY) * 0.5f;
+
+	if (bFirstTime)
+	{
+		for (const URigDependencyGraphNode* Node : Nodes)
+		{
+			// if we found a node which has no dimensions
+			// we haven't updated the panel yet - so we need to delay this.
+			if (Node->Dimensions.IsNearlyZero())
+			{
+				return;
+			}
+		}
+
+		int32 NumNodesRequiringLayout = 0;
+		
+		// compute the dependency depths
+		for (const URigDependencyGraphNode* Node : Nodes)
+		{
+			Node->DependencyDepth.Reset();
+			if (Node->FollowLayout.Get(true))
+			{
+				NumNodesRequiringLayout++;
+			}
+		}
+
+		// if all nodes need to be places - let's do an initial placement via
+		// a grid.
+		if (NumNodesRequiringLayout == NumNodes)
+		{
+			TArray<bool> VisitedFlag;
+			int32 MaxDepth = 0;
+			for (const URigDependencyGraphNode* Node : Nodes)
+			{
+				MaxDepth = FMath::Max<int32>(MaxDepth, Node->GetDependencyDepth_Impl(VisitedFlag));
+			}
+
+			TArray<TArray<URigDependencyGraphNode*>> NodesSortedByDepth;
+			NodesSortedByDepth.SetNum(MaxDepth + 1);
+
+			for (URigDependencyGraphNode* Node : Nodes)
+			{
+				NodesSortedByDepth[Node->GetDependencyDepth()].Add(Node);
+			}
+
+			int32 CurrentColumnX = 0;
+			for (const TArray<URigDependencyGraphNode*>& Column : NodesSortedByDepth)
+			{
+				// Lay out root nodes
+				float MaxWidth = 0.0f;
+				int32 CurrentYOffset = 0;
+				int32 TotalRootY = 0;
+				for (const URigDependencyGraphNode* Node : Column)
+				{
+					TotalRootY += (int32)Node->GetDimensions().Y + NodeMarginY;
+				}
+
+				for (URigDependencyGraphNode* Node : Column)
+				{
+					Node->NodePosX = CurrentColumnX;
+					Node->NodePosY = CurrentYOffset - (TotalRootY / 2);
+
+					CurrentYOffset += (int32)Node->GetDimensions().Y + NodeMarginY;
+					MaxWidth = FMath::Max(MaxWidth, Node->GetDimensions().X);
+				}
+
+				CurrentColumnX += (MaxWidth + NodeMarginX);
+			}
+		}
+		else
+		{
+			bRandomize = false;
+			
+			// for each node place it either to the left or right of a neighbor node which may already have a position
+			for (URigDependencyGraphNode* Node : Nodes)
+			{
+				if (Node->FollowLayout.Get(true))
+				{
+					const bool bUseInputPin = !Node->InputPin->LinkedTo.IsEmpty();
+
+					auto GetPin = [bUseInputPin](const URigDependencyGraphNode* InNode)
+					{
+						return bUseInputPin ? InNode->InputPin : InNode->OutputPin;
+					};
+					auto GetOtherPin = [bUseInputPin](const URigDependencyGraphNode* InNode)
+					{
+						return bUseInputPin ? InNode->OutputPin : InNode->InputPin;
+					};
+
+					if (GetPin(Node)->LinkedTo.IsEmpty())
+					{
+						continue;
+					}
+					
+					const URigDependencyGraphNode* Neighbor = CastChecked<URigDependencyGraphNode>(GetPin(Node)->LinkedTo[0]->GetOwningNode());
+					TArray<const URigDependencyGraphNode*> AllNodesForNeighor;
+					for (const UEdGraphPin* LinkedPin : GetOtherPin(Neighbor)->LinkedTo)
+					{
+						check(!LinkedPin->LinkedTo.IsEmpty());
+						if (LinkedPin->LinkedTo[0] != GetOtherPin(Neighbor))
+						{
+							continue;
+						}
+						const URigDependencyGraphNode* NeighborsNeighbor = CastChecked<URigDependencyGraphNode>(LinkedPin->GetOwningNode());
+						AllNodesForNeighor.Add(NeighborsNeighbor);
+					}
+
+					float MaxWidth = 0.f;
+					float TotalHeight = NodeMarginY * static_cast<float>(AllNodesForNeighor.Num() - 1);
+					
+					for (const URigDependencyGraphNode* NeighborsNeighbor : AllNodesForNeighor)
+					{
+						TotalHeight += NeighborsNeighbor->Dimensions.Y;
+						MaxWidth = FMath::Max(MaxWidth, NeighborsNeighbor->Dimensions.X);
+					}
+
+					Node->NodePosX = Neighbor->NodePosX;
+					if (bUseInputPin)
+					{
+						Node->NodePosX += Neighbor->Dimensions.X + NodeMarginX;
+					}
+					else
+					{
+						Node->NodePosX -= NodeMarginX + MaxWidth;
+					}
+					
+					Node->NodePosY = Neighbor->NodePosY - TotalHeight * 0.5f;
+
+					for (const URigDependencyGraphNode* NeighborsNeighbor : AllNodesForNeighor)
+					{
+						if (NeighborsNeighbor == Node)
+						{
+							break;
+						}
+						Node->NodePosY += NodeMarginY + NeighborsNeighbor->Dimensions.Y;
+					}
+				}
+			}
+		}
+	}
+
+	if (bFirstTime)
+	{
+		InGraph->LayoutEdges.Reset();
+
+		// update the island bounds
+		for (FNodeIsland& Island : InGraph->NodeIslands)
+		{
+			Island.Bounds = InGraph->GetNodesBounds(Island.NodeIds);
+		}
+
+		// compute a deterministic seed
+		int32 Seed = 0;
+		TArray<FNodeId> NodeIds;
+		NodeIds.Reserve(NumNodes);
+		for (int32 NodeIndex = 0; NodeIndex < NumNodes; ++NodeIndex)
+		{
+			NodeIds.Add(Nodes[NodeIndex]->GetNodeId());
+		}
+		NodeIds.Sort();
+		for (const FNodeId& NodeId : NodeIds)
+		{
+			Seed += static_cast<int32>(GetTypeHash(NodeId));
+		}
+
+		// pseudo-randomize the initial positions
+		for (int32 NodeIndex = 0; NodeIndex < NumNodes; ++NodeIndex)
+		{
+			URigDependencyGraphNode* Node = Nodes[NodeIndex];
+			Node->LayoutPosition = FVector2D(Node->GetPosition()) + Node->Dimensions * 0.5f;
+
+			// always consume the random number to ensure the distribution is deterministic
+			const FVector2D RandomDisplacement(
+				URigDependencyGraphSchema_Random_Helper(Seed) * 2.f - 1.f,
+				URigDependencyGraphSchema_Random_Helper(Seed) * 2.f - 1.f);
+			
+			if (bRandomize && Node->FollowLayout.Get(true))
+			{
+				Node->LayoutPosition.X += RandomDisplacement.X * DefaultDisplacement * 0.01f;
+				Node->LayoutPosition.Y += RandomDisplacement.Y * DefaultDisplacement;
+			}
+			
+			Node->LayoutVelocity = FVector2D::ZeroVector;
+			Node->LayoutForce = FVector2D::ZeroVector;
+
+			for (int32 LinkedToIndex = 0; LinkedToIndex < Nodes[NodeIndex]->OutputPin->LinkedTo.Num(); LinkedToIndex++)
+			{
+				UEdGraphPin* LinkedPin = Node->OutputPin->LinkedTo[LinkedToIndex];
+				if (const URigDependencyGraphNode* LinkedNode = Cast<URigDependencyGraphNode>(LinkedPin->GetOwningNode()))
+				{
+					if (ensure(LinkedNode->Index != NodeIndex))
+					{
+						const URigDependencyGraphNode* NodeA = Nodes[NodeIndex];
+						const URigDependencyGraphNode* NodeB = LinkedNode;
+
+						const float DesiredDistance = NodeMarginX + (NodeA->Dimensions.X + NodeB->Dimensions.X) * 0.5f;
+
+						auto ComputeDesiredY = [](const int32 Index, const int32 Num, const float NodeHeight) -> float
+						{
+							if (Num <= 1)
+							{
+								return 0.f;
+							}
+							const float NumVertical = static_cast<float>(FMath::Max(Num - 1, 0));
+							const float Ratio = static_cast<float>(Index) / NumVertical;
+							const float HalfTotalHeight = NumVertical * (NodeHeight + NodeMarginY) * 0.5f;
+							return FMath::Lerp(-HalfTotalHeight, HalfTotalHeight, Ratio);
+						};
+
+						const int32 IndexInStackA = NodeB->InputPin->LinkedTo.IndexOfByKey(NodeA->OutputPin);
+						const int32 SizeOfStackA = NodeB->InputPin->LinkedTo.Num();
+						const float DesiredHeightA = ComputeDesiredY(IndexInStackA, SizeOfStackA, NodeB->Dimensions.Y);
+						const FVector2D AnchorA = FVector2D(-DesiredDistance, DesiredHeightA);
+
+						const int32 IndexInStackB = NodeA->OutputPin->LinkedTo.IndexOfByKey(NodeB->InputPin);
+						const int32 SizeOfStackB = NodeA->OutputPin->LinkedTo.Num();
+						const float DesiredHeightB = ComputeDesiredY(IndexInStackB, SizeOfStackB, NodeA->Dimensions.Y);
+						const FVector2D AnchorB = FVector2D(DesiredDistance, DesiredHeightB);
+
+						float Strength = 1.f;
+						if (NodeA->GetNodeId().IsInstruction() || NodeB->GetNodeId().IsInstruction())
+						{
+							Strength = 2.f;
+						}
+
+						const FLayoutEdge Edge(NodeIndex, LinkedNode->GetIndex(), AnchorA, AnchorB, Strength);
+						InGraph->LayoutEdges.Add(Edge);
+					}
+				}
+			}
+		}
+	}
+
+	// now layout the graph using Fruchterman Reingold
+	for (int32 Iteration = 0; Iteration < InNumIterations; ++Iteration, InGraph->LayoutIterationsLeft--)
+	{
+		const float Ratio = static_cast<float>(InGraph->LayoutIterationsLeft) / static_cast<float>(InGraph->LayoutIterationsTotal);
+		const float ForceMagnitude = 0.15f * Ratio;
+
+		// compute repulsive forces, pairwise between all verts
+		int32 NumSimulatedNodes = 0;
+		for (int32 NodeIndexA = 0; NodeIndexA < NumNodes; ++NodeIndexA)
+		{
+			if (Nodes[NodeIndexA]->FollowLayout.Get(true))
+			{
+				NumSimulatedNodes++;
+			}
+			
+			const FVector2D PositionA = Nodes[NodeIndexA]->LayoutPosition;
+
+			for (int32 NodeIndexB = NodeIndexA + 1; NodeIndexB < NumNodes; ++NodeIndexB)
+			{
+				const FVector2D PositionB = Nodes[NodeIndexB]->LayoutPosition;
+				const FVector2D AtoB = PositionB - PositionA;
+				const FVector2D AtoBNormal = AtoB.GetSafeNormal().GetAbs();
+
+				const float DesiredDistanceX = (Nodes[NodeIndexA]->Dimensions.X + Nodes[NodeIndexB]->Dimensions.X) * AtoBNormal.X;
+				const float DesiredDistanceY = (Nodes[NodeIndexA]->Dimensions.Y + Nodes[NodeIndexB]->Dimensions.Y) * AtoBNormal.Y;
+				const float Margin = NodeMarginX * AtoBNormal.X + NodeMarginY * AtoBNormal.Y;
+				const float DesiredDistance = 1.5f * Margin + 0.5f * (DesiredDistanceX + DesiredDistanceY);
+
+				const float AtoBDistance = AtoB.Length();
+				const float DeltaDistance = AtoBDistance - DesiredDistance;
+				if (AtoBDistance > SMALL_NUMBER && DeltaDistance < -SMALL_NUMBER)
+				{
+					FVector2D Force = ForceMagnitude * DeltaDistance * AtoB / AtoBDistance;
+					Nodes[NodeIndexA]->LayoutForce += Force;
+					Nodes[NodeIndexB]->LayoutForce += -Force;
+				}
+			}
+		}
+
+		// compute attractive forces along edges
+		for (const FLayoutEdge& Edge : InGraph->LayoutEdges)
+		{
+			URigDependencyGraphNode* NodeA = Nodes[Edge.NodeA];
+			URigDependencyGraphNode* NodeB = Nodes[Edge.NodeB];
+
+			const FVector2D A = NodeA->LayoutPosition;
+			const FVector2D B = NodeB->LayoutPosition;
+			const FVector2D DesiredA = B + Edge.AnchorA;
+			const FVector2D DesiredB = A + Edge.AnchorB;
+			const FVector2D DirectionA = DesiredA - A;
+			const FVector2D DirectionB = DesiredB - B;
+			NodeA->LayoutForce += Edge.Strength * ForceMagnitude * DirectionA;
+			NodeB->LayoutForce += Edge.Strength * ForceMagnitude * DirectionB;
+		}
+
+		for (FNodeIsland& Island : InGraph->NodeIslands)
+		{
+			bool bHasSimulatedNodes = false;
+			for (const FNodeId& NodeId : Island.NodeIds)
+			{
+				if (const URigDependencyGraphNode* Node = InGraph->FindNode(NodeId))
+				{
+					if (Node->FollowLayout.Get(true))
+					{
+						bHasSimulatedNodes = true;
+						break;
+					}
+				}
+			}
+
+			if (bHasSimulatedNodes)
+			{
+				const FBox2D CurrentBounds = InGraph->GetNodesBounds(Island.NodeIds);
+				const FVector2D MaintainCenterForce = Island.Bounds.GetCenter() - CurrentBounds.GetCenter();
+				if (!MaintainCenterForce.IsNearlyZero())
+				{
+					for (const FNodeId& NodeId : Island.NodeIds)
+					{
+						if (URigDependencyGraphNode* Node = InGraph->FindNode(NodeId))
+						{
+							if (Node->FollowLayout.Get(true))
+							{
+								Node->LayoutForce += MaintainCenterForce;
+							}
+						}
+					}
+				}
+			}
+		}
+
+		// compute displacements
+		// we treat the "forces" as displacements directly
+		for (int32 NodeIndex = 0; NodeIndex < NumNodes; ++NodeIndex)
+		{
+			Nodes[NodeIndex]->LayoutVelocity = FMath::Lerp(Nodes[NodeIndex]->LayoutVelocity, Nodes[NodeIndex]->LayoutForce, 0.1);
+			Nodes[NodeIndex]->LayoutForce = FVector2D::ZeroVector;
+
+			if (Nodes[NodeIndex]->FollowLayout.Get(true))
+			{
+				Nodes[NodeIndex]->LayoutPosition += Nodes[NodeIndex]->LayoutVelocity;
+				Nodes[NodeIndex]->NodePosX = (float)(Nodes[NodeIndex]->LayoutPosition.X - Nodes[NodeIndex]->Dimensions.X * 0.5);
+				Nodes[NodeIndex]->NodePosY = (float)(Nodes[NodeIndex]->LayoutPosition.Y - Nodes[NodeIndex]->Dimensions.Y * 0.5);
+			}
+
+			Nodes[NodeIndex]->LayoutVelocity *= 0.9;
+
+			// stop the simulation once the node has converged
+			if (Nodes[NodeIndex]->LayoutVelocity.Length() < KINDA_SMALL_NUMBER)
+			{
+				Nodes[NodeIndex]->FollowLayout = false;
+			}
+		}
+
+		// if no more nodes are simulating
+		// the simulation has converged.
+		if (NumSimulatedNodes == 0)
+		{
+			InGraph->LayoutIterationsLeft = 0;
+			break;
+		}
+	}
+
+	if (InGraph->LayoutIterationsLeft <= 0)
+	{
+		for (URigDependencyGraphNode* DependencyGraphNode : InGraph->DependencyGraphNodes)
+		{
+			DependencyGraphNode->FollowLayout = false;
+		}
+	}
+}
+
+const FPinConnectionResponse URigDependencyGraphSchema::CanCreateConnection(const UEdGraphPin* A, const UEdGraphPin* B) const
+{
+	return FPinConnectionResponse(CONNECT_RESPONSE_DISALLOW, LOCTEXT("TheDependencyGraphIsReadOnly", "The dependency graph is read only."));
+}
+
+FPinConnectionResponse URigDependencyGraphSchema::CanCreateNewNodes(UEdGraphPin* InSourcePin) const
+{
+	return FPinConnectionResponse(CONNECT_RESPONSE_DISALLOW, LOCTEXT("TheDependencyGraphIsReadOnly", "The dependency graph is read only."));
+}
+
+bool URigDependencyGraphSchema::SupportsDropPinOnNode(UEdGraphNode* InTargetNode, const FEdGraphPinType& InSourcePinType, EEdGraphPinDirection InSourcePinDirection, FText& OutErrorMessage) const
+{
+	OutErrorMessage = LOCTEXT("TheDependencyGraphIsReadOnly", "The dependency graph is read only.");
+	return false;
+}
+
+void URigDependencyGraphSchema::SetNodePosition(UEdGraphNode* Node, const FVector2f& Position) const
+{
+	Super::SetNodePosition(Node, Position);
+
+	if (URigDependencyGraphNode* DependencyGraphNode = Cast<URigDependencyGraphNode>(Node))
+	{
+		DependencyGraphNode->LayoutPosition.X = DependencyGraphNode->NodePosX;
+		DependencyGraphNode->LayoutPosition.Y = DependencyGraphNode->NodePosY;
+		DependencyGraphNode->FollowLayout = false;
+	}
+}
+
+#undef LOCTEXT_NAMESPACE
diff --git a/Engine/Plugins/Animation/ControlRig/Source/ControlRigEditor/Private/RigDependencyGraph/RigDependencyGraphSchema.h b/Engine/Plugins/Animation/ControlRig/Source/ControlRigEditor/Private/RigDependencyGraph/RigDependencyGraphSchema.h
new file mode 100644
--- /dev/null
+++ b/Engine/Plugins/Animation/ControlRig/Source/ControlRigEditor/Private/RigDependencyGraph/RigDependencyGraphSchema.h
@@ -0,0 +1,43 @@
+// Copyright Epic Games, Inc. All Rights Reserved.
+
+#pragma once
+
+#include "RigDependencyGraph.h"
+#include "EdGraph/EdGraphSchema.h"
+#include "Rigs/RigHierarchy.h"
+#include "RigDependencyGraphNode.h"
+#include "RigDependencyGraphSchema.generated.h"
+
+class URigDependencyGraph;
+class URigDependencyGraphNode;
+
+UCLASS()
+class URigDependencyGraphSchema : public UEdGraphSchema
+{
+	GENERATED_UCLASS_BODY()
+
+public:
+	// UEdGraphSchema interface
+	virtual FName GetParentContextMenuName() const override { return NAME_None; }
+	virtual FLinearColor GetPinTypeColor(const FEdGraphPinType& PinType) const override;
+	virtual void BreakPinLinks(UEdGraphPin& TargetPin, bool bSendsNodeNotifcation) const override;
+	virtual void BreakSinglePinLink(UEdGraphPin* SourcePin, UEdGraphPin* TargetPin) const override;
+	virtual FPinConnectionResponse MovePinLinks(UEdGraphPin& MoveFromPin, UEdGraphPin& MoveToPin, bool bIsIntermediateMove = false, bool bNotifyLinkedNodes = false) const override;
+	virtual FPinConnectionResponse CopyPinLinks(UEdGraphPin& CopyFromPin, UEdGraphPin& CopyToPin, bool bIsIntermediateCopy = false) const override;
+	virtual class FConnectionDrawingPolicy* CreateConnectionDrawingPolicy(int32 InBackLayerID, int32 InFrontLayerID, float InZoomFactor, const FSlateRect& InClippingRect, class FSlateWindowElementList& InDrawElements, class UEdGraph* InGraphObj) const override;
+	virtual bool ShouldAlwaysPurgeOnModification() const override { return false; }
+	virtual FPinConnectionResponse CanCreateNewNodes(UEdGraphPin* InSourcePin) const override;
+	virtual const FPinConnectionResponse CanCreateConnection(const UEdGraphPin* A, const UEdGraphPin* B) const override;
+	virtual bool SupportsDropPinOnNode(UEdGraphNode* InTargetNode, const FEdGraphPinType& InSourcePinType, EEdGraphPinDirection InSourcePinDirection, FText& OutErrorMessage) const override;
+	virtual void SetNodePosition(UEdGraphNode* Node, const FVector2f& Position) const override;
+	// End of UEdGraphSchema interface
+
+	virtual URigDependencyGraphNode* CreateGraphNode(URigDependencyGraph* InGraph, const URigDependencyGraphNode::FNodeId& InNodeId) const;
+
+	virtual void LayoutNodes(URigDependencyGraph* InGraph, int32 InIterations) const;
+
+	using FNodeId = URigDependencyGraphNode::FNodeId;
+	using FLayoutEdge = URigDependencyGraph::FLayoutEdge;
+	using FNodeIsland = URigDependencyGraph::FNodeIsland;
+};
+
diff --git a/Engine/Plugins/Animation/ControlRig/Source/ControlRigEditor/Private/RigDependencyGraph/RigDependencyGraphSummoner.cpp b/Engine/Plugins/Animation/ControlRig/Source/ControlRigEditor/Private/RigDependencyGraph/RigDependencyGraphSummoner.cpp
new file mode 100644
--- /dev/null
+++ b/Engine/Plugins/Animation/ControlRig/Source/ControlRigEditor/Private/RigDependencyGraph/RigDependencyGraphSummoner.cpp
@@ -0,0 +1,44 @@
+// Copyright Epic Games, Inc. All Rights Reserved.
+
+#include "RigDependencyGraph/RigDependencyGraphSummoner.h"
+#include "IDocumentation.h"
+#include "Editor/ControlRigEditor.h"
+#include "RigDependencyGraph/SRigDependencyGraph.h"
+#include "Styling/AppStyle.h"
+
+#define LOCTEXT_NAMESPACE "RigDependencyGraphSummoner"
+
+FRigDependencyGraphSummoner::FRigDependencyGraphSummoner(const TSharedRef<IControlRigBaseEditor>& InControlRigEditor)
+	: FWorkflowTabFactory("RigDependencyGraphView", InControlRigEditor->GetHostingApp())
+	, ControlRigEditor(InControlRigEditor)
+{
+	TabLabel = LOCTEXT("RigDependencyViewer", "Dependency Viewer");
+	TabIcon = FSlateIcon(FAppStyle::GetAppStyleSetName(), "GraphEditor.Event_16x");
+
+	EnableTabPadding();
+	bIsSingleton = false;
+
+	ViewMenuDescription = TabLabel;
+	ViewMenuTooltip = LOCTEXT("RigDependencyGraphView_ToolTip", "Shows the Control Rig Dependency Viewer");
+}
+
+TSharedPtr<SToolTip> FRigDependencyGraphSummoner::CreateTabToolTipWidget(const FWorkflowTabSpawnInfo& Info) const
+{
+	return IDocumentation::Get()->CreateToolTip(LOCTEXT("RigDependencyGraphToolTip", "The Control Rig Dependency Viewer tab lets you debug the relationships caused by your Control Rig."), NULL, TEXT("Shared/Editors/ControlRigEditor"), TEXT("RigDependencyGraph_Window"));
+}
+
+TSharedRef<SDockTab> FRigDependencyGraphSummoner::SpawnTab(const FWorkflowTabSpawnInfo& Info) const
+{
+	TSharedRef<SDockTab> Tab = FWorkflowTabFactory::SpawnTab(Info);
+	Tab->SetTabLabelOverflowPolicy(TOptional<ETextOverflowPolicy>());
+	return Tab;
+}
+
+TSharedRef<SWidget> FRigDependencyGraphSummoner::CreateTabBody(const FWorkflowTabSpawnInfo& Info) const
+{
+	check(ControlRigEditor.IsValid());
+	TSharedRef<SRigDependencyGraph> RigDependencyGraph = SNew(SRigDependencyGraph, ControlRigEditor.Pin().ToSharedRef());
+	return RigDependencyGraph;
+}
+
+#undef LOCTEXT_NAMESPACE
diff --git a/Engine/Plugins/Animation/ControlRig/Source/ControlRigEditor/Private/RigDependencyGraph/RigDependencyGraphSummoner.h b/Engine/Plugins/Animation/ControlRig/Source/ControlRigEditor/Private/RigDependencyGraph/RigDependencyGraphSummoner.h
new file mode 100644
--- /dev/null
+++ b/Engine/Plugins/Animation/ControlRig/Source/ControlRigEditor/Private/RigDependencyGraph/RigDependencyGraphSummoner.h
@@ -0,0 +1,20 @@
+// Copyright Epic Games, Inc. All Rights Reserved.
+
+#pragma once
+
+#include "WorkflowOrientedApp/WorkflowTabFactory.h"
+#include "RigDependencyGraph/SRigDependencyGraph.h"
+
+struct FRigDependencyGraphSummoner : public FWorkflowTabFactory
+{
+public:
+	FRigDependencyGraphSummoner(const TSharedRef<IControlRigBaseEditor>& InControlRigEditor);
+
+	/** FWorkflowTabFactory interface */
+	virtual TSharedRef<SWidget> CreateTabBody(const FWorkflowTabSpawnInfo& Info) const override;
+	virtual TSharedPtr<SToolTip> CreateTabToolTipWidget(const FWorkflowTabSpawnInfo& Info) const override;
+	virtual TSharedRef<SDockTab> SpawnTab(const FWorkflowTabSpawnInfo& Info) const override;
+
+protected:
+	TWeakPtr<IControlRigBaseEditor> ControlRigEditor;
+};
diff --git a/Engine/Plugins/Animation/ControlRig/Source/ControlRigEditor/Private/RigDependencyGraph/SRigDependencyGraph.cpp b/Engine/Plugins/Animation/ControlRig/Source/ControlRigEditor/Private/RigDependencyGraph/SRigDependencyGraph.cpp
new file mode 100644
--- /dev/null
+++ b/Engine/Plugins/Animation/ControlRig/Source/ControlRigEditor/Private/RigDependencyGraph/SRigDependencyGraph.cpp
@@ -0,0 +1,803 @@
+// Copyright Epic Games, Inc. All Rights Reserved.
+
+#include "RigDependencyGraph/SRigDependencyGraph.h"
+#include "GraphEditor.h"
+#include "Toolkits/GlobalEditorCommonCommands.h"
+
+#include "ObjectTools.h"
+#include "Engine/Selection.h"
+#include "RigDependencyGraph/RigDependencyGraph.h"
+#include "RigDependencyGraph/RigDependencyGraphSchema.h"
+#include "RigDependencyGraph/RigDependencyGraphCommands.h"
+#include "Widgets/Input/SSearchBox.h"
+#include "Widgets/Input/SButton.h"
+#include "Widgets/Input/SCheckBox.h"
+#include "Widgets/Layout/SSpacer.h"
+#include "Widgets/Notifications/SProgressBar.h"
+#include "Algo/Transform.h"
+#include "Editor/ControlRigEditor.h"
+#include "Modules/ModuleManager.h"
+#include "Framework/Application/SlateApplication.h"
+#include "Framework/MultiBox/MultiBoxBuilder.h"
+#include "Rigs/RigHierarchyController.h"
+#include "Editor/RigVMEditorStyle.h"
+
+#define LOCTEXT_NAMESPACE "RigDependencyGraph"
+
+SRigDependencyGraph::~SRigDependencyGraph()
+{
+	if (!GExitPurge)
+	{
+		if ( ensure(GraphObj) )
+		{
+			GraphObj->RemoveFromRoot();
+		}		
+	}
+}
+
+void SRigDependencyGraph::Construct(const FArguments& InArgs, const TSharedRef<IControlRigBaseEditor>& InControlRigEditor)
+{
+	WeakControlRigEditor = InControlRigEditor;
+	FlashLightRadiusAttribute = InArgs._FlashLightRadius;
+
+	// Create the graph
+	GraphObj = NewObject<URigDependencyGraph>();
+	GraphObj->Schema = URigDependencyGraphSchema::StaticClass();
+	GraphObj->AddToRoot();
+	GraphObj->Initialize(InControlRigEditor, InControlRigEditor->GetControlRig());
+
+	BindCommands();
+
+	FGraphAppearanceInfo AppearanceInfo;
+	AppearanceInfo.CornerText = LOCTEXT("AppearanceCornerText_DependencyGraph", "Debug");
+
+	SGraphEditor::FGraphEditorEvents GraphEvents;
+	GraphEvents.OnCreateActionMenuAtLocation = SGraphEditor::FOnCreateActionMenuAtLocation::CreateSP(this, &SRigDependencyGraph::OnCreateGraphActionMenu);
+	GraphEvents.OnSelectionChanged = SGraphEditor::FOnSelectionChanged::CreateSP(this, &SRigDependencyGraph::HandleSelectionChanged);
+	GraphEvents.OnNodeDoubleClicked = FSingleNodeEvent::CreateSP(this, &SRigDependencyGraph::HandleNodeDoubleClicked);
+
+	// Create the graph editor
+	GraphEditor = SNew(SGraphEditor)
+		.GraphToEdit(GraphObj)
+		.GraphEvents(GraphEvents)
+		.Appearance(AppearanceInfo)
+		.ShowGraphStateOverlay(false);
+
+	static const FName DefaultForegroundName("DefaultForeground");
+
+	ChildSlot
+	[
+		SNew(SVerticalBox)
+		
+		+SVerticalBox::Slot()
+		.AutoHeight()
+		.HAlign(HAlign_Fill)
+		[
+			SNew(SHorizontalBox)
+			
+			+SHorizontalBox::Slot()
+			.AutoWidth()
+			.HAlign(HAlign_Left)
+			.VAlign(VAlign_Center)
+			.Padding(0.0f, 0.0f, 2.0f, 0.0f)
+			[
+				SNew(SComboButton)
+				.Visibility(EVisibility::Visible)
+				.ComboButtonStyle(&FAppStyle::Get().GetWidgetStyle<FComboButtonStyle>("SimpleComboButtonWithIcon"))
+				.ForegroundColor(FSlateColor::UseStyle())
+				.ContentPadding(0.0f)
+				.OnGetMenuContent(this, &SRigDependencyGraph::CreateFilterMenu)
+				.ButtonContent()
+				[
+					SNew(SImage)
+					.Image(FAppStyle::Get().GetBrush("Icons.Filter"))
+					.ColorAndOpacity(FSlateColor::UseForeground())
+				 ]
+			]
+
+			+SHorizontalBox::Slot()
+			.AutoWidth()
+			.HAlign(HAlign_Left)
+			.VAlign(VAlign_Center)
+			.Padding(0.0f, 0.0f, 2.0f, 0.0f)
+			[
+				SNew(SCheckBox)
+				.Padding(FMargin(8.0f, 2.0f, 8.0f, 2.0f))
+				.Style(FAppStyle::Get(), "ToggleButtonCheckbox")
+				.ToolTipText(LOCTEXT("SyncContentTooltip", "Sync the view with the selection in the hierarchy and / or rig graphs"))
+				.IsChecked_Lambda( [this] () -> ECheckBoxState
+				{
+					return GraphObj->bLockContent ? ECheckBoxState::Unchecked : ECheckBoxState::Checked;
+				})
+				.OnCheckStateChanged_Lambda([this](ECheckBoxState InCheckBoxState)
+				{
+					GraphObj->bLockContent = InCheckBoxState == ECheckBoxState::Unchecked;
+				})
+				[
+					SNew(SImage)
+					.ColorAndOpacity(FSlateColor::UseForeground())
+					.DesiredSizeOverride(FVector2D(20, 20))
+					.Image_Lambda([this]() -> const FSlateBrush*
+					{
+						static const FSlateBrush* Unlocked = FAppStyle::GetBrush("Icons.Unlock");
+						static const FSlateBrush* Locked = FAppStyle::GetBrush("Icons.Lock");
+						return GraphObj->bLockContent ? Locked : Unlocked;
+					})
+				]
+			]
+
+			+SHorizontalBox::Slot()
+			.AutoWidth()
+			.HAlign(HAlign_Left)
+			.VAlign(VAlign_Center)
+			.Padding(0.0f, 0.0f, 2.0f, 0.0f)
+			[
+				SNew(SCheckBox)
+				.Padding(FMargin(8.0f, 2.0f, 8.0f, 2.0f))
+				.Style(FAppStyle::Get(), "ToggleButtonCheckbox")
+				.ToolTipText(LOCTEXT("ShowParentChildRelationShips", "Show the parent-child relationships"))
+				.IsChecked_Lambda( [this] () -> ECheckBoxState
+				{
+					return GraphObj->bFollowParentRelationShips ? ECheckBoxState::Checked : ECheckBoxState::Unchecked;
+				})
+				.OnCheckStateChanged_Lambda([this](ECheckBoxState InCheckBoxState)
+				{
+					GraphObj->SetFollowParentRelationShips(InCheckBoxState != ECheckBoxState::Unchecked);
+				})
+				[
+					SNew(SImage)
+					.ColorAndOpacity(FSlateColor::UseForeground())
+					.DesiredSizeOverride(FVector2D(20,20))
+					.Image(FAppStyle::Get().GetBrush("LevelEditor.CreateBlankBlueprintClass"))
+				]
+			]
+			
+			+SHorizontalBox::Slot()
+			.AutoWidth()
+			.HAlign(HAlign_Left)
+			.VAlign(VAlign_Center)
+			.Padding(0.0f, 0.0f, 2.0f, 0.0f)
+			[
+				SNew(SCheckBox)
+				.Padding(FMargin(8.0f, 2.0f, 8.0f, 2.0f))
+				.Style(FAppStyle::Get(), "ToggleButtonCheckbox")
+				.ToolTipText(LOCTEXT("ShowVMRelationShips", "Show the relationships introduced by the rig nodes"))
+				.IsChecked_Lambda( [this] () -> ECheckBoxState
+				{
+					return GraphObj->bFollowVMRelationShips ? ECheckBoxState::Checked : ECheckBoxState::Unchecked;
+				})
+				.OnCheckStateChanged_Lambda([this](ECheckBoxState InCheckBoxState)
+				{
+					GraphObj->SetFollowVMRelationShips(InCheckBoxState != ECheckBoxState::Unchecked);
+				})
+				[
+					SNew(SImage)
+					.ColorAndOpacity(FSlateColor::UseForeground())
+					.DesiredSizeOverride(FVector2D(20, 20))
+					.Image(FRigVMEditorStyle::Get().GetBrush("RigVM.ReleaseMode"))
+				]
+			]
+
+			+SHorizontalBox::Slot()
+			.AutoWidth()
+			.VAlign(VAlign_Center)
+			.Padding(0.0f, 0.0f, 2.0f, 0.0f)
+			[
+				SNew(SButton)
+				.ButtonStyle(FAppStyle::Get(), "SimpleButton")
+				.OnClicked_Lambda([this]() -> FReply
+				{
+					GraphObj->RemoveUnrelatedNodes();
+					return FReply::Handled();
+				})
+				.ToolTipText(LOCTEXT("HideUnrelatedTooltip", "Removes unrelated nodes from the view"))
+				[
+					SNew(SImage)
+					.DesiredSizeOverride(FVector2D(20, 20))
+					.Image(FControlRigEditorStyle::Get().GetBrush("ControlRig.DependencyGraph.HideUnrelated"))
+				]
+			]
+
+			+SHorizontalBox::Slot()
+			.AutoWidth()
+			.VAlign(VAlign_Center)
+			.Padding(0.0f, 0.0f, 2.0f, 0.0f)
+			[
+				SNew(SButton)
+				.ButtonStyle(FAppStyle::Get(), "SimpleButton")
+				.OnClicked_Lambda([this]() -> FReply
+				{
+					GraphObj->RemoveUnselectedNodes();
+					return FReply::Handled();
+				})
+				.ToolTipText(LOCTEXT("IsolateSelectionTooltip", "Isolates the selection by removing unselected nodes"))
+				[
+					SNew(SImage)
+					.DesiredSizeOverride(FVector2D(20, 20))
+					.Image(FControlRigEditorStyle::Get().GetBrush("ControlRig.DependencyGraph.IsolateSelection"))
+				]
+			]
+
+			+SHorizontalBox::Slot()
+			.MinWidth(100)
+			.MaxWidth(300)
+			.VAlign(VAlign_Center)
+			.Padding(0.0f, 0.0f, 2.0f, 0.0f)
+			[
+				SAssignNew(FilterBox, SSearchBox)
+				.OnTextChanged(this, &SRigDependencyGraph::OnSearchBarTextChanged)
+			]
+
+			+SHorizontalBox::Slot()
+			.AutoWidth()
+			.VAlign(VAlign_Center)
+			.Padding(0.0f, 0.0f, 2.0f, 0.0f)
+			[
+				SNew(SButton)
+				.ButtonStyle(FAppStyle::Get(), "SimpleButton")
+				.OnClicked_Lambda([this]() -> FReply
+				{
+					OnSearchMatchPicked(MatchIndex.Get(0) - 1);
+					return FReply::Handled();
+				})
+				.Visibility_Lambda([this]() -> EVisibility
+				{
+					return NodesMatchingSearch.IsEmpty() ? EVisibility::Collapsed : EVisibility::Visible;
+				})
+				.IsEnabled_Lambda([this]() -> bool
+				{
+					return MatchIndex.Get(0) > 0;
+				})
+				.ToolTipText(LOCTEXT("SelectPreviousMatch", "Select the previous node matching the search"))
+				[
+					SNew(SImage)
+					.DesiredSizeOverride(FVector2D(16, 16))
+					.Image(FAppStyle::Get().GetBrush("Animation.Backward_Step"))
+				]
+			]
+
+			+SHorizontalBox::Slot()
+			.AutoWidth()
+			.VAlign(VAlign_Center)
+			.Padding(0.0f, 0.0f, 2.0f, 0.0f)
+			[
+				SNew(SButton)
+				.ButtonStyle(FAppStyle::Get(), "SimpleButton")
+				.OnClicked_Lambda([this]() -> FReply
+				{
+					OnSearchMatchPicked(MatchIndex.Get(0) + 1);
+					return FReply::Handled();
+				})
+				.Visibility_Lambda([this]() -> EVisibility
+				{
+					return NodesMatchingSearch.IsEmpty() ? EVisibility::Collapsed : EVisibility::Visible;
+				})
+				.IsEnabled_Lambda([this]() -> bool
+				{
+					return MatchIndex.Get(0) < NodesMatchingSearch.Num() - 1;
+				})
+				.ToolTipText(LOCTEXT("SelectNextMatch", "Select the next node matching the search"))
+				[
+					SNew(SImage)
+					.DesiredSizeOverride(FVector2D(16, 16))
+					.Image(FAppStyle::Get().GetBrush("Animation.Forward_Step"))
+				]
+			]
+
+			+SHorizontalBox::Slot()
+			.VAlign(VAlign_Center)
+			.HAlign(HAlign_Fill)
+			.Padding(4.f, 0.f, 4.f, 0.f)
+			[
+				SAssignNew(LayoutProgressBar, SProgressBar)
+				.Visibility_Lambda([this]() -> EVisibility
+				{
+					return GraphObj->LayoutIterationsLeft > 0 ? EVisibility::Visible : EVisibility::Hidden;
+				})
+				.Percent_Lambda([this]() -> TOptional<float>
+				{
+					if (GraphObj->LayoutIterationsLeft > 0)
+					{
+						const float Left = static_cast<float>(GraphObj->LayoutIterationsLeft);
+						const float Total = static_cast<float>(GraphObj->LayoutIterationsTotal);
+						if (Total > SMALL_NUMBER)
+						{
+							return (Total - Left) / Total;
+						}
+					}
+					return 0.f;
+				})
+			]
+
+			+SHorizontalBox::Slot()
+			.AutoWidth()
+			.VAlign(VAlign_Center)
+			.Padding(4)
+			[
+				SNew(SButton)
+				.ButtonStyle(FAppStyle::Get(), "SimpleButton")
+				.Visibility_Lambda([this]() -> EVisibility
+				{
+					return GraphObj->LayoutIterationsLeft > 0 ? EVisibility::Visible : EVisibility::Hidden;
+				})
+				.OnClicked_Lambda([this]() -> FReply
+				{
+					GraphObj->CancelLayout();
+					return FReply::Handled();
+				})
+				.ToolTipText(LOCTEXT("CancelLayoutTooltip", "Cancels the layout simulation"))
+				[
+					SNew(SImage)
+					.DesiredSizeOverride(FVector2D(16, 16))
+					.Image(FAppStyle::Get().GetBrush("Icons.X"))
+				]
+			]
+		]
+
+		+SVerticalBox::Slot()
+		.FillHeight(1.f)
+		[
+			GraphEditor.ToSharedRef()
+		]
+	];
+
+	GEditor->RegisterForUndo(this);
+
+	GraphObj->OnClearSelection.BindSP(this, &SRigDependencyGraph::OnClearNodeSelection);
+}
+
+FActionMenuContent SRigDependencyGraph::OnCreateGraphActionMenu(UEdGraph* InGraph, const FVector2f& InNodePosition, const TArray<UEdGraphPin*>& InDraggedPins, bool bAutoExpand, SGraphEditor::FActionMenuClosed InOnMenuClosed)
+{
+	/*
+	if(InDraggedPins.Num() > 0 && InDraggedPins[0]->Direction == EGPD_Output)
+	{
+		if(URigDependencyGraphNode_Bone* BodyNode = Cast<URigDependencyGraphNode_Bone>(InDraggedPins[0]->GetOwningNode()))
+		{
+			FMenuBuilder MenuBuilder(true, nullptr);
+			TSharedPtr<ISkeletonTree> SkeletonTree = GraphObj->GetRigDependencyEditor()->BuildMenuWidgetNewConstraintForBody(MenuBuilder, BodyNode->BodyIndex, InOnMenuClosed);
+			return FActionMenuContent(MenuBuilder.MakeWidget(), SkeletonTree->GetSearchWidget());
+		}
+	}
+
+	InOnMenuClosed.ExecuteIfBound();
+	*/
+	return FActionMenuContent();
+}
+
+void SRigDependencyGraph::HandleNodeDoubleClicked(UEdGraphNode* InNode)
+{
+	const URigDependencyGraphNode* RigDependencyGraphNode = Cast<URigDependencyGraphNode>(InNode);
+	if (!RigDependencyGraphNode)
+	{
+		return;
+	}
+
+	const TArray<FNodeId> NodesToExpand = GraphObj->GetNodesToExpand({RigDependencyGraphNode->GetNodeId()});
+	if (!NodesToExpand.IsEmpty())
+	{
+		TArray<FNodeId> AllNodeIds;
+		for (const TObjectPtr<URigDependencyGraphNode>& ExistingDependencyGraphNode : GraphObj->DependencyGraphNodes)
+		{
+			AllNodeIds.Add(ExistingDependencyGraphNode->GetNodeId());
+		}
+		for (const FNodeId& NodeId : NodesToExpand)
+		{
+			AllNodeIds.AddUnique(NodeId);
+		}
+		GraphObj->ZoomAndFitDuringLayout = false;
+		GraphObj->ConstructNodes(AllNodeIds, false);
+		return;
+	}
+
+	// if this is an instruction - also navigate to the instruction
+	const bool bIsShiftPressed = FSlateApplication::Get().GetModifierKeys().IsShiftDown();
+	if (RigDependencyGraphNode->GetNodeId().IsInstruction() && bIsShiftPressed)
+	{
+		TSharedPtr<IControlRigBaseEditor> ControlRigEditor = WeakControlRigEditor.Pin();
+		if (ControlRigEditor.IsValid())
+		{
+			if (UControlRigBlueprint* ControlRigBlueprint = ControlRigEditor->GetControlRigBlueprint())
+			{
+				if (FRigVMClient* Client = ControlRigBlueprint->GetRigVMClient())
+				{
+					if (URigVMController* Controller = Client->GetOrCreateController(Client->GetDefaultModel()))
+					{
+						if (const URigVMNode* Node = RigDependencyGraphNode->GetRigVMNodeForInstruction())
+						{
+							Controller->RequestJumpToHyperLink(Node);
+						}
+					}
+				}
+			}
+		}
+	}
+}
+
+void SRigDependencyGraph::OnClearNodeSelection()
+{
+	if (GraphEditor)
+	{
+		GraphEditor->ClearSelectionSet();
+	}
+}
+
+TSharedRef<SWidget> SRigDependencyGraph::CreateFilterMenu()
+{
+	const FRigDependencyGraphCommands& Actions = FRigDependencyGraphCommands::Get();
+
+	const bool CloseAfterSelection = true;
+	FMenuBuilder MenuBuilder(CloseAfterSelection, CommandList);
+
+	MenuBuilder.BeginSection("FilterOptions", LOCTEXT("OptionsMenuHeading", "Options"));
+	{
+		MenuBuilder.AddMenuEntry(Actions.ShowParentChildRelationships);
+		MenuBuilder.AddMenuEntry(Actions.ShowInstructionRelationships);
+	}
+	MenuBuilder.EndSection();
+
+	MenuBuilder.BeginSection("Navigation", LOCTEXT("NavigationMenuHeading", "Navigation"));
+	{
+		MenuBuilder.AddMenuEntry(Actions.FrameSelection);
+	}
+	MenuBuilder.EndSection();
+
+	MenuBuilder.BeginSection("Graph", LOCTEXT("GraphMenuHeading", "Graph"));
+	{
+		MenuBuilder.AddMenuEntry(Actions.SyncContent);
+		MenuBuilder.AddMenuEntry(Actions.SelectAllNodes);
+		MenuBuilder.AddMenuEntry(Actions.SelectInputNodes);
+		MenuBuilder.AddMenuEntry(Actions.SelectOutputNodes);
+		MenuBuilder.AddMenuEntry(Actions.SelectNodeIsland);
+		MenuBuilder.AddMenuEntry(Actions.RemoveAllNodes);
+		MenuBuilder.AddMenuEntry(Actions.RemoveSelected);
+		MenuBuilder.AddMenuEntry(Actions.IsolateSelected);
+		MenuBuilder.AddMenuEntry(Actions.RemoveUnrelated);
+		MenuBuilder.AddMenuEntry(Actions.SyncWithHierarchy);
+		MenuBuilder.AddMenuEntry(Actions.CancelLayoutSimulation);
+	}
+	MenuBuilder.EndSection();
+
+	return MenuBuilder.MakeWidget();
+}
+
+void SRigDependencyGraph::BindCommands()
+{
+	// create new command
+	const FRigDependencyGraphCommands& Commands = FRigDependencyGraphCommands::Get();
+
+	CommandList = MakeShared<FUICommandList>();
+	
+	CommandList->MapAction(Commands.FrameSelection, FExecuteAction::CreateLambda([this] ()
+	{
+		if (GraphEditor)
+		{
+			GraphEditor->ZoomToFit(true);
+		}
+	}));
+
+	CommandList->MapAction(Commands.SelectAllNodes, FExecuteAction::CreateLambda([this] ()
+	{
+		GraphObj->SelectAllNodes();
+	}));
+
+	CommandList->MapAction(Commands.SelectInputNodes, FExecuteAction::CreateLambda([this] ()
+	{
+		GraphObj->SelectLinkedNodes(GraphObj->SelectedNodes, true, true, true);
+	}));
+
+	CommandList->MapAction(Commands.SelectOutputNodes, FExecuteAction::CreateLambda([this] ()
+	{
+		GraphObj->SelectLinkedNodes(GraphObj->SelectedNodes, false, true, true);
+	}));
+
+	CommandList->MapAction(Commands.SelectNodeIsland, FExecuteAction::CreateLambda([this] ()
+	{
+		GraphObj->SelectNodeIsland(GraphObj->SelectedNodes, true);
+	}));
+
+	CommandList->MapAction(Commands.RemoveAllNodes, FExecuteAction::CreateLambda([this] ()
+	{
+		GraphObj->RemoveAllNodes();
+	}));
+
+	CommandList->MapAction(Commands.RemoveSelected, FExecuteAction::CreateLambda([this] ()
+	{
+		TArray<FNodeId> AllNodeIds;
+		GraphObj->RemoveSelectedNodes();
+	}));
+
+	CommandList->MapAction(Commands.IsolateSelected, FExecuteAction::CreateLambda([this] ()
+	{
+		GraphObj->RemoveUnselectedNodes();
+	}));
+
+	CommandList->MapAction(Commands.RemoveUnrelated, FExecuteAction::CreateLambda([this] ()
+	{
+		TArray<FNodeId> AllNodeIds;
+		GraphObj->RemoveUnrelatedNodes();
+	}));
+
+	CommandList->MapAction(Commands.SyncWithHierarchy, FExecuteAction::CreateLambda([this] ()
+	{
+		URigHierarchy* Hierarchy = GraphObj->GetRigHierarchy();
+		if (Hierarchy == nullptr)
+		{
+			return;
+		}
+
+		GraphObj->RemoveAllNodes();
+			
+		const TArray<FRigElementKey> SelectedElements = Hierarchy->GetSelectedKeys();
+
+		TArray<FNodeId> NodesToSelect;
+		for (const FRigElementKey& SelectedKey : SelectedElements)
+		{
+			NodesToSelect.Emplace(FNodeId::EType_RigElement, Hierarchy->GetIndex(SelectedKey), NAME_None);
+		}
+
+		GraphObj->ConstructNodes(NodesToSelect, false);
+	}));
+
+	CommandList->MapAction(Commands.CancelLayoutSimulation,
+		FExecuteAction::CreateLambda([this] ()
+		{
+			GraphObj->CancelLayout();
+		}),
+		FCanExecuteAction::CreateLambda([this] () -> bool
+		{
+			return GraphObj->LayoutIterationsLeft > 0;
+		})
+	);
+
+	CommandList->MapAction(
+		Commands.ShowParentChildRelationships,
+		FExecuteAction::CreateLambda([this]()
+		{
+			GraphObj->bFollowParentRelationShips = !GraphObj->bFollowParentRelationShips;
+			GraphObj->RebuildLinks();
+		}),
+		FCanExecuteAction(),
+		FIsActionChecked::CreateLambda([this]()
+		{
+			return GraphObj->bFollowParentRelationShips;
+		})
+	);
+
+	CommandList->MapAction(
+		Commands.ShowInstructionRelationships,
+		FExecuteAction::CreateLambda([this]()
+		{
+			GraphObj->bFollowVMRelationShips = !GraphObj->bFollowVMRelationShips;
+			GraphObj->RebuildLinks();
+		}),
+		FCanExecuteAction(),
+		FIsActionChecked::CreateLambda([this]()
+		{
+			return GraphObj->bFollowVMRelationShips;
+		})
+	);
+
+	CommandList->MapAction(
+		Commands.SyncContent,
+		FExecuteAction::CreateLambda([this]()
+		{
+			GraphObj->bLockContent = !GraphObj->bLockContent;
+		}),
+		FCanExecuteAction(),
+		FIsActionChecked::CreateLambda([this]()
+		{
+			return !GraphObj->bLockContent;
+		})
+	);
+
+	CommandList->MapAction(Commands.FindNodesByName,
+	FExecuteAction::CreateLambda([this] ()
+	{
+		if (FilterBox.IsValid())
+		{
+			FSlateApplication::Get().ForEachUser([this](FSlateUser& User)
+			{
+				User.SetFocus(FilterBox.ToSharedRef());
+			});
+		}
+	}));
+}
+
+void SRigDependencyGraph::OnSearchBarTextChanged(const FText& NewText)
+{
+	NodesMatchingSearch.Reset();
+	MatchIndex.Reset();
+
+	const FString SearchString = NewText.ToString();
+	if (SearchString.Len() <= 1)
+	{
+		return;
+	}
+	
+	for (const TObjectPtr<URigDependencyGraphNode>& Node : GraphObj->DependencyGraphNodes)
+	{
+		if (Node->GetNodeTitle(ENodeTitleType::FullTitle).ToString().Contains(SearchString, ESearchCase::IgnoreCase))
+		{
+			NodesMatchingSearch.Add(Node->GetNodeId());
+		}
+	}
+
+	OnSearchMatchPicked(0);
+}
+
+void SRigDependencyGraph::OnSearchMatchPicked(int32 InIndex)
+{
+	if (NodesMatchingSearch.IsValidIndex(InIndex))
+	{
+		const FNodeId NodeToSelect = NodesMatchingSearch[InIndex];
+		MatchIndex = InIndex;
+		GraphObj->SelectNodes({NodeToSelect});
+		if (GraphEditor.IsValid())
+		{
+			GraphEditor->JumpToNode(GraphObj->FindNode(NodeToSelect), false, false);
+		}
+	}
+}
+
+void SRigDependencyGraph::SelectNodes(const TArray<FNodeId>& InNodeIds)
+{
+	GraphObj->SelectNodes(InNodeIds);
+}
+
+FReply SRigDependencyGraph::OnKeyDown(const FGeometry& MyGeometry, const FKeyEvent& InKeyEvent)
+{
+	if (CommandList.IsValid() && CommandList->ProcessCommandBindings(InKeyEvent))
+	{
+		return FReply::Handled();
+	}
+	return FReply::Unhandled();
+}
+
+void SRigDependencyGraph::HandleSelectionChanged(const FGraphPanelSelectionSet& SelectionSet)
+{
+	if (GraphObj->BlockSelectionCounter <= 0 && WeakControlRigEditor.IsValid())
+	{
+		GraphObj->BlockSelectionCounter = 2;
+
+		URigHierarchy* Hierarchy = GraphObj->GetRigHierarchy();
+		if (!Hierarchy)
+		{
+			return;
+		}
+
+		GraphObj->SelectedNodes.Reset();
+		for (UObject* SelectedObject : SelectionSet)
+		{
+			if (const URigDependencyGraphNode* RigDependencyGraphNode = Cast<URigDependencyGraphNode>(SelectedObject))
+			{
+				const FNodeId& NodeId = RigDependencyGraphNode->GetNodeId();
+				GraphObj->SelectedNodes.Add(NodeId);
+			}
+		}
+
+		GraphObj->UpdateFadeOutStates();
+
+		TArray<FRigElementKey> SelectedElements;
+		TMap<const URigVMGraph*, TArray<FName>> NodeNamesPerGraph;
+		
+		for (const TObjectPtr<URigDependencyGraphNode>& DependencyGraphNode : GraphObj->DependencyGraphNodes)
+		{
+			const FNodeId& NodeId = DependencyGraphNode->GetNodeId();
+			if (NodeId.IsElement() && GraphObj->SelectedNodes.Contains(NodeId))
+			{
+				SelectedElements.Add(Hierarchy->GetKey(NodeId.Index));
+			}
+			else if (NodeId.IsInstruction())
+			{
+				if (const URigVMNode* Node = DependencyGraphNode->GetRigVMNodeForInstruction())
+				{
+					TArray<FName>& SelectedNodesInGraph = NodeNamesPerGraph.FindOrAdd(Node->GetGraph());
+					if (GraphObj->SelectedNodes.Contains(NodeId))
+					{
+						SelectedNodesInGraph.Add(Node->GetFName());
+					}
+				}
+			}
+		}
+
+		if (URigHierarchyController* Controller = Hierarchy->GetController())
+		{
+			Controller->SetSelection(SelectedElements);
+		}
+		TSharedPtr<IControlRigBaseEditor> ControlRigEditor = WeakControlRigEditor.Pin();
+		if (ControlRigEditor.IsValid())
+		{
+			if (UControlRigBlueprint* ControlRigBlueprint = ControlRigEditor->GetControlRigBlueprint())
+			{
+				if (FRigVMClient* Client = ControlRigBlueprint->GetRigVMClient())
+				{
+					for (TPair<const URigVMGraph*, TArray<FName>>& Pair : NodeNamesPerGraph)
+					{
+						if (URigVMController* Controller = Client->GetOrCreateController(Pair.Key))
+						{
+							Controller->SetNodeSelection(Pair.Value);
+						}
+					}
+				}
+			}
+		}
+	}
+}
+
+void SRigDependencyGraph::Tick(const FGeometry& AllottedGeometry, const double InCurrentTime, const float InDeltaTime)
+{
+	GraphObj->Tick(InDeltaTime);
+	
+	if (GraphObj->NeedsRefreshLayout())
+	{
+		const FBox2D OldBounds = GraphObj->GetAllNodesBounds();
+		GraphObj->GetRigDependencyGraphSchema()->LayoutNodes(GraphObj, 32);
+		if (!GraphObj->GetNodes().IsEmpty() && GraphObj->ZoomAndFitDuringLayout.Get(false))
+		{
+			const FBox2D NewBounds = GraphObj->GetAllNodesBounds();
+			if (!NewBounds.GetCenter().Equals(OldBounds.GetCenter(), 10.f) ||
+				!NewBounds.GetSize().Equals(OldBounds.GetSize(), 10.f))
+			{
+				GraphEditor->ZoomToFit(false);
+			}
+		}
+	}
+
+	// flashlight to lighten up nodes in proximity of the mouse
+	const FGeometry PaintGeometry = GraphEditor->GetPaintSpaceGeometry();
+	const FVector2D MousePosition = FSlateApplication::Get().GetCursorPos();
+	const float FlashLightRadius = FlashLightRadiusAttribute.Get();
+
+	if(PaintGeometry.IsUnderLocation(MousePosition) && FlashLightRadius > SMALL_NUMBER)
+	{
+		FVector2f GraphEditorLocation = FVector2f::ZeroVector;
+		float GraphEditorZoomAmount = 0;
+		GraphEditor->GetViewLocation(GraphEditorLocation, GraphEditorZoomAmount);
+
+		const FVector2f WidgetPosition = PaintGeometry.AbsoluteToLocal(MousePosition) / GraphEditorZoomAmount + GraphEditorLocation;
+		for (URigDependencyGraphNode* Node : GraphObj->DependencyGraphNodes)
+		{
+			Node->ResetFadedOutState();
+
+			const float PreviousFadedOutState = Node->GetFadedOutState();
+			if (PreviousFadedOutState < 1.f - SMALL_NUMBER)
+			{
+				FSlateRect Bounds;
+				if (GraphEditor->GetBoundsForNode(Node, Bounds, 0.f))
+				{
+					const FVector2f Center = Bounds.GetCenter();
+					const FVector2f Size = Bounds.GetSize();
+					if (!Size.IsNearlyZero())
+					{
+						const float Radius = Size.Length() * 0.5f * FlashLightRadius;
+						const float Ratio = FVector2f::DistSquared(Center, WidgetPosition) / (Radius * Radius);
+						if (Ratio < 1.f)
+						{
+							const float HalfRatio = 1.f - FMath::Clamp(Ratio - 0.5f, 0.f, .5f) / 0.5f;
+							Node->OverrideFadeOutState(FMath::Lerp(PreviousFadedOutState, 1.f, HalfRatio));
+						}
+					}
+				}
+			}
+		}
+	}
+	else
+	{
+		for (URigDependencyGraphNode* Node : GraphObj->DependencyGraphNodes)
+		{
+			Node->ResetFadedOutState();
+		}
+	}
+}
+
+void SRigDependencyGraph::PostUndo(bool bSuccess)
+{
+	GraphObj->RebuildGraph();
+}
+
+void SRigDependencyGraph::PostRedo(bool bSuccess)
+{
+	GraphObj->RebuildGraph();
+}
+
+#undef LOCTEXT_NAMESPACE
diff --git a/Engine/Plugins/Animation/ControlRig/Source/ControlRigEditor/Private/RigDependencyGraph/SRigDependencyGraph.h b/Engine/Plugins/Animation/ControlRig/Source/ControlRigEditor/Private/RigDependencyGraph/SRigDependencyGraph.h
new file mode 100644
--- /dev/null
+++ b/Engine/Plugins/Animation/ControlRig/Source/ControlRigEditor/Private/RigDependencyGraph/SRigDependencyGraph.h
@@ -0,0 +1,88 @@
+// Copyright Epic Games, Inc. All Rights Reserved.
+
+#pragma once
+
+#include "PersonaDelegates.h"
+#include "GraphEditor.h"
+#include "Containers/ArrayView.h"
+#include "EditorUndoClient.h"
+#include "RigDependencyGraph.h"
+#include "RigDependencyGraphNode.h"
+
+class UObject;
+class URigDependencyGraph;
+class SRigDependencyGraph;
+class SProgressBar;
+class SSearchBox;
+class IControlRigBaseEditor;
+
+/** Delegate used to communicate graph selection */
+DECLARE_DELEGATE_OneParam(FOnRigDependencyGraphObjectsSelected, const TArrayView<UObject*>& /*InObjects*/);
+
+class SRigDependencyGraph : public SCompoundWidget, public FEditorUndoClient
+{
+public:
+
+	using FNodeId = URigDependencyGraph::FNodeId;
+	
+	SLATE_BEGIN_ARGS(SRigDependencyGraph)
+		: _FlashLightRadius(5.f)
+	{}
+
+	SLATE_ATTRIBUTE(float, FlashLightRadius)
+
+	SLATE_END_ARGS()
+
+	~SRigDependencyGraph();
+
+	/** SWidget interface */
+	virtual void Tick(const FGeometry& AllottedGeometry, const double InCurrentTime, const float InDeltaTime) override;
+
+	/** FEditorUndoClient interface */
+	virtual void PostUndo(bool bSuccess) override;
+	virtual void PostRedo(bool bSuccess) override;
+
+	/** Constructs this widget with InArgs */
+	void Construct(const FArguments& InArgs, const TSharedRef<IControlRigBaseEditor>& InControlRigEditor);
+
+	/** Set the selected nodes */
+	void SelectNodes(const TArray<FNodeId>& InNodeIds);
+
+	virtual FReply OnKeyDown(const FGeometry& MyGeometry, const FKeyEvent& InKeyEvent);
+
+private:
+
+	/** Called to create context menu when right-clicking on graph */
+	FActionMenuContent OnCreateGraphActionMenu(UEdGraph* InGraph, const FVector2f& InNodePosition, const TArray<UEdGraphPin*>& InDraggedPins, bool bAutoExpand, SGraphEditor::FActionMenuClosed InOnMenuClosed);
+
+	/** Called when the path is being edited */
+	void OnSearchBarTextChanged(const FText& NewText);
+	void OnSearchMatchPicked(int32 InIndex);
+
+	void HandleSelectionChanged(const FGraphPanelSelectionSet& SelectionSet);
+
+	void HandleNodeDoubleClicked(UEdGraphNode* InNode);
+
+	void OnClearNodeSelection();
+
+	TSharedRef< SWidget > CreateFilterMenu();
+
+	void BindCommands();
+
+private:
+
+	TSharedPtr<SGraphEditor> GraphEditor;
+	TSharedPtr<SProgressBar> LayoutProgressBar;
+	TWeakPtr<IControlRigBaseEditor> WeakControlRigEditor;
+
+	URigDependencyGraph* GraphObj;
+	TSharedPtr<FUICommandList> CommandList;
+
+	/** Search box widget */
+	TSharedPtr<SSearchBox> FilterBox;
+
+	TArray<FNodeId> NodesMatchingSearch;
+	TOptional<int32> MatchIndex;
+
+	TAttribute<float> FlashLightRadiusAttribute;
+};
diff --git a/Engine/Plugins/Animation/ControlRig/Source/ControlRigEditor/Private/RigDependencyGraph/SRigDependencyGraphNode.cpp b/Engine/Plugins/Animation/ControlRig/Source/ControlRigEditor/Private/RigDependencyGraph/SRigDependencyGraphNode.cpp
new file mode 100644
--- /dev/null
+++ b/Engine/Plugins/Animation/ControlRig/Source/ControlRigEditor/Private/RigDependencyGraph/SRigDependencyGraphNode.cpp
@@ -0,0 +1,271 @@
+// Copyright Epic Games, Inc. All Rights Reserved.
+
+#include "RigDependencyGraph/SRigDependencyGraphNode.h"
+
+#include "Containers/Array.h"
+#include "Delegates/Delegate.h"
+#include "EdGraph/EdGraphNode.h"
+#include "EdGraph/EdGraphPin.h"
+#include "GenericPlatform/ICursor.h"
+#include "HAL/Platform.h"
+#include "HAL/PlatformCrt.h"
+#include "Layout/Margin.h"
+#include "Math/Vector2D.h"
+#include "Misc/AssertionMacros.h"
+#include "Misc/Attribute.h"
+#include "Misc/Optional.h"
+#include "RigDependencyGraph/RigDependencyGraphNode.h"
+#include "SGraphPin.h"
+#include "SNodePanel.h"
+#include "SlotBase.h"
+#include "Styling/AppStyle.h"
+#include "Templates/Casts.h"
+#include "Types/SlateEnums.h"
+#include "Widgets/Images/SImage.h"
+#include "Widgets/Layout/SBorder.h"
+#include "Widgets/Layout/SSpacer.h"
+#include "Widgets/SBoxPanel.h"
+#include "Widgets/SNullWidget.h"
+#include "Widgets/Text/STextBlock.h"
+#include "ControlRigEditorStyle.h"
+#include "RigDependencyGraph.h"
+
+class SWidget;
+class UEdGraphSchema;
+struct FSlateBrush;
+
+#define LOCTEXT_NAMESPACE "RigDependencyGraphEditor"
+
+class SRigDependencyGraphNodePin : public SGraphPin
+{
+public:
+	SLATE_BEGIN_ARGS(SRigDependencyGraphNodePin) {}
+	SLATE_END_ARGS()
+
+	void Construct(const FArguments& InArgs, UEdGraphPin* InPin)
+	{
+		this->SetCursor(EMouseCursor::Default);
+		SetIsEditable(false);
+
+		bShowLabel = false;
+		GraphPinObj = InPin;
+		check(GraphPinObj != NULL);
+
+		// Set up a hover for pins that is tinted the color of the pin.
+		SBorder::Construct(SBorder::FArguments()
+			.BorderImage(FAppStyle::GetBrush("NoBorder"))
+			.BorderBackgroundColor(this, &SRigDependencyGraphNodePin::GetPinColor)
+			.OnMouseButtonDown(this, &SRigDependencyGraphNodePin::OnPinMouseDown)
+			.Cursor(this, &SRigDependencyGraphNodePin::GetPinCursor)
+			.Content()
+			[
+				SNew(SImage)
+				.Image(this, &SRigDependencyGraphNodePin::GetPinIcon)
+			]
+		);
+	}
+
+protected:
+	/** SGraphPin interface */
+	virtual TSharedRef<SWidget>	GetDefaultValueWidget() override
+	{
+		return SNew(SSpacer);
+	}
+};
+
+void SRigDependencyGraphNode::Construct(const FArguments& InArgs, class URigDependencyGraphNode* InNode)
+{
+	GraphNode = InNode;
+	
+	SetCursor(EMouseCursor::CardinalCross);
+
+	if (!ContentWidget.IsValid())
+	{
+		ContentWidget = SNullWidget::NullWidget;
+	}
+	
+	UpdateGraphNode();
+}
+
+void SRigDependencyGraphNode::CreatePinWidgets()
+{
+	URigDependencyGraphNode* RigDependencyGraphNode = CastChecked<URigDependencyGraphNode>(GraphNode);
+
+	UEdGraphPin& InputPin = RigDependencyGraphNode->GetInputPin();
+	if (!InputPin.bHidden)
+	{
+		this->AddPin(SNew(SRigDependencyGraphNodePin, &InputPin));
+	}
+
+	UEdGraphPin& OutputPin = RigDependencyGraphNode->GetOutputPin();
+	if (!OutputPin.bHidden)
+	{
+		this->AddPin(SNew(SRigDependencyGraphNodePin, &OutputPin));
+	}
+}
+
+void SRigDependencyGraphNode::AddPin(const TSharedRef<SGraphPin>& PinToAdd)
+{
+	PinToAdd->SetOwner(SharedThis(this));
+
+	if (PinToAdd->GetDirection() == EGPD_Input)
+	{
+		LeftNodeBox->AddSlot()
+			.HAlign(HAlign_Center)
+			.VAlign(VAlign_Center)
+			.AutoHeight()
+			[
+				PinToAdd
+			];
+	}
+	else
+	{
+		RightNodeBox->AddSlot()
+			.HAlign(HAlign_Center)
+			.VAlign(VAlign_Center)
+			.AutoHeight()
+			[
+				PinToAdd
+			];
+	}
+	OutputPins.Add(PinToAdd);
+}
+
+FText SRigDependencyGraphNode::GetNodeTitle() const
+{
+	return GraphNode->GetNodeTitle(ENodeTitleType::FullTitle);
+}
+
+FText SRigDependencyGraphNode::GetNodeTooltip() const
+{
+	return GraphNode->GetTooltipText();
+}
+
+const FSlateBrush* SRigDependencyGraphNode::GetNodeTitleIcon() const
+{
+	FLinearColor Color = FLinearColor::White;
+	return GraphNode->GetIconAndTint(Color).GetIcon();
+}
+
+FSlateColor SRigDependencyGraphNode::GetNodeBodyColor() const
+{
+	return GraphNode->GetNodeBodyTintColor();
+}
+
+const FSlateBrush* SRigDependencyGraphNode::GetNodeBodyBrush() const
+{
+	return FControlRigEditorStyle::Get().GetBrush(TEXT("ControlRig.DependencyGraph.NodeBody"));
+}
+
+void SRigDependencyGraphNode::UpdateGraphNode()
+{
+	InputPins.Empty();
+	OutputPins.Empty();
+	
+	RightNodeBox.Reset();
+	LeftNodeBox.Reset();
+
+	this->ContentScale.Bind( this, &SGraphNode::GetContentScale );
+
+	this->GetOrAddSlot( ENodeZone::Center )
+		.HAlign(HAlign_Center)
+		.VAlign(VAlign_Center)
+		[
+			SNew(SBorder)
+			.BorderImage( FControlRigEditorStyle::Get().GetBrush("ControlRig.DependencyGraph.NodeBody") )
+			.BorderBackgroundColor(this, &SRigDependencyGraphNode::GetNodeBodyColor)
+			.Padding(0)
+			[
+				SNew(SHorizontalBox)
+
+				+SHorizontalBox::Slot()
+				.AutoWidth()
+				.HAlign(HAlign_Center)
+				.VAlign(VAlign_Center)
+				[
+					SAssignNew(LeftNodeBox, SVerticalBox)
+				]
+
+				+SHorizontalBox::Slot()
+				.AutoWidth()
+				.HAlign(HAlign_Center)
+				.VAlign(VAlign_Center)
+				.Padding(4.0f)
+				[
+					SNew(SVerticalBox)
+					+SVerticalBox::Slot()
+					.AutoHeight()
+					[
+						SNew(SImage)
+						.Image(this, &SRigDependencyGraphNode::GetNodeTitleIcon)
+					]
+				]
+
+				+SHorizontalBox::Slot()
+				.AutoWidth()
+				.HAlign(HAlign_Fill)
+				.VAlign(VAlign_Center)
+				.Padding(4.0f)
+				[
+					SNew(SVerticalBox)
+					+SVerticalBox::Slot()
+					.AutoHeight()
+					[
+						SNew(STextBlock)
+						.Text(this, &SRigDependencyGraphNode::GetNodeTitle)
+					]
+				]
+
+				+SHorizontalBox::Slot()
+				.AutoWidth()
+				.HAlign(HAlign_Center)
+				.VAlign(VAlign_Center)
+				[
+					SAssignNew(RightNodeBox, SVerticalBox)
+				]
+			]
+		];
+
+	CreatePinWidgets();
+}
+
+FReply SRigDependencyGraphNode::OnMouseButtonUp(const FGeometry& MyGeometry, const FPointerEvent& MouseEvent)
+{
+	if (MouseEvent.GetEffectingButton() == EKeys::MiddleMouseButton)
+	{
+		if (URigDependencyGraphNode* DependencyGraphNode = Cast<URigDependencyGraphNode>(GraphNode))
+		{
+			if (URigDependencyGraph* DependencyGraph = DependencyGraphNode->GetRigDependencyGraph())
+			{
+				const bool ClearSelection = !MouseEvent.GetModifierKeys().IsShiftDown();
+				const bool SelectSourceNodes = !MouseEvent.GetModifierKeys().IsControlDown();
+				const bool SelectIsland = MouseEvent.GetModifierKeys().IsAltDown();
+				if (SelectIsland)
+				{
+					DependencyGraph->SelectNodeIsland({DependencyGraphNode->GetNodeId()}, ClearSelection);
+				}
+				else
+				{
+					DependencyGraph->SelectLinkedNodes({DependencyGraphNode->GetNodeId()}, SelectSourceNodes, ClearSelection, true);
+				}
+			}
+			return FReply::Handled();
+		}
+	}
+	return SGraphNode::OnMouseButtonUp(MyGeometry, MouseEvent);
+}
+
+int32 SRigDependencyGraphNode::OnPaint(const FPaintArgs& Args, const FGeometry& AllottedGeometry, const FSlateRect& MyCullingRect,
+	FSlateWindowElementList& OutDrawElements, int32 LayerId, const FWidgetStyle& InWidgetStyle, bool bParentEnabled) const
+{
+	FWidgetStyle WidgetStyle = InWidgetStyle;
+	
+	if (const URigDependencyGraphNode* DependencyGraphNode = Cast<URigDependencyGraphNode>(GraphNode))
+	{
+		WidgetStyle.BlendOpacity(DependencyGraphNode->GetFadedOutState());
+	}
+	
+	return SGraphNode::OnPaint(Args, AllottedGeometry, MyCullingRect, OutDrawElements, LayerId, WidgetStyle, bParentEnabled);
+}
+
+#undef LOCTEXT_NAMESPACE
diff --git a/Engine/Plugins/Animation/ControlRig/Source/ControlRigEditor/Private/RigDependencyGraph/SRigDependencyGraphNode.h b/Engine/Plugins/Animation/ControlRig/Source/ControlRigEditor/Private/RigDependencyGraph/SRigDependencyGraphNode.h
new file mode 100644
--- /dev/null
+++ b/Engine/Plugins/Animation/ControlRig/Source/ControlRigEditor/Private/RigDependencyGraph/SRigDependencyGraphNode.h
@@ -0,0 +1,43 @@
+// Copyright Epic Games, Inc. All Rights Reserved.
+
+#pragma once
+
+#include "Internationalization/Text.h"
+#include "SGraphNode.h"
+#include "Styling/SlateColor.h"
+#include "Templates/SharedPointer.h"
+#include "Widgets/DeclarativeSyntaxSupport.h"
+
+class SGraphPin;
+class SVerticalBox;
+class SWidget;
+struct FSlateBrush;
+
+class SRigDependencyGraphNode : public SGraphNode
+{
+public:
+	SLATE_BEGIN_ARGS(SRigDependencyGraphNode){}
+
+	SLATE_END_ARGS()
+
+	/** Constructs this widget with InArgs */
+	void Construct(const FArguments& InArgs, class URigDependencyGraphNode* InNode);
+
+protected:
+	// SGraphNode interface
+	virtual void CreatePinWidgets() override;
+	virtual void AddPin(const TSharedRef<SGraphPin>& PinToAdd) override;
+
+	FText GetNodeTitle() const;
+	virtual FText GetNodeTooltip() const override;
+	const FSlateBrush* GetNodeTitleIcon() const;
+	virtual FSlateColor GetNodeBodyColor() const override;
+	virtual const FSlateBrush* GetNodeBodyBrush() const override;
+	virtual void UpdateGraphNode() override;
+	virtual FReply OnMouseButtonUp(const FGeometry& MyGeometry, const FPointerEvent& MouseEvent) override;
+	virtual int32 OnPaint(const FPaintArgs& Args, const FGeometry& AllottedGeometry, const FSlateRect& MyCullingRect, FSlateWindowElementList& OutDrawElements, int32 LayerId, const FWidgetStyle& InWidgetStyle, bool bParentEnabled) const override;
+	
+protected:
+	/** The content widget for this node - derived classes can insert what they want */
+	TSharedPtr<SWidget> ContentWidget;
+};
