Control Rig: Guard dependency parser against infinite recursion

#rnx

diff --git a/Engine/Plugins/Animation/ControlRig/Source/ControlRig/Private/Rigs/RigHierarchy.cpp b/Engine/Plugins/Animation/ControlRig/Source/ControlRig/Private/Rigs/RigHierarchy.cpp
--- a/Engine/Plugins/Animation/ControlRig/Source/ControlRig/Private/Rigs/RigHierarchy.cpp
+++ b/Engine/Plugins/Animation/ControlRig/Source/ControlRig/Private/Rigs/RigHierarchy.cpp
@@ -6387,8 +6387,15 @@ bool URigHierarchy::IsDependentOn(FRigBaseElement* InDependent, FRigBaseElement*
 			const IRigDependenciesProvider& InDependencyProvider,
 			bool bIsOnActualTopology,
 			const TTuple<int32,int32>& CacheKey,
-			const FRigHierarchyRecord& InDependentRecord)
+			const FRigHierarchyRecord& InDependentRecord,
+			TSet<FRigHierarchyRecord>& VisitedRecords)
 		{
+			if(VisitedRecords.Contains(InDependentRecord))
+			{
+				return false;
+			}
+			VisitedRecords.Add(InDependentRecord);
+			
 			const TRigHierarchyDependencyMap& DependencyMap = InDependencyProvider.GetRigHierarchyDependencies();
 			
 			if(const TSet<FRigHierarchyRecord>* DependentRecordsPtr = DependencyMap.Find(InDependentRecord))
@@ -6411,7 +6418,7 @@ bool URigHierarchy::IsDependentOn(FRigBaseElement* InDependent, FRigBaseElement*
 					}
 					else if (DependentRecord.IsInstruction() || DependentRecord.IsVariable())
 					{
-						if (FoundDependencyInProvider(InHierarchy, InDependency, InDependencyProvider, bIsOnActualTopology, CacheKey, DependentRecord))
+						if (FoundDependencyInProvider(InHierarchy, InDependency, InDependencyProvider, bIsOnActualTopology, CacheKey, DependentRecord, VisitedRecords))
 						{
 							return true;
 						}
@@ -6424,7 +6431,8 @@ bool URigHierarchy::IsDependentOn(FRigBaseElement* InDependent, FRigBaseElement*
 
 	// check the optional dependency map
 	const FRigHierarchyRecord DependentRecord(FRigHierarchyRecord::EType_RigElement, InDependent->GetIndex(), NAME_None);
-	if (Local::FoundDependencyInProvider(this, InDependency, InDependencyProvider, bIsOnActualTopology, CacheKey, DependentRecord))
+	TSet<FRigHierarchyRecord> VisitedRecords;
+	if (Local::FoundDependencyInProvider(this, InDependency, InDependencyProvider, bIsOnActualTopology, CacheKey, DependentRecord, VisitedRecords))
 	{
 		// we'll only update the caches if we are following edges on the actual topology
 		if(bIsOnActualTopology)
