Control Rig: Improve space switching reporting + stability

#rnx

diff --git a/Engine/Plugins/Animation/ControlRig/Source/ControlRig/Private/Rigs/RigHierarchy.cpp b/Engine/Plugins/Animation/ControlRig/Source/ControlRig/Private/Rigs/RigHierarchy.cpp
--- a/Engine/Plugins/Animation/ControlRig/Source/ControlRig/Private/Rigs/RigHierarchy.cpp
+++ b/Engine/Plugins/Animation/ControlRig/Source/ControlRig/Private/Rigs/RigHierarchy.cpp
@@ -20,6 +20,7 @@
 #include "Algo/Count.h"
 #include "HAL/LowLevelMemTracker.h"
 #include "Misc/Compression.h"
+#include "UObject/ICookInfo.h"
 
 #include UE_INLINE_GENERATED_CPP_BY_NAME(RigHierarchy)
 
@@ -3894,12 +3895,126 @@ bool URigHierarchy::CanSwitchToParent(FRigElementKey InChild, FRigElementKey InP
 		}
 	}
 
-	if(IsParentedTo(Parent, Child, InDependencyProvider))
+	TArray<FRigHierarchyRecord> DependencyChain;
+	if(IsParentedTo(Parent, Child, InDependencyProvider, &DependencyChain))
 	{
+#if WITH_EDITOR
+		
 		if(OutFailureReason)
 		{
-			OutFailureReason->Appendf(TEXT("Cannot switch '%s' to '%s' - would cause a cycle."), *InChild.ToString(), *InParent.ToString());
+			if (DependencyChain.IsEmpty())
+			{
+				OutFailureReason->Appendf(TEXT("Cannot switch '%s' to '%s' - would cause a cycle."), *InChild.ToString(), *InParent.ToString());
+			}
+			else
+			{
+				TArray<FString> Labels;
+				Labels.Reserve(DependencyChain.Num() + 1);
+				for (int32 DependencyIndex = DependencyChain.Num() - 1; DependencyIndex >= 0; --DependencyIndex)
+				{
+					const FRigHierarchyRecord& Dependency = DependencyChain[DependencyIndex];
+					switch (Dependency.Type)
+					{
+						case FRigHierarchyRecord::EType_RigElement:
+						case FRigHierarchyRecord::EType_Metadata:
+						{
+							Labels.Add(GetKey(Dependency.Index).ToString());
+							if (Dependency.IsMetadata())
+							{
+								Labels.Last() += FString::Printf(TEXT(" Metadata(%s)"), *Dependency.Name.ToString());
+							}
+							break;
+						}
+						case FRigHierarchyRecord::EType_Variable:
+						case FRigHierarchyRecord::EType_Instruction:
+						{
+							if (Dependency.IsVariable())
+							{
+								Labels.Add(TEXT("Variable"));
+							}
+							else
+							{
+								Labels.Add(TEXT("Node"));
+							}
+
+							UControlRig* ControlRig = Cast<UControlRig>(GetOuter());
+							if (!ControlRig)
+							{
+								break;
+							}
+
+							UModularRig* ModularRig = Cast<UModularRig>(ControlRig);
+
+							const URigVM* VM = nullptr;
+							if (Dependency.Name.IsNone() || !ModularRig)
+							{
+								VM = ControlRig->GetVM();
+							}
+							else
+							{
+								const FRigModuleInstance* RigModule = ModularRig->FindModule(Dependency.Name);
+								if (!RigModule)
+								{
+									break;
+								}
+								UControlRig* ModuleRig = RigModule->GetRig();
+								if (!ModuleRig)
+								{
+									break;
+								}
+								VM = ModuleRig->GetVM();
+							}
+
+							check(VM);
+
+							if (Dependency.IsVariable())
+							{
+								const TArray<FRigVMExternalVariableDef>& ExternalVariableDefs = VM->GetExternalVariableDefs();
+								if (ExternalVariableDefs.IsValidIndex(Dependency.Index))
+								{
+									Labels.Last() += FString::Printf(TEXT("(%s)"), *ExternalVariableDefs[Dependency.Index].Name.ToString());
+								}
+							}
+							else
+							{
+								if (!VM->GetByteCode().GetInstructions().IsValidIndex(Dependency.Index))
+								{
+									break;
+								}
+	
+								const FRigVMInstruction Instruction = VM->GetByteCode().GetInstructions()[Dependency.Index];
+								if (Instruction.OpCode != ERigVMOpCode::Execute)
+								{
+									break;
+								}
+
+								URigVMNode* Node = Cast<URigVMNode>(VM->GetByteCode().GetSubjectForInstruction(Dependency.Index));
+								if (Node == nullptr)
+								{
+									break;
+								}
+
+								Labels.Last() += FString::Printf(TEXT("(%s)"), *Node->GetNodeTitle());
+							}
+
+							break;
+						}
+						default:
+						{
+							Labels.Add(TEXT("Unknown"));
+							break;
+						}
+					}
+				}
+
+				const FString LastItem = Labels.Last(); 
+				Labels.Insert(LastItem, 0);
+
+				const FString JoinedLabels = FString::Join(Labels, TEXT(" -> ")); 
+				OutFailureReason->Appendf(TEXT("Cannot switch '%s' to '%s' - would cause a cycle.\n\n%s"), *InChild.ToString(), *InParent.ToString(), *JoinedLabels);
+			}
 		}
+#endif
 		return false;
 	}
 
@@ -6286,74 +6401,82 @@ FRigHierarchyKey URigHierarchy::GetPreviousHierarchyParent(const FRigHierarchyKe
 	return FRigHierarchyKey();
 }
 
-bool URigHierarchy::IsParentedTo(FRigBaseElement* InChild, FRigBaseElement* InParent, const IRigDependenciesProvider& InDependencyProvider) const
+bool URigHierarchy::IsParentedTo(FRigBaseElement* InChild, FRigBaseElement* InParent, const IRigDependenciesProvider& InDependencyProvider, TArray<FRigHierarchyRecord>* OutDependencyChain) const
 {
-	ElementDependencyVisited.Reset();
-	if(!InDependencyProvider.GetRigHierarchyDependencies().IsEmpty())
+	TArray<FRigHierarchyRecord> DependencyChain;
+
+	if (OutDependencyChain == nullptr)
+	{
+		DependencyChain.Reserve(64);
+		OutDependencyChain = &DependencyChain;
+	}
+
+	check(OutDependencyChain != nullptr);
+	const FRigHierarchyRecord DependencyRecord(FRigHierarchyRecord::EType_RigElement, InParent->GetIndex(), NAME_None);
+
 	{
-		ElementDependencyVisited.SetNumZeroed(Elements.Num());
+		UE::TScopeLock Lock(IsDependentOnLock);
+		if (!IsDependentOn(InChild, InParent, InDependencyProvider, true, *OutDependencyChain))
+		{
+			return false;
+		}
 	}
-	return IsDependentOn(InChild, InParent, InDependencyProvider, true);
+
+	return true;
 }
 
-bool URigHierarchy::IsDependentOn(FRigBaseElement* InDependent, FRigBaseElement* InDependency, const IRigDependenciesProvider& InDependencyProvider, bool bIsOnActualTopology) const
+bool URigHierarchy::IsDependentOn(FRigBaseElement* InDependent, FRigBaseElement* InDependency, const IRigDependenciesProvider& InDependencyProvider, bool bIsOnActualTopology, TArray<FRigHierarchyRecord>& OutDependencyChain) const
 {
 	if((InDependent == nullptr) || (InDependency == nullptr))
 	{
 		return false;
 	}
 
-	if(InDependent == InDependency)
-	{
-		return true;
-	}
-
 	const int32 DependentElementIndex = InDependent->GetIndex();
 	const int32 DependencyElementIndex = InDependency->GetIndex();
 	const TTuple<int32,int32> CacheKey(DependentElementIndex, DependencyElementIndex);
 
-	if(!ElementDependencyCache.IsValid(GetTopologyVersion()))
+	const FRigHierarchyRecord DependentRecord(FRigHierarchyRecord::EType_RigElement, DependentElementIndex, NAME_None);
+	if (OutDependencyChain.Contains(DependentRecord))
 	{
-		ElementDependencyCache.Set(TMap<TTuple<int32, int32>, bool>(), GetTopologyVersion());
+		return false;
 	}
+	OutDependencyChain.Push(DependentRecord);
 
-	// we'll only update the caches if we are following edges on the actual topology
-	if(const bool* bCachedResult = ElementDependencyCache.Get().Find(CacheKey))
+	if(InDependent == InDependency)
 	{
-		return *bCachedResult;
+		return true;
 	}
 
-	// check if the reverse dependency check has been stored before - if the dependency is dependent
-	// then we don't need to recurse any further.
-	const TTuple<int32,int32> ReverseCacheKey(DependencyElementIndex, DependentElementIndex);
-	if(const bool* bReverseCachedResult = ElementDependencyCache.Get().Find(ReverseCacheKey))
+	if (bIsOnActualTopology)
 	{
-		if(*bReverseCachedResult)
+		if(!ElementDependencyCache.IsValid(GetTopologyVersion()))
 		{
-			return false;
+			ElementDependencyCache.Set(TMap<TTuple<int32, int32>, bool>(), GetTopologyVersion());
 		}
-	}
 
-	if(!ElementDependencyVisited.IsEmpty())
-	{
-		// when running this with a provided dependency map
-		// we may run into a cycle / infinite recursion. this array
-		// keeps track of all elements being visited before.
-		if(!ElementDependencyVisited.IsValidIndex(DependentElementIndex))
+		// we'll only update the caches if we are following edges on the actual topology
+		if(const bool* bCachedResult = ElementDependencyCache.Get().Find(CacheKey))
 		{
-			return false;
+			return *bCachedResult;
 		}
-		if (ElementDependencyVisited[DependentElementIndex])
+
+		// check if the reverse dependency check has been stored before - if the dependency is dependent
+		// then we don't need to recurse any further.
+		const TTuple<int32,int32> ReverseCacheKey(DependencyElementIndex, DependentElementIndex);
+		if(const bool* bReverseCachedResult = ElementDependencyCache.Get().Find(ReverseCacheKey))
 		{
-			return false;
+			if(*bReverseCachedResult)
+			{
+				return false;
+			}
 		}
-		ElementDependencyVisited[DependentElementIndex] = true;
 	}
-
+	
 	// collect all possible parents of the dependent
 	if(const FRigSingleParentElement* SingleParentElement = Cast<FRigSingleParentElement>(InDependent))
 	{
-		if(IsDependentOn(SingleParentElement->ParentElement, InDependency, InDependencyProvider, true))
+		if(IsDependentOn(SingleParentElement->ParentElement, InDependency, InDependencyProvider, true, OutDependencyChain))
 		{
 			// we'll only update the caches if we are following edges on the actual topology
 			if(bIsOnActualTopology)
@@ -6367,7 +6490,7 @@ bool URigHierarchy::IsDependentOn(FRigBaseElement* InDependent, FRigBaseElement*
 	{
 		for(const FRigElementParentConstraint& ParentConstraint : MultiParentElement->ParentConstraints)
 		{
-			if(IsDependentOn(ParentConstraint.ParentElement, InDependency, InDependencyProvider, true))
+			if(IsDependentOn(ParentConstraint.ParentElement, InDependency, InDependencyProvider, true, OutDependencyChain))
 			{
 				// we'll only update the caches if we are following edges on the actual topology
 				if(bIsOnActualTopology)
@@ -6388,14 +6511,8 @@ bool URigHierarchy::IsDependentOn(FRigBaseElement* InDependent, FRigBaseElement*
 			bool bIsOnActualTopology,
 			const TTuple<int32,int32>& CacheKey,
 			const FRigHierarchyRecord& InDependentRecord,
-			TSet<FRigHierarchyRecord>& VisitedRecords)
+			TArray<FRigHierarchyRecord>& OutDependencyChain)
 		{
-			if(VisitedRecords.Contains(InDependentRecord))
-			{
-				return false;
-			}
-			VisitedRecords.Add(InDependentRecord);
-			
 			const TRigHierarchyDependencyMap& DependencyMap = InDependencyProvider.GetRigHierarchyDependencies();
 			
 			if(const TSet<FRigHierarchyRecord>* DependentRecordsPtr = DependencyMap.Find(InDependentRecord))
@@ -6406,7 +6523,7 @@ bool URigHierarchy::IsDependentOn(FRigBaseElement* InDependent, FRigBaseElement*
 					if (DependentRecord.IsElement())
 					{
 						ensure(InHierarchy->Elements.IsValidIndex(DependentRecord.Index));
-						if(InHierarchy->IsDependentOn(InHierarchy->Elements[DependentRecord.Index], InDependency, InDependencyProvider, false))
+						if(InHierarchy->IsDependentOn(InHierarchy->Elements[DependentRecord.Index], InDependency, InDependencyProvider, false, OutDependencyChain))
 						{
 							// we'll only update the caches if we are following edges on the actual topology
 							if(bIsOnActualTopology)
@@ -6418,10 +6535,14 @@ bool URigHierarchy::IsDependentOn(FRigBaseElement* InDependent, FRigBaseElement*
 					}
 					else if (DependentRecord.IsInstruction() || DependentRecord.IsVariable())
 					{
-						if (FoundDependencyInProvider(InHierarchy, InDependency, InDependencyProvider, bIsOnActualTopology, CacheKey, DependentRecord, VisitedRecords))
+						OutDependencyChain.Push(DependentRecord);
+						
+						if (FoundDependencyInProvider(InHierarchy, InDependency, InDependencyProvider, bIsOnActualTopology, CacheKey, DependentRecord, OutDependencyChain))
 						{
 							return true;
 						}
+
+						OutDependencyChain.Pop();
 					}
 				}
 			}
@@ -6430,9 +6551,7 @@ bool URigHierarchy::IsDependentOn(FRigBaseElement* InDependent, FRigBaseElement*
 	};
 
 	// check the optional dependency map
-	const FRigHierarchyRecord DependentRecord(FRigHierarchyRecord::EType_RigElement, InDependent->GetIndex(), NAME_None);
-	TSet<FRigHierarchyRecord> VisitedRecords;
-	if (Local::FoundDependencyInProvider(this, InDependency, InDependencyProvider, bIsOnActualTopology, CacheKey, DependentRecord, VisitedRecords))
+	if (Local::FoundDependencyInProvider(this, InDependency, InDependencyProvider, bIsOnActualTopology, CacheKey, DependentRecord, OutDependencyChain))
 	{
 		// we'll only update the caches if we are following edges on the actual topology
 		if(bIsOnActualTopology)
@@ -6447,6 +6566,9 @@ bool URigHierarchy::IsDependentOn(FRigBaseElement* InDependent, FRigBaseElement*
 	{
 		ElementDependencyCache.Get().FindOrAdd(CacheKey, false);
 	}
+
+	OutDependencyChain.Pop();
+
 	return false;
 }
 
diff --git a/Engine/Plugins/Animation/ControlRig/Source/ControlRig/Private/Rigs/RigHierarchyController.cpp b/Engine/Plugins/Animation/ControlRig/Source/ControlRig/Private/Rigs/RigHierarchyController.cpp
--- a/Engine/Plugins/Animation/ControlRig/Source/ControlRig/Private/Rigs/RigHierarchyController.cpp
+++ b/Engine/Plugins/Animation/ControlRig/Source/ControlRig/Private/Rigs/RigHierarchyController.cpp
@@ -1372,6 +1372,8 @@ TArray<FRigElementKey> URigHierarchyController::ImportBones(
 	bool bSelectBones,
 	bool bSetupUndo)
 {
+	const TGuardValue<bool> CycleCheckGuard(bSuspendParentCycleCheck, true);
+
 	TArray<FRigElementKey> AddedBones;
 
 	if(!IsValid())
@@ -1624,6 +1626,8 @@ TArray<FRigElementKey> URigHierarchyController::ImportBonesFromSkeletalMesh(
 TArray<FRigElementKey> URigHierarchyController::ImportSocketsFromSkeletalMesh(USkeletalMesh* InSkeletalMesh, const FName& InNameSpace,
 	bool bReplaceExistingSockets, bool bRemoveObsoleteSockets, bool bSelectSockets, bool bSetupUndo, bool bPrintPythonCommand)
 {
+	const TGuardValue<bool> CycleCheckGuard(bSuspendParentCycleCheck, true);
+	
 	TArray<FRigElementKey> SocketKeys;
 	
 	if(URigHierarchy* Hierarchy = GetHierarchy())
@@ -2134,6 +2138,8 @@ TArray<FRigElementKey> URigHierarchyController::ImportFromText(FString InContent
 
 TArray<FRigElementKey> URigHierarchyController::ImportFromText(FString InContent, ERigElementType InAllowedTypes, bool bReplaceExistingElements, bool bSelectNewElements, bool bSetupUndo, bool bPrintPythonCommands)
 {
+	const TGuardValue<bool> CycleCheckGuard(bSuspendParentCycleCheck, true);
+
 	TArray<FRigElementKey> PastedKeys;
 	if(!IsValid())
 	{
@@ -2414,6 +2420,8 @@ TArray<FRigElementKey> URigHierarchyController::ImportFromText(FString InContent
 
 TArray<FRigElementKey> URigHierarchyController::ImportFromHierarchyContainer(const FRigHierarchyContainer& InContainer, bool bIsCopyAndPaste)
 {
+	const TGuardValue<bool> CycleCheckGuard(bSuspendParentCycleCheck, true);
+	
 	URigHierarchy* Hierarchy = GetHierarchy();
 
 	TMap<FRigElementKey, FRigElementKey> KeyMap;;
@@ -2905,7 +2913,8 @@ int32 URigHierarchyController::AddElement(FRigBaseElement* InElementToAdd, FRigB
 	}
 	
 	{
-		const TGuardValue<bool> Guard(bSuspendAllNotifications, true);
+		const TGuardValue<bool> NotificationGuard(bSuspendAllNotifications, true);
+		const TGuardValue<bool> CycleCheckGuard(bSuspendParentCycleCheck, true);
 		SetParent(InElementToAdd, InFirstParent, bMaintainGlobalTransform);
 	}
 
@@ -3744,10 +3753,13 @@ bool URigHierarchyController::AddParent(FRigBaseElement* InChild, FRigBaseElemen
 
 	URigHierarchy* Hierarchy = GetHierarchy();
 
-	if(Hierarchy->IsParentedTo(InParent, InChild))
+	if (!(bSuspendParentCycleCheck && Hierarchy->GetChildren(InChild).IsEmpty()))
 	{
-		ReportErrorf(TEXT("Cannot parent '%s' to '%s' - would cause a cycle."), *InChild->Key.ToString(), *InParent->Key.ToString());
-		return false;
+		if(Hierarchy->IsParentedTo(InParent, InChild))
+		{
+			ReportErrorf(TEXT("Cannot parent '%s' to '%s' - would cause a cycle."), *InChild->Key.ToString(), *InParent->Key.ToString());
+			return false;
+		}
 	}
 
 	Hierarchy->EnsureCacheValidity();
diff --git a/Engine/Plugins/Animation/ControlRig/Source/ControlRig/Public/Rigs/RigHierarchy.h b/Engine/Plugins/Animation/ControlRig/Source/ControlRig/Public/Rigs/RigHierarchy.h
--- a/Engine/Plugins/Animation/ControlRig/Source/ControlRig/Public/Rigs/RigHierarchy.h
+++ b/Engine/Plugins/Animation/ControlRig/Source/ControlRig/Public/Rigs/RigHierarchy.h
@@ -3706,7 +3706,7 @@ public:
 	UFUNCTION(BlueprintCallable, Category = URigHierarchy)
 	bool IsParentedTo(FRigElementKey InChild, FRigElementKey InParent) const
 	{
-		return IsParentedTo(GetIndex(InChild), GetIndex(InParent));
+		return IsParentedTo(GetIndex(InChild), GetIndex(InParent), FEmptyRigDependenciesProvider());
 	}
 
 	/**
@@ -3714,13 +3714,14 @@ public:
 	 * @param InChildIndex The index of the child element to check for a parent
 	 * @param InParentIndex The index of the parent element to check for
 	 * @param InDependencyProvider An additional object providing dependencies to respect
+	 * @param OutDependencyChain An optional output error message indicating the parent cycle
 	 * @return True if the given parent and given child have a parent-child relationship
 	 */
-    bool IsParentedTo(int32 InChildIndex, int32 InParentIndex, const IRigDependenciesProvider& InDependencyProvider = FEmptyRigDependenciesProvider()) const
+    bool IsParentedTo(int32 InChildIndex, int32 InParentIndex, const IRigDependenciesProvider& InDependencyProvider, TArray<FRigHierarchyRecord>* OutDependencyChain = nullptr) const
 	{
 		if(Elements.IsValidIndex(InChildIndex) && Elements.IsValidIndex(InParentIndex))
 		{
-			return IsParentedTo(Elements[InChildIndex], Elements[InParentIndex], InDependencyProvider);
+			return IsParentedTo(Elements[InChildIndex], Elements[InParentIndex], InDependencyProvider, OutDependencyChain);
 		}
 		return false;
 	}
@@ -4444,11 +4445,13 @@ public:
 	 * @param InChild The child element to check for a parent
 	 * @param InParent The parent element to check for
 	 * @param InDependencyProvider An additional object providing dependencies to respect
+	 * @param OutDependencyChain An optional output message to describe the relationship between elements
 	 * @return True if the given parent and given child have a parent-child relationship
 	 */
 	CONTROLRIG_API bool IsParentedTo(FRigBaseElement* InChild,
 		FRigBaseElement* InParent,
-		const IRigDependenciesProvider& InDependencyProvider = FEmptyRigDependenciesProvider()) const;
+		const IRigDependenciesProvider& InDependencyProvider = FEmptyRigDependenciesProvider(),
+		TArray<FRigHierarchyRecord>* OutDependencyChain = nullptr) const;
 
 private:
 	/**
@@ -4461,8 +4464,9 @@ private:
 	 */
 	CONTROLRIG_API bool IsDependentOn(FRigBaseElement* InDependent,
 		FRigBaseElement* InDependency,
-		const IRigDependenciesProvider& InDependencyProvider = FEmptyRigDependenciesProvider(),
-		bool bIsOnActualTopology = true) const;
+		const IRigDependenciesProvider& InDependencyProvider,
+		bool bIsOnActualTopology,
+		TArray<FRigHierarchyRecord>& OutDependencyChain) const;
 
 public:
 	/**
@@ -4681,6 +4685,7 @@ private:
 	// Storage for the elements
 	mutable TArray<FRigBaseElement*> Elements;
 	mutable FTransactionallySafeCriticalSection ElementsLock;
+	mutable FTransactionallySafeCriticalSection IsDependentOnLock;
 
 	// Storage for the elements
 	mutable TArray<TArray<FRigBaseElement*>> ElementsPerType;
@@ -4921,7 +4926,6 @@ private:
 	TArray<int32> ElementsToRetainLocalTransform;
 
 	mutable THierarchyCache<TMap<TTuple<int32, int32>, bool>> ElementDependencyCache;
-	mutable TArray<bool> ElementDependencyVisited;
 
 	bool bIsCopyingHierarchy;
 
diff --git a/Engine/Plugins/Animation/ControlRig/Source/ControlRig/Public/Rigs/RigHierarchyController.h b/Engine/Plugins/Animation/ControlRig/Source/ControlRig/Public/Rigs/RigHierarchyController.h
--- a/Engine/Plugins/Animation/ControlRig/Source/ControlRig/Public/Rigs/RigHierarchyController.h
+++ b/Engine/Plugins/Animation/ControlRig/Source/ControlRig/Public/Rigs/RigHierarchyController.h
@@ -24,6 +24,7 @@ public:
 	, bSuspendAllNotifications(false)
 	, bSuspendSelectionNotifications(false)
 	, bSuspendPythonPrinting(false)
+	, bSuspendParentCycleCheck(false)
 	, CurrentInstructionIndex(INDEX_NONE)
 	{}
 
@@ -1132,6 +1133,11 @@ private:
 	*/
 	bool bSuspendPythonPrinting;
 
+	/**
+	 * If set to true parent cycle checking will be disabled.
+	 */
+	bool bSuspendParentCycleCheck;
+
 	/**
 	 * If set the controller will mark new items as procedural and created at the current instruction
 	 */
