Control Rig: Improve doxygen for changes in 44850983

#rnx

diff --git a/Engine/Plugins/Animation/ControlRig/Source/ControlRig/Public/Rigs/RigHierarchy.h b/Engine/Plugins/Animation/ControlRig/Source/ControlRig/Public/Rigs/RigHierarchy.h
--- a/Engine/Plugins/Animation/ControlRig/Source/ControlRig/Public/Rigs/RigHierarchy.h
+++ b/Engine/Plugins/Animation/ControlRig/Source/ControlRig/Public/Rigs/RigHierarchy.h
@@ -4445,7 +4445,7 @@ public:
 	 * @param InChild The child element to check for a parent
 	 * @param InParent The parent element to check for
 	 * @param InDependencyProvider An additional object providing dependencies to respect
-	 * @param OutDependencyChain An optional output message to describe the relationship between elements
+	 * @param OutDependencyChain An optional output structural chain providing insight into the chain causing the dependency cycle
 	 * @return True if the given parent and given child have a parent-child relationship
 	 */
 	CONTROLRIG_API bool IsParentedTo(FRigBaseElement* InChild,
@@ -4460,6 +4460,7 @@ private:
 	 * @param InDependency The dependency element to check for
 	 * @param InDependencyProvider An additional object providing dependencies to respect
 	 * @param bIsOnActualTopology Indicates that the passed dependent and dependency are expected to be on the current topology (if false they are provided with the dependency map)
+	 * @param OutDependencyChain An optional output structural chain providing insight into the chain causing the dependency cycle
 	 * @return True if the given dependent is affected by the given dependency 
 	 */
 	CONTROLRIG_API bool IsDependentOn(FRigBaseElement* InDependent,
diff --git a/Engine/Plugins/Animation/ControlRig/Source/ControlRig/Public/Rigs/RigHierarchyController.h b/Engine/Plugins/Animation/ControlRig/Source/ControlRig/Public/Rigs/RigHierarchyController.h
--- a/Engine/Plugins/Animation/ControlRig/Source/ControlRig/Public/Rigs/RigHierarchyController.h
+++ b/Engine/Plugins/Animation/ControlRig/Source/ControlRig/Public/Rigs/RigHierarchyController.h
@@ -1134,7 +1134,8 @@ private:
 	bool bSuspendPythonPrinting;
 
 	/**
-	 * If set to true parent cycle checking will be disabled.
+	 * If set to true parent cycle checking will be only deployed for elements
+	 * that already have children. This is an optimization used in a few cases only.
 	 */
 	bool bSuspendParentCycleCheck;
 
