Control Rig: Fix infinite recursion in dependency traverser - variables case

#rnx

diff --git a/Engine/Plugins/Animation/ControlRig/Source/ControlRig/Private/Rigs/RigHierarchy.cpp b/Engine/Plugins/Animation/ControlRig/Source/ControlRig/Private/Rigs/RigHierarchy.cpp
--- a/Engine/Plugins/Animation/ControlRig/Source/ControlRig/Private/Rigs/RigHierarchy.cpp
+++ b/Engine/Plugins/Animation/ControlRig/Source/ControlRig/Private/Rigs/RigHierarchy.cpp
@@ -6535,6 +6535,10 @@ bool URigHierarchy::IsDependentOn(FRigBaseElement* InDependent, FRigBaseElement*
 					}
 					else if (DependentRecord.IsInstruction() || DependentRecord.IsVariable())
 					{
+						if (OutDependencyChain.Contains(DependentRecord))
+						{
+							return false;
+						}
 						OutDependencyChain.Push(DependentRecord);
 						
 						if (FoundDependencyInProvider(InHierarchy, InDependency, InDependencyProvider, bIsOnActualTopology, CacheKey, DependentRecord, OutDependencyChain))
