Control Rig: Avoid respecting the dependency provider when executing space switching from sequencer

#rnx

diff --git a/Engine/Plugins/Animation/ControlRig/Source/ControlRig/Private/ControlRig.cpp b/Engine/Plugins/Animation/ControlRig/Source/ControlRig/Private/ControlRig.cpp
--- a/Engine/Plugins/Animation/ControlRig/Source/ControlRig/Private/ControlRig.cpp
+++ b/Engine/Plugins/Animation/ControlRig/Source/ControlRig/Private/ControlRig.cpp
@@ -2362,8 +2362,10 @@ void UControlRig::SwitchToParent(const FRigElementKey& InElementKey, const FRigE
 	else
 	{
 #if WITH_EDITOR
-		FRigDependenciesProviderForControlRig DependenciesProvider(this);
-		DynamicHierarchy->SwitchToParent(InElementKey, InNewParentKey, bInitial, bAffectChildren, DependenciesProvider, nullptr);
+		// when simply executing space switching we need to make sure it runs through,
+		// and during that we'll only respect parent child relationships (while ignoring VM dependencies).
+		static const FEmptyRigDependenciesProvider EmptyRigDependenciesProvider;
+		DynamicHierarchy->SwitchToParent(InElementKey, InNewParentKey, bInitial, bAffectChildren, EmptyRigDependenciesProvider, nullptr);
 #else
 		DynamicHierarchy->SwitchToParent(InElementKey, InNewParentKey, bInitial, bAffectChildren);
 #endif
