Control Rig: Fix crash when selecting to swap module when no module is selected
#rxn

diff --git a/Engine/Plugins/Animation/ControlRig/Source/ControlRigEditor/Private/Editor/SModularRigModel.cpp b/Engine/Plugins/Animation/ControlRig/Source/ControlRigEditor/Private/Editor/SModularRigModel.cpp
--- a/Engine/Plugins/Animation/ControlRig/Source/ControlRigEditor/Private/Editor/SModularRigModel.cpp
+++ b/Engine/Plugins/Animation/ControlRig/Source/ControlRigEditor/Private/Editor/SModularRigModel.cpp
@@ -1140,6 +1140,12 @@ void SModularRigModel::HandleSwapClassForModules()
 			}
 			return FName(NAME_None);
 		});
+
+		if (SelectedModuleNames.IsEmpty())
+		{
+			return;
+		}
+		
 		HandleSwapClassForModules(SelectedModuleNames);
 	}
 }
