Control Settings: avoid sending useless ControlModified() notifications if the settings haven't actually changed

#rnx

diff --git a/Engine/Plugins/Animation/ControlRig/Source/ControlRig/Private/Units/Execution/RigUnit_DynamicHierarchy.cpp b/Engine/Plugins/Animation/ControlRig/Source/ControlRig/Private/Units/Execution/RigUnit_DynamicHierarchy.cpp
--- a/Engine/Plugins/Animation/ControlRig/Source/ControlRig/Private/Units/Execution/RigUnit_DynamicHierarchy.cpp
+++ b/Engine/Plugins/Animation/ControlRig/Source/ControlRig/Private/Units/Execution/RigUnit_DynamicHierarchy.cpp
@@ -1532,10 +1532,22 @@ FRigUnit_HierarchySetShapeSettings_Execute()
 	{
 		if(FRigControlElement* ControlElement = Hierarchy->Find<FRigControlElement>(Item))
 		{
-			ControlElement->Settings.bShapeVisible = Settings.bVisible;
-			ControlElement->Settings.ShapeColor = Settings.Color;
-			ControlElement->Settings.ShapeName = Settings.Name;
-			Hierarchy->Notify(ERigHierarchyNotification::ControlSettingChanged, ControlElement);
+			FRigControlSettings& CurrentSettings = ControlElement->Settings;
+
+			const bool bNotifySettingsChanged =
+				CurrentSettings.bShapeVisible != Settings.bVisible ||
+				CurrentSettings.ShapeColor != Settings.Color ||
+				CurrentSettings.ShapeName != Settings.Name;
+
+			CurrentSettings.bShapeVisible = Settings.bVisible;
+			CurrentSettings.ShapeColor = Settings.Color;
+			CurrentSettings.ShapeName = Settings.Name;
+
+			if (bNotifySettingsChanged)
+			{
+				Hierarchy->Notify(ERigHierarchyNotification::ControlSettingChanged, ControlElement);
+			}
+			
 			Hierarchy->SetControlShapeTransformByIndex(ControlElement->GetIndex(), Settings.Transform, true, false);
 			Hierarchy->SetControlShapeTransformByIndex(ControlElement->GetIndex(), Settings.Transform, false, false);
 		}
