Control Rig: Fix URigHierarchy::GetControlValue would not return the initial value for transforms with value type ERigControlValueType::Initial

#rnx

diff --git a/Engine/Plugins/Animation/ControlRig/Source/ControlRig/Private/Rigs/RigHierarchy.cpp b/Engine/Plugins/Animation/ControlRig/Source/ControlRig/Private/Rigs/RigHierarchy.cpp
--- a/Engine/Plugins/Animation/ControlRig/Source/ControlRig/Private/Rigs/RigHierarchy.cpp
+++ b/Engine/Plugins/Animation/ControlRig/Source/ControlRig/Private/Rigs/RigHierarchy.cpp
@@ -5867,6 +5867,7 @@ FRigControlValue URigHierarchy::GetControlValue(FRigControlElement* InControlEle
 			}
 			
 			const bool bInitial = InValueType == ERigControlValueType::Initial;
+			const ERigTransformType::Type TransformType = bInitial ? InitialLocal : CurrentLocal;
 			switch(InControlElement->Settings.ControlType)
 			{
 				case ERigControlType::Rotator:
@@ -5876,7 +5877,7 @@ FRigControlValue URigHierarchy::GetControlValue(FRigControlElement* InControlEle
 				}
 				case ERigControlType::EulerTransform:
 				{
-					FEulerTransform EulerTransform(GetTransform(InControlElement, CurrentLocal));
+					FEulerTransform EulerTransform(GetTransform(InControlElement, TransformType));
 					EulerTransform.Rotation = InControlElement->PreferredEulerAngles.GetRotator(bInitial);
 					Value = MakeControlValueFromEulerTransform(EulerTransform); 
 					return true;
