Anim in Engine - Selected Anim Details properties now navigates to their value widget when pressing tab

diff --git a/Engine/Plugins/Animation/ControlRig/Source/ControlRigEditor/Private/AnimDetails/AnimDetailsNavigableWidgetRegistrar.cpp b/Engine/Plugins/Animation/ControlRig/Source/ControlRigEditor/Private/AnimDetails/AnimDetailsNavigableWidgetRegistrar.cpp
--- a/Engine/Plugins/Animation/ControlRig/Source/ControlRigEditor/Private/AnimDetails/AnimDetailsNavigableWidgetRegistrar.cpp
+++ b/Engine/Plugins/Animation/ControlRig/Source/ControlRigEditor/Private/AnimDetails/AnimDetailsNavigableWidgetRegistrar.cpp
@@ -10,23 +10,44 @@ namespace UE::ControlRigEditor
 	{
 		if (!IsEngineExitRequested() &&
 			WeakProxyManager.IsValid() &&
-			WeakWidget.IsValid())
+			WeakOwnerWidget.IsValid())
 		{
-			WeakProxyManager->GetNavigableWidgetRegistry().Unregister(WeakWidget.Pin().ToSharedRef());
+			WeakProxyManager->GetNavigableWidgetRegistry().UnregisterNavigableWidget(WeakOwnerWidget.Pin().ToSharedRef());
+			WeakProxyManager->GetNavigableWidgetRegistry().UnregisterNavigatorWidget(WeakOwnerWidget.Pin().ToSharedRef());
 		}
 	}
 
-	void FAnimDetailsNavigableWidgetRegistrar::Register(UAnimDetailsProxyManager& ProxyManager, const TSharedRef<SWidget> Widget)
+	void FAnimDetailsNavigableWidgetRegistrar::RegisterAsNavigable(
+		UAnimDetailsProxyManager& ProxyManager, 
+		const TSharedRef<SWidget>& OwnerWidget,
+		const TSharedRef<SWidget>& NavigableWidget)
 	{
 		if (WeakProxyManager.IsValid() &&
-			WeakWidget.IsValid())
+			WeakOwnerWidget.IsValid())
 		{
-			WeakProxyManager->GetNavigableWidgetRegistry().Unregister(WeakWidget.Pin().ToSharedRef());
+			WeakProxyManager->GetNavigableWidgetRegistry().UnregisterNavigableWidget(WeakOwnerWidget.Pin().ToSharedRef());
 		}
 
-		ProxyManager.GetNavigableWidgetRegistry().Register(Widget);
+		ProxyManager.GetNavigableWidgetRegistry().RegisterNavigableWidget(OwnerWidget, NavigableWidget);
 
 		WeakProxyManager = &ProxyManager;
-		WeakWidget = Widget;
+		WeakOwnerWidget = NavigableWidget;
+	}
+
+	void FAnimDetailsNavigableWidgetRegistrar::RegisterAsNavigator(
+		UAnimDetailsProxyManager& ProxyManager, 
+		const TSharedRef<SWidget>& NavigatorWidget, 
+		const TSharedRef<SWidget>& NavigateToOwner)
+	{
+		if (WeakProxyManager.IsValid() &&
+			WeakOwnerWidget.IsValid())
+		{
+			WeakProxyManager->GetNavigableWidgetRegistry().UnregisterNavigatorWidget(WeakOwnerWidget.Pin().ToSharedRef());
+		}
+
+		ProxyManager.GetNavigableWidgetRegistry().RegisterNavigatorWidget(NavigatorWidget, NavigateToOwner);
+
+		WeakProxyManager = &ProxyManager;
+		WeakOwnerWidget = NavigatorWidget;
 	}
 }
diff --git a/Engine/Plugins/Animation/ControlRig/Source/ControlRigEditor/Private/AnimDetails/AnimDetailsNavigableWidgetRegistrar.h b/Engine/Plugins/Animation/ControlRig/Source/ControlRigEditor/Private/AnimDetails/AnimDetailsNavigableWidgetRegistrar.h
--- a/Engine/Plugins/Animation/ControlRig/Source/ControlRigEditor/Private/AnimDetails/AnimDetailsNavigableWidgetRegistrar.h
+++ b/Engine/Plugins/Animation/ControlRig/Source/ControlRigEditor/Private/AnimDetails/AnimDetailsNavigableWidgetRegistrar.h
@@ -15,11 +15,33 @@ namespace UE::ControlRigEditor
 	{
 		~FAnimDetailsNavigableWidgetRegistrar();
 
-		/** Registers a widget with the navigable widget registry */
-		void Register(UAnimDetailsProxyManager& ProxyManager, const TSharedRef<SWidget> Widget);
+		/** 
+		 * Registers a widget that can be navigated to with the navigable widget registry 
+		 * 
+		 * @param ProxyManager		The proxy manager that owns the registry
+		 * @param OwnerWidget		The widget that owns the navigable widget
+		 * @param NavigableWidget	The widget to navigate to
+		 */
+		void RegisterAsNavigable(
+			UAnimDetailsProxyManager& ProxyManager, 
+			const TSharedRef<SWidget>& OwnerWidget,  
+			const TSharedRef<SWidget>& NavigableWidget);
+
+		/** Registers a widget that can navigate to an owner with the navigable widget registry. 
+		 * @param ProxyManager		The proxy manager that owns the registry
+		 * @param NavigatorWidget	The widget that invokes the navigation
+		 * @param NavigateToOwner	The owner widget to navigate to.
+		 */
+		void RegisterAsNavigator(
+			UAnimDetailsProxyManager& ProxyManager, 
+			const TSharedRef<SWidget>& NavigatorWidget, 
+			const TSharedRef<SWidget>& NavigateToOwner);
 
 	private:
+		/** The proxy manager the widgets are registered with */
 		TWeakObjectPtr<UAnimDetailsProxyManager> WeakProxyManager;
-		TWeakPtr<SWidget> WeakWidget;
+
+		/** The widget that owns this registrar */
+		TWeakPtr<SWidget> WeakOwnerWidget;
 	};
 }
diff --git a/Engine/Plugins/Animation/ControlRig/Source/ControlRigEditor/Private/AnimDetails/AnimDetailsNavigableWidgetRegistry.cpp b/Engine/Plugins/Animation/ControlRig/Source/ControlRigEditor/Private/AnimDetails/AnimDetailsNavigableWidgetRegistry.cpp
--- a/Engine/Plugins/Animation/ControlRig/Source/ControlRigEditor/Private/AnimDetails/AnimDetailsNavigableWidgetRegistry.cpp
+++ b/Engine/Plugins/Animation/ControlRig/Source/ControlRigEditor/Private/AnimDetails/AnimDetailsNavigableWidgetRegistry.cpp
@@ -2,28 +2,72 @@
 
 #include "AnimDetailsNavigableWidgetRegistry.h"
 
+#include "Algo/AnyOf.h"
 #include "Framework/Application/SlateApplication.h"
 
 namespace UE::ControlRigEditor
 {
-	const TSharedPtr<SWidget> FAnimDetailsNavigableWidgetRegistry::GetNext(const TSharedRef<SWidget>& Current) const
+	const TSharedPtr<SWidget> FAnimDetailsNavigableWidgetRegistry::GetNext(const TSharedRef<SWidget>& OwnerWidget) const
 	{
-		WeakWidgets.Remove(nullptr);
+		// Lazily cleanup
+		NavigableWidgets.RemoveAll(
+			[](const FNavigableWidget& NavigableWidget)
+			{
+				return 
+					!NavigableWidget.WeakOwner.IsValid() || 
+					!NavigableWidget.WeakNavigableWidget.IsValid();
+			});
+
+		for (auto It = WeakNavigatorToNavigableWidgetMap.CreateIterator(); It; ++It)
+		{
+			if (!(*It).Key.IsValid() || !(*It).Value.IsValid())
+			{
+				It.RemoveCurrent();
+			}
+		}
+
+		// If the widget is a navigator, navigate to its navigable widget
+		const TWeakPtr<SWidget>* OwnerWidgetPtr = WeakNavigatorToNavigableWidgetMap.Find(OwnerWidget);
+		if (OwnerWidgetPtr)
+		{
+			const FNavigableWidget* NavigableWidgetPtr = NavigableWidgets.FindByPredicate(
+				[OwnerWidgetPtr](const FNavigableWidget& NavigableWidget)
+				{
+					return NavigableWidget.WeakOwner == *OwnerWidgetPtr;
+				});
+			
+			const TSharedPtr<SWidget> NavigateToWidget = NavigableWidgetPtr && NavigableWidgetPtr->WeakNavigableWidget.IsValid() ?
+				NavigableWidgetPtr->WeakNavigableWidget.Pin() :
+				nullptr;
+
+			if (NavigateToWidget.IsValid() &&
+				NavigateToWidget->IsEnabled() &&
+				IsWidgetVisible(NavigateToWidget.ToSharedRef()))
+			{
+				return NavigateToWidget;
+			}
+		}
+
+		// If the widget is navigable, navigate to the next index
+		const int32 CurrentIndex = NavigableWidgets.IndexOfByPredicate(
+			[&OwnerWidget](const FNavigableWidget& NavigableWidget)
+			{
+				return NavigableWidget.WeakOwner == OwnerWidget;
+			});
 
-		const int32 CurrentIndex = WeakWidgets.IndexOfByKey(Current);
 		if (CurrentIndex == INDEX_NONE)
 		{
 			return nullptr;
 		}
 
-		for (int32 NextIndex = CurrentIndex + 1; NextIndex < WeakWidgets.Num(); NextIndex++)
+		for (int32 NextIndex = CurrentIndex + 1; NextIndex < NavigableWidgets.Num(); NextIndex++)
 		{
-			if (!WeakWidgets[NextIndex].IsValid())
+			if (!NavigableWidgets[NextIndex].WeakNavigableWidget.IsValid())
 			{
 				continue;
 			}
 
-			const TSharedRef<SWidget> NextWidget = WeakWidgets[NextIndex].Pin().ToSharedRef();
+			const TSharedRef<SWidget> NextWidget = NavigableWidgets[NextIndex].WeakNavigableWidget.Pin().ToSharedRef();
 			if (NextWidget->IsEnabled()) 
 			{
 				if (IsWidgetVisible(NextWidget))
@@ -53,13 +97,46 @@ namespace UE::ControlRigEditor
 		return true;
 	}
 
-	void FAnimDetailsNavigableWidgetRegistry::Register(const TSharedRef<SWidget>& InWidget)
+	void FAnimDetailsNavigableWidgetRegistry::RegisterNavigableWidget(
+		const TSharedRef<SWidget>& InOwnerWidget,
+		const TSharedRef<SWidget>& InNavigableWidget)
 	{
-		WeakWidgets.Add(InWidget);
+		NavigableWidgets.Emplace(InOwnerWidget, InNavigableWidget);
 	}
 
-	void FAnimDetailsNavigableWidgetRegistry::Unregister(const TSharedRef<SWidget>& InWidget)
+	void FAnimDetailsNavigableWidgetRegistry::UnregisterNavigableWidget(const TSharedRef<SWidget>& InOwnerWidget)
 	{
-		WeakWidgets.Remove(InWidget);
+		NavigableWidgets.RemoveAll(
+			[&InOwnerWidget](const FNavigableWidget& NavigableWidget)
+			{
+				return NavigableWidget.WeakOwner == InOwnerWidget;
+			});
 	}
+
+	void FAnimDetailsNavigableWidgetRegistry::RegisterNavigatorWidget(
+		const TSharedRef<SWidget>& InNavigatorWidget, 
+		const TSharedRef<SWidget>& InNavigateToWidget)
+	{
+		const bool bNavigableWidgetRegistered = Algo::AnyOf(NavigableWidgets, 
+			[&InNavigateToWidget](const FNavigableWidget& NavigableWidget)
+			{
+				return NavigableWidget.WeakOwner == InNavigateToWidget;
+			});
+		if (ensureMsgf(bNavigableWidgetRegistered, TEXT("Cannot register navigator widget if the navigate to widget is not already registered")))
+		{
+			WeakNavigatorToNavigableWidgetMap.Add(InNavigatorWidget, InNavigateToWidget);
+		}
+	}
+
+	void FAnimDetailsNavigableWidgetRegistry::UnregisterNavigatorWidget(const TSharedRef<SWidget>& InWidget)
+	{
+		WeakNavigatorToNavigableWidgetMap.Remove(InWidget);
+	}
+
+	FAnimDetailsNavigableWidgetRegistry::FNavigableWidget::FNavigableWidget(
+		const TSharedRef<SWidget>& InWeakOwner, 
+		const TSharedRef<SWidget>& InWeakNavigableWidget)
+		: WeakOwner(InWeakOwner)
+		, WeakNavigableWidget(InWeakNavigableWidget)
+	{}
 }
diff --git a/Engine/Plugins/Animation/ControlRig/Source/ControlRigEditor/Private/AnimDetails/AnimDetailsNavigableWidgetRegistry.h b/Engine/Plugins/Animation/ControlRig/Source/ControlRigEditor/Private/AnimDetails/AnimDetailsNavigableWidgetRegistry.h
--- a/Engine/Plugins/Animation/ControlRig/Source/ControlRigEditor/Private/AnimDetails/AnimDetailsNavigableWidgetRegistry.h
+++ b/Engine/Plugins/Animation/ControlRig/Source/ControlRigEditor/Private/AnimDetails/AnimDetailsNavigableWidgetRegistry.h
@@ -8,26 +8,58 @@ class SWidget;
 
 namespace UE::ControlRigEditor
 {
+	/** 
+	 * Registry for navigable widgets, assuming an owner widget and an inner widget that actually can receive navigation.
+	 * Used to navigate from property names or value wrappers (e.g. SAnimDetailsValueNumeric) to the actual value widgets. 
+	 * Given such an owner widget, GetNext will return the correct value widget to navigate to.
+	 */
 	class FAnimDetailsNavigableWidgetRegistry
 	{
 		// Let the registrar register and unregister widgets
 		friend struct FAnimDetailsNavigableWidgetRegistrar;
 
 	public:
-		/** Returns the next widget given the current, or nullptr if there is no next widget. Useful for navigation */
-		const TSharedPtr<SWidget> GetNext(const TSharedRef<SWidget>& Current) const;
+		/**
+		 * Returns the next widget given the owner widget, whereas the owner widget 
+		 * can be either registered as navigator or as navigable owner.
+		 * 
+		 * Returns nullptr if the widget is not registered with this registry. 
+		 */
+		const TSharedPtr<SWidget> GetNext(const TSharedRef<SWidget>& OwnerWidget) const;
 
 	private:
 		/** Does a deep test if a widget is truly visible. */
 		bool IsWidgetVisible(const TSharedRef<SWidget>& Widget) const;
 
-		/** Registers a widget */
-		void Register(const TSharedRef<SWidget>& InWidget);
+		/** Registers a widget that can be navigated to */
+		void RegisterNavigableWidget(
+			const TSharedRef<SWidget>& InOwnerWidget,
+			const TSharedRef<SWidget>& InNavigableWidget);
 
-		/** Unregisters a widget */
-		void Unregister(const TSharedRef<SWidget>& InWidget);
+		/** Unregisters a navigable widget by its owner */
+		void UnregisterNavigableWidget(const TSharedRef<SWidget>& InOwnerWidget);
 
-		/** The current widgets */
-		mutable TArray<TWeakPtr<SWidget>> WeakWidgets;
+		/** Registers a widget that can be navigated to a navigable widget. Requires the navigable widget to be registered. */
+		void RegisterNavigatorWidget(
+			const TSharedRef<SWidget>& InNavigatorWidget, 
+			const TSharedRef<SWidget>& InNavigateToWidget);
+
+		/** Unregisters a navigator widget */
+		void UnregisterNavigatorWidget(const TSharedRef<SWidget>& InOwnerWidget);
+
+		/** A widget that can be navigated to */
+		struct FNavigableWidget
+		{
+			FNavigableWidget(const TSharedRef<SWidget>& InWeakOwner, const TSharedRef<SWidget>& InWeakNavigableWidget);
+
+			const TWeakPtr<SWidget> WeakOwner;
+			const TWeakPtr<SWidget> WeakNavigableWidget;
+		};
+
+		/** The current navigable widgets */
+		mutable TArray<FNavigableWidget> NavigableWidgets;
+
+		/** The current navigator widgets with their navigate to widget */
+		mutable TMap<TWeakPtr<SWidget>, TWeakPtr<SWidget>> WeakNavigatorToNavigableWidgetMap;
 	};
 }
diff --git a/Engine/Plugins/Animation/ControlRig/Source/ControlRigEditor/Private/AnimDetails/Customizations/AnimDetailsValueCustomization.cpp b/Engine/Plugins/Animation/ControlRig/Source/ControlRigEditor/Private/AnimDetails/Customizations/AnimDetailsValueCustomization.cpp
--- a/Engine/Plugins/Animation/ControlRig/Source/ControlRigEditor/Private/AnimDetails/Customizations/AnimDetailsValueCustomization.cpp
+++ b/Engine/Plugins/Animation/ControlRig/Source/ControlRigEditor/Private/AnimDetails/Customizations/AnimDetailsValueCustomization.cpp
@@ -87,7 +87,11 @@ namespace UE::ControlRigEditor
 		IDetailChildrenBuilder& InStructBuilder, 
 		IPropertyTypeCustomizationUtils& InStructCustomizationUtils)
 	{
-		if (IsStructPropertyHiddenByFilter(InStructPropertyHandle))
+		FControlRigEditMode* EditMode = static_cast<FControlRigEditMode*>(GLevelEditorModeTools().GetActiveMode(FControlRigEditMode::ModeName));
+		UAnimDetailsProxyManager* ProxyManager = EditMode ? EditMode->GetAnimDetailsProxyManager() : nullptr;
+
+		if (!ProxyManager ||
+			IsStructPropertyHiddenByFilter(InStructPropertyHandle))
 		{
 			InStructPropertyHandle->MarkHiddenByCustomization();
 			return;
@@ -144,21 +148,20 @@ namespace UE::ControlRigEditor
 				.HAlign(HAlign_Fill)
 				.VAlign(VAlign_Fill)
 				[
-					SNew(SAnimDetailsPropertySelectionBorder, ChildHandle)
+					SNew(SAnimDetailsPropertySelectionBorder, *ProxyManager, ChildHandle)
+					.NavigateToWidget(ValueWidget)
 					[
 						SNew(STextBlock)
+						.Text(PropertyDisplayText)
 						.Font(FAppStyle::GetFontStyle("PropertyWindow.NormalFont"))
-						.Text_Lambda([PropertyDisplayText]()
-							{
-								return PropertyDisplayText;
-							})
 					]
 				]
 				.ValueContent()
 				.HAlign(HAlign_Fill)
 				[
-					SNew(SAnimDetailsPropertySelectionBorder, ChildHandle)
-					.RequiresModifierKeys(true)
+					SNew(SAnimDetailsPropertySelectionBorder, *ProxyManager, ChildHandle)
+					.NavigateToWidget(ValueWidget)
+					.RequiresModifierKeys(true)				
 					[
 						ValueWidget
 					]
@@ -328,6 +331,7 @@ namespace UE::ControlRigEditor
 		TSharedRef<IPropertyHandle>& PropertyHandle)
 	{
 		const FProperty* Property = PropertyHandle->GetProperty();
+
 		if (!Property)
 		{
 			return SNullWidget::NullWidget;
@@ -341,44 +345,32 @@ namespace UE::ControlRigEditor
 		if (PropertyClass == FDoubleProperty::StaticClass())
 		{
 			return
-				SNew(SAnimDetailsPropertySelectionBorder, PropertyHandle)
-				.RequiresModifierKeys(true)
+				SNew(SAnimDetailsValueNumeric<double>, PropertyHandle)
+				.IsEnabled(this, &FMathStructCustomization::IsValueEnabled, PropertyHandle.ToWeakPtr())
+				.LabelPadding(FMargin(3))
+				.LabelLocation(SAnimDetailsValueNumeric<double>::ELabelLocation::Inside)
+				.Label()
 				[
-					SNew(SAnimDetailsValueNumeric<double>, PropertyHandle)
-					.IsEnabled(this, &FMathStructCustomization::IsValueEnabled, PropertyHandle.ToWeakPtr())
-					.LabelPadding(FMargin(3))
-					.LabelLocation(SAnimDetailsValueNumeric<double>::ELabelLocation::Inside)
-					.Label()
-					[
-						Label
-					]
+					Label
 				];
 		}
 		else if (PropertyClass == FInt64Property::StaticClass())
 		{
 			return
-				SNew(SAnimDetailsPropertySelectionBorder, PropertyHandle)
-				.RequiresModifierKeys(true)
+				SNew(SAnimDetailsValueNumeric<int64>, PropertyHandle)
+				.IsEnabled(this, &FMathStructCustomization::IsValueEnabled, PropertyHandle.ToWeakPtr())
+				.LabelPadding(FMargin(3))
+				.LabelLocation(SAnimDetailsValueNumeric<int64>::ELabelLocation::Inside)
+				.Label()
 				[
-					SNew(SAnimDetailsValueNumeric<int64>, PropertyHandle)
-					.IsEnabled(this, &FMathStructCustomization::IsValueEnabled, PropertyHandle.ToWeakPtr())
-					.LabelPadding(FMargin(3))
-					.LabelLocation(SAnimDetailsValueNumeric<int64>::ELabelLocation::Inside)
-					.Label()
-					[
-						Label
-					]
+					Label
 				];
 		}
 		else if (PropertyClass == FBoolProperty::StaticClass())
 		{
 			return
-				SNew(SAnimDetailsPropertySelectionBorder, PropertyHandle)
-				.RequiresModifierKeys(true)
-				[
-					SNew(SAnimDetailsValueBoolean, PropertyHandle)
-					.IsEnabled(this, &FMathStructCustomization::IsValueEnabled, PropertyHandle.ToWeakPtr())
-				];
+				SNew(SAnimDetailsValueBoolean, PropertyHandle)
+				.IsEnabled(this, &FMathStructCustomization::IsValueEnabled, PropertyHandle.ToWeakPtr());
 		}
 
 		ensureMsgf(false, TEXT("Unsupported property class, cannot create an Anim Detail Values customization."));
diff --git a/Engine/Plugins/Animation/ControlRig/Source/ControlRigEditor/Private/AnimDetails/Widgets/SAnimDetailsPropertySelectionBorder.cpp b/Engine/Plugins/Animation/ControlRig/Source/ControlRigEditor/Private/AnimDetails/Widgets/SAnimDetailsPropertySelectionBorder.cpp
--- a/Engine/Plugins/Animation/ControlRig/Source/ControlRigEditor/Private/AnimDetails/Widgets/SAnimDetailsPropertySelectionBorder.cpp
+++ b/Engine/Plugins/Animation/ControlRig/Source/ControlRigEditor/Private/AnimDetails/Widgets/SAnimDetailsPropertySelectionBorder.cpp
@@ -45,8 +45,6 @@ namespace UE::ControlRigEditor
 				{
 					Instance = MakeUnique<FAnimDetailsInteractiveSelection>();
 
-					FSlateApplication::Get().SetKeyboardFocus(InstigatorWidget);
-
 					FCoreDelegates::OnEndFrame.AddRaw(Instance.Get(), &FAnimDetailsInteractiveSelection::TickInteractiveSelection);
 				}
 			}
@@ -76,8 +74,6 @@ namespace UE::ControlRigEditor
 					FCoreDelegates::OnEndFrame.RemoveAll(Instance.Get());
 
 					Instance.Reset();
-
-					FSlateApplication::Get().ClearKeyboardFocus();
 				}
 			}
 
@@ -88,12 +84,21 @@ namespace UE::ControlRigEditor
 		TUniquePtr<FAnimDetailsInteractiveSelection> FAnimDetailsInteractiveSelection::Instance;
 	}
 
-	void SAnimDetailsPropertySelectionBorder::Construct(const FArguments& InArgs, const TSharedRef<IPropertyHandle>& PropertyHandle)
+	void SAnimDetailsPropertySelectionBorder::Construct(
+		const FArguments& InArgs, 
+		UAnimDetailsProxyManager& InProxyManager, 
+		const TSharedRef<IPropertyHandle>& PropertyHandle)
 	{
 		SetCanTick(true);
 
+		WeakProxyManager = &InProxyManager;
 		PropertyName = PropertyHandle->GetProperty()->GetFName();
 			
+		if (InArgs._NavigateToWidget.IsValid())
+		{
+			NavigableWidgetRegistrar.RegisterAsNavigator(InProxyManager, AsShared(), InArgs._NavigateToWidget.ToSharedRef());
+		}
+
 		TArray<UObject*> OuterObjects;
 		PropertyHandle->GetOuterObjects(OuterObjects);
 
@@ -153,18 +158,6 @@ namespace UE::ControlRigEditor
 						}
 					})
 				.OnMouseButtonDown(this, &SAnimDetailsPropertySelectionBorder::OnBorderMouseButtonDown)
-				.OnMouseButtonUp_Lambda([this](const FGeometry& MyGeometry, const FPointerEvent& MouseEvent) 
-					{ 
-						// Prevent detail rows from handling events
-						if (bRequiresModifierKeys)
-						{
-							const bool bModifierKeyDown = FSlateApplication::Get().GetModifierKeys().IsControlDown() || FSlateApplication::Get().GetModifierKeys().IsShiftDown();
-
-							return bModifierKeyDown ? FReply::Handled() : FReply::Unhandled();
-						}
-						
-						return FReply::Handled();
-					})
 				.BorderImage(FAppStyle::GetBrush("NoBorder"))
 			]
 		];
@@ -200,10 +193,26 @@ namespace UE::ControlRigEditor
 	{
 		using namespace UE::ControlRigEditor::Private;
 
-		// Block any keyboard input to prevent from any keyboard shortcuts while interactively changing selection
-		return FAnimDetailsInteractiveSelection::IsChangingSelection() ?
-			FReply::Handled() :
-			FReply::Unhandled();
+		if (InKeyEvent.GetKey() == EKeys::Tab)
+		{
+			// Navigate to the value widget
+			const TSharedPtr<SWidget> NavigateToWidget = WeakProxyManager.IsValid() ?
+				WeakProxyManager->GetNavigableWidgetRegistry().GetNext(AsShared()) :
+				nullptr;
+
+			if (NavigateToWidget.IsValid())
+			{
+				return FReply::Handled().SetNavigation(NavigateToWidget.ToSharedRef(), ENavigationGenesis::Keyboard);
+			}
+		}
+		
+		// Block any keyboard besides tab-navigating to prevent from any keyboard shortcuts while interactively changing selection
+		if (FAnimDetailsInteractiveSelection::IsChangingSelection())
+		{
+			return FReply::Handled();
+		}
+
+		return FReply::Unhandled();
 	}
 
 	FReply SAnimDetailsPropertySelectionBorder::OnBorderMouseButtonDown(const FGeometry& MyGeometry, const FPointerEvent& MouseEvent)
@@ -257,7 +266,7 @@ namespace UE::ControlRigEditor
 			Selection->SelectPropertyInProxies(Proxies, PropertyName, SelectionType);
 		}
 
-		// Always handle clicks, it should not get to the details row below
+		// Always handle clicks, it should not get to the details row below.
 		return FReply::Handled();
 	}
 
diff --git a/Engine/Plugins/Animation/ControlRig/Source/ControlRigEditor/Private/AnimDetails/Widgets/SAnimDetailsPropertySelectionBorder.h b/Engine/Plugins/Animation/ControlRig/Source/ControlRigEditor/Private/AnimDetails/Widgets/SAnimDetailsPropertySelectionBorder.h
--- a/Engine/Plugins/Animation/ControlRig/Source/ControlRigEditor/Private/AnimDetails/Widgets/SAnimDetailsPropertySelectionBorder.h
+++ b/Engine/Plugins/Animation/ControlRig/Source/ControlRigEditor/Private/AnimDetails/Widgets/SAnimDetailsPropertySelectionBorder.h
@@ -2,6 +2,7 @@
 
 #pragma once
 
+#include "AnimDetails/AnimDetailsNavigableWidgetRegistrar.h"
 #include "Widgets/SCompoundWidget.h"
 
 class IPropertyHandle;
@@ -27,15 +28,22 @@ namespace UE::ControlRigEditor
 			/** If set to true, selection only occurs when a modifier key is pressed */
 			SLATE_ARGUMENT(bool, RequiresModifierKeys)
 
+			/** A widget that is navigated when tab-clicking on this border */
+			SLATE_ARGUMENT(TSharedPtr<SWidget>, NavigateToWidget)
+
 		SLATE_END_ARGS()
 
 		/** 
 		 * Constructs this widget
 		 * 
+		 * @param InArgs					Slate arguments
 		 * @param ProxyManager				The proxy manager that owns the proxies
 		 * @param PropertyHandle			Property handle for the property that can be selected
 		 */
-		void Construct(const FArguments& InArgs, const TSharedRef<IPropertyHandle>& PropertyHandle);
+		void Construct(
+			const FArguments& InArgs, 
+			UAnimDetailsProxyManager& InProxyManager, 
+			const TSharedRef<IPropertyHandle>& PropertyHandle);
 
 	protected:
 		//~ Begin SWidget interface
@@ -63,6 +71,9 @@ namespace UE::ControlRigEditor
 		/** The name of the property that can be selected */
 		FName PropertyName;
 
+		/** Registrar for the navigable widget registry */
+		FAnimDetailsNavigableWidgetRegistrar NavigableWidgetRegistrar;
+
 		/** If ture, selection should only occur when a modifier key is pressed */
 		bool bRequiresModifierKeys = false;
 	};
diff --git a/Engine/Plugins/Animation/ControlRig/Source/ControlRigEditor/Private/AnimDetails/Widgets/SAnimDetailsValueBoolean.cpp b/Engine/Plugins/Animation/ControlRig/Source/ControlRigEditor/Private/AnimDetails/Widgets/SAnimDetailsValueBoolean.cpp
--- a/Engine/Plugins/Animation/ControlRig/Source/ControlRigEditor/Private/AnimDetails/Widgets/SAnimDetailsValueBoolean.cpp
+++ b/Engine/Plugins/Animation/ControlRig/Source/ControlRigEditor/Private/AnimDetails/Widgets/SAnimDetailsValueBoolean.cpp
@@ -37,14 +37,14 @@ namespace UE::ControlRigEditor
 
 		ChildSlot
 		[
-			SAssignNew(SelectionBorder, SAnimDetailsPropertySelectionBorder, InPropertyHandle)
+			SAssignNew(SelectionBorder, SAnimDetailsPropertySelectionBorder, *ProxyManager, InPropertyHandle)
 			.RequiresModifierKeys(true)
 			[
 				CheckBox.ToSharedRef()
 			]
 		];
 
-		NavigableWidgetRegistrar.Register(*ProxyManager, CheckBox.ToSharedRef());
+		NavigableWidgetRegistrar.RegisterAsNavigable(*ProxyManager, AsShared(), CheckBox.ToSharedRef());
 
 		FAnimDetailsMultiEditUtil::Get().Join(ProxyManager, InPropertyHandle);
 	}
@@ -66,7 +66,7 @@ namespace UE::ControlRigEditor
 						return SNullWidget::NullWidget;
 					}
 
-					const TSharedPtr<SWidget> NextWidget = WeakProxyManager->GetNavigableWidgetRegistry().GetNext(CheckBox.ToSharedRef());
+					const TSharedPtr<SWidget> NextWidget = WeakProxyManager->GetNavigableWidgetRegistry().GetNext(AsShared());
 					if (!NextWidget.IsValid())
 					{
 						return SNullWidget::NullWidget;
diff --git a/Engine/Plugins/Animation/ControlRig/Source/ControlRigEditor/Private/AnimDetails/Widgets/SAnimDetailsValueNumeric.cpp b/Engine/Plugins/Animation/ControlRig/Source/ControlRigEditor/Private/AnimDetails/Widgets/SAnimDetailsValueNumeric.cpp
--- a/Engine/Plugins/Animation/ControlRig/Source/ControlRigEditor/Private/AnimDetails/Widgets/SAnimDetailsValueNumeric.cpp
+++ b/Engine/Plugins/Animation/ControlRig/Source/ControlRigEditor/Private/AnimDetails/Widgets/SAnimDetailsValueNumeric.cpp
@@ -206,7 +206,7 @@ namespace UE::ControlRigEditor
 						return SNullWidget::NullWidget;
 					}
 
-					const TSharedPtr<SWidget> NextWidget = WeakProxyManager->GetNavigableWidgetRegistry().GetNext(SpinBox.ToSharedRef());
+					const TSharedPtr<SWidget> NextWidget = WeakProxyManager->GetNavigableWidgetRegistry().GetNext(AsShared());
 					if (!NextWidget.IsValid())
 					{
 						return SNullWidget::NullWidget;
@@ -267,7 +267,7 @@ namespace UE::ControlRigEditor
 			.ClearKeyboardFocusOnCommit(true)
 			.TypeInterface(TypeInterface);
 
-		NavigableWidgetRegistrar.Register(*ProxyManager, SpinBox.ToSharedRef());
+		NavigableWidgetRegistrar.RegisterAsNavigable(*ProxyManager, AsShared(), SpinBox.ToSharedRef());
 
 		const TSharedRef<SHorizontalBox> ContentBox = SNew(SHorizontalBox);
 
