Control Rig: Allow spawning of visual cue controls
#rnx

diff --git a/Engine/Plugins/Animation/ControlRig/Source/ControlRig/Private/Units/Execution/RigUnit_DynamicHierarchy.cpp b/Engine/Plugins/Animation/ControlRig/Source/ControlRig/Private/Units/Execution/RigUnit_DynamicHierarchy.cpp
--- a/Engine/Plugins/Animation/ControlRig/Source/ControlRig/Private/Units/Execution/RigUnit_DynamicHierarchy.cpp
+++ b/Engine/Plugins/Animation/ControlRig/Source/ControlRig/Private/Units/Execution/RigUnit_DynamicHierarchy.cpp
@@ -490,11 +490,13 @@ FTransform FRigUnit_HierarchyAddControlElement::ProjectOffsetTransform(const FTr
 void FRigUnit_HierarchyAddControl_Settings::ConfigureFrom(const FRigControlElement* InControlElement, const FRigControlSettings& InSettings)
 {
 	DisplayName = InSettings.DisplayName;
+	bSelectable = InSettings.AnimationType != ERigControlAnimationType::VisualCue;
 }
 
 void FRigUnit_HierarchyAddControl_Settings::Configure(FRigControlSettings& OutSettings) const
 {
 	OutSettings.DisplayName = DisplayName;
+	OutSettings.AnimationType = bSelectable ? ERigControlAnimationType::AnimationControl : ERigControlAnimationType::VisualCue;
 }
 
 void FRigUnit_HierarchyAddControl_ShapeSettings::ConfigureFrom(const FRigControlElement* InControlElement, const FRigControlSettings& InSettings)
@@ -522,7 +524,10 @@ void FRigUnit_HierarchyAddControl_ProxySettings::ConfigureFrom(const FRigControl
 
 void FRigUnit_HierarchyAddControl_ProxySettings::Configure(FRigControlSettings& OutSettings) const
 {
-	OutSettings.AnimationType = bIsProxy ? ERigControlAnimationType::ProxyControl : ERigControlAnimationType::AnimationControl;
+	if (bIsProxy)
+	{
+		OutSettings.AnimationType = ERigControlAnimationType::ProxyControl;
+	}
 	OutSettings.DrivenControls = DrivenControls;
 	OutSettings.ShapeVisibility = ShapeVisibility;
 }
diff --git a/Engine/Plugins/Animation/ControlRig/Source/ControlRig/Public/Units/Execution/RigUnit_DynamicHierarchy.h b/Engine/Plugins/Animation/ControlRig/Source/ControlRig/Public/Units/Execution/RigUnit_DynamicHierarchy.h
--- a/Engine/Plugins/Animation/ControlRig/Source/ControlRig/Public/Units/Execution/RigUnit_DynamicHierarchy.h
+++ b/Engine/Plugins/Animation/ControlRig/Source/ControlRig/Public/Units/Execution/RigUnit_DynamicHierarchy.h
@@ -496,7 +496,7 @@ struct FRigUnit_HierarchyAddControl_Settings
 	GENERATED_BODY()
 
 	FRigUnit_HierarchyAddControl_Settings()
-		: DisplayName(NAME_None)
+		: DisplayName(NAME_None), bSelectable(true)
 	{}
 	virtual ~FRigUnit_HierarchyAddControl_Settings(){}
 
@@ -505,6 +505,9 @@ struct FRigUnit_HierarchyAddControl_Settings
 
 	UPROPERTY(EditAnywhere, BlueprintReadWrite, Category = "Settings")
 	FName DisplayName;
+
+	UPROPERTY(EditAnywhere, BlueprintReadWrite, Category = "Settings")
+	bool bSelectable;
 };
 
 USTRUCT(BlueprintType)
