Fix for BakeToControlRig function. Fetch the templateID for the sequence that we're in so that if the sequence we are working on is a nested, instance, or template sequence we can resolve objects within the function.

diff --git a/Engine/Plugins/Animation/ControlRig/Source/ControlRigEditor/Private/ControlRigSequencerEditorLibrary.cpp b/Engine/Plugins/Animation/ControlRig/Source/ControlRigEditor/Private/ControlRigSequencerEditorLibrary.cpp
--- a/Engine/Plugins/Animation/ControlRig/Source/ControlRigEditor/Private/ControlRigSequencerEditorLibrary.cpp
+++ b/Engine/Plugins/Animation/ControlRig/Source/ControlRigEditor/Private/ControlRigSequencerEditorLibrary.cpp
@@ -1226,7 +1226,8 @@ bool UControlRigSequencerEditorLibrary::BakeToControlRig(UWorld* World, ULevelSe
 
 	ALevelSequenceActor* OutActor = nullptr;
 	FMovieSceneSequencePlaybackSettings Settings;
-	FMovieSceneSequenceIDRef Template = MovieSceneSequenceID::Root;
+	const TSharedPtr<ISequencer> ParentSequencer = WeakSequencer.Pin();
+	FMovieSceneSequenceIDRef Template = ParentSequencer->GetFocusedTemplateID();
 	FMovieSceneSequenceTransform RootToLocalTransform;
 	IMovieScenePlayer* Player = nullptr;
 	ULevelSequencePlayer* LevelPlayer = nullptr;
@@ -1255,7 +1256,7 @@ bool UControlRigSequencerEditorLibrary::BakeToControlRig(UWorld* World, ULevelSe
 			LevelPlayer->SetPlaybackPosition(FMovieSceneSequencePlaybackParams(StartTime, EUpdatePositionMethod::Play));
 		}
 		TArrayView<TWeakObjectPtr<>>  Result = Player->FindBoundObjects(Binding.BindingID, Template);
-	
+		
 		if (Result.Num() > 0 && Result[0].IsValid())
 		{
 			UObject* BoundObject = Result[0].Get();
