Improvement to previous code submission to protect against a dereferenced pointer for the Sequencer interface.

diff --git a/Engine/Plugins/Animation/ControlRig/Source/ControlRigEditor/Private/ControlRigSequencerEditorLibrary.cpp b/Engine/Plugins/Animation/ControlRig/Source/ControlRigEditor/Private/ControlRigSequencerEditorLibrary.cpp
--- a/Engine/Plugins/Animation/ControlRig/Source/ControlRigEditor/Private/ControlRigSequencerEditorLibrary.cpp
+++ b/Engine/Plugins/Animation/ControlRig/Source/ControlRigEditor/Private/ControlRigSequencerEditorLibrary.cpp
@@ -1226,8 +1226,11 @@ bool UControlRigSequencerEditorLibrary::BakeToControlRig(UWorld* World, ULevelSe
 
 	ALevelSequenceActor* OutActor = nullptr;
 	FMovieSceneSequencePlaybackSettings Settings;
-	const TSharedPtr<ISequencer> ParentSequencer = WeakSequencer.Pin();
-	FMovieSceneSequenceIDRef Template = ParentSequencer->GetFocusedTemplateID();
+	FMovieSceneSequenceID Template = MovieSceneSequenceID::Root;
+	if (const TSharedPtr<ISequencer> ParentSequencer = WeakSequencer.Pin())
+	{
+		Template = ParentSequencer->GetFocusedTemplateID();
+	}
 	FMovieSceneSequenceTransform RootToLocalTransform;
 	IMovieScenePlayer* Player = nullptr;
 	ULevelSequencePlayer* LevelPlayer = nullptr;
@@ -1255,7 +1258,7 @@ bool UControlRigSequencerEditorLibrary::BakeToControlRig(UWorld* World, ULevelSe
 			FFrameTime StartTime = FFrameRate::TransformTime(UE::MovieScene::DiscreteInclusiveLower(MovieScene->GetPlaybackRange()).Value, MovieScene->GetTickResolution(), MovieScene->GetDisplayRate());
 			LevelPlayer->SetPlaybackPosition(FMovieSceneSequencePlaybackParams(StartTime, EUpdatePositionMethod::Play));
 		}
-		TArrayView<TWeakObjectPtr<>>  Result = Player->FindBoundObjects(Binding.BindingID, Template);
+		TArrayView<TWeakObjectPtr<>>  Result = Player->FindBoundObjects(Binding.BindingID, static_cast<FMovieSceneSequenceIDRef>(Template));
 		
 		if (Result.Num() > 0 && Result[0].IsValid())
 		{
