Control Rig: Fix crash when swapping functions
#rnx

diff --git a/Engine/Plugins/Animation/ControlRig/Source/ControlRigDeveloper/Private/Graph/ControlRigGraph.cpp b/Engine/Plugins/Animation/ControlRig/Source/ControlRigDeveloper/Private/Graph/ControlRigGraph.cpp
--- a/Engine/Plugins/Animation/ControlRig/Source/ControlRigDeveloper/Private/Graph/ControlRigGraph.cpp
+++ b/Engine/Plugins/Animation/ControlRig/Source/ControlRigDeveloper/Private/Graph/ControlRigGraph.cpp
@@ -166,7 +166,7 @@ const TArray<TSharedPtr<FRigVMStringWithTag>>* UControlRigGraph::GetElementNameL
 	
 	if(!ElementNameLists.Contains(InElementType))
 	{
-		const UControlRigBlueprint* Blueprint = Cast<UControlRigBlueprint>(GetAsset()->GetObject());
+		const UControlRigBlueprint* Blueprint = Cast<UControlRigBlueprint>(GetAsset().GetObject());
 		if(Blueprint == nullptr)
 		{
 			return &EmptyElementNameList;
