Control Rig & compensation: do not compute space & constraints compensation if the rig is already evaluating

when keyframes are added via the SendEvent nodes, compensation must not happen while the rig is being evaluated

#rnx

diff --git a/Engine/Plugins/Animation/ControlRig/Source/ControlRigEditor/Private/Sequencer/ControlRigParameterTrackEditor.cpp b/Engine/Plugins/Animation/ControlRig/Source/ControlRigEditor/Private/Sequencer/ControlRigParameterTrackEditor.cpp
--- a/Engine/Plugins/Animation/ControlRig/Source/ControlRigEditor/Private/Sequencer/ControlRigParameterTrackEditor.cpp
+++ b/Engine/Plugins/Animation/ControlRig/Source/ControlRigEditor/Private/Sequencer/ControlRigParameterTrackEditor.cpp
@@ -4641,16 +4641,23 @@ FKeyPropertyResult FControlRigParameterTrackEditor::AddKeysToControlRigHandle(UO
 			UMovieSceneControlRigParameterSection* ParamSection = Cast<UMovieSceneControlRigParameterSection>(Track->GetSectionToKey(ControlRigName));
 			if (UControlRig* SectionControlRig = ParamSection ? ParamSection->GetControlRig() : nullptr)
 			{
-				TOptional<FFrameNumber> OptionalKeyTime = KeyTime;
-
-				// compensate spaces
-				FControlRigSpaceChannelHelpers::CompensateIfNeeded(
-					SectionControlRig, Sequencer.Get(), ParamSection,
-					OptionalKeyTime, true /*comp previous*/);
+				// avoid reentrant evaluation: space switching & constraints evaluates the rig at T-1
+				// but this must not happen if it's already being evaluated.
+				if (!SectionControlRig->IsEvaluating())
+				{
+					TOptional<FFrameNumber> OptionalKeyTime = KeyTime;
 
-				// compensate constraints
-				const uint32 ControlHash = UTransformableControlHandle::ComputeHash(SectionControlRig, RigControlName);
-				FMovieSceneConstraintChannelHelper::CompensateIfNeeded(Sequencer, ParamSection, OptionalKeyTime, true /*bCompPreviousTick*/, ControlHash);
+					constexpr bool bComputePreviousTick = true;
+					
+					// compensate spaces
+					FControlRigSpaceChannelHelpers::CompensateIfNeeded(
+						SectionControlRig, Sequencer.Get(), ParamSection,
+						OptionalKeyTime, bComputePreviousTick);
+
+					// compensate constraints
+					const uint32 ControlHash = UTransformableControlHandle::ComputeHash(SectionControlRig, RigControlName);
+					FMovieSceneConstraintChannelHelper::CompensateIfNeeded(Sequencer, ParamSection, OptionalKeyTime, bComputePreviousTick, ControlHash);
+				}
 			}
 		}
 	}
