Fix for client crash in FRigUnit_HierarchyInstantiateFromPhysicsAsset, due to modular control rig overrides getting garbage collected

diff --git a/Engine/Plugins/Animation/ControlRig/Source/ControlRig/Private/ControlRigOverride.cpp b/Engine/Plugins/Animation/ControlRig/Source/ControlRig/Private/ControlRigOverride.cpp
--- a/Engine/Plugins/Animation/ControlRig/Source/ControlRig/Private/ControlRigOverride.cpp
+++ b/Engine/Plugins/Animation/ControlRig/Source/ControlRig/Private/ControlRigOverride.cpp
@@ -348,6 +348,54 @@ uint32 GetTypeHash(const FControlRigOverrideValue& InOverride)
 	return InOverride.Hash;
 }
 
+void FControlRigOverrideValue::AddStructReferencedObjects(FReferenceCollector& Collector) const
+{
+	if (const FProperty* LeafProperty = GetLeafProperty())
+	{
+		if (const FArrayProperty* ArrayProperty = CastField<FArrayProperty>(LeafProperty))
+		{
+			if (const FObjectProperty* ObjProperty = CastField<FObjectProperty>(ArrayProperty->Inner))
+			{
+				const void* DataPtr = GetData();
+				FScriptArrayHelper ArrayHelper(ArrayProperty, DataPtr);
+				for (int32 Index = 0; Index < ArrayHelper.Num(); Index++)
+				{
+					void* ElementData = ArrayHelper.GetRawPtr(Index);
+					if (TObjectPtr<UObject>* ObjPtr = (TObjectPtr<UObject>*)(ElementData))
+					{
+						Collector.AddReferencedObject(*ObjPtr);
+					}
+				}
+			}
+			else if (const FStructProperty* StructProperty = CastField<FStructProperty>(ArrayProperty->Inner))
+			{
+				const void* DataPtr = GetData();
+				FScriptArrayHelper ArrayHelper(ArrayProperty, DataPtr);
+				for (int32 Index = 0; Index < ArrayHelper.Num(); Index++)
+				{
+					UScriptStruct* Struct = StructProperty->Struct;
+					void* StructData = ArrayHelper.GetRawPtr(Index);
+					Collector.AddPropertyReferencesWithStructARO(Struct, StructData);
+				}
+			}
+		}
+		else
+		{
+			if (const FObjectProperty* ObjProperty = CastField<FObjectProperty>(LeafProperty))
+			{
+				if (TObjectPtr<const UObject>* ObjPtr = ObjProperty->ContainerPtrToValuePtr<TObjectPtr<const UObject>>(const_cast<void*>(GetData())))
+				{
+					Collector.AddReferencedObject(*ObjPtr);
+				}
+			}
+			else if (const FStructProperty* StructProperty = CastField<FStructProperty>(LeafProperty))
+			{
+				Collector.AddPropertyReferencesWithStructARO(StructProperty->Struct, const_cast<void*>(GetData()));
+			}
+		}
+	}
+}
+
 const FProperty* FControlRigOverrideValue::GetLeafProperty() const
 {
 	if(Properties.IsEmpty())
@@ -1363,6 +1411,8 @@ void FControlRigOverrideContainer::AddStructReferencedObjects(FReferenceCollecto
 				TWeakObjectPtr<UStruct> WeakStruct(Struct);
 				Collector.AddReferencedObject(WeakStruct);
 			}
+
+			Override.AddStructReferencedObjects(Collector);
 		}
 	}
 }
diff --git a/Engine/Plugins/Animation/ControlRig/Source/ControlRig/Public/ControlRigOverride.h b/Engine/Plugins/Animation/ControlRig/Source/ControlRig/Public/ControlRigOverride.h
--- a/Engine/Plugins/Animation/ControlRig/Source/ControlRig/Public/ControlRigOverride.h
+++ b/Engine/Plugins/Animation/ControlRig/Source/ControlRig/Public/ControlRigOverride.h
@@ -194,6 +194,8 @@ public:
 
 	UE_API friend uint32 GetTypeHash(const FControlRigOverrideValue& InOverride);
 
+	void AddStructReferencedObjects(FReferenceCollector& Collector) const;
+
 private:
 
 	UE_API bool SetPropertiesFromPath(const FString& InPath, const UStruct* InOwnerStruct);
