Fixed control module iterator lambda parameters creating copies of some expensive objects

diff --git a/Engine/Plugins/Animation/ControlRig/Source/ControlRig/Private/ModularRig.cpp b/Engine/Plugins/Animation/ControlRig/Source/ControlRig/Private/ModularRig.cpp
--- a/Engine/Plugins/Animation/ControlRig/Source/ControlRig/Private/ModularRig.cpp
+++ b/Engine/Plugins/Animation/ControlRig/Source/ControlRig/Private/ModularRig.cpp
@@ -338,7 +338,7 @@ bool UModularRig::InitializeVMs(const FName& InEventName)
 	URigVMHost::InitializeVM(InEventName);
 	UpdateModuleHierarchyFromCDO();
 
-	ForEachModule([InEventName](const FRigModuleInstance* Module) -> bool
+	ForEachModule([&InEventName](const FRigModuleInstance* Module) -> bool
 	{
 		if (UControlRig* ModuleRig = Module->GetRig())
 		{
@@ -380,7 +380,7 @@ void UModularRig::UpdateModuleHierarchyFromCDO()
 		// for each module (data only) reference in the model.
 		// Note: The CDO does not contain any UObject module instances itself.
 		const FModularRigModel& Model = GetModularRigModel();
-		Model.ForEachModule([this, Model](const FRigModuleReference* InModuleReference) -> bool
+		Model.ForEachModule([this, &Model](const FRigModuleReference* InModuleReference) -> bool
 		{
 			check(InModuleReference);
 			if (IsInGameThread() && !InModuleReference->Class.IsValid())
@@ -410,7 +410,7 @@ void UModularRig::UpdateModuleHierarchyFromCDO()
 
 		// update the module variable bindings now - since for this all
 		// modules have to exist first
-		ForEachModule([this, Model](const FRigModuleInstance* Module) -> bool
+		ForEachModule([this, &Model](const FRigModuleInstance* Module) -> bool
 		{
 			if(const FRigModuleReference* ModuleReference = Model.FindModule(Module->Name))
 			{
@@ -437,7 +437,7 @@ bool UModularRig::Execute_Internal(const FName& InEventName)
 		const FRigUnitContext& UnitContext = PublicContext.UnitContext;
 		const URigHierarchy* Hierarchy = GetHierarchy();
 
-		ForEachModule([&InEventName, this, Hierarchy, UnitContext](FRigModuleInstance* Module) -> bool
+		ForEachModule([this, &InEventName, Hierarchy, &UnitContext](FRigModuleInstance* Module) -> bool
 		{
 			if (const UControlRig* ModuleRig = Module->GetRig())
 			{
@@ -536,7 +536,7 @@ void UModularRig::GetControlsInOrder(TArray<FRigControlElement*>& SortedControls
 		bContinue = true;
 	});
 
-	ForEachModule([ControlsByModule, &SortedControls](const FRigModuleInstance* Module) -> bool
+	ForEachModule([&ControlsByModule, &SortedControls](const FRigModuleInstance* Module) -> bool
 	{
 		if (const TArray<FRigControlElement*>* Controls = ControlsByModule.Find(Module->Name))
 		{
@@ -1293,7 +1293,7 @@ void UModularRig::ForEachModule(TFunctionRef<bool(FRigModuleInstance*)> PerModul
 
 void UModularRig::ForEachModule(TFunctionRef<bool(const FRigModuleInstance*)> PerModuleFunction, bool bDepthFirst) const
 {
-	const_cast<UModularRig*>(this)->ForEachModule([PerModuleFunction](FRigModuleInstance* InModuleInstance) -> bool
+	const_cast<UModularRig*>(this)->ForEachModule([&PerModuleFunction](FRigModuleInstance* InModuleInstance) -> bool
 	{
 		return PerModuleFunction(InModuleInstance);
 	}, bDepthFirst);
@@ -1396,7 +1396,7 @@ TArray<FName> UModularRig::ExecuteEventOnAllModules(FName InEvent)
 {
 	TArray<FRigModuleExecutionElement> QueueForEvent;
 
-	ForEachModule([&QueueForEvent, InEvent](FRigModuleInstance* Module) -> bool
+	ForEachModule([&QueueForEvent, &InEvent](FRigModuleInstance* Module) -> bool
 	{
 		if(const UControlRig* ModuleRig = Module->GetRig())
 		{
diff --git a/Engine/Plugins/Animation/ControlRig/Source/ControlRig/Private/ModularRigController.cpp b/Engine/Plugins/Animation/ControlRig/Source/ControlRig/Private/ModularRigController.cpp
--- a/Engine/Plugins/Animation/ControlRig/Source/ControlRig/Private/ModularRigController.cpp
+++ b/Engine/Plugins/Animation/ControlRig/Source/ControlRig/Private/ModularRigController.cpp
@@ -1109,7 +1109,7 @@ TArray<FString> UModularRigController::GetPossibleBindings(const FName& InModule
 	}
 
 	// Add possible module variables
-	Model->ForEachModule([this, &PossibleBindings, InModuleName, InVariableName](const FRigModuleReference* InModule) -> bool
+	Model->ForEachModule([this, &PossibleBindings, &InModuleName, &InVariableName](const FRigModuleReference* InModule) -> bool
 	{
 		const FName CurModuleName = InModule->Name;
 		if (InModuleName != CurModuleName)
@@ -1950,7 +1950,7 @@ bool UModularRigController::SwapModulesOfClass(TSubclassOf<UControlRig> InOldCla
 	}
 #endif
 	
-	Model->ForEachModule([this, InOldClass, InNewClass, bSetupUndo](const FRigModuleReference* Module) -> bool
+	Model->ForEachModule([this, &InOldClass, &InNewClass, bSetupUndo](const FRigModuleReference* Module) -> bool
 	{
 		if (Module->Class.Get() == InOldClass)
 		{
@@ -2173,7 +2173,7 @@ void UModularRigController::RefreshModuleVariables(const FRigModuleReference* In
 	}
 
 	// If the module is the source of another module's binding, make sure it is still a valid binding
-	Model->ForEachModule([this, InModule, ModuleName, ModuleClass](const FRigModuleReference* OtherModule) -> bool
+	Model->ForEachModule([this, InModule, &ModuleName, ModuleClass](const FRigModuleReference* OtherModule) -> bool
 	{
 		if (InModule == OtherModule)
 		{
diff --git a/Engine/Plugins/Animation/ControlRig/Source/ControlRig/Private/ModularRigModel.cpp b/Engine/Plugins/Animation/ControlRig/Source/ControlRig/Private/ModularRigModel.cpp
--- a/Engine/Plugins/Animation/ControlRig/Source/ControlRig/Private/ModularRigModel.cpp
+++ b/Engine/Plugins/Animation/ControlRig/Source/ControlRig/Private/ModularRigModel.cpp
@@ -449,7 +449,7 @@ TArray<const FRigModuleReference*> FModularRigModel::FindModuleInstancesOfClass(
 	TArray<const FRigModuleReference*> Result;
 	FString ModuleClassPath = InModuleClassPath;
 	ModuleClassPath.RemoveFromEnd(TEXT("_C"));
-	ForEachModule([&Result, ModuleClassPath](const FRigModuleReference* Module) -> bool
+	ForEachModule([&Result, &ModuleClassPath](const FRigModuleReference* Module) -> bool
 	{
 		FString PackageName = Module->Class.ToSoftObjectPath().GetAssetPathString();
 		PackageName.RemoveFromEnd(TEXT("_C"));
@@ -517,7 +517,7 @@ bool FModularRigModel::TraverseModules(const FRigModuleReference* InModuleRefere
 TArray<FName> FModularRigModel::SortModuleNames(const TArray<FName>& InModuleNames) const
 {
 	TArray<FName> SortedModuleNames;
-	ForEachModule([InModuleNames, &SortedModuleNames](const FRigModuleReference* Module) -> bool
+	ForEachModule([&InModuleNames, &SortedModuleNames](const FRigModuleReference* Module) -> bool
 	{
 		if(InModuleNames.Contains(Module->GetFName()))
 		{
diff --git a/Engine/Plugins/Animation/ControlRig/Source/ControlRigEditor/Private/Editor/ControlRigEditor.cpp b/Engine/Plugins/Animation/ControlRig/Source/ControlRigEditor/Private/Editor/ControlRigEditor.cpp
--- a/Engine/Plugins/Animation/ControlRig/Source/ControlRigEditor/Private/Editor/ControlRigEditor.cpp
+++ b/Engine/Plugins/Animation/ControlRig/Source/ControlRigEditor/Private/Editor/ControlRigEditor.cpp
@@ -3471,7 +3471,7 @@ void FControlRigBaseEditor::HandlePreviewMeshChanged(USkeletalMesh* InOldSkeleta
 						UModularRigController* ModularRigController = ControlRigBP->GetModularRigController();
 						const bool bAutoResolve = ControlRigBP->ModularRigSettings.bAutoResolve;
 						Model->ForEachModule(
-							[Model, Hierarchy, ModularRigController, PreviousConnections, bAutoResolve]
+							[&Model, Hierarchy, ModularRigController, &PreviousConnections, bAutoResolve]
 							(const FRigModuleReference* Module) -> bool
 							{
 								bool bContinueResolval;
