Added null checks to ControlRigOverride AddStructReferencedObjects funtion after submit feedback

#rnx

diff --git a/Engine/Plugins/Animation/ControlRig/Source/ControlRig/Private/ControlRigOverride.cpp b/Engine/Plugins/Animation/ControlRig/Source/ControlRig/Private/ControlRigOverride.cpp
--- a/Engine/Plugins/Animation/ControlRig/Source/ControlRig/Private/ControlRigOverride.cpp
+++ b/Engine/Plugins/Animation/ControlRig/Source/ControlRig/Private/ControlRigOverride.cpp
@@ -356,26 +356,30 @@ void FControlRigOverrideValue::AddStructReferencedObjects(FReferenceCollector& C
 		{
 			if (const FObjectProperty* ObjProperty = CastField<FObjectProperty>(ArrayProperty->Inner))
 			{
-				const void* DataPtr = GetData();
-				FScriptArrayHelper ArrayHelper(ArrayProperty, DataPtr);
-				for (int32 Index = 0; Index < ArrayHelper.Num(); Index++)
+				if (const void* DataPtr = GetData())
 				{
-					void* ElementData = ArrayHelper.GetRawPtr(Index);
-					if (TObjectPtr<UObject>* ObjPtr = (TObjectPtr<UObject>*)(ElementData))
+					FScriptArrayHelper ArrayHelper(ArrayProperty, DataPtr);
+					for (int32 Index = 0; Index < ArrayHelper.Num(); Index++)
 					{
-						Collector.AddReferencedObject(*ObjPtr);
+						void* ElementData = ArrayHelper.GetRawPtr(Index);
+						if (TObjectPtr<UObject>* ObjPtr = (TObjectPtr<UObject>*)(ElementData))
+						{
+							Collector.AddReferencedObject(*ObjPtr);
+						}
 					}
 				}
 			}
 			else if (const FStructProperty* StructProperty = CastField<FStructProperty>(ArrayProperty->Inner))
 			{
-				const void* DataPtr = GetData();
-				FScriptArrayHelper ArrayHelper(ArrayProperty, DataPtr);
-				for (int32 Index = 0; Index < ArrayHelper.Num(); Index++)
+				if (const void* DataPtr = GetData())
 				{
-					UScriptStruct* Struct = StructProperty->Struct;
-					void* StructData = ArrayHelper.GetRawPtr(Index);
-					Collector.AddPropertyReferencesWithStructARO(Struct, StructData);
+					FScriptArrayHelper ArrayHelper(ArrayProperty, DataPtr);
+					for (int32 Index = 0; Index < ArrayHelper.Num(); Index++)
+					{
+						UScriptStruct* Struct = StructProperty->Struct;
+						void* StructData = ArrayHelper.GetRawPtr(Index);
+						Collector.AddPropertyReferencesWithStructARO(Struct, StructData);
+					}
 				}
 			}
 		}
@@ -390,7 +394,10 @@ void FControlRigOverrideValue::AddStructReferencedObjects(FReferenceCollector& C
 			}
 			else if (const FStructProperty* StructProperty = CastField<FStructProperty>(LeafProperty))
 			{
-				Collector.AddPropertyReferencesWithStructARO(StructProperty->Struct, const_cast<void*>(GetData()));
+				if (const void* DataPtr = GetData())
+				{
+					Collector.AddPropertyReferencesWithStructARO(StructProperty->Struct, const_cast<void*>(DataPtr));
+				}
 			}
 		}
 	}
