Control Rig: Add rig class property in FRigModularInstance so that the class is loaded before the modular rig
#rnx

diff --git a/Engine/Plugins/Animation/ControlRig/Source/ControlRig/Private/ModularRig.cpp b/Engine/Plugins/Animation/ControlRig/Source/ControlRig/Private/ModularRig.cpp
--- a/Engine/Plugins/Animation/ControlRig/Source/ControlRig/Private/ModularRig.cpp
+++ b/Engine/Plugins/Animation/ControlRig/Source/ControlRig/Private/ModularRig.cpp
@@ -94,6 +94,7 @@ void FRigModuleInstance::SetRig(UControlRig* InRig)
 
 	// update the cache
 	RigPtr = InRig;
+	RigClass = InRig->GetClass();
 }
 
 bool FRigModuleInstance::ContainsRig(const UControlRig* InRig) const
diff --git a/Engine/Plugins/Animation/ControlRig/Source/ControlRig/Public/ModularRig.h b/Engine/Plugins/Animation/ControlRig/Source/ControlRig/Public/ModularRig.h
--- a/Engine/Plugins/Animation/ControlRig/Source/ControlRig/Public/ModularRig.h
+++ b/Engine/Plugins/Animation/ControlRig/Source/ControlRig/Public/ModularRig.h
@@ -66,6 +66,9 @@ public:
 	UPROPERTY()
 	FName Name;
 
+	UPROPERTY()
+	TSubclassOf<UControlRig> RigClass;
+
 private:
 
 	UPROPERTY(transient)
