Control Rig: Fix possible crash when setting instance rig to nullptr
#rnx

diff --git a/Engine/Plugins/Animation/ControlRig/Source/ControlRig/Private/ModularRig.cpp b/Engine/Plugins/Animation/ControlRig/Source/ControlRig/Private/ModularRig.cpp
--- a/Engine/Plugins/Animation/ControlRig/Source/ControlRig/Private/ModularRig.cpp
+++ b/Engine/Plugins/Animation/ControlRig/Source/ControlRig/Private/ModularRig.cpp
@@ -94,7 +94,15 @@ void FRigModuleInstance::SetRig(UControlRig* InRig)
 
 	// update the cache
 	RigPtr = InRig;
-	RigClass = InRig->GetClass();
+	if (InRig)
+	{
+		RigClass = InRig->GetClass();
+	}
+	else
+	{
+		RigClass = nullptr;
+	}
+	
 }
 
 bool FRigModuleInstance::ContainsRig(const UControlRig* InRig) const
