Added missing null pointer check after review

#rnx

diff --git a/Engine/Plugins/Animation/ControlRig/Source/ControlRig/Private/ControlRigOverride.cpp b/Engine/Plugins/Animation/ControlRig/Source/ControlRig/Private/ControlRigOverride.cpp
--- a/Engine/Plugins/Animation/ControlRig/Source/ControlRig/Private/ControlRigOverride.cpp
+++ b/Engine/Plugins/Animation/ControlRig/Source/ControlRig/Private/ControlRigOverride.cpp
@@ -387,9 +387,12 @@ void FControlRigOverrideValue::AddStructReferencedObjects(FReferenceCollector& C
 		{
 			if (const FObjectProperty* ObjProperty = CastField<FObjectProperty>(LeafProperty))
 			{
-				if (TObjectPtr<const UObject>* ObjPtr = ObjProperty->ContainerPtrToValuePtr<TObjectPtr<const UObject>>(const_cast<void*>(GetData())))
+				if (const void* DataPtr = GetData())
 				{
-					Collector.AddReferencedObject(*ObjPtr);
+					if (TObjectPtr<const UObject>* ObjPtr = ObjProperty->ContainerPtrToValuePtr<TObjectPtr<const UObject>>(const_cast<void*>(DataPtr)))
+					{
+						Collector.AddReferencedObject(*ObjPtr);
+					}
 				}
 			}
 			else if (const FStructProperty* StructProperty = CastField<FStructProperty>(LeafProperty))
