Fix for Control Rig curve container undo / redo breaking the curve list

diff --git a/Engine/Plugins/Animation/ControlRig/Source/ControlRigEditor/Private/Editor/SRigCurveContainer.cpp b/Engine/Plugins/Animation/ControlRig/Source/ControlRigEditor/Private/Editor/SRigCurveContainer.cpp
--- a/Engine/Plugins/Animation/ControlRig/Source/ControlRigEditor/Private/Editor/SRigCurveContainer.cpp
+++ b/Engine/Plugins/Animation/ControlRig/Source/ControlRigEditor/Private/Editor/SRigCurveContainer.cpp
@@ -407,7 +407,7 @@ void SRigCurveContainer::CreateRigCurveList( const FString& SearchText )
 		return;
 	}
 	
-	URigHierarchy* Hierarchy = GetHierarchy();
+	URigHierarchy* Hierarchy = GetDefaultHierarchy();
 	if (Hierarchy)
 	{
 		RigCurveList.Reset();
@@ -537,7 +537,7 @@ bool SRigCurveContainer::CanDelete()
 
 void SRigCurveContainer::SetCurveValue(const FName& CurveName, float CurveValue)
 {
-	URigHierarchy* Hierarchy = GetHierarchy();
+	URigHierarchy* Hierarchy = GetDefaultHierarchy();
 	if (Hierarchy)
 	{
 		Hierarchy->SetCurveValue(FRigElementKey(CurveName, ERigElementType::Curve), CurveValue, true);
@@ -546,6 +546,8 @@ void SRigCurveContainer::SetCurveValue(const FName& CurveName, float CurveValue)
 
 float SRigCurveContainer::GetCurveValue(const FName& CurveName)
 {
+	// For actions that involve topological changes, the code uses the default hierarchy
+	// But here, we use the instance, in order to be able to see the curve values while they change in the debugged instance
 	URigHierarchy* Hierarchy = GetHierarchy();
 	if (Hierarchy)
 	{
@@ -575,7 +577,7 @@ void SRigCurveContainer::OnSelectionChanged(FDisplayedRigCurveInfoPtr Selection,
 		return;
 	}
 
-	URigHierarchy* Hierarchy = GetHierarchy();
+	URigHierarchy* Hierarchy = GetDefaultHierarchy();
 	if (Hierarchy)
 	{
 		URigHierarchyController* Controller = Hierarchy->GetController();
