Control Rig: Fix tween widget preventing interaction with PIE

diff --git a/Engine/Plugins/Animation/ControlRig/Source/ControlRigEditor/Private/Misc/Flyout/FlyoutOverlayManager.cpp b/Engine/Plugins/Animation/ControlRig/Source/ControlRigEditor/Private/Misc/Flyout/FlyoutOverlayManager.cpp
--- a/Engine/Plugins/Animation/ControlRig/Source/ControlRigEditor/Private/Misc/Flyout/FlyoutOverlayManager.cpp
+++ b/Engine/Plugins/Animation/ControlRig/Source/ControlRigEditor/Private/Misc/Flyout/FlyoutOverlayManager.cpp
@@ -2,6 +2,7 @@
 
 #include "FlyoutOverlayManager.h"
 
+#include "Editor.h"
 #include "Framework/Application/IInputProcessor.h"
 #include "Framework/Application/SlateApplication.h"
 #include "Host/IWidgetHost.h"
@@ -162,7 +163,12 @@ void FFlyoutOverlayManager::ShowWidget()
 		DragWidget = SNew(ToolWidgets::SDraggableBoxOverlay)
 			.OnUserDraggedToNewPosition_Raw(this, &FFlyoutOverlayManager::SaveWidgetState)
 			.IsDraggable_Raw(this, &FFlyoutOverlayManager::CanDragWidget)
-			.Visibility_Lambda([this]{ return bShouldShowWidget ? EVisibility::Visible : EVisibility::Collapsed; })
+			.Visibility_Lambda([this]
+			{
+				// Don't show the widget while in PIE because it prevents interaction with UMG widgets.
+				const bool bIsPlaying = GEditor && (GEditor->bIsSimulatingInEditor || GEditor->PlayWorld);
+				return !bIsPlaying && bShouldShowWidget ? EVisibility::Visible : EVisibility::Collapsed;
+			})
 			.Content()
 			[
 				SAssignNew(ContentWidget, Private::SFlyoutWidgetWrapper)
