Anim in Engine: Fix an issue in Anim Details where selecting a component in a plain location, rotation or scale control would not select the control in Curve Editor

diff --git a/Engine/Plugins/Animation/ControlRig/Source/ControlRigEditor/Private/AnimDetails/Proxies/AnimDetailsProxyLocation.cpp b/Engine/Plugins/Animation/ControlRig/Source/ControlRigEditor/Private/AnimDetails/Proxies/AnimDetailsProxyLocation.cpp
--- a/Engine/Plugins/Animation/ControlRig/Source/ControlRigEditor/Private/AnimDetails/Proxies/AnimDetailsProxyLocation.cpp
+++ b/Engine/Plugins/Animation/ControlRig/Source/ControlRigEditor/Private/AnimDetails/Proxies/AnimDetailsProxyLocation.cpp
@@ -129,15 +129,15 @@ EControlRigContextChannelToKey UAnimDetailsProxyLocation::GetChannelToKeyFromPro
 
 EControlRigContextChannelToKey UAnimDetailsProxyLocation::GetChannelToKeyFromChannelName(const FString& InChannelName) const
 {
-	if (InChannelName == TEXT("X"))
+	if (InChannelName == TEXT("Location.X"))
 	{
 		return EControlRigContextChannelToKey::TranslationX;
 	}
-	else if (InChannelName == TEXT("Y"))
+	else if (InChannelName == TEXT("Location.Y"))
 	{
 		return EControlRigContextChannelToKey::TranslationY;
 	}
-	else if (InChannelName == TEXT("Z"))
+	else if (InChannelName == TEXT("Location.Z"))
 	{
 		return EControlRigContextChannelToKey::TranslationZ;
 	}
diff --git a/Engine/Plugins/Animation/ControlRig/Source/ControlRigEditor/Private/AnimDetails/Proxies/AnimDetailsProxyRotation.cpp b/Engine/Plugins/Animation/ControlRig/Source/ControlRigEditor/Private/AnimDetails/Proxies/AnimDetailsProxyRotation.cpp
--- a/Engine/Plugins/Animation/ControlRig/Source/ControlRigEditor/Private/AnimDetails/Proxies/AnimDetailsProxyRotation.cpp
+++ b/Engine/Plugins/Animation/ControlRig/Source/ControlRigEditor/Private/AnimDetails/Proxies/AnimDetailsProxyRotation.cpp
@@ -153,18 +153,19 @@ EControlRigContextChannelToKey UAnimDetailsProxyRotation::GetChannelToKeyFromPro
 
 EControlRigContextChannelToKey UAnimDetailsProxyRotation::GetChannelToKeyFromChannelName(const FString& InChannelName) const
 {
-	if (InChannelName == TEXT("X"))
+	if (InChannelName == TEXT("Rotation.X") || InChannelName == TEXT("Rotation.Roll"))
 	{
 		return EControlRigContextChannelToKey::RotationX;
 	}
-	else if (InChannelName == TEXT("Y"))
+	else if (InChannelName == TEXT("Rotation.Y") || InChannelName == TEXT("Rotation.Pitch"))
 	{
 		return EControlRigContextChannelToKey::RotationY;
 	}
-	else if (InChannelName == TEXT("Z"))
+	else if (InChannelName == TEXT("Rotation.Z") || InChannelName == TEXT("Rotation.Yaw"))
 	{
 		return EControlRigContextChannelToKey::RotationZ;
 	}
+
 	return EControlRigContextChannelToKey::AllTransform;
 }
 
diff --git a/Engine/Plugins/Animation/ControlRig/Source/ControlRigEditor/Private/AnimDetails/Proxies/AnimDetailsProxyScale.cpp b/Engine/Plugins/Animation/ControlRig/Source/ControlRigEditor/Private/AnimDetails/Proxies/AnimDetailsProxyScale.cpp
--- a/Engine/Plugins/Animation/ControlRig/Source/ControlRigEditor/Private/AnimDetails/Proxies/AnimDetailsProxyScale.cpp
+++ b/Engine/Plugins/Animation/ControlRig/Source/ControlRigEditor/Private/AnimDetails/Proxies/AnimDetailsProxyScale.cpp
@@ -140,18 +140,19 @@ EControlRigContextChannelToKey UAnimDetailsProxyScale::GetChannelToKeyFromProper
 
 EControlRigContextChannelToKey UAnimDetailsProxyScale::GetChannelToKeyFromChannelName(const FString& InChannelName) const
 {
-	if (InChannelName == TEXT("X"))
+	if (InChannelName == TEXT("Scale.X"))
 	{
 		return EControlRigContextChannelToKey::ScaleX;
 	}
-	else if (InChannelName == TEXT("Y"))
+	else if (InChannelName == TEXT("Scale.Y"))
 	{
 		return EControlRigContextChannelToKey::ScaleY;
 	}
-	else if (InChannelName == TEXT("Z"))
+	else if (InChannelName == TEXT("Scale.Z"))
 	{
 		return EControlRigContextChannelToKey::ScaleZ;
 	}
+
 	return EControlRigContextChannelToKey::AllTransform;
 }
 
