Anim in Engine: Fix an issue where FAnimDetailsScale::ToVector3f set the Y component from X instead of Y

#rnx

diff --git a/Engine/Plugins/Animation/ControlRig/Source/ControlRigEditor/Private/AnimDetails/Proxies/AnimDetailsProxyScale.cpp b/Engine/Plugins/Animation/ControlRig/Source/ControlRigEditor/Private/AnimDetails/Proxies/AnimDetailsProxyScale.cpp
--- a/Engine/Plugins/Animation/ControlRig/Source/ControlRigEditor/Private/AnimDetails/Proxies/AnimDetailsProxyScale.cpp
+++ b/Engine/Plugins/Animation/ControlRig/Source/ControlRigEditor/Private/AnimDetails/Proxies/AnimDetailsProxyScale.cpp
@@ -50,7 +50,7 @@ FVector FAnimDetailsScale::ToVector() const
 
 FVector3f FAnimDetailsScale::ToVector3f() const
 {
-	return FVector3f(SX, SX, SZ); 
+	return FVector3f(SX, SY, SZ); 
 }
 
 FName UAnimDetailsProxyScale::GetCategoryName() const
