Anim in Engine: Anim Details properties now reset their values to the intial values of controls when clicking reset to default, cleanup relate code, reduce reduntant evaluation calls to the rig for better performance

diff --git a/Engine/Plugins/Animation/ControlRig/Source/ControlRigEditor/Private/AnimDetails/AnimDetailsProxyManager.cpp b/Engine/Plugins/Animation/ControlRig/Source/ControlRigEditor/Private/AnimDetails/AnimDetailsProxyManager.cpp
--- a/Engine/Plugins/Animation/ControlRig/Source/ControlRigEditor/Private/AnimDetails/AnimDetailsProxyManager.cpp
+++ b/Engine/Plugins/Animation/ControlRig/Source/ControlRigEditor/Private/AnimDetails/AnimDetailsProxyManager.cpp
@@ -442,7 +442,7 @@ void UAnimDetailsProxyManager::ForceUpdateProxyValues()
 
 	for (UAnimDetailsProxyBase* Proxy : ExternalSelection)
 	{
-		Proxy->UpdateProxyValues();
+		Proxy->AdoptValues();
 	}
 }
 
@@ -644,7 +644,7 @@ void UAnimDetailsProxyManager::ReevaluateConstraints()
 		Proxy->Modify();
 
 		Controller.EvaluateAllConstraints();
-		Proxy->UpdateProxyValues();
+		Proxy->AdoptValues();
 	}
 }
 
diff --git a/Engine/Plugins/Animation/ControlRig/Source/ControlRigEditor/Private/AnimDetails/Customizations/AnimDetailsValueCustomization.cpp b/Engine/Plugins/Animation/ControlRig/Source/ControlRigEditor/Private/AnimDetails/Customizations/AnimDetailsValueCustomization.cpp
--- a/Engine/Plugins/Animation/ControlRig/Source/ControlRigEditor/Private/AnimDetails/Customizations/AnimDetailsValueCustomization.cpp
+++ b/Engine/Plugins/Animation/ControlRig/Source/ControlRigEditor/Private/AnimDetails/Customizations/AnimDetailsValueCustomization.cpp
@@ -2,6 +2,7 @@
 
 #include "AnimDetailsValueCustomization.h"
 
+#include "Algo/AnyOf.h"
 #include "AnimDetails/AnimDetailsProxyManager.h"
 #include "AnimDetails/Proxies/AnimDetailsProxyLocation.h"
 #include "AnimDetails/Proxies/AnimDetailsProxyRotation.h"
@@ -15,6 +16,7 @@
 #include "EditorModeManager.h"
 #include "IDetailChildrenBuilder.h"
 #include "PropertyEditorClipboard.h"
+#include "ScopedTransaction.h"
 #include "Widgets/Input/SNumericEntryBox.h"
 
 #define LOCTEXT_NAMESPACE "AnimDetailsValueCustomization"
@@ -133,10 +135,10 @@ namespace UE::ControlRigEditor
 					return ChildHandle->GetPropertyDisplayName();
 				}();
 
-			const TSharedRef<SWidget> ValueWidget = MakeChildWidget(InStructPropertyHandle, ChildHandle);
-			
+
 			// Handle visibilty for child content widgets. This is required so that  
 			// FAnimDetailsNavigableWidgetRegistry can find the right widget to navigate to.
+			const TSharedRef<SWidget> ValueWidget = MakeChildWidget(InStructPropertyHandle, ChildHandle);
 			if (const bool bCanBeExpanded = SortedChildHandles.Num() > 1)
 			{
 				ChildContentWidgets.Add(ValueWidget);
@@ -170,6 +172,12 @@ namespace UE::ControlRigEditor
 				[
 					ChildHandle->CreateDefaultPropertyButtonWidgets()
 				]
+				.OverrideResetToDefault(
+					FResetToDefaultOverride::Create
+					(
+						TAttribute<bool>::CreateSP(this, &FAnimDetailsValueCustomization::IsResetToDefaultVisible, ChildHandle.ToWeakPtr()),
+						FSimpleDelegate::CreateSP(this, &FAnimDetailsValueCustomization::OnResetToDefaultClicked, ChildHandle.ToWeakPtr())
+					))
 				.CopyAction(
 					FUIAction(
 						FExecuteAction::CreateLambda([WeakChildHandle = ChildHandle->AsWeak()]
@@ -229,7 +237,7 @@ namespace UE::ControlRigEditor
 		}
 
 		StructPropertyHandle = InStructPropertyHandle;
-		
+
 		// Make enough space for each child handle
 		const float DesiredWidth = 125.f * SortedChildHandles.Num();
 
@@ -256,7 +264,13 @@ namespace UE::ControlRigEditor
 				SAssignNew(HorizontalBox, SHorizontalBox)
 				.Visibility(this, &FAnimDetailsValueCustomization::GetVisibilityFromExpansionState)
 				.IsEnabled(this, &FMathStructCustomization::IsValueEnabled, StructPropertyHandle.ToWeakPtr())
-			];
+			]
+			.OverrideResetToDefault(
+				FResetToDefaultOverride::Create
+				(
+					TAttribute<bool>::CreateSP(this, &FAnimDetailsValueCustomization::IsResetToDefaultVisible, StructPropertyHandle.ToWeakPtr()),
+					FSimpleDelegate::CreateSP(this, &FAnimDetailsValueCustomization::OnResetToDefaultClicked, StructPropertyHandle.ToWeakPtr())
+				));
 
 		// Create inline children if collapsed, the children can only be selected in the property value row
 		for (int32 ChildIndex = 0; ChildIndex < SortedChildHandles.Num(); ++ChildIndex)
@@ -428,6 +442,65 @@ namespace UE::ControlRigEditor
 
 		return FLinearColor::White;
 	}
+
+	bool FAnimDetailsValueCustomization::IsResetToDefaultVisible(const TWeakPtr<IPropertyHandle> WeakPropertyHandle) const
+	{
+		if (!WeakPropertyHandle.IsValid())
+		{
+			return false;
+		}
+		const TSharedRef<IPropertyHandle> PropertyHandle = WeakPropertyHandle.Pin().ToSharedRef();
+
+		const FName PropertyName = PropertyHandle->GetProperty()->GetFName();
+
+		TArray<UObject*> OuterObjects;
+		PropertyHandle->GetOuterObjects(OuterObjects);
+
+		return Algo::AnyOf(OuterObjects,
+			[&PropertyName](const UObject* OuterObject)
+			{
+				if (const UAnimDetailsProxyBase* Proxy = Cast<const UAnimDetailsProxyBase>(OuterObject))
+				{
+					return Proxy && !Proxy->HasDefaultValue(PropertyName);
+				}
+
+				return false;
+			});
+	}
+
+	void FAnimDetailsValueCustomization::OnResetToDefaultClicked(const TWeakPtr<IPropertyHandle> WeakPropertyHandle)
+	{
+		if (!WeakPropertyHandle.IsValid())
+		{
+			return;
+		}
+		const TSharedRef<IPropertyHandle> PropertyHandle = WeakPropertyHandle.Pin().ToSharedRef();
+
+		const FName PropertyName = PropertyHandle->GetProperty()->GetFName();
+
+		TArray<UObject*> OuterObjects;
+		PropertyHandle->GetOuterObjects(OuterObjects);
+
+		const FScopedTransaction ResetValueToDefaultTransaction(LOCTEXT("ResetValueToDefaultTransaction", "Reset Value To Default"));
+
+		// Reset all then propagonate all values so all is evaluated in the new state
+		for (UObject* OuterObject : OuterObjects)
+		{
+			if (UAnimDetailsProxyBase* Proxy = Cast<UAnimDetailsProxyBase>(OuterObject))
+			{
+				Proxy->Modify();
+				Proxy->ResetPropertyToDefault(PropertyName);
+			}
+		}
+
+		for (UObject* OuterObject : OuterObjects)
+		{
+			if (UAnimDetailsProxyBase* Proxy = Cast<UAnimDetailsProxyBase>(OuterObject))
+			{
+				Proxy->PropagonateValues();
+			}
+		}
+	}
 }
 
 #undef LOCTEXT_NAMESPACE
diff --git a/Engine/Plugins/Animation/ControlRig/Source/ControlRigEditor/Private/AnimDetails/Customizations/AnimDetailsValueCustomization.h b/Engine/Plugins/Animation/ControlRig/Source/ControlRigEditor/Private/AnimDetails/Customizations/AnimDetailsValueCustomization.h
--- a/Engine/Plugins/Animation/ControlRig/Source/ControlRigEditor/Private/AnimDetails/Customizations/AnimDetailsValueCustomization.h
+++ b/Engine/Plugins/Animation/ControlRig/Source/ControlRigEditor/Private/AnimDetails/Customizations/AnimDetailsValueCustomization.h
@@ -65,6 +65,12 @@ namespace UE::ControlRigEditor
 		/** Gets a color from the property */
 		FLinearColor GetColorFromProperty(const FName& PropertyName) const;
 
+		/** Returns true if the property handle's reset to default button is visible */
+		bool IsResetToDefaultVisible(const TWeakPtr<IPropertyHandle> WeakPropertyHandle) const;
+
+		/** Called when the property handle's reset to default button was clicked */
+		void OnResetToDefaultClicked(const TWeakPtr<IPropertyHandle> WeakPropertyHandle);
+
 		/** Handles visibility for child widgets in this customization */
 		TUniquePtr<Private::FAnimDetailsChildWidgetVisiblityHandler> ChildContentVisibilityHandler;
 
diff --git a/Engine/Plugins/Animation/ControlRig/Source/ControlRigEditor/Private/AnimDetails/Proxies/AnimDetailsProxyBase.cpp b/Engine/Plugins/Animation/ControlRig/Source/ControlRigEditor/Private/AnimDetails/Proxies/AnimDetailsProxyBase.cpp
--- a/Engine/Plugins/Animation/ControlRig/Source/ControlRigEditor/Private/AnimDetails/Proxies/AnimDetailsProxyBase.cpp
+++ b/Engine/Plugins/Animation/ControlRig/Source/ControlRigEditor/Private/AnimDetails/Proxies/AnimDetailsProxyBase.cpp
@@ -599,6 +599,16 @@ void UAnimDetailsProxyBase::SetControlFromControlRig(UControlRig* InControlRig,
 	{
 		WeakControlRig = InControlRig;
 		CachedRigElement = FCachedRigElement(FRigElementKey(InName, ERigElementType::Control), Hierarchy);
+
+		// Test if this yields a valid proxy
+		const FRigControlElement* ControlElement = GetControlElement();
+		const bool bValidControlElement = ControlElement && GetSupportedControlTypes().Contains(ControlElement->Settings.ControlType);
+		
+		if (!ensureMsgf(bValidControlElement, TEXT("Created invalid anim details proxy, control element is invalid, or control type does not match proxy type")))
+		{
+			WeakControlRig.Reset();
+			CachedRigElement.Reset();
+		}
 	}
 }
 
@@ -659,14 +669,14 @@ void UAnimDetailsProxyBase::PropagonateValues()
 	FRigControlElement* ControlElement = GetControlElement();
 	if (ControlRig && ControlElement)
 	{
-		SetControlRigElementValueFromCurrent(ControlRig, ControlElement, Context);
+		SetControlRigElementValueFromCurrent(Context);
 		FControlRigEditMode::NotifyDrivenControls(ControlRig, ControlElement->GetKey(), NotifyDrivenContext);
 			
 		ControlRig->Evaluate_AnyThread();
 	}
-	else if (SequencerItem.IsValid())
+	else
 	{
-		SetBindingValueFromCurrent(SequencerItem.GetBoundObject(), SequencerItem.GetBinding(), Context);
+		SetSequencerBindingValueFromCurrent(Context);
 	}
 }
 
@@ -722,7 +732,9 @@ void UAnimDetailsProxyBase::SetKey(const IPropertyHandle& KeyedPropertyHandle)
 			Context.SetKey = EControlRigSetKey::Always;
 			Context.KeyMask = (uint32)GetChannelToKeyFromPropertyName(PropertyName);
 
-			SetControlRigElementValueFromCurrent(ControlRig, ControlElement, Context);
+			SetControlRigElementValueFromCurrent(Context);
+
+			ControlRig->Evaluate_AnyThread();
 
 			FRigControlModifiedContext NotifyDrivenContext; // always key ever
 			NotifyDrivenContext.SetKey = EControlRigSetKey::Always;
@@ -848,7 +860,13 @@ void UAnimDetailsProxyBase::PostEditChangeChainProperty(FPropertyChangedChainEve
 	{
 		return;
 	}
-
+	
+	// Reset to default is handled in the anim details value customizations
+	if (PropertyChangedEvent.ChangeType & EPropertyChangeType::ResetToDefault)
+	{
+		return;
+	}
+	
 	if (const FProperty* Property = PropertyChangedEvent.Property)
 	{
 		const FProperty* MemberProperty = PropertyChangedEvent.PropertyChain.GetActiveMemberNode() ? 
@@ -857,9 +875,12 @@ void UAnimDetailsProxyBase::PostEditChangeChainProperty(FPropertyChangedChainEve
 
 		if (PropertyIsOnProxy(Property, MemberProperty))
 		{
-			EControlRigContextChannelToKey ChannelToKeyContext = GetChannelToKeyFromPropertyName(Property->GetFName());		
+
+			// Let the control rig know it's interacted with
+			EControlRigContextChannelToKey ChannelToKeyContext = GetChannelToKeyFromPropertyName(Property->GetFName());
 			AddControlRigInteractionScope(ChannelToKeyContext, PropertyChangedEvent.ChangeType);
 
+			// Propagonate the values to the control rig
 			PropagonateValues();
 
 			if (PropertyChangedEvent.ChangeType != EPropertyChangeType::Interactive)
@@ -867,7 +888,9 @@ void UAnimDetailsProxyBase::PostEditChangeChainProperty(FPropertyChangedChainEve
 				InteractionScopes.Reset();
 			}
 
-			UpdateProxyValues();
+			// Adopt values from the rig or sequencer binding. 
+			// They may be different than what was propagonated, e.g. due to constraints.
+			AdoptValues(ERigControlValueType::Current);
 		}
 	}
 }
diff --git a/Engine/Plugins/Animation/ControlRig/Source/ControlRigEditor/Private/AnimDetails/Proxies/AnimDetailsProxyBase.h b/Engine/Plugins/Animation/ControlRig/Source/ControlRigEditor/Private/AnimDetails/Proxies/AnimDetailsProxyBase.h
--- a/Engine/Plugins/Animation/ControlRig/Source/ControlRigEditor/Private/AnimDetails/Proxies/AnimDetailsProxyBase.h
+++ b/Engine/Plugins/Animation/ControlRig/Source/ControlRigEditor/Private/AnimDetails/Proxies/AnimDetailsProxyBase.h
@@ -96,11 +96,20 @@ public:
 	/** Returns the localized property name, useful for filtering. Currently transforms override this as they have more than one inner struct. */
 	virtual void GetLocalizedPropertyName(const FName& InPropertyName, FText& OutPropertyDisplayName, TOptional<FText>& OutOptionalStructDisplayName) const;
 
+	/** Returns the control types the proxy supports */
+	virtual TSet<ERigControlType> GetSupportedControlTypes() const PURE_VIRTUAL(UAnimDetailsProxyBase::GetSupportedControlTypes, return {};);
+
 	/** Returns true if the property is handled by this proxy */
 	virtual bool PropertyIsOnProxy(const FProperty* Property, const FProperty* MemberProperty) PURE_VIRTUAL(UAnimDetailsProxyBase::PropertyIsOnProxy, return false;)
 
-	/** Updates the proxy values from the sequencer binding or the control in the control rig respectively */
-	virtual void UpdateProxyValues() PURE_VIRTUAL(UAnimDetailsProxyBase::UpdateProxyValues, return;)
+	/** Tries to adopt the current value(s) from the sequencer binding or the control in the control rig. */
+	virtual void AdoptValues(const ERigControlValueType RigControlValueType = ERigControlValueType::Current) PURE_VIRTUAL(UAnimDetailsProxyBase::AdoptValues, return;)
+
+	/** Resets a property in the proxy to the initial value in the rig. If the proxy does not relate to a control rig, this function has no effect. */
+	virtual void ResetPropertyToDefault(const FName& PropertyName) PURE_VIRTUAL(UAnimDetailsProxyBase::ResetPropertyToDefault, return;)
+
+	/** Returns true if the property value is the default value */
+	virtual bool HasDefaultValue(const FName& PropertyName) const PURE_VIRTUAL(UAnimDetailsProxyBase::HasDefaultValue, return true;)
 
 	/** Updates the proxy shape override properties such as DisplayName and shape */
 	virtual void UpdateOverrideableProperties();
@@ -132,10 +141,10 @@ public:
 
 protected:
 	/** Sets the control rig element value from the current proxy value */
-	virtual void SetControlRigElementValueFromCurrent(UControlRig* ControlRig, FRigControlElement* ControlElement, const FRigControlModifiedContext& Context) PURE_VIRTUAL(UAnimDetailsProxyBase::SetControlRigElementValueFromCurrent, return;)
+	virtual void SetControlRigElementValueFromCurrent(const FRigControlModifiedContext& Context) PURE_VIRTUAL(UAnimDetailsProxyBase::SetControlRigElementValueFromCurrent, return;)
 
 	/** Sets the sequencer binding value from the current proxy value */
-	virtual void SetBindingValueFromCurrent(UObject* InObject, const TSharedPtr<FTrackInstancePropertyBindings>&Binding, const FRigControlModifiedContext & Context, bool bInteractive = false) {}
+	virtual void SetSequencerBindingValueFromCurrent(const FRigControlModifiedContext & Context) {}
 
 	/** The control rig that holds the control or nullptr if there is no control rig. */
 	TWeakObjectPtr<UControlRig> WeakControlRig;
diff --git a/Engine/Plugins/Animation/ControlRig/Source/ControlRigEditor/Private/AnimDetails/Proxies/AnimDetailsProxyBool.cpp b/Engine/Plugins/Animation/ControlRig/Source/ControlRigEditor/Private/AnimDetails/Proxies/AnimDetailsProxyBool.cpp
--- a/Engine/Plugins/Animation/ControlRig/Source/ControlRigEditor/Private/AnimDetails/Proxies/AnimDetailsProxyBool.cpp
+++ b/Engine/Plugins/Animation/ControlRig/Source/ControlRigEditor/Private/AnimDetails/Proxies/AnimDetailsProxyBool.cpp
@@ -53,37 +53,68 @@ TArray<FName> UAnimDetailsProxyBool::GetPropertyNames() const
 	};
 }
 
+TSet<ERigControlType> UAnimDetailsProxyBool::GetSupportedControlTypes() const
+{
+	static const TSet<ERigControlType> SupportedControlTypes = { ERigControlType::Bool };
+	return SupportedControlTypes;
+}
+
 bool UAnimDetailsProxyBool::PropertyIsOnProxy(const FProperty* Property, const FProperty* MemberProperty)
 {
-	return ((Property && Property->GetFName() == GET_MEMBER_NAME_CHECKED(FAnimDetailsBool, Bool)) ||
-		(MemberProperty && MemberProperty->GetFName() == GET_MEMBER_NAME_CHECKED(UAnimDetailsProxyBool, Bool)));
+	return (Property && Property->GetFName() == GET_MEMBER_NAME_CHECKED(FAnimDetailsBool, Bool)) ||
+		(MemberProperty && MemberProperty->GetFName() == GET_MEMBER_NAME_CHECKED(UAnimDetailsProxyBool, Bool));
 }
 
-void UAnimDetailsProxyBool::UpdateProxyValues()
+void UAnimDetailsProxyBool::AdoptValues(const ERigControlValueType RigControlValueType)
 {
-	bool Value = false;
+	const bool Value = [RigControlValueType, this]()
+		{
+			UControlRig* ControlRig = GetControlRig();
+			FRigControlElement* ControlElement = GetControlElement();
+			if (ControlRig && ControlElement)
+			{
+				const FRigControlValue ControlValue = ControlRig->GetControlValue(ControlElement, RigControlValueType);
+				return ControlValue.Get<bool>();
+			}
+			else if (SequencerItem.IsValid())
+			{
+				const TOptional<bool> OptionalValue = SequencerItem.GetBinding()->GetOptionalValue<bool>(*SequencerItem.GetBoundObject());
+				return OptionalValue.IsSet() ? OptionalValue.GetValue() : false;
+			}
+
+			return false;
+		}();
 
-	UControlRig* ControlRig = GetControlRig();
-	FRigControlElement* ControlElement = GetControlElement();
-	if (ControlRig && ControlElement)
+	const FAnimDetailsBool AnimDetailProxyBool(Value);
+
+	constexpr const TCHAR* PropertyName = TEXT("Bool");
+	FTrackInstancePropertyBindings Binding(PropertyName, PropertyName);
+	Binding.CallFunction<FAnimDetailsBool>(*this, AnimDetailProxyBool);
+}
+
+void UAnimDetailsProxyBool::ResetPropertyToDefault(const FName& PropertyName)
+{
+	if (const FRigControlElement* ControlElement = GetControlElement())
 	{
-		if (ControlElement->Settings.ControlType == ERigControlType::Bool)
-		{
-			FRigControlValue ControlValue = ControlRig->GetControlValue(ControlElement, ERigControlValueType::Current);
-			Value = ControlValue.Get<bool>();
-		}
+		AdoptValues(ERigControlValueType::Initial);
 	}
-	else if (SequencerItem.IsValid())
+	else
 	{
-		const TOptional<bool> OptionalValue = SequencerItem.GetBinding()->GetOptionalValue<bool>(*SequencerItem.GetBoundObject());
-		Value = OptionalValue.IsSet() ? OptionalValue.GetValue() : false;
+		Bool = FAnimDetailsBool();
 	}
+}
 
-	const FAnimDetailsBool AnimDetailProxyBool(Value);
+bool UAnimDetailsProxyBool::HasDefaultValue(const FName& PropertyName) const
+{
+	UControlRig* ControlRig = GetControlRig();
+	FRigControlElement* ControlElement = GetControlElement();
+	if (ControlRig && ControlElement)
+	{
+		const FRigControlValue ControlValue = ControlRig->GetControlValue(ControlElement, ERigControlValueType::Initial);
+		return Bool.Bool == ControlValue.Get<bool>();
+	}
 
-	const FName PropName("Bool");
-	FTrackInstancePropertyBindings Binding(PropName, PropName.ToString());
-	Binding.CallFunction<FAnimDetailsBool>(*this, AnimDetailProxyBool);
+	return Bool.Bool == false;
 }
 
 EControlRigContextChannelToKey UAnimDetailsProxyBool::GetChannelToKeyFromPropertyName(const FName& PropertyName) const
@@ -113,26 +144,26 @@ EControlRigContextChannelToKey UAnimDetailsProxyBool::GetChannelToKeyFromChannel
 	return EControlRigContextChannelToKey::AllTransform;
 }
 
-void UAnimDetailsProxyBool::SetControlRigElementValueFromCurrent(UControlRig* ControlRig, FRigControlElement* ControlElement, const FRigControlModifiedContext& Context)
+void UAnimDetailsProxyBool::SetControlRigElementValueFromCurrent(const FRigControlModifiedContext& Context)
 {
-	if (ControlElement && 
-		ControlRig && 
-		ControlElement->Settings.ControlType == ERigControlType::Bool)
+	UControlRig* ControlRig = GetControlRig();
+	FRigControlElement* ControlElement = GetControlElement();
+	if (ControlElement && ControlRig)
 	{
 		const bool Value = Bool.Bool;
 
 		constexpr bool bNotify = true;
 		constexpr bool bSetupUndo = false;
 		ControlRig->SetControlValue<bool>(ControlElement->GetKey().Name, Value, bNotify, Context, bSetupUndo);
-
-		ControlRig->Evaluate_AnyThread();
 	}
 }
 
-void UAnimDetailsProxyBool::SetBindingValueFromCurrent(UObject* InObject, const TSharedPtr<FTrackInstancePropertyBindings>& Binding, const FRigControlModifiedContext& Context, bool bInteractive)
+void UAnimDetailsProxyBool::SetSequencerBindingValueFromCurrent(const FRigControlModifiedContext& Context)
 {
-	if (InObject && Binding.IsValid())
+	const TSharedPtr<FTrackInstancePropertyBindings> Binding = SequencerItem.GetBinding();
+	UObject* BoundObject = SequencerItem.GetBoundObject();
+	if (Binding.IsValid() && BoundObject)
 	{
-		Binding->SetCurrentValue<bool>(*InObject, Bool.Bool);
+		Binding->SetCurrentValue<bool>(*BoundObject, Bool.Bool);
 	}
 }
diff --git a/Engine/Plugins/Animation/ControlRig/Source/ControlRigEditor/Private/AnimDetails/Proxies/AnimDetailsProxyBool.h b/Engine/Plugins/Animation/ControlRig/Source/ControlRigEditor/Private/AnimDetails/Proxies/AnimDetailsProxyBool.h
--- a/Engine/Plugins/Animation/ControlRig/Source/ControlRigEditor/Private/AnimDetails/Proxies/AnimDetailsProxyBool.h
+++ b/Engine/Plugins/Animation/ControlRig/Source/ControlRigEditor/Private/AnimDetails/Proxies/AnimDetailsProxyBool.h
@@ -28,12 +28,15 @@ public:
 	virtual FName GetCategoryName() const override;
 	virtual void UpdatePropertyDisplayNames(IDetailLayoutBuilder& DetailBuilder);
 	virtual TArray<FName> GetPropertyNames() const override;
+	virtual TSet<ERigControlType> GetSupportedControlTypes() const override;
 	virtual bool PropertyIsOnProxy(const FProperty* Property, const FProperty* MemberProperty) override;
-	virtual void UpdateProxyValues() override;
+	virtual void AdoptValues(const ERigControlValueType RigControlValueType = ERigControlValueType::Current) override;
+	virtual void ResetPropertyToDefault(const FName& PropertyName) override;
+	virtual bool HasDefaultValue(const FName& PropertyName) const override;
 	virtual EControlRigContextChannelToKey GetChannelToKeyFromPropertyName(const FName& InPropertyName) const override;
 	virtual EControlRigContextChannelToKey GetChannelToKeyFromChannelName(const FString& InChannelName) const override;
-	virtual void SetControlRigElementValueFromCurrent(UControlRig* ControlRig, FRigControlElement* ControlElement, const FRigControlModifiedContext& Context) override;
-	virtual void SetBindingValueFromCurrent(UObject* InObject, const TSharedPtr<FTrackInstancePropertyBindings>& Binding, const FRigControlModifiedContext& Context, bool bInteractive = false) override;
+	virtual void SetControlRigElementValueFromCurrent(const FRigControlModifiedContext& Context) override;
+	virtual void SetSequencerBindingValueFromCurrent(const FRigControlModifiedContext& Context) override;
 	//~ End UAnimDetailsProxyBase interface
 
 	UPROPERTY(EditAnywhere, Interp, Meta = (ShowOnlyInnerProperties), Category = "Bool")
diff --git a/Engine/Plugins/Animation/ControlRig/Source/ControlRigEditor/Private/AnimDetails/Proxies/AnimDetailsProxyEnum.cpp b/Engine/Plugins/Animation/ControlRig/Source/ControlRigEditor/Private/AnimDetails/Proxies/AnimDetailsProxyEnum.cpp
--- a/Engine/Plugins/Animation/ControlRig/Source/ControlRigEditor/Private/AnimDetails/Proxies/AnimDetailsProxyEnum.cpp
+++ b/Engine/Plugins/Animation/ControlRig/Source/ControlRigEditor/Private/AnimDetails/Proxies/AnimDetailsProxyEnum.cpp
@@ -63,44 +63,65 @@ TArray<FName> UAnimDetailsProxyEnum::GetPropertyNames() const
 	};
 }
 
+TSet<ERigControlType> UAnimDetailsProxyEnum::GetSupportedControlTypes() const
+{
+	static const TSet<ERigControlType> SupportedControlTypes = { ERigControlType::Integer };
+	return SupportedControlTypes;
+}
 
 bool UAnimDetailsProxyEnum::PropertyIsOnProxy(const FProperty* Property, const FProperty* MemberProperty)
 {
-	return ((Property && Property->GetFName() == GET_MEMBER_NAME_CHECKED(FAnimDetailsEnum, EnumIndex)) ||
-		(MemberProperty && MemberProperty->GetFName() == GET_MEMBER_NAME_CHECKED(UAnimDetailsProxyEnum, Enum)));
+	return (Property && Property->GetFName() == GET_MEMBER_NAME_CHECKED(FAnimDetailsEnum, EnumIndex)) ||
+		(MemberProperty && MemberProperty->GetFName() == GET_MEMBER_NAME_CHECKED(UAnimDetailsProxyEnum, Enum));
 }
 
-void UAnimDetailsProxyEnum::UpdateProxyValues()
+void UAnimDetailsProxyEnum::AdoptValues(const ERigControlValueType RigControlValueType)
 {
 	UControlRig* ControlRig = GetControlRig();
 	FRigControlElement* ControlElement = GetControlElement();
-	if (!ControlRig || !ControlElement)
+	if (ControlRig &&
+		ControlElement &&
+		ControlElement->Settings.ControlEnum)
 	{
-		return;
-	}
+		UEnum& EnumType = *ControlElement->Settings.ControlEnum;
 
-	TObjectPtr<UEnum> EnumType = nullptr;
+		const FRigControlValue ControlValue = ControlRig->GetControlValue(ControlElement, RigControlValueType);
 
-	int32 Value = 0;
-	if (ControlElement->Settings.ControlType == ERigControlType::Integer &&
-		ControlElement->Settings.ControlEnum)
-	{
-		EnumType = ControlElement->Settings.ControlEnum;
-		const FRigControlValue ControlValue = ControlRig->GetControlValue(ControlElement, ERigControlValueType::Current);
+		FAnimDetailsEnum EnumValue;
+		EnumValue.EnumType = &EnumType;
+		EnumValue.EnumIndex = ControlValue.Get<int32>();
 
-		Value = ControlValue.Get<int32>();
+		constexpr const TCHAR* PropertyName = TEXT("Enum");
+		FTrackInstancePropertyBindings Binding(PropertyName, PropertyName);
+		Binding.CallFunction<FAnimDetailsEnum>(*this, EnumValue);
 	}
 
-	if (EnumType)
+	// Not implemented for sequencer items (altho they have an enum type @todo)
+}
+
+void UAnimDetailsProxyEnum::ResetPropertyToDefault(const FName& PropertyName)
+{
+	if (const FRigControlElement* ControlElement = GetControlElement())
 	{
-		FAnimDetailsEnum EnumValue;
-		EnumValue.EnumType = EnumType;
-		EnumValue.EnumIndex = Value;
+		AdoptValues(ERigControlValueType::Initial);
+	}
+	else
+	{
+		Enum.EnumIndex = 0;
+	}
+}
 
-		const FName PropertyName("Enum");
-		FTrackInstancePropertyBindings Binding(PropertyName, PropertyName.ToString());
-		Binding.CallFunction<FAnimDetailsEnum>(*this, EnumValue);
+bool UAnimDetailsProxyEnum::HasDefaultValue(const FName& PropertyName) const
+{
+	UControlRig* ControlRig = GetControlRig();
+	FRigControlElement* ControlElement = GetControlElement();
+	if (ControlRig && ControlElement)
+	{
+		const FRigControlValue ControlValue = ControlRig->GetControlValue(ControlElement, ERigControlValueType::Initial);
+		return Enum.EnumIndex == ControlValue.Get<int32>();
 	}
+
+	return Enum.EnumIndex == 0;
 }
 
 EControlRigContextChannelToKey UAnimDetailsProxyEnum::GetChannelToKeyFromPropertyName(const FName& PropertyName) const
@@ -132,19 +153,16 @@ EControlRigContextChannelToKey UAnimDetailsProxyEnum::GetChannelToKeyFromChannel
 	return EControlRigContextChannelToKey::AllTransform;
 }
 
-void UAnimDetailsProxyEnum::SetControlRigElementValueFromCurrent(UControlRig* ControlRig, FRigControlElement* ControlElement, const FRigControlModifiedContext& Context)
+void UAnimDetailsProxyEnum::SetControlRigElementValueFromCurrent(const FRigControlModifiedContext& Context)
 {
-	if (ControlRig && 
-		ControlElement &&
-		ControlElement->Settings.ControlType == ERigControlType::Integer &&
-		ControlElement->Settings.ControlEnum)
+	UControlRig* ControlRig = GetControlRig();
+	FRigControlElement* ControlElement = GetControlElement();
+	if (ControlRig && ControlElement)
 	{
 		int32 Val = Enum.EnumIndex;
 
 		constexpr bool bNotify = true;
 		constexpr bool bSetupUndo = false;
 		ControlRig->SetControlValue<int32>(ControlElement->GetKey().Name, Val, bNotify, Context, bSetupUndo);
-
-		ControlRig->Evaluate_AnyThread();
 	}
 }
diff --git a/Engine/Plugins/Animation/ControlRig/Source/ControlRigEditor/Private/AnimDetails/Proxies/AnimDetailsProxyEnum.h b/Engine/Plugins/Animation/ControlRig/Source/ControlRigEditor/Private/AnimDetails/Proxies/AnimDetailsProxyEnum.h
--- a/Engine/Plugins/Animation/ControlRig/Source/ControlRigEditor/Private/AnimDetails/Proxies/AnimDetailsProxyEnum.h
+++ b/Engine/Plugins/Animation/ControlRig/Source/ControlRigEditor/Private/AnimDetails/Proxies/AnimDetailsProxyEnum.h
@@ -38,11 +38,14 @@ public:
 	virtual FName GetDetailRowID() const override;
 	virtual void UpdatePropertyDisplayNames(IDetailLayoutBuilder& DetailBuilder);
 	virtual TArray<FName> GetPropertyNames() const override;
+	virtual TSet<ERigControlType> GetSupportedControlTypes() const override;
 	virtual bool PropertyIsOnProxy(const FProperty* Property, const FProperty* MemberProperty) override;
-	virtual void UpdateProxyValues() override;
+	virtual void AdoptValues(const ERigControlValueType RigControlValueType = ERigControlValueType::Current) override;
+	virtual void ResetPropertyToDefault(const FName& PropertyName) override;
+	virtual bool HasDefaultValue(const FName& PropertyName) const override;
 	virtual EControlRigContextChannelToKey GetChannelToKeyFromPropertyName(const FName& InPropertyName) const override;
 	virtual EControlRigContextChannelToKey GetChannelToKeyFromChannelName(const FString& InChannelName) const override;
-	virtual void SetControlRigElementValueFromCurrent(UControlRig* ControlRig, FRigControlElement* ControlElement, const FRigControlModifiedContext& Context) override;
+	virtual void SetControlRigElementValueFromCurrent(const FRigControlModifiedContext& Context) override;
 	//~ End UAnimDetailsProxyBase interface
 
 	UPROPERTY(EditAnywhere, Interp, Meta = (ShowOnlyInnerProperties), Category = "Enum")
diff --git a/Engine/Plugins/Animation/ControlRig/Source/ControlRigEditor/Private/AnimDetails/Proxies/AnimDetailsProxyFloat.cpp b/Engine/Plugins/Animation/ControlRig/Source/ControlRigEditor/Private/AnimDetails/Proxies/AnimDetailsProxyFloat.cpp
--- a/Engine/Plugins/Animation/ControlRig/Source/ControlRigEditor/Private/AnimDetails/Proxies/AnimDetailsProxyFloat.cpp
+++ b/Engine/Plugins/Animation/ControlRig/Source/ControlRigEditor/Private/AnimDetails/Proxies/AnimDetailsProxyFloat.cpp
@@ -53,57 +53,93 @@ TArray<FName> UAnimDetailsProxyFloat::GetPropertyNames() const
 	};
 }
 
-bool UAnimDetailsProxyFloat::PropertyIsOnProxy(const FProperty* Property, const FProperty* MemberProperty)
+TSet<ERigControlType> UAnimDetailsProxyFloat::GetSupportedControlTypes() const
 {
-	return ((Property && Property->GetFName() == GET_MEMBER_NAME_CHECKED(FAnimDetailsFloat, Float)) ||
-		(MemberProperty && MemberProperty->GetFName() == GET_MEMBER_NAME_CHECKED(UAnimDetailsProxyFloat, Float)));
+	static const TSet<ERigControlType> SupportedControlTypes = 
+	{ 
+		ERigControlType::Float,
+		ERigControlType::ScaleFloat
+	};
+
+	return SupportedControlTypes;
 }
 
-void UAnimDetailsProxyFloat::UpdateProxyValues()
+bool UAnimDetailsProxyFloat::PropertyIsOnProxy(const FProperty* Property, const FProperty* MemberProperty)
 {
-	UControlRig* ControlRig = GetControlRig();
-	FRigControlElement* ControlElement = GetControlElement();
+	return 
+		(Property && Property->GetFName() == GET_MEMBER_NAME_CHECKED(FAnimDetailsFloat, Float)) ||
+		(MemberProperty && MemberProperty->GetFName() == GET_MEMBER_NAME_CHECKED(UAnimDetailsProxyFloat, Float));
+}
 
-	double Value = 0.f;
-	if (ControlRig && ControlElement)
-	{
-		if (ControlElement->Settings.ControlType == ERigControlType::Float ||
-			(ControlElement->Settings.ControlType == ERigControlType::ScaleFloat))
-		{
-			FRigControlValue ControlValue = ControlRig->GetControlValue(ControlElement, ERigControlValueType::Current);
-			Value = ControlValue.Get<float>();
-		}
-	}
-	else if (SequencerItem.IsValid())
-	{
-		if (FProperty* Property = SequencerItem.GetProperty())
+void UAnimDetailsProxyFloat::AdoptValues(const ERigControlValueType RigControlValueType)
+{
+	const double Value = [this, RigControlValueType]() -> double
 		{
-			if (Property->IsA(FDoubleProperty::StaticClass()))
+			UControlRig* ControlRig = GetControlRig();
+			FRigControlElement* ControlElement = GetControlElement();
+			if (ControlRig && ControlElement)
 			{
-				const TOptional<double> OptionalValue = SequencerItem.GetBinding()->GetOptionalValue<double>(*SequencerItem.GetBoundObject());
-				if (OptionalValue.IsSet())
-				{
-					Value = OptionalValue.GetValue();
-				}
+				const FRigControlValue ControlValue = ControlRig->GetControlValue(ControlElement, RigControlValueType);
+				return ControlValue.Get<float>();
 			}
-			else if (Property->IsA(FFloatProperty::StaticClass()))
+			else if (SequencerItem.IsValid())
 			{
-				const TOptional<float> OptionalValue = SequencerItem.GetBinding()->GetOptionalValue<float>(*SequencerItem.GetBoundObject());
-				if (OptionalValue.IsSet())
+				if (FProperty* Property = SequencerItem.GetProperty())
 				{
-					Value = OptionalValue.GetValue();
+					if (Property->IsA(FDoubleProperty::StaticClass()))
+					{
+						const TOptional<double> OptionalValue = SequencerItem.GetBinding()->GetOptionalValue<double>(*SequencerItem.GetBoundObject());
+						if (OptionalValue.IsSet())
+						{
+							return OptionalValue.GetValue();
+						}
+					}
+					else if (Property->IsA(FFloatProperty::StaticClass()))
+					{
+						const TOptional<float> OptionalValue = SequencerItem.GetBinding()->GetOptionalValue<float>(*SequencerItem.GetBoundObject());
+						if (OptionalValue.IsSet())
+						{
+							return OptionalValue.GetValue();
+						}
+					}
 				}
 			}
-		}
-	}
+
+			return 0.0;
+		}();
 
 	const FAnimDetailsFloat AnimDetailProxyFloat(Value);
 
-	const FName PropName("Float");
-	FTrackInstancePropertyBindings Binding(PropName, PropName.ToString());
+	constexpr const TCHAR* PropertyName = TEXT("Float");
+	FTrackInstancePropertyBindings Binding(PropertyName, PropertyName);
 	Binding.CallFunction<FAnimDetailsFloat>(*this, AnimDetailProxyFloat);
 }
 
+void UAnimDetailsProxyFloat::ResetPropertyToDefault(const FName& PropertyName)
+{
+	if (const FRigControlElement* ControlElement = GetControlElement())
+	{
+		AdoptValues(ERigControlValueType::Initial);
+	}
+	else
+	{
+		Float = FAnimDetailsFloat();
+	}
+}
+
+bool UAnimDetailsProxyFloat::HasDefaultValue(const FName& PropertyName) const
+{
+	UControlRig* ControlRig = GetControlRig();
+	FRigControlElement* ControlElement = GetControlElement();
+	if (ControlRig && ControlElement)
+	{
+		const FRigControlValue ControlValue = ControlRig->GetControlValue(ControlElement, ERigControlValueType::Initial);
+		return Float.Float == ControlValue.Get<float>();
+	}
+
+	return Float.Float == 0.0;
+}
+
 EControlRigContextChannelToKey UAnimDetailsProxyFloat::GetChannelToKeyFromPropertyName(const FName& PropertyName) const
 {
 	if (PropertyName == GET_MEMBER_NAME_CHECKED(FAnimDetailsFloat, Float))
@@ -133,36 +169,36 @@ EControlRigContextChannelToKey UAnimDetailsProxyFloat::GetChannelToKeyFromChanne
 	return EControlRigContextChannelToKey::AllTransform;
 }
 
-void UAnimDetailsProxyFloat::SetControlRigElementValueFromCurrent(UControlRig* ControlRig, FRigControlElement* ControlElement, const FRigControlModifiedContext& Context)
+void UAnimDetailsProxyFloat::SetControlRigElementValueFromCurrent(const FRigControlModifiedContext& Context)
 {
-	if (ControlRig && 
-		ControlElement &&
-		(ControlElement->Settings.ControlType == ERigControlType::Float || ControlElement->Settings.ControlType == ERigControlType::ScaleFloat))
+	UControlRig* ControlRig = GetControlRig();
+	FRigControlElement* ControlElement = GetControlElement();
+	if (ControlRig && ControlElement)
 	{
 		float Value = Float.Float;
 
 		constexpr bool bNotify = true;
 		constexpr bool bSetupUndo = false;
 		ControlRig->SetControlValue<float>(ControlElement->GetKey().Name, Value, bNotify, Context, bSetupUndo);
-
-		ControlRig->Evaluate_AnyThread();
 	}
 }
 
-void UAnimDetailsProxyFloat::SetBindingValueFromCurrent(UObject* InObject, const TSharedPtr<FTrackInstancePropertyBindings>& Binding, const FRigControlModifiedContext& Context, bool bInteractive)
+void UAnimDetailsProxyFloat::SetSequencerBindingValueFromCurrent(const FRigControlModifiedContext& Context)
 {
-	if (InObject && Binding.IsValid())
+	const TSharedPtr<FTrackInstancePropertyBindings> Binding = SequencerItem.GetBinding();
+	UObject* BoundObject = SequencerItem.GetBoundObject();
+	if (Binding.IsValid() && BoundObject)
 	{
-		if (FProperty* Property = Binding->GetProperty((*(InObject))))
+		if (FProperty* Property = Binding->GetProperty((*BoundObject)))
 		{
 			if (Property->IsA(FDoubleProperty::StaticClass()))
 			{
-				Binding->SetCurrentValue<double>(*InObject, Float.Float);
+				Binding->SetCurrentValue<double>(*BoundObject, Float.Float);
 			}
 			else if (Property->IsA(FFloatProperty::StaticClass()))
 			{
 				float FVal = (float)Float.Float;
-				Binding->SetCurrentValue<float>(*InObject, FVal);
+				Binding->SetCurrentValue<float>(*BoundObject, FVal);
 			}
 		}
 	}
diff --git a/Engine/Plugins/Animation/ControlRig/Source/ControlRigEditor/Private/AnimDetails/Proxies/AnimDetailsProxyFloat.h b/Engine/Plugins/Animation/ControlRig/Source/ControlRigEditor/Private/AnimDetails/Proxies/AnimDetailsProxyFloat.h
--- a/Engine/Plugins/Animation/ControlRig/Source/ControlRigEditor/Private/AnimDetails/Proxies/AnimDetailsProxyFloat.h
+++ b/Engine/Plugins/Animation/ControlRig/Source/ControlRigEditor/Private/AnimDetails/Proxies/AnimDetailsProxyFloat.h
@@ -33,12 +33,15 @@ public:
 	virtual FName GetCategoryName() const override;
 	virtual void UpdatePropertyDisplayNames(IDetailLayoutBuilder& DetailBuilder);
 	virtual TArray<FName> GetPropertyNames() const override;
+	virtual TSet<ERigControlType> GetSupportedControlTypes() const override;
 	virtual bool PropertyIsOnProxy(const FProperty* Property, const FProperty* MemberProperty) override;
-	virtual void UpdateProxyValues() override;
+	virtual void AdoptValues(const ERigControlValueType RigControlValueType = ERigControlValueType::Current) override;
+	virtual void ResetPropertyToDefault(const FName& PropertyName) override;
+	virtual bool HasDefaultValue(const FName& PropertyName) const override;
 	virtual EControlRigContextChannelToKey GetChannelToKeyFromPropertyName(const FName& InPropertyName) const override;
 	virtual EControlRigContextChannelToKey GetChannelToKeyFromChannelName(const FString& InChannelName) const override;
-	virtual void SetControlRigElementValueFromCurrent(UControlRig* ControlRig, FRigControlElement* ControlElement, const FRigControlModifiedContext& Context) override;
-	virtual void SetBindingValueFromCurrent(UObject* InObject, const TSharedPtr<FTrackInstancePropertyBindings>& Binding, const FRigControlModifiedContext& Context, bool bInteractive = false) override;
+	virtual void SetControlRigElementValueFromCurrent(const FRigControlModifiedContext& Context) override;
+	virtual void SetSequencerBindingValueFromCurrent(const FRigControlModifiedContext& Context) override;
 	//~ End UAnimDetailsProxyBase interface
 
 	UPROPERTY(EditAnywhere, Interp, Category = "Float", meta = (ShowOnlyInnerProperties, Delta = "0.05", SliderExponent = "1", LinearDeltaSensitivity = "1"))
diff --git a/Engine/Plugins/Animation/ControlRig/Source/ControlRigEditor/Private/AnimDetails/Proxies/AnimDetailsProxyInteger.cpp b/Engine/Plugins/Animation/ControlRig/Source/ControlRigEditor/Private/AnimDetails/Proxies/AnimDetailsProxyInteger.cpp
--- a/Engine/Plugins/Animation/ControlRig/Source/ControlRigEditor/Private/AnimDetails/Proxies/AnimDetailsProxyInteger.cpp
+++ b/Engine/Plugins/Animation/ControlRig/Source/ControlRigEditor/Private/AnimDetails/Proxies/AnimDetailsProxyInteger.cpp
@@ -53,44 +53,75 @@ TArray<FName> UAnimDetailsProxyInteger::GetPropertyNames() const
 	};
 }
 
-bool UAnimDetailsProxyInteger::PropertyIsOnProxy(const FProperty* Property, const FProperty* MemberProperty)
+TSet<ERigControlType> UAnimDetailsProxyInteger::GetSupportedControlTypes() const
 {
-	return ((Property && Property->GetFName() == GET_MEMBER_NAME_CHECKED(FAnimDetailsInteger, Integer)) ||
-		(MemberProperty && MemberProperty->GetFName() == GET_MEMBER_NAME_CHECKED(UAnimDetailsProxyInteger, Integer)));
+	static const TSet<ERigControlType> SupportedControlTypes = { ERigControlType::Integer };
+	return SupportedControlTypes;
 }
 
-void UAnimDetailsProxyInteger::UpdateProxyValues()
+bool UAnimDetailsProxyInteger::PropertyIsOnProxy(const FProperty* Property, const FProperty* MemberProperty)
 {
-	int64 Value = 0;
+	return (Property && Property->GetFName() == GET_MEMBER_NAME_CHECKED(FAnimDetailsInteger, Integer)) ||
+		(MemberProperty && MemberProperty->GetFName() == GET_MEMBER_NAME_CHECKED(UAnimDetailsProxyInteger, Integer));
+}
 
-	UControlRig* ControlRig = GetControlRig();
-	FRigControlElement* ControlElement = GetControlElement();
-	if (ControlRig && ControlElement)
-	{
-		if (ControlElement->Settings.ControlType == ERigControlType::Integer)
-		{
-			FRigControlValue ControlValue = ControlRig->GetControlValue(ControlElement, ERigControlValueType::Current);
-			Value = ControlValue.Get<int32>();
-		}
-	}
-	else if (SequencerItem.IsValid())
-	{
-		if (const TOptional<int64> OptionalValue = SequencerItem.GetBinding()->GetOptionalValue<int64>(*SequencerItem.GetBoundObject()))
+void UAnimDetailsProxyInteger::AdoptValues(const ERigControlValueType RigControlValueType)
+{
+	const int64 Value = [this, RigControlValueType]() -> int64
 		{
-			if (OptionalValue.IsSet())
+			UControlRig* ControlRig = GetControlRig();
+			FRigControlElement* ControlElement = GetControlElement();
+			if (ControlRig && ControlElement)
 			{
-				Value = OptionalValue.GetValue();
+				const FRigControlValue ControlValue = ControlRig->GetControlValue(ControlElement, RigControlValueType);
+				return ControlValue.Get<int32>();
 			}
-		}
-	}
+			else if (SequencerItem.IsValid())
+			{
+				if (const TOptional<int64> OptionalValue = SequencerItem.GetBinding()->GetOptionalValue<int64>(*SequencerItem.GetBoundObject()))
+				{
+					if (OptionalValue.IsSet())
+					{
+						return OptionalValue.GetValue();
+					}
+				}
+			}
+
+			return 0;
+		}();
 
 	const FAnimDetailsInteger AnimDetailProxyInteger(Value);
 
-	const FName PropName("Integer");
-	FTrackInstancePropertyBindings Binding(PropName, PropName.ToString());
+	constexpr const TCHAR* PropertyName = TEXT("Integer");
+	FTrackInstancePropertyBindings Binding(PropertyName, PropertyName);
 	Binding.CallFunction<FAnimDetailsInteger>(*this, AnimDetailProxyInteger);
 }
 
+void UAnimDetailsProxyInteger::ResetPropertyToDefault(const FName& PropertyName)
+{
+	if (const FRigControlElement* ControlElement = GetControlElement())
+	{
+		AdoptValues(ERigControlValueType::Initial);
+	}
+	else
+	{
+		Integer = FAnimDetailsInteger();
+	}
+}
+
+bool UAnimDetailsProxyInteger::HasDefaultValue(const FName& PropertyName) const
+{
+	UControlRig* ControlRig = GetControlRig();
+	FRigControlElement* ControlElement = GetControlElement();
+	if (ControlRig && ControlElement)
+	{
+		const FRigControlValue ControlValue = ControlRig->GetControlValue(ControlElement, ERigControlValueType::Initial);
+		return Integer.Integer == ControlValue.Get<int32>();
+	}
+
+	return Integer.Integer == 0;
+}
+
 EControlRigContextChannelToKey UAnimDetailsProxyInteger::GetChannelToKeyFromPropertyName(const FName& PropertyName) const
 {
 	if (PropertyName == GET_MEMBER_NAME_CHECKED(UAnimDetailsProxyInteger, Integer))
@@ -120,27 +151,26 @@ EControlRigContextChannelToKey UAnimDetailsProxyInteger::GetChannelToKeyFromChan
 	return EControlRigContextChannelToKey::AllTransform;
 }
 
-void UAnimDetailsProxyInteger::SetControlRigElementValueFromCurrent(UControlRig* ControlRig, FRigControlElement* ControlElement, const FRigControlModifiedContext& Context)
+void UAnimDetailsProxyInteger::SetControlRigElementValueFromCurrent(const FRigControlModifiedContext& Context)
 {
-	if (ControlRig && 
-		ControlElement &&		
-		ControlElement->Settings.ControlType == ERigControlType::Integer &&
-		!ControlElement->Settings.ControlEnum)
+	UControlRig* ControlRig = GetControlRig();
+	FRigControlElement* ControlElement = GetControlElement();
+	if (ControlRig && ControlElement)
 	{
 		int32 Val = (int32)Integer.Integer;
 
 		constexpr bool bNotify = true;
 		constexpr bool bSetupUndo = false;
 		ControlRig->SetControlValue<int32>(ControlElement->GetKey().Name, Val, bNotify, Context, bSetupUndo);
-
-		ControlRig->Evaluate_AnyThread();
 	}
 }
 
-void UAnimDetailsProxyInteger::SetBindingValueFromCurrent(UObject* InObject, const TSharedPtr<FTrackInstancePropertyBindings>& Binding, const FRigControlModifiedContext& Context, bool bInteractive)
+void UAnimDetailsProxyInteger::SetSequencerBindingValueFromCurrent(const FRigControlModifiedContext& Context)
 {
-	if (InObject && Binding.IsValid())
+	const TSharedPtr<FTrackInstancePropertyBindings> Binding = SequencerItem.GetBinding();
+	UObject* BoundObject = SequencerItem.GetBoundObject();
+	if (Binding.IsValid() && BoundObject)
 	{
-		Binding->SetCurrentValue<int64>(*InObject, Integer.Integer);
+		Binding->SetCurrentValue<int64>(*BoundObject, Integer.Integer);
 	}
 }
\ No newline at end of file
diff --git a/Engine/Plugins/Animation/ControlRig/Source/ControlRigEditor/Private/AnimDetails/Proxies/AnimDetailsProxyInteger.h b/Engine/Plugins/Animation/ControlRig/Source/ControlRigEditor/Private/AnimDetails/Proxies/AnimDetailsProxyInteger.h
--- a/Engine/Plugins/Animation/ControlRig/Source/ControlRigEditor/Private/AnimDetails/Proxies/AnimDetailsProxyInteger.h
+++ b/Engine/Plugins/Animation/ControlRig/Source/ControlRigEditor/Private/AnimDetails/Proxies/AnimDetailsProxyInteger.h
@@ -28,12 +28,15 @@ public:
 	virtual FName GetCategoryName() const override;
 	virtual void UpdatePropertyDisplayNames(IDetailLayoutBuilder& DetailBuilder);
 	virtual TArray<FName> GetPropertyNames() const override;
+	virtual TSet<ERigControlType> GetSupportedControlTypes() const override;
 	virtual bool PropertyIsOnProxy(const FProperty* Property, const FProperty* MemberProperty) override;
-	virtual void UpdateProxyValues() override;
+	virtual void AdoptValues(const ERigControlValueType RigControlValueType = ERigControlValueType::Current) override;
+	virtual void ResetPropertyToDefault(const FName& PropertyName) override;
+	virtual bool HasDefaultValue(const FName& PropertyName) const override;
 	virtual EControlRigContextChannelToKey GetChannelToKeyFromPropertyName(const FName& InPropertyName) const override;
 	virtual EControlRigContextChannelToKey GetChannelToKeyFromChannelName(const FString& InChannelName) const override;
-	virtual void SetControlRigElementValueFromCurrent(UControlRig* ControlRig, FRigControlElement* ControlElement, const FRigControlModifiedContext& Context) override;
-	virtual void SetBindingValueFromCurrent(UObject* InObject, const TSharedPtr<FTrackInstancePropertyBindings>& Binding, const FRigControlModifiedContext& Context, bool bInteractive = false) override;
+	virtual void SetControlRigElementValueFromCurrent(const FRigControlModifiedContext& Context) override;
+	virtual void SetSequencerBindingValueFromCurrent(const FRigControlModifiedContext& Context) override;
 	//~ End UAnimDetailsProxyBase interface
 
 	UPROPERTY(EditAnywhere, Interp, Meta = (ShowOnlyInnerProperties), Category = "Integer")
diff --git a/Engine/Plugins/Animation/ControlRig/Source/ControlRigEditor/Private/AnimDetails/Proxies/AnimDetailsProxyLocation.cpp b/Engine/Plugins/Animation/ControlRig/Source/ControlRigEditor/Private/AnimDetails/Proxies/AnimDetailsProxyLocation.cpp
--- a/Engine/Plugins/Animation/ControlRig/Source/ControlRigEditor/Private/AnimDetails/Proxies/AnimDetailsProxyLocation.cpp
+++ b/Engine/Plugins/Animation/ControlRig/Source/ControlRigEditor/Private/AnimDetails/Proxies/AnimDetailsProxyLocation.cpp
@@ -28,6 +28,22 @@ namespace UE::ControlRigEditor::LocationUtils
 			TLocation.Z = Value.Z;
 		}
 	}
+
+	/** Returns the default location for this proxy */
+	static FAnimDetailsLocation GetDefaultLocation(const UAnimDetailsProxyLocation& Proxy)
+	{
+		UControlRig* ControlRig = Proxy.GetControlRig();
+		FRigControlElement* ControlElement = Proxy.GetControlElement();
+		if (ControlRig && ControlElement)
+		{
+			const FRigControlValue InitialControlValue = ControlRig->GetControlValue(ControlElement, ERigControlValueType::Initial);
+			const FVector3f InitialValue = InitialControlValue.Get<FVector3f>();
+
+			return FAnimDetailsLocation(InitialValue);
+		}
+
+		return FAnimDetailsLocation();
+	}
 }
 
 FAnimDetailsLocation::FAnimDetailsLocation(const FVector& InVector)
@@ -79,34 +95,93 @@ void UAnimDetailsProxyLocation::GetLocalizedPropertyName(const FName& InProperty
 	}
 }
 
+TSet<ERigControlType> UAnimDetailsProxyLocation::GetSupportedControlTypes() const
+{
+	static const TSet<ERigControlType> SupportedControlTypes = { ERigControlType::Position };
+	return SupportedControlTypes;
+}
+
 bool UAnimDetailsProxyLocation::PropertyIsOnProxy(const FProperty* Property, const FProperty* MemberProperty)
 {
-	return ((Property && Property->GetFName() == GET_MEMBER_NAME_CHECKED(UAnimDetailsProxyLocation, Location)) ||
-		(MemberProperty && MemberProperty->GetFName() == GET_MEMBER_NAME_CHECKED(UAnimDetailsProxyLocation, Location)));
+	return (Property && Property->GetFName() == GET_MEMBER_NAME_CHECKED(UAnimDetailsProxyLocation, Location)) ||
+		(MemberProperty && MemberProperty->GetFName() == GET_MEMBER_NAME_CHECKED(UAnimDetailsProxyLocation, Location));
 }
 
-void UAnimDetailsProxyLocation::UpdateProxyValues()
+void UAnimDetailsProxyLocation::AdoptValues(const ERigControlValueType RigControlValueType)
 {
 	UControlRig* ControlRig = GetControlRig();
 	FRigControlElement* ControlElement = GetControlElement();
-	if (!ControlRig || !ControlElement)
+	if (ControlRig && ControlElement)
+	{
+		const FRigControlValue ControlValue = ControlRig->GetControlValue(ControlElement, RigControlValueType);
+		const FVector3f Value = ControlValue.Get<FVector3f>();
+
+		// Alwasys set the value, even if it's the default
+		const FAnimDetailsLocation AnimDetailProxyLocation(Value);
+
+		constexpr const TCHAR* PropertyName = TEXT("Location");
+		FTrackInstancePropertyBindings LocationBinding(PropertyName, PropertyName);
+		LocationBinding.CallFunction<FAnimDetailsLocation>(*this, AnimDetailProxyLocation);
+	}
+
+	// Not implemented for sequencer items, they don't have a Location type
+}
+
+void UAnimDetailsProxyLocation::ResetPropertyToDefault(const FName& PropertyName)
+{
+	using namespace UE::ControlRigEditor;
+	if (PropertyName == FAnimDetailsLocation::StaticStruct()->GetFName())
+	{
+		Location = LocationUtils::GetDefaultLocation(*this);
+	}
+	else
 	{
-		return;
+		FProperty* Property = FAnimDetailsLocation::StaticStruct()->FindPropertyByName(PropertyName);
+		if (!ensureMsgf(Property, TEXT("Cannot find property in struct, cannot reset anim details property to default")))
+		{
+			return;
+		}
+
+		// Reset the specific property name
+		const FAnimDetailsLocation DefaultLocation = LocationUtils::GetDefaultLocation(*this);
+
+		*Property->ContainerPtrToValuePtr<double>(&Location) = *Property->ContainerPtrToValuePtr<double>(&DefaultLocation);
+
+		constexpr const TCHAR* StructPropertyName = TEXT("Location");
+		FTrackInstancePropertyBindings Binding(StructPropertyName, StructPropertyName);
+		Binding.CallFunction<FAnimDetailsLocation>(*this, Location);
 	}
+}
+
+bool UAnimDetailsProxyLocation::HasDefaultValue(const FName& PropertyName) const
+{
+	using namespace UE::ControlRigEditor;
+
+	const FAnimDetailsLocation DefaultLocation = LocationUtils::GetDefaultLocation(*this);
 
-	FVector3f Value = FVector3f::ZeroVector;
-	if (ControlElement->Settings.ControlType == ERigControlType::Position)
+	if (PropertyName == FAnimDetailsLocation::StaticStruct()->GetFName())
 	{
-		FRigControlValue ControlValue = ControlRig->GetControlValue(ControlElement, ERigControlValueType::Current);
-		Value = ControlValue.Get<FVector3f>();
+		return
+			Location.LX == DefaultLocation.LX &&
+			Location.LY == DefaultLocation.LY &&
+			Location.LZ == DefaultLocation.LZ;
+	}
+	else if (PropertyName == GET_MEMBER_NAME_CHECKED(FAnimDetailsLocation, LX))
+	{
+		return Location.LX == DefaultLocation.LX;
+	}
+	else if (PropertyName == GET_MEMBER_NAME_CHECKED(FAnimDetailsLocation, LY))
+	{
+		return Location.LY == DefaultLocation.LY;
+	}
+	else if (PropertyName == GET_MEMBER_NAME_CHECKED(FAnimDetailsLocation, LZ))
+	{
+		return Location.LZ == DefaultLocation.LZ;
 	}
 
-	//set from values, note that the fact if they were multiples or not was already set so we need to set them on these values.
-	const FAnimDetailsLocation AnimDetailProxyLocation(Value);
+	ensureMsgf(0, TEXT("Failed finding property for anim details location. Cannot determine if property has its default value"));
 
-	const FName LocationName("Location");
-	FTrackInstancePropertyBindings LocationBinding(LocationName, LocationName.ToString());
-	LocationBinding.CallFunction<FAnimDetailsLocation>(*this, AnimDetailProxyLocation);
+	return true;
 }
 
 EControlRigContextChannelToKey UAnimDetailsProxyLocation::GetChannelToKeyFromPropertyName(const FName& PropertyName) const
@@ -145,13 +220,13 @@ EControlRigContextChannelToKey UAnimDetailsProxyLocation::GetChannelToKeyFromCha
 	return EControlRigContextChannelToKey::AllTransform;
 }
 
-void UAnimDetailsProxyLocation::SetControlRigElementValueFromCurrent(UControlRig* ControlRig, FRigControlElement* ControlElement, const FRigControlModifiedContext& Context)
+void UAnimDetailsProxyLocation::SetControlRigElementValueFromCurrent(const FRigControlModifiedContext& Context)
 {
 	using namespace UE::ControlRigEditor;
 
-	if (ControlRig && 
-		ControlElement &&		
-		ControlElement->Settings.ControlType == ERigControlType::Position)
+	UControlRig* ControlRig = GetControlRig();
+	FRigControlElement* ControlElement = GetControlElement();
+	if (ControlRig && ControlElement)
 	{
 		FVector3f TLocation = Location.ToVector3f();
 		LocationUtils::SetLocationValuesFromContext(ControlRig, ControlElement, Context, TLocation);
@@ -159,7 +234,5 @@ void UAnimDetailsProxyLocation::SetControlRigElementValueFromCurrent(UControlRig
 		constexpr bool bNotify = true;
 		constexpr bool bSetupUndo = false;
 		ControlRig->SetControlValue<FVector3f>(ControlElement->GetKey().Name, TLocation, bNotify, Context, bSetupUndo);
-
-		ControlRig->Evaluate_AnyThread();
 	}
 }
diff --git a/Engine/Plugins/Animation/ControlRig/Source/ControlRigEditor/Private/AnimDetails/Proxies/AnimDetailsProxyLocation.h b/Engine/Plugins/Animation/ControlRig/Source/ControlRigEditor/Private/AnimDetails/Proxies/AnimDetailsProxyLocation.h
--- a/Engine/Plugins/Animation/ControlRig/Source/ControlRigEditor/Private/AnimDetails/Proxies/AnimDetailsProxyLocation.h
+++ b/Engine/Plugins/Animation/ControlRig/Source/ControlRigEditor/Private/AnimDetails/Proxies/AnimDetailsProxyLocation.h
@@ -41,11 +41,14 @@ public:
 	virtual FName GetCategoryName() const override;
 	virtual TArray<FName> GetPropertyNames() const override;
 	virtual void GetLocalizedPropertyName(const FName& InPropertyName, FText& OutPropertyDisplayName, TOptional<FText>& OutOptionalStructDisplayName) const override;
+	virtual TSet<ERigControlType> GetSupportedControlTypes() const override;
 	virtual bool PropertyIsOnProxy(const FProperty* Property, const FProperty* MemberProperty) override;
-	virtual void UpdateProxyValues() override;
+	virtual void AdoptValues(const ERigControlValueType RigControlValueType = ERigControlValueType::Current) override;
+	virtual void ResetPropertyToDefault(const FName& PropertyName) override;
+	virtual bool HasDefaultValue(const FName& PropertyName) const override;
 	virtual EControlRigContextChannelToKey GetChannelToKeyFromPropertyName(const FName& InPropertyName) const override;
 	virtual EControlRigContextChannelToKey GetChannelToKeyFromChannelName(const FString& InChannelName) const override;
-	virtual void SetControlRigElementValueFromCurrent(UControlRig* ControlRig, FRigControlElement* ControlElement, const FRigControlModifiedContext& Context) override;
+	virtual void SetControlRigElementValueFromCurrent(const FRigControlModifiedContext& Context) override;
 	//~ End UAnimDetailsProxyBase interface
 
 	UPROPERTY(EditAnywhere, Interp, Category = Location)
diff --git a/Engine/Plugins/Animation/ControlRig/Source/ControlRigEditor/Private/AnimDetails/Proxies/AnimDetailsProxyRotation.cpp b/Engine/Plugins/Animation/ControlRig/Source/ControlRigEditor/Private/AnimDetails/Proxies/AnimDetailsProxyRotation.cpp
--- a/Engine/Plugins/Animation/ControlRig/Source/ControlRigEditor/Private/AnimDetails/Proxies/AnimDetailsProxyRotation.cpp
+++ b/Engine/Plugins/Animation/ControlRig/Source/ControlRigEditor/Private/AnimDetails/Proxies/AnimDetailsProxyRotation.cpp
@@ -28,6 +28,20 @@ namespace UE::ControlRigEditor::RotationUtils
 			Val.Z = Value.Z;
 		}
 	}
+
+	/** Returns the default rotation for this proxy */
+	static FAnimDetailsRotation GetDefaultRotation(const UAnimDetailsProxyRotation& Proxy)
+	{
+		UControlRig* ControlRig = Proxy.GetControlRig();
+		FRigControlElement* ControlElement = Proxy.GetControlElement();
+		if (ControlRig && ControlElement)
+		{
+			const FRigControlValue InitialControlValue = ControlRig->GetControlValue(ControlElement, ERigControlValueType::Initial);
+			return FAnimDetailsRotation(InitialControlValue.Get<FVector3f>());
+		}
+
+		return FAnimDetailsRotation();
+	}
 }
 
 FAnimDetailsRotation::FAnimDetailsRotation(const FRotator& InRotator)
@@ -104,33 +118,94 @@ void UAnimDetailsProxyRotation::GetLocalizedPropertyName(const FName& InProperty
 	}
 }
 
+TSet<ERigControlType> UAnimDetailsProxyRotation::GetSupportedControlTypes() const
+{
+	static const TSet<ERigControlType> SupportedControlTypes = { ERigControlType::Rotator };
+	return SupportedControlTypes;
+}
+
 bool UAnimDetailsProxyRotation::PropertyIsOnProxy(const FProperty* Property, const FProperty* MemberProperty)
 {
-	return ((Property && Property->GetFName() == GET_MEMBER_NAME_CHECKED(UAnimDetailsProxyRotation, Rotation)) ||
-		(MemberProperty && MemberProperty->GetFName() == GET_MEMBER_NAME_CHECKED(UAnimDetailsProxyRotation, Rotation)));
+	return (Property && Property->GetFName() == GET_MEMBER_NAME_CHECKED(UAnimDetailsProxyRotation, Rotation)) ||
+		(MemberProperty && MemberProperty->GetFName() == GET_MEMBER_NAME_CHECKED(UAnimDetailsProxyRotation, Rotation));
 }
 
-void UAnimDetailsProxyRotation::UpdateProxyValues()
+void UAnimDetailsProxyRotation::AdoptValues(const ERigControlValueType RigControlValueType)
 {
 	UControlRig* ControlRig = GetControlRig();
 	FRigControlElement* ControlElement = GetControlElement();
-	if (!ControlRig || !ControlElement)
+	if (ControlRig && ControlElement)
 	{
-		return;
+		FRigControlValue ControlValue = ControlRig->GetControlValue(ControlElement, RigControlValueType);
+		const FVector3f Value = ControlValue.Get<FVector3f>();
+
+		const FAnimDetailsRotation AnimDetailProxyRotation(Value);
+
+		constexpr const TCHAR* PropertyName = TEXT("Rotation");
+		FTrackInstancePropertyBindings Binding(PropertyName, PropertyName);
+		Binding.CallFunction<FAnimDetailsRotation>(*this, AnimDetailProxyRotation);
 	}
-	
-	FVector3f Value = FVector3f::ZeroVector;
-	if (ControlElement->Settings.ControlType == ERigControlType::Rotator)
+
+	// Not implemented for sequencer items, they don't have a rotation type
+}
+
+void UAnimDetailsProxyRotation::ResetPropertyToDefault(const FName& PropertyName)
+{
+	using namespace UE::ControlRigEditor;
+
+	if (PropertyName == FAnimDetailsRotation::StaticStruct()->GetFName())
 	{
-		FRigControlValue ControlValue = ControlRig->GetControlValue(ControlElement, ERigControlValueType::Current);
-		Value = ControlValue.Get<FVector3f>();
+		// Reset the whole rotation struct
+		Rotation = RotationUtils::GetDefaultRotation(*this);
 	}
+	else
+	{
+		FProperty* Property = FAnimDetailsRotation::StaticStruct()->FindPropertyByName(PropertyName);
+		if (!ensureMsgf(Property, TEXT("Cannot find property in struct, cannot reset anim details property to default")))
+		{
+			return;
+		}
 
-	const FAnimDetailsRotation AnimDetailProxyRotation(Value);
+		// Reset the specific property name
+		const FAnimDetailsRotation DefaultRotation = RotationUtils::GetDefaultRotation(*this);
 
-	const FName PropName("Rotation");
-	FTrackInstancePropertyBindings Binding(PropName, PropName.ToString());
-	Binding.CallFunction<FAnimDetailsRotation>(*this, AnimDetailProxyRotation);
+		*Property->ContainerPtrToValuePtr<double>(&Rotation) = *Property->ContainerPtrToValuePtr<double>(&DefaultRotation);
+
+		constexpr const TCHAR* StructPropertyName = TEXT("Rotation");
+		FTrackInstancePropertyBindings Binding(StructPropertyName, StructPropertyName);
+		Binding.CallFunction<FAnimDetailsRotation>(*this, Rotation);
+	}
+}
+
+bool UAnimDetailsProxyRotation::HasDefaultValue(const FName& PropertyName) const
+{
+	using namespace UE::ControlRigEditor;
+
+	const FAnimDetailsRotation DefaultRotation = RotationUtils::GetDefaultRotation(*this);
+
+	if (PropertyName == FAnimDetailsRotation::StaticStruct()->GetFName())
+	{
+		return
+			Rotation.RX == DefaultRotation.RX &&
+			Rotation.RY == DefaultRotation.RY &&
+			Rotation.RZ == DefaultRotation.RZ;
+	}
+	else if (PropertyName == GET_MEMBER_NAME_CHECKED(FAnimDetailsRotation, RX))
+	{
+		return Rotation.RX == DefaultRotation.RX;
+	}
+	else if (PropertyName == GET_MEMBER_NAME_CHECKED(FAnimDetailsRotation, RY))
+	{
+		return Rotation.RY == DefaultRotation.RY;
+	}
+	else if (PropertyName == GET_MEMBER_NAME_CHECKED(FAnimDetailsRotation, RZ))
+	{
+		return Rotation.RZ == DefaultRotation.RZ;
+	}
+
+	ensureMsgf(0, TEXT("Failed finding property for anim details rotation. Cannot determine if property has its default value"));
+
+	return true;
 }
 
 EControlRigContextChannelToKey UAnimDetailsProxyRotation::GetChannelToKeyFromPropertyName(const FName& PropertyName) const
@@ -169,15 +244,14 @@ EControlRigContextChannelToKey UAnimDetailsProxyRotation::GetChannelToKeyFromCha
 	return EControlRigContextChannelToKey::AllTransform;
 }
 
-void UAnimDetailsProxyRotation::SetControlRigElementValueFromCurrent(UControlRig* ControlRig, FRigControlElement* ControlElement, const FRigControlModifiedContext& Context)
+void UAnimDetailsProxyRotation::SetControlRigElementValueFromCurrent(const FRigControlModifiedContext& Context)
 {
 	using namespace UE::ControlRigEditor;
 
+	UControlRig* ControlRig = GetControlRig();
 	URigHierarchy* Hierarchy = ControlRig ? ControlRig->GetHierarchy() : nullptr;
-	if (ControlRig && 
-		ControlElement &&		
-		ControlElement->Settings.ControlType == ERigControlType::Rotator &&
-		Hierarchy)
+	FRigControlElement* ControlElement = GetControlElement();
+	if (ControlRig && ControlElement)
 	{
 		FVector3f Value = Rotation.ToVector3f();
 		RotationUtils::SetRotationValuesFromContext(ControlRig, ControlElement, Context, Value);
diff --git a/Engine/Plugins/Animation/ControlRig/Source/ControlRigEditor/Private/AnimDetails/Proxies/AnimDetailsProxyRotation.h b/Engine/Plugins/Animation/ControlRig/Source/ControlRigEditor/Private/AnimDetails/Proxies/AnimDetailsProxyRotation.h
--- a/Engine/Plugins/Animation/ControlRig/Source/ControlRigEditor/Private/AnimDetails/Proxies/AnimDetailsProxyRotation.h
+++ b/Engine/Plugins/Animation/ControlRig/Source/ControlRigEditor/Private/AnimDetails/Proxies/AnimDetailsProxyRotation.h
@@ -43,11 +43,14 @@ public:
 	virtual FName GetCategoryName() const override;
 	virtual TArray<FName> GetPropertyNames() const override;
 	virtual void GetLocalizedPropertyName(const FName& InPropertyName, FText& OutPropertyDisplayName, TOptional<FText>& OutOptionalStructDisplayName) const override;
+	virtual TSet<ERigControlType> GetSupportedControlTypes() const override;
 	virtual bool PropertyIsOnProxy(const FProperty* Property, const FProperty* MemberProperty) override;
-	virtual void UpdateProxyValues() override;
+	virtual void AdoptValues(const ERigControlValueType RigControlValueType = ERigControlValueType::Current) override;
+	virtual void ResetPropertyToDefault(const FName& PropertyName) override;
+	virtual bool HasDefaultValue(const FName& PropertyName) const override;
 	virtual EControlRigContextChannelToKey GetChannelToKeyFromPropertyName(const FName& InPropertyName) const override;
 	virtual EControlRigContextChannelToKey GetChannelToKeyFromChannelName(const FString& InChannelName) const override;
-	virtual void SetControlRigElementValueFromCurrent(UControlRig* ControlRig, FRigControlElement* ControlElement, const FRigControlModifiedContext& Context) override;
+	virtual void SetControlRigElementValueFromCurrent(const FRigControlModifiedContext& Context) override;
 	//~ End UAnimDetailsProxyBase interface
 
 	UPROPERTY(EditAnywhere, Interp, Category = Rotation)
diff --git a/Engine/Plugins/Animation/ControlRig/Source/ControlRigEditor/Private/AnimDetails/Proxies/AnimDetailsProxyScale.cpp b/Engine/Plugins/Animation/ControlRig/Source/ControlRigEditor/Private/AnimDetails/Proxies/AnimDetailsProxyScale.cpp
--- a/Engine/Plugins/Animation/ControlRig/Source/ControlRigEditor/Private/AnimDetails/Proxies/AnimDetailsProxyScale.cpp
+++ b/Engine/Plugins/Animation/ControlRig/Source/ControlRigEditor/Private/AnimDetails/Proxies/AnimDetailsProxyScale.cpp
@@ -29,6 +29,20 @@ namespace UE::ControlRigEditor::ScaleUtils
 			TScale.Z = Value.Z;
 		}
 	}
+
+	/** Returns the default scale for this proxy */
+	static FAnimDetailsScale GetDefaultScale(const UAnimDetailsProxyScale& Proxy)
+	{
+		UControlRig* ControlRig = Proxy.GetControlRig();
+		FRigControlElement* ControlElement = Proxy.GetControlElement();
+		if (ControlRig && ControlElement)
+		{
+			const FRigControlValue InitialControlValue = ControlRig->GetControlValue(ControlElement, ERigControlValueType::Initial);
+			return InitialControlValue.Get<FVector3f>();
+		}
+
+		return FAnimDetailsScale();
+	}
 }
 
 FAnimDetailsScale::FAnimDetailsScale(const FVector& InVector)
@@ -90,34 +104,94 @@ void UAnimDetailsProxyScale::GetLocalizedPropertyName(const FName& InPropertyNam
 	}
 }
 
+TSet<ERigControlType> UAnimDetailsProxyScale::GetSupportedControlTypes() const
+{
+	static const TSet<ERigControlType> SupportedControlTypes = { ERigControlType::Scale };
+	return SupportedControlTypes;
+}
+
 bool UAnimDetailsProxyScale::PropertyIsOnProxy(const FProperty* Property, const FProperty* MemberProperty)
 {
-	return ((Property && Property->GetFName() == GET_MEMBER_NAME_CHECKED(UAnimDetailsProxyScale, Scale)) ||
-		(MemberProperty && MemberProperty->GetFName() == GET_MEMBER_NAME_CHECKED(UAnimDetailsProxyScale, Scale)));
+	return (Property && Property->GetFName() == GET_MEMBER_NAME_CHECKED(UAnimDetailsProxyScale, Scale)) ||
+		(MemberProperty && MemberProperty->GetFName() == GET_MEMBER_NAME_CHECKED(UAnimDetailsProxyScale, Scale));
 }
 
-void UAnimDetailsProxyScale::UpdateProxyValues()
+void UAnimDetailsProxyScale::AdoptValues(const ERigControlValueType RigControlValueType)
 {
 	UControlRig* ControlRig = GetControlRig();
 	FRigControlElement* ControlElement = GetControlElement();
-	if (!ControlRig || !ControlElement)
+	if (ControlRig && ControlElement)
 	{
-		return;
+		FRigControlValue ControlValue = ControlRig->GetControlValue(ControlElement, RigControlValueType);
+		const FVector3f Value = ControlValue.Get<FVector3f>();
+
+		const FAnimDetailsScale AnimDetailProxyScale(Value);
+
+		constexpr const TCHAR* PropertyName = TEXT("Scale");
+		FTrackInstancePropertyBindings Binding(PropertyName, PropertyName);
+		Binding.CallFunction<FAnimDetailsScale>(*this, AnimDetailProxyScale);
 	}
 
-	FVector3f Value = FVector3f::ZeroVector;
-	if (ControlElement->Settings.ControlType == ERigControlType::Scale)
+	// Not implemented for sequencer items, they don't have a scale type
+}
+
+void UAnimDetailsProxyScale::ResetPropertyToDefault(const FName& PropertyName)
+{
+	using namespace UE::ControlRigEditor;
+
+	if (PropertyName == FAnimDetailsScale::StaticStruct()->GetFName())
 	{
-		FRigControlValue ControlValue = ControlRig->GetControlValue(ControlElement, ERigControlValueType::Current);
-		Value = ControlValue.Get<FVector3f>();
+		// Reset the whole scale struct
+		Scale = ScaleUtils::GetDefaultScale(*this);
 	}
+	else
+	{
+		FProperty* Property = FAnimDetailsScale::StaticStruct()->FindPropertyByName(PropertyName);
+		if (!ensureMsgf(Property, TEXT("Cannot find property in struct, cannot reset anim details property to default")))
+		{
+			return;
+		}
+
+		// Reset the specific property name
+		const FAnimDetailsScale DefaultScale = ScaleUtils::GetDefaultScale(*this);
 
-	//set from values, note that the fact if they were multiples or not was already set so we need to set them on these values.
-	const FAnimDetailsScale AnimDetailProxyScale(Value);
+		*Property->ContainerPtrToValuePtr<double>(&Scale) = *Property->ContainerPtrToValuePtr<double>(&DefaultScale);
 
-	const FName PropName("Scale");
-	FTrackInstancePropertyBindings Binding(PropName, PropName.ToString());
-	Binding.CallFunction<FAnimDetailsScale>(*this, AnimDetailProxyScale);
+		constexpr const TCHAR* StructPropertyName = TEXT("Scale");
+		FTrackInstancePropertyBindings Binding(StructPropertyName, StructPropertyName);
+		Binding.CallFunction<FAnimDetailsScale>(*this, Scale);
+	}
+}
+
+bool UAnimDetailsProxyScale::HasDefaultValue(const FName& PropertyName) const
+{
+	using namespace UE::ControlRigEditor;
+
+	const FAnimDetailsScale DefaultScale = ScaleUtils::GetDefaultScale(*this);
+
+	if (PropertyName == FAnimDetailsScale::StaticStruct()->GetFName())
+	{
+		return
+			Scale.SX == DefaultScale.SX &&
+			Scale.SY == DefaultScale.SY &&
+			Scale.SZ == DefaultScale.SZ;
+	}
+	else if (PropertyName == GET_MEMBER_NAME_CHECKED(FAnimDetailsScale, SX))
+	{
+		return Scale.SX == DefaultScale.SX;
+	}
+	else if (PropertyName == GET_MEMBER_NAME_CHECKED(FAnimDetailsScale, SY))
+	{
+		return Scale.SY == DefaultScale.SY;
+	}
+	else if (PropertyName == GET_MEMBER_NAME_CHECKED(FAnimDetailsScale, SZ))
+	{
+		return Scale.SZ == DefaultScale.SZ;
+	}
+
+	ensureMsgf(0, TEXT("Failed finding property for anim details scale. Cannot determine if property has its default value"));
+
+	return true;
 }
 
 EControlRigContextChannelToKey UAnimDetailsProxyScale::GetChannelToKeyFromPropertyName(const FName& PropertyName) const
@@ -156,13 +230,13 @@ EControlRigContextChannelToKey UAnimDetailsProxyScale::GetChannelToKeyFromChanne
 	return EControlRigContextChannelToKey::AllTransform;
 }
 
-void UAnimDetailsProxyScale::SetControlRigElementValueFromCurrent(UControlRig* ControlRig, FRigControlElement* ControlElement, const FRigControlModifiedContext& Context)
+void UAnimDetailsProxyScale::SetControlRigElementValueFromCurrent(const FRigControlModifiedContext& Context)
 {
 	using namespace UE::ControlRigEditor;
 
-	if (ControlRig && 
-		ControlElement &&		
-		ControlElement->Settings.ControlType == ERigControlType::Scale)
+	UControlRig* ControlRig = GetControlRig();
+	FRigControlElement* ControlElement = GetControlElement();
+	if (ControlRig && ControlElement)
 	{
 		FVector3f Val = Scale.ToVector3f();
 		ScaleUtils::SetScaleValuesFromContext(ControlRig, ControlElement, Context, Val);
diff --git a/Engine/Plugins/Animation/ControlRig/Source/ControlRigEditor/Private/AnimDetails/Proxies/AnimDetailsProxyScale.h b/Engine/Plugins/Animation/ControlRig/Source/ControlRigEditor/Private/AnimDetails/Proxies/AnimDetailsProxyScale.h
--- a/Engine/Plugins/Animation/ControlRig/Source/ControlRigEditor/Private/AnimDetails/Proxies/AnimDetailsProxyScale.h
+++ b/Engine/Plugins/Animation/ControlRig/Source/ControlRigEditor/Private/AnimDetails/Proxies/AnimDetailsProxyScale.h
@@ -41,11 +41,14 @@ public:
 	virtual FName GetCategoryName() const override;
 	virtual TArray<FName> GetPropertyNames() const override;
 	virtual void GetLocalizedPropertyName(const FName& InPropertyName, FText& OutPropertyDisplayName, TOptional<FText>& OutOptionalStructDisplayName) const override;
+	virtual TSet<ERigControlType> GetSupportedControlTypes() const override;
 	virtual bool PropertyIsOnProxy(const FProperty* Property, const FProperty* MemberProperty) override;
-	virtual void UpdateProxyValues() override;
+	virtual void AdoptValues(const ERigControlValueType RigControlValueType = ERigControlValueType::Current) override;
+	virtual void ResetPropertyToDefault(const FName& PropertyName) override;
+	virtual bool HasDefaultValue(const FName& PropertyName) const override;
 	virtual EControlRigContextChannelToKey GetChannelToKeyFromPropertyName(const FName& InPropertyName) const override;
 	virtual EControlRigContextChannelToKey GetChannelToKeyFromChannelName(const FString& InChannelName) const override;
-	virtual void SetControlRigElementValueFromCurrent(UControlRig* ControlRig, FRigControlElement* ControlElement, const FRigControlModifiedContext& Context) override;
+	virtual void SetControlRigElementValueFromCurrent(const FRigControlModifiedContext& Context) override;
 	//~ End UAnimDetailsProxyBase interface
 
 	UPROPERTY(EditAnywhere, Interp, Category = Scale)
diff --git a/Engine/Plugins/Animation/ControlRig/Source/ControlRigEditor/Private/AnimDetails/Proxies/AnimDetailsProxyTransform.cpp b/Engine/Plugins/Animation/ControlRig/Source/ControlRigEditor/Private/AnimDetails/Proxies/AnimDetailsProxyTransform.cpp
--- a/Engine/Plugins/Animation/ControlRig/Source/ControlRigEditor/Private/AnimDetails/Proxies/AnimDetailsProxyTransform.cpp
+++ b/Engine/Plugins/Animation/ControlRig/Source/ControlRigEditor/Private/AnimDetails/Proxies/AnimDetailsProxyTransform.cpp
@@ -61,7 +61,7 @@ namespace UE::ControlRigEditor
 			}
 		}
 
-		static FEulerTransform GetCurrentValue(UControlRig* ControlRig, FRigControlElement* ControlElement)
+		static FEulerTransform GetValue(UControlRig* ControlRig, FRigControlElement* ControlElement, const ERigControlValueType ValueType)
 		{
 			FEulerTransform EulerTransform = FEulerTransform::Identity;
 
@@ -69,14 +69,14 @@ namespace UE::ControlRigEditor
 			{
 			case ERigControlType::Transform:
 			{
-				const FTransform NewTransform = ControlRig->GetControlValue(ControlElement, ERigControlValueType::Current).Get<FRigControlValue::FTransform_Float>().ToTransform();
+				const FTransform NewTransform = ControlRig->GetControlValue(ControlElement, ValueType).Get<FRigControlValue::FTransform_Float>().ToTransform();
 				EulerTransform = FEulerTransform(NewTransform);
 				break;
 
 			}
 			case ERigControlType::TransformNoScale:
 			{
-				const FTransformNoScale NewTransform = ControlRig->GetControlValue(ControlElement, ERigControlValueType::Current).Get<FRigControlValue::FTransformNoScale_Float>().ToTransform();
+				const FTransformNoScale NewTransform = ControlRig->GetControlValue(ControlElement, ValueType).Get<FRigControlValue::FTransformNoScale_Float>().ToTransform();
 				EulerTransform.Location = NewTransform.Location;
 				EulerTransform.Rotation = FRotator(NewTransform.Rotation);
 				break;
@@ -84,14 +84,15 @@ namespace UE::ControlRigEditor
 			}
 			case ERigControlType::EulerTransform:
 			{
-				EulerTransform = ControlRig->GetControlValue(ControlElement, ERigControlValueType::Current).Get<FRigControlValue::FEulerTransform_Float>().ToTransform();
+				EulerTransform = ControlRig->GetControlValue(ControlElement, ValueType).Get<FRigControlValue::FEulerTransform_Float>().ToTransform();
 				break;
 			}
 			};
 
 			if (ControlRig->GetHierarchy()->UsesPreferredEulerAngles())
 			{
-				EulerTransform.Rotation = ControlRig->GetHierarchy()->GetControlPreferredRotator(ControlElement);
+				const bool bInitial = ValueType == ERigControlValueType::Initial;
+				EulerTransform.Rotation = ControlRig->GetHierarchy()->GetControlPreferredRotator(ControlElement, bInitial);
 			}
 
 			return EulerTransform;
@@ -115,7 +116,7 @@ namespace UE::ControlRigEditor
 			}
 		}
 
-		static FEulerTransform GetCurrentValue(UObject* InObject, const TSharedPtr<FTrackInstancePropertyBindings>& Binding)
+		static FEulerTransform GetValue(UObject* InObject, const TSharedPtr<FTrackInstancePropertyBindings>& Binding)
 		{
 			const FStructProperty* TransformProperty = Binding.IsValid() ? CastField<FStructProperty>(Binding->GetProperty(*InObject)) : nullptr;
 			if (TransformProperty && TransformProperty->Struct == TBaseStructure<FTransform>::Get())
@@ -217,6 +218,19 @@ namespace UE::ControlRigEditor
 
 			return true;
 		}
+
+		/** Returns the default transfrom for this proxy */
+		static FEulerTransform GetDefaultTransform(const UAnimDetailsProxyTransform& Proxy)
+		{
+			UControlRig* ControlRig = Proxy.GetControlRig();
+			FRigControlElement* ControlElement = Proxy.GetControlElement();
+			if (ControlRig && ControlElement)
+			{
+				return TransformUtils::GetValue(ControlRig, ControlElement, ERigControlValueType::Initial);
+			}
+
+			return FEulerTransform();
+		}
 	}
 }
 
@@ -296,6 +310,18 @@ void UAnimDetailsProxyTransform::GetLocalizedPropertyName(const FName& InPropert
 	}
 }
 
+TSet<ERigControlType> UAnimDetailsProxyTransform::GetSupportedControlTypes() const
+{
+	static const TSet<ERigControlType> SupportedControlTypes =
+	{
+		ERigControlType::Transform,
+		ERigControlType::EulerTransform,
+		ERigControlType::TransformNoScale
+	};
+
+	return SupportedControlTypes;
+}
+
 bool UAnimDetailsProxyTransform::PropertyIsOnProxy(const FProperty* Property, const FProperty* MemberProperty)
 {
 	return ((Property && Property->GetFName() == GET_MEMBER_NAME_CHECKED(UAnimDetailsProxyTransform, Location)) ||
@@ -306,40 +332,148 @@ bool UAnimDetailsProxyTransform::PropertyIsOnProxy(const FProperty* Property, co
 		(MemberProperty && MemberProperty->GetFName() == GET_MEMBER_NAME_CHECKED(UAnimDetailsProxyTransform, Scale)));
 }
 
-void UAnimDetailsProxyTransform::UpdateProxyValues()
+void UAnimDetailsProxyTransform::AdoptValues(const ERigControlValueType RigControlValueType)
 {
 	using namespace UE::ControlRigEditor;
 
+	FEulerTransform Value;
+
 	UControlRig* ControlRig = GetControlRig();
 	FRigControlElement* ControlElement = GetControlElement();
-
-	FEulerTransform EulerTransform;
 	if (ControlRig && ControlElement)
 	{
-		EulerTransform = TransformUtils::GetCurrentValue(ControlRig, ControlElement);
+		Value = TransformUtils::GetValue(ControlRig, ControlElement, RigControlValueType);
 	}
 	else if (SequencerItem.IsValid())
 	{
-		EulerTransform = TransformUtils::GetCurrentValue(SequencerItem.GetBoundObject(), SequencerItem.GetBinding());
+		Value = TransformUtils::GetValue(SequencerItem.GetBoundObject(), SequencerItem.GetBinding());
 	}
 
-	FAnimDetailsLocation TLocation(EulerTransform.Location);
-	FAnimDetailsRotation TRotation(EulerTransform.Rotation);
-	FAnimDetailsScale TScale(EulerTransform.Scale);
+	const FAnimDetailsLocation TLocation(Value.Location);
+	const FAnimDetailsRotation TRotation(Value.Rotation);
+	const FAnimDetailsScale TScale(Value.Scale);
 
-	const FName LocationName("Location");
-	FTrackInstancePropertyBindings LocationBinding(LocationName, LocationName.ToString());
+	constexpr const TCHAR* LocationPropertyName = TEXT("Location");
+	FTrackInstancePropertyBindings LocationBinding(LocationPropertyName, LocationPropertyName);
 	LocationBinding.CallFunction<FAnimDetailsLocation>(*this, TLocation);
 
-	const FName RotationName("Rotation");
-	FTrackInstancePropertyBindings RotationBinding(RotationName, RotationName.ToString());
+	constexpr const TCHAR* RotationPropertyName = TEXT("Rotation");
+	FTrackInstancePropertyBindings RotationBinding(RotationPropertyName, RotationPropertyName);
 	RotationBinding.CallFunction<FAnimDetailsRotation>(*this, TRotation);
 
-	const FName ScaleName("Scale");
-	FTrackInstancePropertyBindings ScaleBinding(ScaleName, ScaleName.ToString());
+	constexpr const TCHAR* ScalePropertyName = TEXT("Scale");
+	FTrackInstancePropertyBindings ScaleBinding(ScalePropertyName, ScalePropertyName);
 	ScaleBinding.CallFunction<FAnimDetailsScale>(*this, TScale);
 }
 
+void UAnimDetailsProxyTransform::ResetPropertyToDefault(const FName& PropertyName)
+{
+	using namespace UE::ControlRigEditor;
+
+	FEulerTransform CurrentValue;
+	CurrentValue.Location = Location.ToVector();
+	CurrentValue.Rotation = Rotation.ToRotator();
+	CurrentValue.Scale = Scale.ToVector();
+
+	const FEulerTransform DefaulTransform = TransformUtils::GetDefaultTransform(*this);
+
+	TOptional<FAnimDetailsLocation> NewLocation;
+	TOptional<FAnimDetailsRotation> NewRotation;
+	TOptional<FAnimDetailsScale> NewScale;
+
+	if (PropertyName == GET_MEMBER_NAME_CHECKED(UAnimDetailsProxyTransform, Location))
+	{
+		NewLocation = FAnimDetailsLocation(DefaulTransform.Location);
+	}
+	else if (PropertyName == GET_MEMBER_NAME_CHECKED(UAnimDetailsProxyTransform, Rotation))
+	{
+		NewRotation = FAnimDetailsRotation(DefaulTransform.Rotation);
+	}
+	else if (PropertyName == GET_MEMBER_NAME_CHECKED(UAnimDetailsProxyTransform, Scale))
+	{
+		NewScale = FAnimDetailsScale(DefaulTransform.Scale);
+	}
+	else if (FProperty* LocationProperty = FAnimDetailsLocation::StaticStruct()->FindPropertyByName(PropertyName))
+	{
+		NewLocation = Location;
+		*LocationProperty->ContainerPtrToValuePtr<double>(&NewLocation.GetValue()) = 
+			*LocationProperty->ContainerPtrToValuePtr<double>(&DefaulTransform.Location);
+	}
+	else if (FProperty* RotationProperty = FAnimDetailsRotation::StaticStruct()->FindPropertyByName(PropertyName))
+	{
+		NewRotation = Rotation;
+		*RotationProperty->ContainerPtrToValuePtr<double>(&NewRotation.GetValue()) = 
+			*RotationProperty->ContainerPtrToValuePtr<double>(&DefaulTransform.Rotation);
+	}
+	else if (FProperty* ScaleProperty = FAnimDetailsScale::StaticStruct()->FindPropertyByName(PropertyName))
+	{
+		NewScale = Scale;
+		*ScaleProperty->ContainerPtrToValuePtr<double>(&NewScale.GetValue()) = 
+			*ScaleProperty->ContainerPtrToValuePtr<double>(&DefaulTransform.Scale);
+	}
+
+	if (NewLocation.IsSet())
+	{
+		constexpr const TCHAR* LocationPropertyName = TEXT("Location");
+		FTrackInstancePropertyBindings LocationBinding(LocationPropertyName, LocationPropertyName);
+		LocationBinding.CallFunction<FAnimDetailsLocation>(*this, NewLocation.GetValue());
+	}
+	else if (NewRotation.IsSet())
+	{
+		constexpr const TCHAR* RotationPropertyName = TEXT("Rotation");
+		FTrackInstancePropertyBindings RotationBinding(RotationPropertyName, RotationPropertyName);
+		RotationBinding.CallFunction<FAnimDetailsRotation>(*this, NewRotation.GetValue());
+	}
+	else if (NewScale.IsSet())
+	{
+		constexpr const TCHAR* ScalePropertyName = TEXT("Scale");
+		FTrackInstancePropertyBindings ScaleBinding(ScalePropertyName, ScalePropertyName);
+		ScaleBinding.CallFunction<FAnimDetailsScale>(*this, NewScale.GetValue());
+	}
+}
+
+bool UAnimDetailsProxyTransform::HasDefaultValue(const FName& PropertyName) const
+{
+	using namespace UE::ControlRigEditor;
+
+	const FEulerTransform DefaulTransform = TransformUtils::GetDefaultTransform(*this);
+
+	if (PropertyName == GET_MEMBER_NAME_CHECKED(UAnimDetailsProxyTransform, Location))
+	{
+		return Location.ToVector() == DefaulTransform.Location;
+	}
+	else if (PropertyName == GET_MEMBER_NAME_CHECKED(UAnimDetailsProxyTransform, Rotation))
+	{
+		return Rotation.ToRotator() == DefaulTransform.Rotation;
+	}
+	else if (PropertyName == GET_MEMBER_NAME_CHECKED(UAnimDetailsProxyTransform, Scale))
+	{
+		return Scale.ToVector() == DefaulTransform.Scale;
+	}
+	else if (FProperty* LocationProperty = FAnimDetailsLocation::StaticStruct()->FindPropertyByName(PropertyName))
+	{
+		const FAnimDetailsLocation& DefaultLocation = DefaulTransform.Location;
+		return (*LocationProperty->ContainerPtrToValuePtr<double>(&Location)) == 
+			(*LocationProperty->ContainerPtrToValuePtr<double>(&DefaultLocation));
+	}
+	else if (FProperty* RotationProperty = FAnimDetailsRotation::StaticStruct()->FindPropertyByName(PropertyName))
+	{
+		const FAnimDetailsRotation& DefaultRotation = DefaulTransform.Rotation;
+		return (*RotationProperty->ContainerPtrToValuePtr<double>(&Rotation)) == 
+			(*RotationProperty->ContainerPtrToValuePtr<double>(&DefaultRotation));
+	}
+	else if (FProperty* ScaleProperty = FAnimDetailsScale::StaticStruct()->FindPropertyByName(PropertyName))
+	{
+		const FAnimDetailsScale& DefaultScale = DefaulTransform.Scale;
+		return (*ScaleProperty->ContainerPtrToValuePtr<double>(&Scale)) == 
+			(*ScaleProperty->ContainerPtrToValuePtr<double>(&DefaultScale));
+	}
+
+	ensureMsgf(0, TEXT("Failed finding property %s for anim details transform. Cannot determine if property has its default value"), *PropertyName.ToString());
+
+	return true;
+}
+
 EControlRigContextChannelToKey UAnimDetailsProxyTransform::GetChannelToKeyFromPropertyName(const FName& PropertyName) const
 {
 	if (PropertyName == GET_MEMBER_NAME_CHECKED(UAnimDetailsProxyTransform, Location))
@@ -435,10 +569,12 @@ EControlRigContextChannelToKey UAnimDetailsProxyTransform::GetChannelToKeyFromCh
 	return EControlRigContextChannelToKey::AllTransform;
 }
 
-void UAnimDetailsProxyTransform::SetControlRigElementValueFromCurrent(UControlRig* ControlRig, FRigControlElement* ControlElement, const FRigControlModifiedContext& Context)
+void UAnimDetailsProxyTransform::SetControlRigElementValueFromCurrent(const FRigControlModifiedContext& Context)
 {
 	using namespace UE::ControlRigEditor;
 
+	UControlRig* ControlRig = GetControlRig();
+	FRigControlElement* ControlElement = GetControlElement();
 	if (!ControlElement || !ControlRig)
 	{
 		return;
@@ -450,7 +586,7 @@ void UAnimDetailsProxyTransform::SetControlRigElementValueFromCurrent(UControlRi
 	FVector TLocation = Location.ToVector();
 	FRotator TRotation = Rotation.ToRotator();
 	FVector TScale = Scale.ToVector();
-	FEulerTransform EulerTransform = TransformUtils::GetCurrentValue(ControlRig, ControlElement);
+	FEulerTransform EulerTransform = TransformUtils::GetValue(ControlRig, ControlElement, ERigControlValueType::Current);
 
 	TransformUtils::SetValuesFromContext(EulerTransform, Context, TLocation, TRotation, TScale);
 
@@ -458,7 +594,7 @@ void UAnimDetailsProxyTransform::SetControlRigElementValueFromCurrent(UControlRi
 	const FTransform RealTransform(TRotation, TLocation, TScale);
 	if (TransformUtils::SetConstrainedTransform(RealTransform, ControlRig, ControlElement, Context))
 	{
-		UpdateProxyValues();
+		AdoptValues(ERigControlValueType::Current);
 		return;
 	}
 	switch (ControlElement->Settings.ControlType)
@@ -510,11 +646,13 @@ void UAnimDetailsProxyTransform::SetControlRigElementValueFromCurrent(UControlRi
 	ControlRig->Evaluate_AnyThread();
 }
 
-void UAnimDetailsProxyTransform::SetBindingValueFromCurrent(UObject* InObject, const TSharedPtr<FTrackInstancePropertyBindings>& Binding, const FRigControlModifiedContext& Context, bool bInteractive)
+void UAnimDetailsProxyTransform::SetSequencerBindingValueFromCurrent(const FRigControlModifiedContext& Context)
 {
 	using namespace UE::ControlRigEditor;
-
-	if (!InObject || !Binding.IsValid())
+	
+	const TSharedPtr<FTrackInstancePropertyBindings> Binding = SequencerItem.GetBinding();
+	UObject* BoundObject = SequencerItem.GetBoundObject();
+	if (!Binding.IsValid() || !BoundObject)
 	{
 		return;
 	}
@@ -522,25 +660,25 @@ void UAnimDetailsProxyTransform::SetBindingValueFromCurrent(UObject* InObject, c
 	FVector TLocation = Location.ToVector();
 	FRotator TRotation = Rotation.ToRotator();
 	FVector TScale = Scale.ToVector();
-	FEulerTransform EulerTransform = TransformUtils::GetCurrentValue(InObject, Binding);
+	FEulerTransform EulerTransform = TransformUtils::GetValue(BoundObject, Binding);
 
 	TransformUtils::SetValuesFromContext(EulerTransform, Context, TLocation, TRotation, TScale);
 
 	const FTransform RealTransform(TRotation, TLocation, TScale);
 
-	const FStructProperty* TransformProperty = Binding.IsValid() ? CastField<FStructProperty>(Binding->GetProperty(*InObject)) : nullptr;
+	const FStructProperty* TransformProperty = Binding.IsValid() ? CastField<FStructProperty>(Binding->GetProperty(*BoundObject)) : nullptr;
 	if (TransformProperty)
 	{
 		if (TransformProperty->Struct == TBaseStructure<FEulerTransform>::Get())
 		{
 			EulerTransform = FEulerTransform(TLocation, TRotation, TScale);
-			Binding->SetCurrentValue<FTransform>(*InObject, RealTransform);
+			Binding->SetCurrentValue<FTransform>(*BoundObject, RealTransform);
 		}
 	}
 
 	AActor* ActorThatChanged = nullptr;
 	USceneComponent* SceneComponentThatChanged = nullptr;
-	TransformUtils::GetActorAndSceneComponentFromObject(InObject, ActorThatChanged, SceneComponentThatChanged);
+	TransformUtils::GetActorAndSceneComponentFromObject(BoundObject, ActorThatChanged, SceneComponentThatChanged);
 	if (SceneComponentThatChanged)
 	{
 		FProperty* ValueProperty = nullptr;
@@ -599,8 +737,8 @@ void UAnimDetailsProxyTransform::SetBindingValueFromCurrent(UObject* InObject, c
 		// Have to downcast here because of function overloading and inheritance not playing nicely
 		if (ValueProperty)
 		{
-			const TArray<UObject*> ModifiedObjects = { InObject };
-			FPropertyChangedEvent PropertyChangedEvent(ValueProperty, bInteractive ? EPropertyChangeType::Interactive : EPropertyChangeType::ValueSet, MakeArrayView(ModifiedObjects));
+			const TArray<UObject*> ModifiedObjects = { BoundObject };
+			FPropertyChangedEvent PropertyChangedEvent(ValueProperty, EPropertyChangeType::ValueSet, MakeArrayView(ModifiedObjects));
 			
 			FEditPropertyChain PropertyChain;
 			if (AxisProperty)
@@ -625,8 +763,8 @@ void UAnimDetailsProxyTransform::SetBindingValueFromCurrent(UObject* InObject, c
 
 		if (ValueProperty)
 		{
-			const TArray<UObject*> ModifiedObjects = { InObject };
-			FPropertyChangedEvent PropertyChangedEvent(ValueProperty, bInteractive ? EPropertyChangeType::Interactive : EPropertyChangeType::ValueSet, MakeArrayView(ModifiedObjects));
+			const TArray<UObject*> ModifiedObjects = { BoundObject };
+			FPropertyChangedEvent PropertyChangedEvent(ValueProperty, EPropertyChangeType::ValueSet, MakeArrayView(ModifiedObjects));
 
 			FEditPropertyChain PropertyChain;
 			if (AxisProperty)
diff --git a/Engine/Plugins/Animation/ControlRig/Source/ControlRigEditor/Private/AnimDetails/Proxies/AnimDetailsProxyTransform.h b/Engine/Plugins/Animation/ControlRig/Source/ControlRigEditor/Private/AnimDetails/Proxies/AnimDetailsProxyTransform.h
--- a/Engine/Plugins/Animation/ControlRig/Source/ControlRigEditor/Private/AnimDetails/Proxies/AnimDetailsProxyTransform.h
+++ b/Engine/Plugins/Animation/ControlRig/Source/ControlRigEditor/Private/AnimDetails/Proxies/AnimDetailsProxyTransform.h
@@ -21,12 +21,15 @@ public:
 	virtual FName GetCategoryName() const override;
 	virtual TArray<FName> GetPropertyNames() const override;
 	virtual void GetLocalizedPropertyName(const FName& InPropertyName, FText& OutPropertyDisplayName, TOptional<FText>& OutOptionalStructDisplayName) const override;
+	virtual TSet<ERigControlType> GetSupportedControlTypes() const override;
 	virtual bool PropertyIsOnProxy(const FProperty* Property, const FProperty* MemberProperty) override;
-	virtual void UpdateProxyValues() override;
+	virtual void AdoptValues(const ERigControlValueType RigControlValueType = ERigControlValueType::Current) override;
+	virtual void ResetPropertyToDefault(const FName& PropertyName) override;
+	virtual bool HasDefaultValue(const FName& PropertyName) const override;
 	virtual EControlRigContextChannelToKey GetChannelToKeyFromPropertyName(const FName& InPropertyName) const override;
 	virtual EControlRigContextChannelToKey GetChannelToKeyFromChannelName(const FString& InChannelName) const override;
-	virtual void SetControlRigElementValueFromCurrent(UControlRig* ControlRig, FRigControlElement* ControlElement, const FRigControlModifiedContext& Context) override;
-	virtual void SetBindingValueFromCurrent(UObject* InObject, const TSharedPtr<FTrackInstancePropertyBindings>& Binding, const FRigControlModifiedContext& Context, bool bInteractive = false) override;
+	virtual void SetControlRigElementValueFromCurrent(const FRigControlModifiedContext& Context) override;
+	virtual void SetSequencerBindingValueFromCurrent(const FRigControlModifiedContext& Context) override;
 	//~ End UAnimDetailsProxyBase interface
 
 	UPROPERTY(EditAnywhere, Interp, Category = Transform)
diff --git a/Engine/Plugins/Animation/ControlRig/Source/ControlRigEditor/Private/AnimDetails/Proxies/AnimDetailsProxyVector2D.cpp b/Engine/Plugins/Animation/ControlRig/Source/ControlRigEditor/Private/AnimDetails/Proxies/AnimDetailsProxyVector2D.cpp
--- a/Engine/Plugins/Animation/ControlRig/Source/ControlRigEditor/Private/AnimDetails/Proxies/AnimDetailsProxyVector2D.cpp
+++ b/Engine/Plugins/Animation/ControlRig/Source/ControlRigEditor/Private/AnimDetails/Proxies/AnimDetailsProxyVector2D.cpp
@@ -68,35 +68,111 @@ void UAnimDetailsProxyVector2D::GetLocalizedPropertyName(const FName& InProperty
 	}
 }
 
+TSet<ERigControlType> UAnimDetailsProxyVector2D::GetSupportedControlTypes() const
+{
+	static const TSet<ERigControlType> SupportedControlTypes = { ERigControlType::Vector2D };
+	return SupportedControlTypes;
+}
+
 bool UAnimDetailsProxyVector2D::PropertyIsOnProxy(const FProperty* Property, const FProperty* MemberProperty)
 {
-	return ((Property && Property->GetFName() == GET_MEMBER_NAME_CHECKED(UAnimDetailsProxyVector2D, Vector2D)) ||
-		(MemberProperty && MemberProperty->GetFName() == GET_MEMBER_NAME_CHECKED(UAnimDetailsProxyVector2D, Vector2D)));
+	return (Property && Property->GetFName() == GET_MEMBER_NAME_CHECKED(UAnimDetailsProxyVector2D, Vector2D)) ||
+		(MemberProperty && MemberProperty->GetFName() == GET_MEMBER_NAME_CHECKED(UAnimDetailsProxyVector2D, Vector2D));
 }
 
-void UAnimDetailsProxyVector2D::UpdateProxyValues()
+void UAnimDetailsProxyVector2D::AdoptValues(const ERigControlValueType RigControlValueType)
 {
+	FVector2D Value = FVector2D::ZeroVector;
+
 	UControlRig* ControlRig = GetControlRig();
 	FRigControlElement* ControlElement = GetControlElement();
-	if (!ControlRig || !ControlElement)
+	if (ControlRig &&
+		ControlElement &&
+		ControlElement->Settings.ControlType == ERigControlType::Vector2D)
 	{
-		return;
+		const FRigControlValue ControlValue = ControlRig->GetControlValue(ControlElement, RigControlValueType);
+		const FVector3f Vec3Value = ControlValue.Get<FVector3f>();
+		Value = FVector2D(Vec3Value.X, Vec3Value.Y);
 	}
 
-	FVector2D Value = FVector2D::ZeroVector;
-	if (ControlElement->Settings.ControlType == ERigControlType::Vector2D)
+	const FAnimDetailsVector2D AnimDetailProxyVector2D(Value);
+
+	constexpr const TCHAR* PropertyName = TEXT("Vector2D");
+	FTrackInstancePropertyBindings Binding(PropertyName, PropertyName);
+	Binding.CallFunction<FAnimDetailsVector2D>(*this, AnimDetailProxyVector2D);
+}
+
+void UAnimDetailsProxyVector2D::ResetPropertyToDefault(const FName& PropertyName)
+{
+	if (PropertyName == FAnimDetailsVector2D::StaticStruct()->GetFName())
 	{
-		FRigControlValue ControlValue = ControlRig->GetControlValue(ControlElement, ERigControlValueType::Current);
-		//FVector2D version deleted for some reason so need to convert
-		FVector3f Val = ControlValue.Get<FVector3f>();
-		Value = FVector2D(Val.X, Val.Y);
+		// Reset the whole rotation struct
+		AdoptValues(ERigControlValueType::Initial);
 	}
+	else
+	{
+		FProperty* Property = FAnimDetailsVector2D::StaticStruct()->FindPropertyByName(PropertyName);
+		if (!ensureMsgf(Property, TEXT("Cannot find property in struct, cannot reset anim details property to default")))
+		{
+			return;
+		}
 
-	const FAnimDetailsVector2D AnimDetailProxyVector2D(Value);
+		// Reset the specific property name
+		UControlRig* ControlRig = GetControlRig();
+		FRigControlElement* ControlElement = GetControlElement();
+		if (ControlRig &&
+			ControlElement &&
+			ControlElement->Settings.ControlType == ERigControlType::Vector2D)
+		{
+			const FRigControlValue InitialControlValue = ControlRig->GetControlValue(ControlElement, ERigControlValueType::Initial);
+			const FVector3f InitialVec3Value = InitialControlValue.Get<FVector3f>();
+			const FVector2D InitialValue = FVector2D(InitialVec3Value.X, InitialVec3Value.Y);
+			const FAnimDetailsVector2D InitialVector2D(InitialValue);
 
-	const FName PropName("Vector2D");
-	FTrackInstancePropertyBindings Binding(PropName, PropName.ToString());
-	Binding.CallFunction<FAnimDetailsVector2D>(*this, AnimDetailProxyVector2D);
+			*Property->ContainerPtrToValuePtr<double>(&Vector2D) = *Property->ContainerPtrToValuePtr<double>(&InitialVector2D);
+
+			constexpr const TCHAR* StructPropertName = TEXT("Vector2D");
+			FTrackInstancePropertyBindings Binding(StructPropertName, StructPropertName);
+			Binding.CallFunction<FAnimDetailsVector2D>(*this, InitialVector2D);
+		}
+	}
+}
+
+bool UAnimDetailsProxyVector2D::HasDefaultValue(const FName& PropertyName) const
+{
+	const FAnimDetailsVector2D DefaultVector2D = [this]() -> FAnimDetailsVector2D
+		{
+
+			UControlRig* ControlRig = GetControlRig();
+			FRigControlElement* ControlElement = GetControlElement();
+			if (ControlRig && ControlElement)
+			{
+				const FRigControlValue InitialControlValue = ControlRig->GetControlValue(ControlElement, ERigControlValueType::Initial);
+				const FVector3f InitialVec3Value = InitialControlValue.Get<FVector3f>();
+				return FVector2D(InitialVec3Value.X, InitialVec3Value.Y);
+			}
+
+			return FAnimDetailsVector2D();
+		}();
+
+	if (PropertyName == FAnimDetailsVector2D::StaticStruct()->GetFName())
+	{
+		return
+			Vector2D.X == DefaultVector2D.X &&
+			Vector2D.Y == DefaultVector2D.Y;
+	}
+	else if (PropertyName == GET_MEMBER_NAME_CHECKED(FAnimDetailsVector2D, X))
+	{
+		return Vector2D.X == DefaultVector2D.X;
+	}
+	else if (PropertyName == GET_MEMBER_NAME_CHECKED(FAnimDetailsVector2D, Y))
+	{
+		return Vector2D.Y == DefaultVector2D.Y;
+	}
+
+	ensureMsgf(0, TEXT("Failed finding property for anim details vector 2D. Cannot determine if property has its default value"));
+
+	return true;
 }
 
 EControlRigContextChannelToKey UAnimDetailsProxyVector2D::GetChannelToKeyFromPropertyName(const FName& PropertyName) const
@@ -127,13 +203,13 @@ EControlRigContextChannelToKey UAnimDetailsProxyVector2D::GetChannelToKeyFromCha
 	return EControlRigContextChannelToKey::AllTransform;
 }
 
-void UAnimDetailsProxyVector2D::SetControlRigElementValueFromCurrent(UControlRig* ControlRig, FRigControlElement* ControlElement, const FRigControlModifiedContext& Context)
+void UAnimDetailsProxyVector2D::SetControlRigElementValueFromCurrent(const FRigControlModifiedContext& Context)
 {
 	using namespace UE::ControlRigEditor;
 
-	if (ControlRig && 
-		ControlElement && 
-		ControlElement->Settings.ControlType == ERigControlType::Vector2D)
+	UControlRig* ControlRig = GetControlRig();
+	FRigControlElement* ControlElement = GetControlElement();
+	if (ControlRig && ControlElement)
 	{
 		FVector2D Value = Vector2D.ToVector2D();
 		Vector2DUtils::SetVector2DValuesFromContext(ControlRig, ControlElement, Context, Value);
diff --git a/Engine/Plugins/Animation/ControlRig/Source/ControlRigEditor/Private/AnimDetails/Proxies/AnimDetailsProxyVector2D.h b/Engine/Plugins/Animation/ControlRig/Source/ControlRigEditor/Private/AnimDetails/Proxies/AnimDetailsProxyVector2D.h
--- a/Engine/Plugins/Animation/ControlRig/Source/ControlRigEditor/Private/AnimDetails/Proxies/AnimDetailsProxyVector2D.h
+++ b/Engine/Plugins/Animation/ControlRig/Source/ControlRigEditor/Private/AnimDetails/Proxies/AnimDetailsProxyVector2D.h
@@ -36,11 +36,14 @@ public:
 	virtual FName GetCategoryName() const override;
 	virtual TArray<FName> GetPropertyNames() const override;
 	virtual void GetLocalizedPropertyName(const FName& InPropertyName, FText& OutPropertyDisplayName, TOptional<FText>& OutOptionalStructDisplayName) const override;
+	virtual TSet<ERigControlType> GetSupportedControlTypes() const override;
 	virtual bool PropertyIsOnProxy(const FProperty* Property, const FProperty* MemberProperty) override;
-	virtual void UpdateProxyValues() override;
+	virtual void AdoptValues(const ERigControlValueType RigControlValueType = ERigControlValueType::Current) override;
+	virtual void ResetPropertyToDefault(const FName& PropertyName) override;
+	virtual bool HasDefaultValue(const FName& PropertyName) const override;
 	virtual EControlRigContextChannelToKey GetChannelToKeyFromPropertyName(const FName& InPropertyName) const override;
 	virtual EControlRigContextChannelToKey GetChannelToKeyFromChannelName(const FString& InChannelName) const override;
-	virtual void SetControlRigElementValueFromCurrent(UControlRig* ControlRig, FRigControlElement* ControlElement, const FRigControlModifiedContext& Context) override;
+	virtual void SetControlRigElementValueFromCurrent(const FRigControlModifiedContext& Context) override;
 	//~ End UAnimDetailsProxyBase interface
 
 	UPROPERTY(EditAnywhere, Interp, Category = Vector2D)
