Anim in Engine: Minor code clean up for Anim Details

* Fix comment referring to rotation in UAnimDetailsProxyVector2D.cpp
* Improve comment for UAnimDetailsProxyBase::ResetPropertyToDefault
* Add missing test for nullptr in AnimDetailsProxyRotation.cpp

#rnx

diff --git a/Engine/Plugins/Animation/ControlRig/Source/ControlRigEditor/Private/AnimDetails/Proxies/AnimDetailsProxyBase.h b/Engine/Plugins/Animation/ControlRig/Source/ControlRigEditor/Private/AnimDetails/Proxies/AnimDetailsProxyBase.h
--- a/Engine/Plugins/Animation/ControlRig/Source/ControlRigEditor/Private/AnimDetails/Proxies/AnimDetailsProxyBase.h
+++ b/Engine/Plugins/Animation/ControlRig/Source/ControlRigEditor/Private/AnimDetails/Proxies/AnimDetailsProxyBase.h
@@ -105,7 +105,7 @@ public:
 	/** Tries to adopt the current value(s) from the sequencer binding or the control in the control rig. */
 	virtual void AdoptValues(const ERigControlValueType RigControlValueType = ERigControlValueType::Current) PURE_VIRTUAL(UAnimDetailsProxyBase::AdoptValues, return;)
 
-	/** Resets a property in the proxy to the initial value in the rig. If the proxy does not relate to a control rig, this function has no effect. */
+	/** Resets a property in the proxy to the initial value in the rig, or a zero value when the proxy relates to a sequencer binding. */
 	virtual void ResetPropertyToDefault(const FName& PropertyName) PURE_VIRTUAL(UAnimDetailsProxyBase::ResetPropertyToDefault, return;)
 
 	/** Returns true if the property value is the default value */
diff --git a/Engine/Plugins/Animation/ControlRig/Source/ControlRigEditor/Private/AnimDetails/Proxies/AnimDetailsProxyRotation.cpp b/Engine/Plugins/Animation/ControlRig/Source/ControlRigEditor/Private/AnimDetails/Proxies/AnimDetailsProxyRotation.cpp
--- a/Engine/Plugins/Animation/ControlRig/Source/ControlRigEditor/Private/AnimDetails/Proxies/AnimDetailsProxyRotation.cpp
+++ b/Engine/Plugins/Animation/ControlRig/Source/ControlRigEditor/Private/AnimDetails/Proxies/AnimDetailsProxyRotation.cpp
@@ -251,7 +251,7 @@ void UAnimDetailsProxyRotation::SetControlRigElementValueFromCurrent(const FRigC
 	UControlRig* ControlRig = GetControlRig();
 	URigHierarchy* Hierarchy = ControlRig ? ControlRig->GetHierarchy() : nullptr;
 	FRigControlElement* ControlElement = GetControlElement();
-	if (ControlRig && ControlElement)
+	if (ControlRig && Hierarchy && ControlElement)
 	{
 		FVector3f Value = Rotation.ToVector3f();
 		RotationUtils::SetRotationValuesFromContext(ControlRig, ControlElement, Context, Value);
diff --git a/Engine/Plugins/Animation/ControlRig/Source/ControlRigEditor/Private/AnimDetails/Proxies/AnimDetailsProxyVector2D.cpp b/Engine/Plugins/Animation/ControlRig/Source/ControlRigEditor/Private/AnimDetails/Proxies/AnimDetailsProxyVector2D.cpp
--- a/Engine/Plugins/Animation/ControlRig/Source/ControlRigEditor/Private/AnimDetails/Proxies/AnimDetailsProxyVector2D.cpp
+++ b/Engine/Plugins/Animation/ControlRig/Source/ControlRigEditor/Private/AnimDetails/Proxies/AnimDetailsProxyVector2D.cpp
@@ -106,7 +106,7 @@ void UAnimDetailsProxyVector2D::ResetPropertyToDefault(const FName& PropertyName
 {
 	if (PropertyName == FAnimDetailsVector2D::StaticStruct()->GetFName())
 	{
-		// Reset the whole rotation struct
+		// Reset the whole vector 2D struct
 		AdoptValues(ERigControlValueType::Initial);
 	}
 	else
