ControlRig & Sequencer: ensure that the space channel stores the control's active parent instead of the first one

when adding a space channel with a parent which is not the first one (but defined at construction time via a switch parent node for example), it's important that the channel stores this information as the default value.
sequencer wrongly assumed that the active parent was the first one, so a single evaluation was enough to switch to the default parent, setting the hierarchy in a different state than it actually was when creating the space channel.

#rnx

diff --git a/Engine/Plugins/Animation/ControlRig/Source/ControlRig/Private/Sequencer/MovieSceneControlRigParameterSection.cpp b/Engine/Plugins/Animation/ControlRig/Source/ControlRig/Private/Sequencer/MovieSceneControlRigParameterSection.cpp
--- a/Engine/Plugins/Animation/ControlRig/Source/ControlRig/Private/Sequencer/MovieSceneControlRigParameterSection.cpp
+++ b/Engine/Plugins/Animation/ControlRig/Source/ControlRig/Private/Sequencer/MovieSceneControlRigParameterSection.cpp
@@ -1725,15 +1725,26 @@ void UMovieSceneControlRigParameterSection::AddSpaceChannel(FName InControlName,
 	if (UMovieSceneControlRigParameterTrack* Track = GetTypedOuter<UMovieSceneControlRigParameterTrack>())
 	{
 		const TArray<UMovieSceneSection*>& Sections = Track->GetAllSections();
-		if (Sections[0] == this)
+		if (!Sections.IsEmpty() && Sections[0] == this)
 		{
-			
 			if (CanCreateSpaceChannel(InControlName) && !HasSpaceChannel(InControlName))
 			{
-				SpaceChannels.Add(FSpaceControlNameAndChannel(InControlName));
+				FSpaceControlNameAndChannel& NameAndChannel = SpaceChannels.Emplace_GetRef(InControlName);
+				
+				if (const FRigControlElement* Control = ControlRig ? ControlRig->FindControl(InControlName) : nullptr)
+				{
+					const URigHierarchy* Hierarchy = ControlRig->GetHierarchy();
+					const FRigElementKey CurrentParentKey = Hierarchy->GetActiveParent(Control->GetKey());
+
+					// set the default channel value to the current parent if it differs from the one initially constructed
+					if (CurrentParentKey != Hierarchy->GetFirstParent(Control->GetKey()))
+					{
+						NameAndChannel.SpaceCurve.SetDefault({EMovieSceneControlRigSpaceType::ControlRig, CurrentParentKey});
+					}
+				}
+				
 				if (OnSpaceChannelAdded.IsBound())
 				{
-					FSpaceControlNameAndChannel& NameAndChannel = SpaceChannels[SpaceChannels.Num() - 1];
 					OnSpaceChannelAdded.Broadcast(this, InControlName, &NameAndChannel.SpaceCurve);
 				}
 			}
diff --git a/Engine/Plugins/Animation/ControlRig/Source/ControlRig/Private/Sequencer/MovieSceneControlRigSpaceChannel.cpp b/Engine/Plugins/Animation/ControlRig/Source/ControlRig/Private/Sequencer/MovieSceneControlRigSpaceChannel.cpp
--- a/Engine/Plugins/Animation/ControlRig/Source/ControlRig/Private/Sequencer/MovieSceneControlRigSpaceChannel.cpp
+++ b/Engine/Plugins/Animation/ControlRig/Source/ControlRig/Private/Sequencer/MovieSceneControlRigSpaceChannel.cpp
@@ -11,14 +11,20 @@
 
 bool FMovieSceneControlRigSpaceChannel::Evaluate(FFrameTime InTime, FMovieSceneControlRigSpaceBaseKey& OutValue) const
 {
-	FMovieSceneControlRigSpaceBaseKey DefaultValue; //this will be in parent space
-	if (KeyTimes.Num())
+	if (!KeyTimes.IsEmpty())
 	{
 		const int32 Index = FMath::Max(0, Algo::UpperBound(KeyTimes, InTime.FrameNumber) - 1);
 		OutValue = KeyValues[Index];
 		return true;
 	}
-	OutValue = DefaultValue;
+	
+	if (bHasDefaultValue)
+	{
+		OutValue = DefaultValue;
+		return true;
+	}
+
+	OutValue = FMovieSceneControlRigSpaceBaseKey();
 	return true;
 }
 
@@ -203,6 +209,23 @@ void FMovieSceneControlRigSpaceChannel::BroadcastSpaceNoLongerUsed(const TArray<
 	}
 }
 
+void FMovieSceneControlRigSpaceChannel::SetDefault(const FMovieSceneControlRigSpaceBaseKey& InDefaultValue)
+{
+	bHasDefaultValue = true;
+	DefaultValue = InDefaultValue;
+}
+
+TOptional<FMovieSceneControlRigSpaceBaseKey> FMovieSceneControlRigSpaceChannel::GetDefault() const
+{
+	return bHasDefaultValue ? TOptional<FMovieSceneControlRigSpaceBaseKey>(DefaultValue) : TOptional<FMovieSceneControlRigSpaceBaseKey>();
+}
+
+void FMovieSceneControlRigSpaceChannel::ClearDefault()
+{
+	bHasDefaultValue = false;
+	DefaultValue = FMovieSceneControlRigSpaceBaseKey();
+}
+
 FName FMovieSceneControlRigSpaceBaseKey::GetName() const
 {
 	switch (SpaceType)
@@ -218,20 +241,3 @@ FName FMovieSceneControlRigSpaceBaseKey::GetName() const
 	};
 	return NAME_None;
 }
-
-/* Is this needed if not remove mz todoo
-TSharedPtr<FStructOnScope> GetKeyStruct(TMovieSceneChannelHandle<FMovieSceneControlRigSpaceChannel> Channel, FKeyHandle InHandle)
-{
-	int32 KeyValueIndex = Channel.Get()->GetData().GetIndex(InHandle);
-	if (KeyValueIndex != INDEX_NONE)
-	{
-		FNiagaraTypeDefinition KeyType = Channel.Get()->GetData().GetValues()[KeyValueIndex].Value.GetType();
-		uint8* KeyData = Channel.Get()->GetData().GetValues()[KeyValueIndex].Value.GetData();
-		return MakeShared<FStructOnScope>(KeyType.GetStruct(), KeyData);
-		
-	}
-	return TSharedPtr<FStructOnScope>();
-}
-*/
-
-
diff --git a/Engine/Plugins/Animation/ControlRig/Source/ControlRig/Public/Sequencer/MovieSceneControlRigSpaceChannel.h b/Engine/Plugins/Animation/ControlRig/Source/ControlRig/Public/Sequencer/MovieSceneControlRigSpaceChannel.h
--- a/Engine/Plugins/Animation/ControlRig/Source/ControlRig/Public/Sequencer/MovieSceneControlRigSpaceChannel.h
+++ b/Engine/Plugins/Animation/ControlRig/Source/ControlRig/Public/Sequencer/MovieSceneControlRigSpaceChannel.h
@@ -123,6 +123,24 @@ public:
 
 	UE_API TArray <FSpaceRange> FindSpaceIntervals();
 
+	/**
+	 * Set this channel's default value that should be used when no keys are present
+	 * @param InDefaultValue The desired default value
+	 */
+	void SetDefault(const FMovieSceneControlRigSpaceBaseKey& InDefaultValue);
+
+	/**
+	 * Get this channel's default value that will be used when no keys are present
+	 * @return (Optional) The channel's default value
+	 */
+	TOptional<FMovieSceneControlRigSpaceBaseKey> GetDefault() const;
+
+	/**
+	 * Clear the default value on this channel
+	 */
+	virtual void ClearDefault() override;
+
+	
 private:
 
 	UE_API void BroadcastSpaceNoLongerUsed(const TArray<FRigElementKey>& BeforeKeys, const TArray<FRigElementKey>& AfterKeys);
@@ -141,6 +159,12 @@ private:
 	
 	FMovieSceneControlRigSpaceChannelSpaceNoLongerUsedEvent SpaceNoLongerUsedEvent;
 
+	UPROPERTY()
+	FMovieSceneControlRigSpaceBaseKey DefaultValue;
+
+	UPROPERTY()
+	bool bHasDefaultValue;
+	
 	friend struct FControlRigSpaceChannelHelpers;
 };
 
@@ -149,7 +173,6 @@ template<>
 struct TMovieSceneChannelTraits<FMovieSceneControlRigSpaceChannel> : TMovieSceneChannelTraitsBase<FMovieSceneControlRigSpaceChannel>
 {
 	enum { SupportsDefaults = false };
-
 };
 
 inline bool EvaluateChannel(const FMovieSceneControlRigSpaceChannel* InChannel, FFrameTime InTime, FMovieSceneControlRigSpaceBaseKey& OutValue)
diff --git a/Engine/Plugins/Animation/ControlRig/Source/ControlRigEditor/Private/ControlRigSequencerEditorLibrary.cpp b/Engine/Plugins/Animation/ControlRig/Source/ControlRigEditor/Private/ControlRigSequencerEditorLibrary.cpp
--- a/Engine/Plugins/Animation/ControlRig/Source/ControlRigEditor/Private/ControlRigSequencerEditorLibrary.cpp
+++ b/Engine/Plugins/Animation/ControlRig/Source/ControlRigEditor/Private/ControlRigSequencerEditorLibrary.cpp
@@ -2986,8 +2986,21 @@ bool UControlRigSequencerEditorLibrary::SetControlRigSpace(ULevelSequence* Level
 		if (FRigControlElement* Element = ControlRig->FindControl(ControlName))
 		{
 			TSharedPtr<ISequencer>  Sequencer = WeakSequencer.Pin();
+			
+			constexpr bool bDoNotCreate = false;
+			FSpaceChannelAndSection SpaceChannelAndSection = FControlRigSpaceChannelHelpers::FindSpaceChannelAndSectionForControl(ControlRig, ControlName, Sequencer.Get(), bDoNotCreate);
+			if (!SpaceChannelAndSection.SpaceChannel)
+			{
+				// no need to create a space channel if InSpaceKey is the current parent
+				const FRigElementKey ControlKey(ControlName, ERigElementType::Control);
+				if (ControlRig->GetHierarchy()->GetActiveParent(ControlKey) == InSpaceKey)
+				{
+					return true;
+				}
+			}
+			
 			FScopedTransaction Transaction(LOCTEXT("KeyControlRigSpace", "Key Control Rig Space"));
-			FSpaceChannelAndSection SpaceChannelAndSection = FControlRigSpaceChannelHelpers::FindSpaceChannelAndSectionForControl(ControlRig, ControlName, Sequencer.Get(), true /*bCreateIfNeeded*/);
+			SpaceChannelAndSection = FControlRigSpaceChannelHelpers::FindSpaceChannelAndSectionForControl(ControlRig, ControlName, Sequencer.Get(), true /*bCreateIfNeeded*/);
 			if (SpaceChannelAndSection.SpaceChannel)
 			{
 				if (TimeUnit == EMovieSceneTimeUnit::DisplayRate)
diff --git a/Engine/Plugins/Animation/ControlRig/Source/ControlRigEditor/Private/EditMode/ControlRigEditMode.cpp b/Engine/Plugins/Animation/ControlRig/Source/ControlRigEditor/Private/EditMode/ControlRigEditMode.cpp
--- a/Engine/Plugins/Animation/ControlRig/Source/ControlRigEditor/Private/EditMode/ControlRigEditMode.cpp
+++ b/Engine/Plugins/Animation/ControlRig/Source/ControlRigEditor/Private/EditMode/ControlRigEditMode.cpp
@@ -3434,8 +3434,19 @@ void FControlRigEditMode::OpenSpacePickerWidget()
 					ISequencer* Sequencer = WeakSequencer.Pin().Get();
 					if (Sequencer)
 					{
+						constexpr bool bDoNotCreate = false;
+						FSpaceChannelAndSection SpaceChannelAndSection = FControlRigSpaceChannelHelpers::FindSpaceChannelAndSectionForControl(RuntimeRig, InControlKey.Name, Sequencer, bDoNotCreate);
+						if (!SpaceChannelAndSection.SpaceChannel)
+						{
+							// no need to create a space channel if InSpaceKey is the current parent 
+							if (InHierarchy->GetActiveParent(InControlKey) == InSpaceKey)
+							{
+								return;
+							}
+						}
+						
 						FScopedTransaction Transaction(LOCTEXT("KeyControlRigSpace", "Key Control Rig Space"));
-						FSpaceChannelAndSection SpaceChannelAndSection = FControlRigSpaceChannelHelpers::FindSpaceChannelAndSectionForControl(RuntimeRig, InControlKey.Name, Sequencer, true /*bCreateIfNeeded*/);
+						SpaceChannelAndSection = FControlRigSpaceChannelHelpers::FindSpaceChannelAndSectionForControl(RuntimeRig, InControlKey.Name, Sequencer, true /*bCreateIfNeeded*/);
 						if (SpaceChannelAndSection.SpaceChannel)
 						{
 							const FFrameRate TickResolution = Sequencer->GetFocusedTickResolution();
diff --git a/Engine/Plugins/Animation/ControlRig/Source/ControlRigEditor/Private/EditMode/SControlRigEditModeTools.cpp b/Engine/Plugins/Animation/ControlRig/Source/ControlRigEditor/Private/EditMode/SControlRigEditModeTools.cpp
--- a/Engine/Plugins/Animation/ControlRig/Source/ControlRigEditor/Private/EditMode/SControlRigEditModeTools.cpp
+++ b/Engine/Plugins/Animation/ControlRig/Source/ControlRigEditor/Private/EditMode/SControlRigEditModeTools.cpp
@@ -951,10 +951,21 @@ void SControlRigEditModeTools::HandleActiveSpaceChanged(URigHierarchy* InHierarc
 					ISequencer* Sequencer = WeakSequencer.Pin().Get();
 					if (Sequencer)
 					{
+						constexpr bool bDoNotCreate = false;
+						FSpaceChannelAndSection SpaceChannelAndSection = FControlRigSpaceChannelHelpers::FindSpaceChannelAndSectionForControl(ControlRig.Get(), InControlKey.Name, Sequencer, bDoNotCreate);
+						if (!SpaceChannelAndSection.SpaceChannel)
+						{
+							// no need to create a space channel if InSpaceKey is the current parent 
+							if (InHierarchy->GetActiveParent(InControlKey) == InSpaceKey)
+							{
+								return;
+							}
+						}
+						
 						FScopedTransaction Transaction(LOCTEXT("KeyControlRigSpace", "Key Control Rig Space"));
 						ControlRig->Modify();
 
-						FSpaceChannelAndSection SpaceChannelAndSection = FControlRigSpaceChannelHelpers::FindSpaceChannelAndSectionForControl(ControlRig.Get(), InControlKey.Name, Sequencer, true /*bCreateIfNeeded*/);
+						SpaceChannelAndSection = FControlRigSpaceChannelHelpers::FindSpaceChannelAndSectionForControl(ControlRig.Get(), InControlKey.Name, Sequencer, true /*bCreateIfNeeded*/);
 						if (SpaceChannelAndSection.SpaceChannel)
 						{
 							const FFrameRate TickResolution = Sequencer->GetFocusedTickResolution();
diff --git a/Engine/Plugins/Animation/ControlRig/Source/ControlRigEditor/Private/EditMode/SControlRigSpacePicker.cpp b/Engine/Plugins/Animation/ControlRig/Source/ControlRigEditor/Private/EditMode/SControlRigSpacePicker.cpp
--- a/Engine/Plugins/Animation/ControlRig/Source/ControlRigEditor/Private/EditMode/SControlRigSpacePicker.cpp
+++ b/Engine/Plugins/Animation/ControlRig/Source/ControlRigEditor/Private/EditMode/SControlRigSpacePicker.cpp
@@ -177,10 +177,21 @@ void SControlRigSpacePicker::HandleActiveSpaceChanged(URigHierarchy* InHierarchy
 
 			if (const FRigControlElement* ControlElement = InHierarchy->Find<FRigControlElement>(InControlKey))
 			{
+				constexpr bool bDoNotCreate = false;
+				FSpaceChannelAndSection SpaceChannelAndSection = FControlRigSpaceChannelHelpers::FindSpaceChannelAndSectionForControl(ControlRig, InControlKey.Name, Sequencer, bDoNotCreate);
+				if (!SpaceChannelAndSection.SpaceChannel)
+				{
+					// no need to create a space channel if InSpaceKey is the current parent 
+					if (InHierarchy->GetActiveParent(InControlKey) == InSpaceKey)
+					{
+						return;
+					}
+				}
+				
 				FScopedTransaction Transaction(LOCTEXT("KeyControlRigSpace", "Key Control Rig Space"));
 				ControlRig->Modify();
 
-				FSpaceChannelAndSection SpaceChannelAndSection = FControlRigSpaceChannelHelpers::FindSpaceChannelAndSectionForControl(ControlRig, InControlKey.Name, Sequencer, true /*bCreateIfNeeded*/);
+				SpaceChannelAndSection = FControlRigSpaceChannelHelpers::FindSpaceChannelAndSectionForControl(ControlRig, InControlKey.Name, Sequencer, true /*bCreateIfNeeded*/);
 				if (SpaceChannelAndSection.SpaceChannel)
 				{
 					const FFrameRate TickResolution = Sequencer->GetFocusedTickResolution();
