fixed bDefaultValue not initialized

#rnx

diff --git a/Engine/Plugins/Animation/ControlRig/Source/ControlRig/Public/Sequencer/MovieSceneControlRigSpaceChannel.h b/Engine/Plugins/Animation/ControlRig/Source/ControlRig/Public/Sequencer/MovieSceneControlRigSpaceChannel.h
--- a/Engine/Plugins/Animation/ControlRig/Source/ControlRig/Public/Sequencer/MovieSceneControlRigSpaceChannel.h
+++ b/Engine/Plugins/Animation/ControlRig/Source/ControlRig/Public/Sequencer/MovieSceneControlRigSpaceChannel.h
@@ -163,7 +163,7 @@ private:
 	FMovieSceneControlRigSpaceBaseKey DefaultValue;
 
 	UPROPERTY()
-	bool bHasDefaultValue;
+	bool bHasDefaultValue = false;
 	
 	friend struct FControlRigSpaceChannelHelpers;
 };
