ControlRig: SetTransform nodes to handle bInitial with PostConstruction event as well
+ renamed IsRunningConstructionEvent to IsRunningAConstructionEvent

#rnx

diff --git a/Engine/Plugins/Animation/ControlRig/Source/ControlRig/Private/Units/Execution/RigUnit_DynamicHierarchy.cpp b/Engine/Plugins/Animation/ControlRig/Source/ControlRig/Private/Units/Execution/RigUnit_DynamicHierarchy.cpp
--- a/Engine/Plugins/Animation/ControlRig/Source/ControlRig/Private/Units/Execution/RigUnit_DynamicHierarchy.cpp
+++ b/Engine/Plugins/Animation/ControlRig/Source/ControlRig/Private/Units/Execution/RigUnit_DynamicHierarchy.cpp
@@ -23,7 +23,7 @@ bool FRigUnit_DynamicHierarchyBase::IsValidToRunInContext(
 
 	if(bAllowOnlyConstructionEvent)
 	{
-		if(!InExecuteContext.IsRunningConstructionEvent())
+		if(!InExecuteContext.IsRunningAConstructionEvent())
 		{
 			if(OutErrorMessage)
 			{
@@ -225,8 +225,8 @@ FRigUnit_SwitchParent_Execute()
 				return;
 			}
 
-			// during construction event also change the initial weights
-			if(ExecuteContext.IsRunningConstructionEvent())
+			// during construction and post-construction events also change the initial weights
+			if(ExecuteContext.IsRunningAConstructionEvent())
 			{
 				if(!ExecuteContext.Hierarchy->SwitchToParent(ChildElement, ParentElement, true, true, NoDependencyProvider, &FailureReason))
 				{
@@ -299,8 +299,8 @@ FRigUnit_HierarchySetParentWeights_Execute()
 
 	ExecuteContext.Hierarchy->SetParentWeightArray(ChildElement, Weights, false, true);
 
-	// during construction event also change the initial weights
-	if(ExecuteContext.IsRunningConstructionEvent())
+	// during construction and post-construction events also change the initial weights
+	if(ExecuteContext.IsRunningAConstructionEvent())
 	{
 		ExecuteContext.Hierarchy->SetParentWeightArray(ChildElement, Weights, true, true);
 	}
diff --git a/Engine/Plugins/Animation/ControlRig/Source/ControlRig/Private/Units/Hierarchy/RigUnit_SetTransform.cpp b/Engine/Plugins/Animation/ControlRig/Source/ControlRig/Private/Units/Hierarchy/RigUnit_SetTransform.cpp
--- a/Engine/Plugins/Animation/ControlRig/Source/ControlRig/Private/Units/Hierarchy/RigUnit_SetTransform.cpp
+++ b/Engine/Plugins/Animation/ControlRig/Source/ControlRig/Private/Units/Hierarchy/RigUnit_SetTransform.cpp
@@ -93,8 +93,8 @@ FRigUnit_SetTransform_Execute()
 #if WITH_EDITOR
 			if(bInitial)
 			{
-				// provide some user feedback when changing initial transforms during forward solve
-				if(ExecuteContext.GetEventName() == FRigUnit_BeginExecution::EventName)
+				// provide some user feedback when changing initial transforms out of a construction event
+				if(!ExecuteContext.IsRunningAConstructionEvent())
 				{
 					UE_CONTROLRIG_RIGUNIT_LOG_MESSAGE(TEXT("Changing initial transforms during %s is not recommended."), *ExecuteContext.GetEventName().ToString());
 				}
@@ -107,7 +107,7 @@ FRigUnit_SetTransform_Execute()
 				FTransform TransformMutable = Value;
 				FRigUnit_SetControlOffset::StaticExecute(ExecuteContext, CachedIndex.GetKey().Name, TransformMutable, Space, CachedIndex);
 				
-				if (ExecuteContext.GetEventName() == FRigUnit_PrepareForExecution::EventName)
+				if(ExecuteContext.IsRunningAConstructionEvent())
 				{
 					Hierarchy->SetLocalTransformByIndex(CachedIndex, FTransform::Identity, true, bPropagateToChildren);
 					Hierarchy->SetLocalTransformByIndex(CachedIndex, FTransform::Identity, false, bPropagateToChildren);
@@ -145,7 +145,7 @@ FRigUnit_SetTransform_Execute()
 				{
 					Hierarchy->SetGlobalTransformByIndex(CachedIndex, WeightedTransform, bInitial, bPropagateToChildren);
 
-					if (bInitial && ExecuteContext.GetEventName() == FRigUnit_PrepareForExecution::EventName)
+					if (bInitial && ExecuteContext.IsRunningAConstructionEvent())
 					{
 						Hierarchy->SetGlobalTransformByIndex(CachedIndex, WeightedTransform, false, bPropagateToChildren);
 					}
@@ -155,7 +155,7 @@ FRigUnit_SetTransform_Execute()
 				{
 					Hierarchy->SetLocalTransformByIndex(CachedIndex, WeightedTransform, bInitial, bPropagateToChildren);
 
-					if (bInitial && ExecuteContext.GetEventName() == FRigUnit_PrepareForExecution::EventName)
+					if (bInitial && ExecuteContext.IsRunningAConstructionEvent())
 					{
 						Hierarchy->SetLocalTransformByIndex(CachedIndex, WeightedTransform, false, bPropagateToChildren);
 					}
diff --git a/Engine/Plugins/Animation/ControlRig/Source/ControlRig/Private/Units/RigUnitContext.cpp b/Engine/Plugins/Animation/ControlRig/Source/ControlRig/Private/Units/RigUnitContext.cpp
--- a/Engine/Plugins/Animation/ControlRig/Source/ControlRig/Private/Units/RigUnitContext.cpp
+++ b/Engine/Plugins/Animation/ControlRig/Source/ControlRig/Private/Units/RigUnitContext.cpp
@@ -7,7 +7,7 @@
 
 #include UE_INLINE_GENERATED_CPP_BY_NAME(RigUnitContext)
 
-bool FControlRigExecuteContext::IsRunningConstructionEvent() const
+bool FControlRigExecuteContext::IsRunningAConstructionEvent() const
 {
 	return GetEventName() == FRigUnit_PrepareForExecution::EventName ||
 		GetEventName() == FRigUnit_PostPrepareForExecution::EventName;
diff --git a/Engine/Plugins/Animation/ControlRig/Source/ControlRig/Public/Units/RigUnitContext.h b/Engine/Plugins/Animation/ControlRig/Source/ControlRig/Public/Units/RigUnitContext.h
--- a/Engine/Plugins/Animation/ControlRig/Source/ControlRig/Public/Units/RigUnitContext.h
+++ b/Engine/Plugins/Animation/ControlRig/Source/ControlRig/Public/Units/RigUnitContext.h
@@ -172,9 +172,9 @@ public:
 	}
 
 	/**
-	 * Returns true if the event currently running is considered a construction event
+	 * Returns true if the event currently running is considered a construction or post-construction event
 	 */
-	bool IsRunningConstructionEvent() const;
+	bool IsRunningAConstructionEvent() const;
 
 	/**
 	 * Returns true if this context is used on a module currently
