Control Rig: Fix crash when undoing add to sequencer while CR editor is opened
#rnx

diff --git a/Engine/Plugins/Animation/ControlRig/Source/ControlRig/Private/ControlRig.cpp b/Engine/Plugins/Animation/ControlRig/Source/ControlRig/Private/ControlRig.cpp
--- a/Engine/Plugins/Animation/ControlRig/Source/ControlRig/Private/ControlRig.cpp
+++ b/Engine/Plugins/Animation/ControlRig/Source/ControlRig/Private/ControlRig.cpp
@@ -1015,11 +1015,14 @@ bool UControlRig::Execute(const FName& InEventName)
 	// allow access to the default hierarchy to allow to reset
 	if(!HasAnyFlags(RF_ClassDefaultObject))
 	{
-		if (UControlRig* CDO = Cast<UControlRig>(GetClass()->GetDefaultObject()))
+		if (PublicContext.Hierarchy)
 		{
-			if(URigHierarchy* DefaultHierarchy = CDO->GetHierarchy())
+			if (UControlRig* CDO = Cast<UControlRig>(GetClass()->GetDefaultObject()))
 			{
-				PublicContext.Hierarchy->DefaultHierarchyPtr = DefaultHierarchy;
+				if(URigHierarchy* DefaultHierarchy = CDO->GetHierarchy())
+				{
+					PublicContext.Hierarchy->DefaultHierarchyPtr = DefaultHierarchy;
+				}
 			}
 		}
 	}
diff --git a/Engine/Plugins/Animation/ControlRig/Source/ControlRigDeveloper/Private/ControlRigAsset.cpp b/Engine/Plugins/Animation/ControlRig/Source/ControlRigDeveloper/Private/ControlRigAsset.cpp
--- a/Engine/Plugins/Animation/ControlRig/Source/ControlRigDeveloper/Private/ControlRigAsset.cpp
+++ b/Engine/Plugins/Animation/ControlRig/Source/ControlRigDeveloper/Private/ControlRigAsset.cpp
@@ -1902,10 +1902,13 @@ void IControlRigAssetInterface::ClearTransientControls()
 			UControlRig* InstancedControlRig = Cast<UControlRig>(ArchetypeInstance);
 			if (InstancedControlRig)
 			{
-				if(!InstancedControlRig->GetHierarchy()->GetTransientControls().IsEmpty())
+				if (URigHierarchy* InstanceHierarchy = InstancedControlRig->GetHierarchy())
 				{
-					bHasAnyTransientControls = true;
-					break;
+					if(!InstanceHierarchy->GetTransientControls().IsEmpty())
+					{
+						bHasAnyTransientControls = true;
+						break;
+					}
 				}
 			}
 		}
diff --git a/Engine/Plugins/Animation/ControlRig/Source/ControlRigEditor/Private/Editor/ControlRigEditor.cpp b/Engine/Plugins/Animation/ControlRig/Source/ControlRigEditor/Private/Editor/ControlRigEditor.cpp
--- a/Engine/Plugins/Animation/ControlRig/Source/ControlRigEditor/Private/Editor/ControlRigEditor.cpp
+++ b/Engine/Plugins/Animation/ControlRig/Source/ControlRigEditor/Private/Editor/ControlRigEditor.cpp
@@ -3415,19 +3415,19 @@ void FControlRigBaseEditor::HandlePreviewMeshChanged(USkeletalMesh* InOldSkeleta
 								const FRigElementKey& Parent = BPHierarchy->GetDefaultParent(Key);
 								ConnectorsAndSocketsToParents.Emplace(Key, Parent, BPHierarchy->GetLocalTransform(Key));
 							}
-							(void)Controller->RemoveElement(Key, true, true);
+							(void)Controller->RemoveElement(Key, bTransact, true);
 						}
 						
 						USkeleton* Skeleton = InNewSkeletalMesh ? InNewSkeletalMesh->GetSkeleton() : nullptr;
-						Controller->ImportBones(Skeleton, NAME_None, true, true, false, true, true);
+						Controller->ImportBones(Skeleton, NAME_None, true, true, false, bTransact, true);
 						if(InNewSkeletalMesh)
 						{
-							Controller->ImportCurvesFromSkeletalMesh(InNewSkeletalMesh, NAME_None, false, true, true);
-							Controller->ImportSocketsFromSkeletalMesh(InNewSkeletalMesh, NAME_None, true, true, false, true, true);
+							Controller->ImportCurvesFromSkeletalMesh(InNewSkeletalMesh, NAME_None, false, bTransact, true);
+							Controller->ImportSocketsFromSkeletalMesh(InNewSkeletalMesh, NAME_None, true, true, false, bTransact, true);
 						}
 						else
 						{
-							Controller->ImportCurves(Skeleton, NAME_None, false, true, true);
+							Controller->ImportCurves(Skeleton, NAME_None, false, bTransact, true);
 						}
 
 						// Recreate manually created elements
@@ -3441,8 +3441,8 @@ void FControlRigBaseEditor::HandlePreviewMeshChanged(USkeletalMesh* InOldSkeleta
 							{
 								switch (Key.Type)
 								{
-									case ERigElementType::Socket: Controller->AddSocket(Key.Name, Parent, Transform, false); break;
-									case ERigElementType::Connector: Controller->AddConnector(Key.Name); break;
+									case ERigElementType::Socket: Controller->AddSocket(Key.Name, Parent, Transform, false, FLinearColor::White, FString(), bTransact); break;
+									case ERigElementType::Connector: Controller->AddConnector(Key.Name, FRigConnectorSettings(), bTransact); break;
 								}
 							}
 						}
@@ -3476,7 +3476,7 @@ void FControlRigBaseEditor::HandlePreviewMeshChanged(USkeletalMesh* InOldSkeleta
 						UModularRigController* ModularRigController = ControlRigBP->GetModularRigController();
 						const bool bAutoResolve = ControlRigBP->GetModularRigSettings().bAutoResolve;
 						Model->ForEachModule(
-							[&Model, Hierarchy, ModularRigController, &PreviousConnections, bAutoResolve]
+							[&Model, Hierarchy, ModularRigController, &PreviousConnections, bAutoResolve, bTransact]
 							(const FRigModuleReference* Module) -> bool
 							{
 								bool bContinueResolval;
@@ -3518,7 +3518,7 @@ void FControlRigBaseEditor::HandlePreviewMeshChanged(USkeletalMesh* InOldSkeleta
 											if(PreviousConnections.HasConnection(ConnectorKey, Hierarchy))
 											{
 												const FRigElementKey Target = PreviousConnections.FindTargetFromConnector(ConnectorKey);
-												if(ModularRigController->ConnectConnectorToElement(ConnectorKey, Target, true))
+												if(ModularRigController->ConnectConnectorToElement(ConnectorKey, Target, bTransact))
 												{
 													bContinueResolval = true;
 												}
@@ -3527,7 +3527,7 @@ void FControlRigBaseEditor::HandlePreviewMeshChanged(USkeletalMesh* InOldSkeleta
 											// try to auto resolve it
 											if(!bContinueResolval && bIsSecondary && bAutoResolve)
 											{
-												if(ModularRigController->AutoConnectSecondaryConnectors({ConnectorKey}, true, true))
+												if(ModularRigController->AutoConnectSecondaryConnectors({ConnectorKey}, true, bTransact))
 												{
 													bContinueResolval = true;
 												}
