ControlRigEditMode: fix rotations not being properly reset when not using the default parent

#rnx

diff --git a/Engine/Plugins/Animation/ControlRig/Source/ControlRigEditor/Private/EditMode/ControlRigEditMode.cpp b/Engine/Plugins/Animation/ControlRig/Source/ControlRigEditor/Private/EditMode/ControlRigEditMode.cpp
--- a/Engine/Plugins/Animation/ControlRig/Source/ControlRigEditor/Private/EditMode/ControlRigEditMode.cpp
+++ b/Engine/Plugins/Animation/ControlRig/Source/ControlRigEditor/Private/EditMode/ControlRigEditMode.cpp
@@ -3872,6 +3872,8 @@ void FControlRigEditMode::ZeroTransforms(UControlRig* ControlRig, const FRigCont
 	}
 
 	FControlRigInteractionScope InteractionScope(ControlRig, ControlsInteracting);
+	URigHierarchy* Hierarchy = ControlRig->GetHierarchy();
+	
 	for (const FRigElementKey& ElementToReset : TransformElementsToReset)
 	{
 		FRigControlElement* ControlElement = nullptr;
@@ -3891,12 +3893,12 @@ void FControlRigEditMode::ZeroTransforms(UControlRig* ControlRig, const FRigCont
 		ControlRig->Modify();
 		if (bHasNonDefaultParent == true) //possibly not at default parent so switch to it
 		{
-			ControlRig->GetHierarchy()->SwitchToDefaultParent(ElementToReset);
+			Hierarchy->SwitchToDefaultParent(ElementToReset);
 		}
 		if (ControlElement)
 		{
-			const FVector InitialAngles = ControlRig->GetHierarchy()->GetControlPreferredEulerAngles(ControlElement, ControlElement->Settings.PreferredRotationOrder, true);
-			ControlRig->GetHierarchy()->SetControlPreferredEulerAngles(ControlElement, InitialAngles, ControlElement->Settings.PreferredRotationOrder);
+			const FVector InitialAngles = Hierarchy->GetControlPreferredEulerAngles(ControlElement, ControlElement->Settings.PreferredRotationOrder, true);
+			Hierarchy->SetControlPreferredEulerAngles(ControlElement, InitialAngles, ControlElement->Settings.PreferredRotationOrder);
 			ControlRig->SetControlLocalTransform(ElementToReset.Name, InitialLocalTransform, true, Context, true, true);
 
 			NotifyDrivenControls(ControlRig, ElementToReset, Context);
@@ -3908,7 +3910,7 @@ void FControlRigEditMode::ZeroTransforms(UControlRig* ControlRig, const FRigCont
 		}
 		else
 		{
-			ControlRig->GetHierarchy()->SetLocalTransform(ElementToReset, InitialLocalTransform, false, true, true);
+			Hierarchy->SetLocalTransform(ElementToReset, InitialLocalTransform, false, true, true);
 		}
 
 		//@helge not sure what to do if the non-default parent
@@ -3931,14 +3933,14 @@ void FControlRigEditMode::ZeroTransforms(UControlRig* ControlRig, const FRigCont
 				FRigControlElement* ControlElement = ControlRig->FindControl(ElementToReset.Name);
 				if (ControlElement && !ControlElement->Settings.bIsTransientControl)
 				{
-					FTransform GlobalTransform = ControlRig->GetHierarchy()->GetGlobalTransform(ElementToReset);
-					GlobalTransforms.Add(ElementToReset, GlobalTransform);
+					FTransform GlobalTransform = Hierarchy->GetGlobalTransform(ElementToReset);
+					GlobalTransforms.Emplace(ElementToReset, GlobalTransform);
 				}
 				NotifyDrivenControls(ControlRig, ElementToReset,Context);
 			}
 			else
 			{
-				FTransform GlobalTransform = ControlRig->GetHierarchy()->GetGlobalTransform(ElementToReset);
+				FTransform GlobalTransform = Hierarchy->GetGlobalTransform(ElementToReset);
 				GlobalTransforms.Add(ElementToReset, GlobalTransform);
 			}
 		}
@@ -3952,16 +3954,20 @@ void FControlRigEditMode::ZeroTransforms(UControlRig* ControlRig, const FRigCont
 					FRigControlElement* ControlElement = ControlRig->FindControl(ElementToReset.Name);
 					if (ControlElement && !ControlElement->Settings.bIsTransientControl)
 					{
-						ControlRig->GetHierarchy()->SwitchToParent(ElementToReset, *SpaceKey);
+						Hierarchy->SwitchToParent(ElementToReset, *SpaceKey);
 					}
 				}
 				else
 				{
-					ControlRig->GetHierarchy()->SwitchToParent(ElementToReset, *SpaceKey);
+					Hierarchy->SwitchToParent(ElementToReset, *SpaceKey);
 				}
 			}
 		}
+		
 		//set global transforms in this space // do it twice since ControlsInOrder is not really always in order
+		const FRigControlModifiedContext DefaultContext;
+		constexpr bool bNotify = true, bSetupUndo = true, bNoPython = false, bFixEuler = true;
+		
 		for (int32 SetHack = 0; SetHack < 2; ++SetHack)
 		{
 			const TGuardValue<bool> GuardEvaluationType(ControlRig->bEvaluationTriggeredFromInteraction, true);
@@ -3975,20 +3981,20 @@ void FControlRigEditMode::ZeroTransforms(UControlRig* ControlRig, const FRigCont
 						FRigControlElement* ControlElement = ControlRig->FindControl(ElementToReset.Name);
 						if (ControlElement && !ControlElement->Settings.bIsTransientControl)
 						{
-							ControlRig->SetControlGlobalTransform(ElementToReset.Name, *GlobalTransform, true);
+							ControlRig->SetControlGlobalTransform(ElementToReset.Name, *GlobalTransform, bNotify, DefaultContext, bSetupUndo, bNoPython, bFixEuler);
 							ControlRig->Evaluate_AnyThread();
 							NotifyDrivenControls(ControlRig, ElementToReset, Context);
 						}
 					}
 					else
 					{
-						ControlRig->GetHierarchy()->SetGlobalTransform(ElementToReset, *GlobalTransform, false, true, true);
+						Hierarchy->SetGlobalTransform(ElementToReset, *GlobalTransform, false, true, bSetupUndo);
 					}
 				}
 			}
 		}
+		
 		//send notifies
-
 		for (const FRigElementKey& ControlToReset : ControlsToReset)
 		{
 			FRigControlElement* ControlElement = ControlRig->FindControl(ControlToReset.Name);
