Constraints & Control Rig: avoid re-entrant rig evaluation

#rnx

diff --git a/Engine/Plugins/Animation/ControlRig/Source/ControlRig/Private/Constraints/ControlRigTransformableHandle.cpp b/Engine/Plugins/Animation/ControlRig/Source/ControlRig/Private/Constraints/ControlRigTransformableHandle.cpp
--- a/Engine/Plugins/Animation/ControlRig/Source/ControlRig/Private/Constraints/ControlRigTransformableHandle.cpp
+++ b/Engine/Plugins/Animation/ControlRig/Source/ControlRig/Private/Constraints/ControlRigTransformableHandle.cpp
@@ -260,7 +260,10 @@ FTransform UTransformableControlHandle::GetGlobalTransform() const
 	const FTransform& ComponentTransform = BoundComponent->GetComponentTransform();
 	if (TransformableHandleUtils::SkipTicking())
 	{
-		TransformableHandleUtils::EvaluateComponent( GetSkeletalMesh(), GetEvaluationTask() );
+		if (!ControlRig->IsEvaluating())
+		{
+			TransformableHandleUtils::EvaluateComponent( GetSkeletalMesh(), GetEvaluationTask() );
+		}
 		return Cache * ComponentTransform;
 	}
 
