Control Rig: Fix false cycles being detected when modifying array connectors
#rnx

diff --git a/Engine/Plugins/Animation/ControlRig/Source/ControlRig/Private/ModularRigController.cpp b/Engine/Plugins/Animation/ControlRig/Source/ControlRig/Private/ModularRigController.cpp
--- a/Engine/Plugins/Animation/ControlRig/Source/ControlRig/Private/ModularRigController.cpp
+++ b/Engine/Plugins/Animation/ControlRig/Source/ControlRig/Private/ModularRigController.cpp
@@ -284,7 +284,7 @@ bool UModularRigController::ConnectConnectorToElements(const FRigElementKey& InC
 			{
 				TargetKeyString = TargetKeyString + TEXT(", ") + TargetKeys[Index].ToString();
 			}
-			TargetKeyString = TEXT("{") + TargetKeyString = TEXT("}");
+			TargetKeyString = TEXT("{") + TargetKeyString + TEXT("}");
 		}
 		return TargetKeyString;
 	};
@@ -409,19 +409,17 @@ bool UModularRigController::ConnectConnectorToElements(const FRigElementKey& InC
 	}
 #endif
 
-	const TArray<FRigElementKey> DisconnectedConnectors = DisconnectCyclicConnectors();
-
-	// If the connection that was attempted generated cycles, try to reestablish the previous connections
-	if (DisconnectedConnectors.Contains(InConnectorKey))
-	{
-		FString TargetKeyString = GenerateTargetKeyString();
-		UE_LOG(LogControlRig, Error, TEXT("Could not connect %s to %s: cycles detected"), *InConnectorKey.ToString(), *TargetKeyString);
-		
-		if (TArray<FRigElementKey>* Previous = PreviousConnections.Find(InConnectorKey))
-		{
-			return ConnectConnectorToElements(InConnectorKey, *Previous, bSetupUndo, bAutomaticReparenting, false);
-		}
-	}
+	// todo: Check if this is even necessary, for now it is disabled 
+	// const TArray<FRigElementKey> DisconnectedConnectors = DisconnectCyclicConnectors();
+	//
+	// // If the connection that was attempted generated cycles, try to reestablish the previous connections
+	// if (DisconnectedConnectors.Contains(InConnectorKey))
+	// {
+	// 	FString TargetKeyString = GenerateTargetKeyString();
+	// 	UE_LOG(LogControlRig, Error, TEXT("Could not connect %s to %s: cycles detected"), *InConnectorKey.ToString(), *TargetKeyString);
+	// 	
+	// 	Model->Connections = ConnectionsSnapshot; // reestablish the previous connections
+	// }
 
 	Notify(EModularRigNotification::ConnectionChanged, Module);
 
