Control Rig: Update default settings and remember them per editor per user

* Display Axes On Selection is now enabled by default
* Display Axes On Selection, Axis Scale, Display Nulls, Display Sockets, Gizmo Scale, Hide Control Shapes, Show Controls as Overlays, Only Select Rig Controls are now saved per editor per user

diff --git a/Engine/Plugins/Animation/ControlRig/Source/ControlRigEditor/Private/EditMode/ControlRigEditMode.cpp b/Engine/Plugins/Animation/ControlRig/Source/ControlRigEditor/Private/EditMode/ControlRigEditMode.cpp
--- a/Engine/Plugins/Animation/ControlRig/Source/ControlRigEditor/Private/EditMode/ControlRigEditMode.cpp
+++ b/Engine/Plugins/Animation/ControlRig/Source/ControlRigEditor/Private/EditMode/ControlRigEditMode.cpp
@@ -3291,6 +3291,8 @@ void FControlRigEditMode::IncreaseShapeSize()
 {
 	UControlRigEditModeSettings* Settings = GetMutableSettings();
 	Settings->GizmoScale += 0.1f;
+	Settings->SaveConfig();
+
 	GetModeManager()->SetWidgetScale(Settings->GizmoScale);
 }
 
@@ -3298,6 +3300,8 @@ void FControlRigEditMode::DecreaseShapeSize()
 {
 	UControlRigEditModeSettings* Settings = GetMutableSettings();
 	Settings->GizmoScale -= 0.1f;
+	Settings->SaveConfig();
+
 	GetModeManager()->SetWidgetScale(Settings->GizmoScale);
 }
 
@@ -3635,12 +3639,13 @@ void FControlRigEditMode::ToggleManipulators()
 	{
 		UControlRigEditModeSettings* Settings = GetMutableSettings();
 		Settings->bHideControlShapes = !Settings->bHideControlShapes;
+		Settings->SaveConfig();
 	}
 }
 
 void FControlRigEditMode::ToggleModuleManipulators()
 {
-	UControlRigEditModeSettings* Settings = GetMutableSettings();
+	const UControlRigEditModeSettings* Settings = GetSettings();
 
 	if (!AreEditingControlRigDirectly() && !Settings->bHideControlShapes)
 	{
@@ -3687,6 +3692,7 @@ void FControlRigEditMode::ToggleAllManipulators()
 {	
 	UControlRigEditModeSettings* Settings = GetMutableSettings();
 	Settings->bHideControlShapes = !Settings->bHideControlShapes;
+	Settings->SaveConfig();
 
 	//turn on all if in level editor in case any where off
 	if (!AreEditingControlRigDirectly() && Settings->bHideControlShapes)
@@ -3713,6 +3719,8 @@ void FControlRigEditMode::ToggleControlsAsOverlay()
 {
 	UControlRigEditModeSettings* Settings = GetMutableSettings();
 	Settings->bShowControlsAsOverlay = !Settings->bShowControlsAsOverlay;
+	Settings->SaveConfig();
+
 	OnSettingsChanged(Settings);
 }
 
@@ -6454,6 +6462,7 @@ void FControlRigEditMode::SetOnlySelectRigControls(bool Val)
 {
 	UControlRigEditModeSettings* Settings = GetMutableSettings();
 	Settings->bOnlySelectRigControls = Val;
+	Settings->SaveConfig();
 }
 
 bool FControlRigEditMode::GetOnlySelectRigControls()const
diff --git a/Engine/Plugins/Animation/ControlRig/Source/ControlRigEditor/Private/EditMode/ControlRigEditModeSettings.h b/Engine/Plugins/Animation/ControlRig/Source/ControlRigEditor/Private/EditMode/ControlRigEditModeSettings.h
--- a/Engine/Plugins/Animation/ControlRig/Source/ControlRigEditor/Private/EditMode/ControlRigEditModeSettings.h
+++ b/Engine/Plugins/Animation/ControlRig/Source/ControlRigEditor/Private/EditMode/ControlRigEditModeSettings.h
@@ -24,7 +24,7 @@ class UControlRigEditModeSettings : public UDeveloperSettings
 		, bHideControlShapes(false)
 		, bShowAllProxyControls(false)
 		, bShowControlsAsOverlay(false)
-		, bDisplayAxesOnSelection(false)
+		, bDisplayAxesOnSelection(true)
 		, AxisScale(10.f)
 		, bCoordSystemPerWidgetMode(true)
 		, bOnlySelectRigControls(false)
diff --git a/Engine/Plugins/Animation/ControlRig/Source/ControlRigEditor/Private/Editor/ControlRigEditor.cpp b/Engine/Plugins/Animation/ControlRig/Source/ControlRigEditor/Private/Editor/ControlRigEditor.cpp
--- a/Engine/Plugins/Animation/ControlRig/Source/ControlRigEditor/Private/Editor/ControlRigEditor.cpp
+++ b/Engine/Plugins/Animation/ControlRig/Source/ControlRigEditor/Private/Editor/ControlRigEditor.cpp
@@ -151,6 +151,7 @@ void OnToolbarAxesScaleChanged(float InValue)
 	if (UControlRigEditModeSettings* Settings = GetMutableDefault<UControlRigEditModeSettings>())
 	{
 		Settings->AxisScale = InValue;
+		Settings->SaveConfig();
 	}
 }
 
@@ -3056,6 +3057,7 @@ void FControlRigBaseEditor::HandleToggleToolbarDrawAxesOnSelection()
 	if (UControlRigEditModeSettings* Settings = GetMutableDefault<UControlRigEditModeSettings>())
 	{
 		Settings->bDisplayAxesOnSelection = !Settings->bDisplayAxesOnSelection;
+		Settings->SaveConfig();
 	}
 }
 
@@ -3085,6 +3087,7 @@ void FControlRigBaseEditor::HandleToggleToolbarDrawNulls()
 	if (UControlRigEditModeSettings* Settings = GetMutableDefault<UControlRigEditModeSettings>())
 	{
 		Settings->bDisplayNulls = !Settings->bDisplayNulls;
+		Settings->SaveConfig();
 	}
 }
 
@@ -3114,6 +3117,7 @@ void FControlRigBaseEditor::HandleToggleToolbarDrawSockets()
 	if (UControlRigEditModeSettings* Settings = GetMutableDefault<UControlRigEditModeSettings>())
 	{
 		Settings->bDisplaySockets = !Settings->bDisplaySockets;
+		Settings->SaveConfig();
 	}
 }
 
