Control Rig: Fix layered rigs not working properly with rigs that import the skeleton during construction
#rnx

diff --git a/Engine/Plugins/Animation/ControlRig/Source/ControlRigEditor/Private/Sequencer/ControlRigParameterTrackEditor.cpp b/Engine/Plugins/Animation/ControlRig/Source/ControlRigEditor/Private/Sequencer/ControlRigParameterTrackEditor.cpp
--- a/Engine/Plugins/Animation/ControlRig/Source/ControlRigEditor/Private/Sequencer/ControlRigParameterTrackEditor.cpp
+++ b/Engine/Plugins/Animation/ControlRig/Source/ControlRigEditor/Private/Sequencer/ControlRigParameterTrackEditor.cpp
@@ -2167,8 +2167,9 @@ static UMovieSceneControlRigParameterTrack* AddControlRig(TSharedPtr<ISequencer>
 	FScopedTransaction AddControlRigTrackTransaction(LOCTEXT("AddControlRigTrack", "Add Control Rig Track"));
 
 	TArray<UMovieSceneControlRigParameterTrack*> ExistingRigTracks;
-	
-	if (USkeletalMeshComponent* SkeletalMeshComponent = AcquireSkeletalMeshFromObjectGuid(ObjectBinding, SharedSequencer))
+
+	USkeletalMeshComponent* SkeletalMeshComponent = AcquireSkeletalMeshFromObjectGuid(ObjectBinding, SharedSequencer);
+	if (SkeletalMeshComponent)
 	{
 		ExistingRigTracks = GetExistingControlRigTracksForSkeletalMeshComponent(OwnerMovieScene, SkeletalMeshComponent);
 	}
@@ -2224,7 +2225,13 @@ static UMovieSceneControlRigParameterTrack* AddControlRig(TSharedPtr<ISequencer>
 		{
 			ControlRig->SetIsAdditive(bIsAdditiveControlRig);
 		}
+
+		// The bound object and the data source must be initialized before construction is executed so that
+		// things like FRigUnit_HierarchyImportFromSkeleton can work properly 
 		ControlRig->SetObjectBinding(MakeShared<FControlRigObjectBinding>());
+		ControlRig->GetObjectBinding()->BindToObject(SkeletalMeshComponent);
+		ControlRig->GetDataSourceRegistry()->RegisterDataSource(UControlRig::OwnerComponent, SkeletalMeshComponent);
+		
 		// Do not re-initialize existing control rig
 		if (!InExistingControlRig)
 		{
