Control Rig: Fix binding CR to possible nullptr
#rnx

diff --git a/Engine/Plugins/Animation/ControlRig/Source/ControlRigEditor/Private/Sequencer/ControlRigParameterTrackEditor.cpp b/Engine/Plugins/Animation/ControlRig/Source/ControlRigEditor/Private/Sequencer/ControlRigParameterTrackEditor.cpp
--- a/Engine/Plugins/Animation/ControlRig/Source/ControlRigEditor/Private/Sequencer/ControlRigParameterTrackEditor.cpp
+++ b/Engine/Plugins/Animation/ControlRig/Source/ControlRigEditor/Private/Sequencer/ControlRigParameterTrackEditor.cpp
@@ -2229,8 +2229,11 @@ static UMovieSceneControlRigParameterTrack* AddControlRig(TSharedPtr<ISequencer>
 		// The bound object and the data source must be initialized before construction is executed so that
 		// things like FRigUnit_HierarchyImportFromSkeleton can work properly 
 		ControlRig->SetObjectBinding(MakeShared<FControlRigObjectBinding>());
-		ControlRig->GetObjectBinding()->BindToObject(SkeletalMeshComponent);
-		ControlRig->GetDataSourceRegistry()->RegisterDataSource(UControlRig::OwnerComponent, SkeletalMeshComponent);
+		if (SkeletalMeshComponent)
+		{
+			ControlRig->GetObjectBinding()->BindToObject(SkeletalMeshComponent);
+			ControlRig->GetDataSourceRegistry()->RegisterDataSource(UControlRig::OwnerComponent, SkeletalMeshComponent);
+		}
 		
 		// Do not re-initialize existing control rig
 		if (!InExistingControlRig)
