Control Rig: The editor now retains the Rig Solve Mode after compiling

diff --git a/Engine/Plugins/Animation/ControlRig/Source/ControlRigEditor/Private/Editor/ControlRigEditor.cpp b/Engine/Plugins/Animation/ControlRig/Source/ControlRigEditor/Private/Editor/ControlRigEditor.cpp
--- a/Engine/Plugins/Animation/ControlRig/Source/ControlRigEditor/Private/Editor/ControlRigEditor.cpp
+++ b/Engine/Plugins/Animation/ControlRig/Source/ControlRigEditor/Private/Editor/ControlRigEditor.cpp
@@ -1809,7 +1809,26 @@ bool FControlRigBaseEditor::DetailViewShowsRigModule(FName InModuleName) const
 }
 
 void FControlRigBaseEditor::CompileBaseImpl()
-{
+{	
+	// Compiling can change the event queue. Restore the editor state after compiling.
+	struct FScopedRestoreEventQueue : FNoncopyable
+	{
+		FScopedRestoreEventQueue(FControlRigBaseEditor& InEditor)
+			: Editor(InEditor)
+			, RestoreEventQueue(Editor.GetEventQueue())
+		{
+		}
+		~FScopedRestoreEventQueue()
+		{
+			constexpr bool bCompile = false;
+			Editor.SetEventQueue(RestoreEventQueue, bCompile);
+		}
+	private:
+		FControlRigBaseEditor& Editor;
+		const TArray<FName> RestoreEventQueue;
+	};
+	const FScopedRestoreEventQueue ScopedRestoreEventQueue(*this);
+
 	{
 		DECLARE_SCOPE_HIERARCHICAL_COUNTER_FUNC()
 
