Anim In Engine: : Pose library can paste onto hidden control rigs
#trivial

diff --git a/Engine/Plugins/Animation/ControlRig/Source/ControlRigEditor/Private/EditMode/SControlRigBaseListWidget.cpp b/Engine/Plugins/Animation/ControlRig/Source/ControlRigEditor/Private/EditMode/SControlRigBaseListWidget.cpp
--- a/Engine/Plugins/Animation/ControlRig/Source/ControlRigEditor/Private/EditMode/SControlRigBaseListWidget.cpp
+++ b/Engine/Plugins/Animation/ControlRig/Source/ControlRigEditor/Private/EditMode/SControlRigBaseListWidget.cpp
@@ -987,7 +987,7 @@ void SControlRigBaseListWidget::OnAssetsActivated(const TArray<FAssetData>& Sele
 					FControlRigEditMode* ControlRigEditMode = GetEditMode();
 					if (ControlRigEditMode)
 					{
-						TArray<UControlRig*> ControlRigs = ControlRigEditMode->GetControlRigsArray(true /*bIsVisible*/);
+						TArray<UControlRig*> ControlRigs = ControlRigEditMode->GetControlRigsArray(false /*bIsVisible*/);
 						if (ControlRigs.Num() > 0)
 						{
 							const FScopedTransaction Transaction(LOCTEXT("SelectControls", "Select Controls"));
@@ -1365,7 +1365,7 @@ void SControlRigBaseListWidget::ExecutePastePose(UControlRigPoseAsset* PoseAsset
 		FControlRigEditMode* ControlRigEditMode = GetEditMode();
 		if (ControlRigEditMode)
 		{
-			TArray<UControlRig*> ControlRigs = ControlRigEditMode->GetControlRigsArray(true /*bIsVisible*/);
+			TArray<UControlRig*> ControlRigs = ControlRigEditMode->GetControlRigsArray(false /*bIsVisible*/);
 			if (ControlRigs.Num() > 0)
 			{
 				const FScopedTransaction Transaction(LOCTEXT("PastePose", "Paste Pose"));
@@ -1394,7 +1394,7 @@ void SControlRigBaseListWidget::ExecuteSelectControls(UControlRigPoseAsset* Pose
 	FControlRigEditMode* ControlRigEditMode = GetEditMode();
 	if (ControlRigEditMode)
 	{
-		TArray<UControlRig*> ControlRigs = ControlRigEditMode->GetControlRigsArray(true /*bIsVisible*/);
+		TArray<UControlRig*> ControlRigs = ControlRigEditMode->GetControlRigsArray(false /*bIsVisible*/);
 		if (ControlRigs.Num() > 0)
 		{
 			const FScopedTransaction Transaction(LOCTEXT("SelectControls", "Select Controls"));
@@ -1431,7 +1431,7 @@ void SControlRigBaseListWidget::ExecutePasteMirrorPose(UControlRigPoseAsset* Pos
 		FControlRigEditMode* ControlRigEditMode = GetEditMode();
 		if (ControlRigEditMode)
 		{
-			TArray<UControlRig*> ControlRigs = ControlRigEditMode->GetControlRigsArray(true /*bIsVisible*/);
+			TArray<UControlRig*> ControlRigs = ControlRigEditMode->GetControlRigsArray(false /*bIsVisible*/);
 			if (ControlRigs.Num() > 0)
 			{
 				const FScopedTransaction Transaction(LOCTEXT("PasteMirrorPose", "Paste Mirror Pose"));
