Control Rig: Dragging a rig element onto the graph now offers an option to Get Metadata and Set Metadata

diff --git a/Engine/Plugins/Animation/ControlRig/Source/ControlRigEditor/Private/Editor/ControlRigEditor.cpp b/Engine/Plugins/Animation/ControlRig/Source/ControlRigEditor/Private/Editor/ControlRigEditor.cpp
--- a/Engine/Plugins/Animation/ControlRig/Source/ControlRigEditor/Private/Editor/ControlRigEditor.cpp
+++ b/Engine/Plugins/Animation/ControlRig/Source/ControlRigEditor/Private/Editor/ControlRigEditor.cpp
@@ -64,6 +64,7 @@
 #include "Units/Hierarchy/RigUnit_SetCurveValue.h"
 #include "Units/Hierarchy/RigUnit_AddBoneTransform.h"
 #include "Units/Hierarchy/RigUnit_Component.h"
+#include "Units/Hierarchy/RigUnit_Metadata.h"
 #include "EdGraph/NodeSpawners/RigVMEdGraphUnitNodeSpawner.h"
 #include "Graph/ControlRigGraphSchema.h"
 #include "ControlRigObjectVersion.h"
@@ -5417,6 +5418,57 @@ void FControlRigBaseEditor::CreateRigHierarchyToGraphDragAndDropMenu() const
 							CreateItemArrayEntry.InsertPosition.Position = EToolMenuInsertType::After;
 							Section.AddEntry(CreateItemArrayEntry);
 						}
+
+						// Meta data 
+						{
+							FToolMenuEntry& MetadataSeparator = Section.AddSeparator(TEXT("MetadataSeparator"));
+							MetadataSeparator.InsertPosition.Name = TEXT("CreateItemArray"),
+							MetadataSeparator.InsertPosition.Position = EToolMenuInsertType::After;
+
+							constexpr bool bGetMetadataEntry_IsGetter = true;
+							FToolMenuEntry GetMetadataEntry = FToolMenuEntry::InitMenuEntry(
+								TEXT("GetMetadataEntry"),
+								LOCTEXT("GetMetadata", "Get Metadata"),
+								LOCTEXT("GetMetadata_ToolTip", "Getter for Metadata"),
+								FSlateIcon(),
+								FUIAction(
+									FExecuteAction::CreateSP(
+										ControlRigEditor->SharedRef(),
+										&FControlRigBaseEditor::HandleMakeMetadataGetterSetter,
+										bGetMetadataEntry_IsGetter, 
+										DraggedElements,
+										Graph, 
+										NodePosition),
+									FCanExecuteAction()
+								)
+							);
+
+							GetMetadataEntry.InsertPosition.Name = MetadataSeparator.Name;
+							GetMetadataEntry.InsertPosition.Position = EToolMenuInsertType::After;
+							Section.AddEntry(GetMetadataEntry);
+
+							constexpr bool bSetMetadataEntry_IsGetter = false;
+							FToolMenuEntry SetMetadataEntry = FToolMenuEntry::InitMenuEntry(
+								TEXT("SetMetadataEntry"),
+								LOCTEXT("SetMetadata", "Set Metadata"),
+								LOCTEXT("SetMetadata_ToolTip", "Setter for Metadata"),
+								FSlateIcon(),
+								FUIAction(
+									FExecuteAction::CreateSP(
+										ControlRigEditor->SharedRef(),
+										&FControlRigBaseEditor::HandleMakeMetadataGetterSetter,
+										bSetMetadataEntry_IsGetter,
+										DraggedElements,
+										Graph,
+										NodePosition),
+									FCanExecuteAction()
+								)
+							);
+
+							SetMetadataEntry.InsertPosition.Name = GetMetadataEntry.Name;
+							SetMetadataEntry.InsertPosition.Position = EToolMenuInsertType::After;
+							Section.AddEntry(SetMetadataEntry);
+						}
 					}
 				}
 			})
@@ -5978,6 +6030,99 @@ void FControlRigBaseEditor::HandleMakeElementGetterSetter(ERigElementGetterSette
 	}
 }
 
+void FControlRigBaseEditor::HandleMakeMetadataGetterSetter(const bool bIsGetter, const TArray<FRigElementKey> Keys, UEdGraph* Graph, FVector2D NodePosition)
+{
+	DECLARE_SCOPE_HIERARCHICAL_COUNTER_FUNC()
+
+	URigHierarchy* Hierarchy = GetHierarchyBeingDebugged();
+	URigVMController* Controller = GetFocusedController();
+	if (Keys.IsEmpty() ||
+		!Hierarchy ||
+		!Controller)
+	{
+		return;
+	}
+
+	Controller->OpenUndoBracket(TEXT("Adding Nodes from Hierarchy"));
+
+	const FRigVMDispatchFactory* GetterDispatchFactory = FRigVMRegistry::Get().FindOrAddDispatchFactory(FRigDispatch_GetMetadata::StaticStruct());
+	const FRigVMDispatchFactory* SetterDispatchFactory = FRigVMRegistry::Get().FindOrAddDispatchFactory(FRigDispatch_SetMetadata::StaticStruct());
+	if (!ensureMsgf(GetterDispatchFactory && SetterDispatchFactory, TEXT("Cannot find factory for meta data getter and setter nodes, skip creating nodes.")))
+	{
+		return;
+	}
+
+	const FRigVMDispatchFactory* Factory = bIsGetter ? GetterDispatchFactory : SetterDispatchFactory;
+	check(Factory);
+
+	const FVector2D NodePositionIncrement = bIsGetter ? FVector2D(0.0, 200.0) : FVector2D(360.0, 0.0);
+
+	TArray<FName> NewNodes;
+	for (int32 Index = 0; Index < Keys.Num(); Index++)
+	{
+		const FRigElementKey Key = Keys[Index];
+
+		constexpr const TCHAR* NodeNameOverride = TEXT("");
+		constexpr bool bSetupUndoRedo_ForAddTemplateNode = true;
+		constexpr bool bPrintPythonCommand_ForAddTemplateNode = true;
+		if (URigVMTemplateNode* ModelNode = Controller->AddTemplateNode(
+			Factory->GetTemplateNotation(),
+			NodePosition,
+			NodeNameOverride,
+			bSetupUndoRedo_ForAddTemplateNode,
+			bPrintPythonCommand_ForAddTemplateNode))
+		{
+			NewNodes.Add(ModelNode->GetFName());
+
+			constexpr bool bResizeArrays = true;
+			constexpr bool bSetupUndoRedo = true;
+			constexpr bool bPrintPythonCommand_ForSetPinDefaultValue = false;
+			constexpr bool bSetValueOnLinkedPin_ForSetPinDefaultValue = true;
+
+			// Set default item type
+			{
+				const FString ItemTypePinPath = FString::Printf(TEXT("%s.%s.Type"), *ModelNode->GetName(), TEXT("Item"));
+				const FString ItemTypeStr = StaticEnum<ERigElementType>()->GetDisplayNameTextByValue((int64)Key.Type).ToString();
+				Controller->SetPinDefaultValue(
+					ItemTypePinPath,
+					ItemTypeStr,
+					bResizeArrays,
+					bSetupUndoRedo,
+					bPrintPythonCommand_ForSetPinDefaultValue,
+					bSetValueOnLinkedPin_ForSetPinDefaultValue);
+			}
+
+			// Set default item name
+			{
+				const FString ItemNamePinPath = FString::Printf(TEXT("%s.%s.Name"), *ModelNode->GetName(), TEXT("Item"));
+				Controller->SetPinDefaultValue(
+					ItemNamePinPath,
+					Key.Name.ToString(),
+					bResizeArrays,
+					bSetupUndoRedo,
+					bPrintPythonCommand_ForSetPinDefaultValue,
+					bSetValueOnLinkedPin_ForSetPinDefaultValue);
+			}
+
+			URigVMEdGraphUnitNodeSpawner::HookupMutableNode(ModelNode, GetControlRigAssetInterface()->GetRigVMAssetInterface());
+
+			NewNodes.Add(ModelNode->GetFName());
+		}
+
+		NodePosition += NodePositionIncrement;
+	}
+
+	if (NewNodes.Num() > 0)
+	{
+		Controller->SetNodeSelection(NewNodes);
+		Controller->CloseUndoBracket();
+	}
+	else
+	{
+		Controller->CancelUndoBracket();
+	}
+}
+
 void FControlRigBaseEditor::HandleMakeComponentGetterSetter(bool bIsGetter, TArray<FRigComponentKey> Keys, UEdGraph* Graph, FVector2D NodePosition)
 {
 	DECLARE_SCOPE_HIERARCHICAL_COUNTER_FUNC()
diff --git a/Engine/Plugins/Animation/ControlRig/Source/ControlRigEditor/Private/Editor/ControlRigEditor.h b/Engine/Plugins/Animation/ControlRig/Source/ControlRigEditor/Private/Editor/ControlRigEditor.h
--- a/Engine/Plugins/Animation/ControlRig/Source/ControlRigEditor/Private/Editor/ControlRigEditor.h
+++ b/Engine/Plugins/Animation/ControlRig/Source/ControlRigEditor/Private/Editor/ControlRigEditor.h
@@ -415,6 +415,7 @@ protected:
 
 	virtual void FilterDraggedKeys(TArray<FRigElementKey>& Keys, bool bRemoveNameSpace) override;
 	void HandleMakeElementGetterSetter(ERigElementGetterSetterType Type, bool bIsGetter, TArray<FRigElementKey> Keys, UEdGraph* Graph, FVector2D NodePosition);
+	void HandleMakeMetadataGetterSetter(const bool bIsGetter, const TArray<FRigElementKey> Keys, UEdGraph* Graph, FVector2D NodePosition);
 	void HandleMakeComponentGetterSetter(bool bIsGetter, TArray<FRigComponentKey> Keys, UEdGraph* Graph, FVector2D NodePosition);
 
 	void HandleOnControlModified(UControlRig* Subject, FRigControlElement* ControlElement, const FRigControlModifiedContext& Context);
