Control Rig: Extract snapper, space picker, and constraints edition to a separate spawnable tab.

diff --git a/Engine/Plugins/Animation/ControlRig/Source/ControlRigEditor/Private/ControlRigEditorStyle.h b/Engine/Plugins/Animation/ControlRig/Source/ControlRigEditor/Private/ControlRigEditorStyle.h
--- a/Engine/Plugins/Animation/ControlRig/Source/ControlRigEditor/Private/ControlRigEditorStyle.h
+++ b/Engine/Plugins/Animation/ControlRig/Source/ControlRigEditor/Private/ControlRigEditorStyle.h
@@ -225,6 +225,24 @@ public:
 			}
 
 		}
+
+		// Tabs
+		{
+			FCheckBoxStyle BaseCheckBoxStyle = FAppStyle::Get().GetWidgetStyle<FCheckBoxStyle>("ToggleButtonCheckbox");
+			Set("ControlRig.TabButton", BaseCheckBoxStyle
+				.SetUncheckedImage(FSlateRoundedBoxBrush(FStyleColors::Header, 4.0f, FStyleColors::Input, 1.0f))
+				.SetUncheckedHoveredImage(FSlateRoundedBoxBrush(FStyleColors::Hover, 4.0f, FStyleColors::Input, 1.0f))
+				.SetUncheckedPressedImage(FSlateRoundedBoxBrush(FStyleColors::Hover, 4.0f, FStyleColors::Input, 1.0f))
+				.SetCheckedImage(FSlateRoundedBoxBrush(FStyleColors::Primary, 4.0f, FStyleColors::Input, 1.0f))
+				.SetCheckedHoveredImage(FSlateRoundedBoxBrush(FStyleColors::PrimaryHover, 4.0f, FStyleColors::Input, 1.0f))
+				.SetCheckedPressedImage(FSlateRoundedBoxBrush(FStyleColors::PrimaryHover, 4.0f, FStyleColors::Input, 1.0f))
+				.SetPadding(FMargin(16, 6))
+			);
+			
+			Set("ControlRig.Constrain.Spaces", new IMAGE_BRUSH_SVG("Slate/Tab_Spaces_16x", Icon16x16));
+			Set("ControlRig.Constrain.Constraints", new IMAGE_BRUSH_SVG("Slate/Tab_Constraints_16x", Icon16x16));
+			Set("ControlRig.Constrain.Snapper", new IMAGE_BRUSH_SVG("Slate/Tab_RigSnapper_16x", Icon16x16));
+		}
 		
 		FSlateStyleRegistry::RegisterSlateStyle(*this);
 	}
diff --git a/Engine/Plugins/Animation/ControlRig/Source/ControlRigEditor/Private/EditMode/Constrain/ConstrainToolsManager.cpp b/Engine/Plugins/Animation/ControlRig/Source/ControlRigEditor/Private/EditMode/Constrain/ConstrainToolsManager.cpp
new file mode 100644
--- /dev/null
+++ b/Engine/Plugins/Animation/ControlRig/Source/ControlRigEditor/Private/EditMode/Constrain/ConstrainToolsManager.cpp
@@ -0,0 +1,87 @@
+﻿// Copyright Epic Games, Inc. All Rights Reserved.
+
+#include "ConstrainToolsManager.h"
+
+#include "EditMode/ControlRigEditModeToolkit.h"
+#include "Framework/Docking/TabManager.h"
+#include "SConstrainToolsRoot.h"
+#include "Widgets/Docking/SDockTab.h"
+
+#define LOCTEXT_NAMESPACE "FConstrainToolsManager"
+
+namespace UE::ControlRigEditor
+{
+FConstrainToolsManager::FConstrainToolsManager(
+	const TSharedRef<FTabManager>& InToolkitTabManager, 
+	FControlRigEditMode& InOwningEditMode,
+	const TSharedRef<FRigSelectionViewModel>& InSelectionViewModel
+	)
+	: TabManager(InToolkitTabManager)
+	, OwningEditMode(InOwningEditMode)
+	, SelectionViewModel(InSelectionViewModel)
+{
+	RegisterTabSpawner(TabManager);
+}
+
+FConstrainToolsManager::~FConstrainToolsManager()
+{
+	HideWidget();
+	UnregisterTabSpawner(TabManager);
+}
+
+void FConstrainToolsManager::ToggleVisibility() const
+{
+	if (IsShowingWidget())
+	{
+		HideWidget();
+	}
+	else
+	{
+		ShowWidget();
+	}
+}
+
+void FConstrainToolsManager::ShowWidget() const
+{
+	TabManager->TryInvokeTab(FControlRigEditModeToolkit::ConstrainingTabName);
+}
+
+void FConstrainToolsManager::HideWidget() const
+{
+	if (const TSharedPtr<SDockTab> Tab = SpawnedTab.Pin())
+	{
+		Tab->RequestCloseTab();
+	}
+}
+
+bool FConstrainToolsManager::IsShowingWidget() const
+{
+	return SpawnedTab.IsValid();
+}
+
+void FConstrainToolsManager::RegisterTabSpawner(const TSharedRef<FTabManager>& InTabManager)
+{
+	InTabManager->RegisterTabSpawner(
+		FControlRigEditModeToolkit::ConstrainingTabName, FOnSpawnTab::CreateRaw(this, &FConstrainToolsManager::SpawnSnapperTab)
+		)
+		.SetDisplayName(LOCTEXT("Tab.Title", "Constraining"))
+		.SetTooltipText(LOCTEXT("Tab.ToolTip", "Snap child objects to a parent object over a set of frames."))
+		.SetIcon(FSlateIcon(TEXT("ControlRigEditorStyle"), TEXT("ControlRig.SnapperTool")));
+	InTabManager->RegisterDefaultTabWindowSize(FControlRigEditModeToolkit::ConstrainingTabName, FVector2D(380, 500));
+}
+
+void FConstrainToolsManager::UnregisterTabSpawner(const TSharedRef<FTabManager>& InTabManager)
+{
+	InTabManager->UnregisterTabSpawner(FControlRigEditModeToolkit::ConstrainingTabName);
+}
+
+TSharedRef<SDockTab> FConstrainToolsManager::SpawnSnapperTab(const FSpawnTabArgs& Args)
+{
+	return SAssignNew(SpawnedTab, SDockTab)
+	[
+		SNew(SConstrainToolsRoot, OwningEditMode, SelectionViewModel)
+	];
+}
+}
+
+#undef LOCTEXT_NAMESPACE
\ No newline at end of file
diff --git a/Engine/Plugins/Animation/ControlRig/Source/ControlRigEditor/Private/EditMode/Constrain/ConstrainToolsManager.h b/Engine/Plugins/Animation/ControlRig/Source/ControlRigEditor/Private/EditMode/Constrain/ConstrainToolsManager.h
new file mode 100644
--- /dev/null
+++ b/Engine/Plugins/Animation/ControlRig/Source/ControlRigEditor/Private/EditMode/Constrain/ConstrainToolsManager.h
@@ -0,0 +1,61 @@
+﻿// Copyright Epic Games, Inc. All Rights Reserved.
+
+#pragma once
+
+#include "Templates/SharedPointer.h"
+#include "Templates/UnrealTemplate.h"
+
+class FControlRigEditMode;
+class FSpawnTabArgs;
+class SDockTab;
+class FTabManager;
+class FToolBarBuilder;
+
+namespace UE::ControlRigEditor
+{
+class FRigSelectionViewModel;
+	
+/**
+ * Manages UI that is used for constraining, namely:
+ * - SControlRigSnapper
+ * - SRigSpacePickerWidget
+ * - SConstraintsEditionWidget
+ *
+ * This class responsibility is spawning & managing a window with tab-like buttons to switch between these widgets. 
+ */
+class FConstrainToolsManager : public FNoncopyable
+{
+public:
+
+	explicit FConstrainToolsManager(
+		const TSharedRef<FTabManager>& InToolkitTabManager,
+		FControlRigEditMode& InOwningEditMode UE_LIFETIMEBOUND,
+		const TSharedRef<FRigSelectionViewModel>& InSelectionViewModel
+		);
+	~FConstrainToolsManager();
+	
+	void ToggleVisibility() const;
+	void ShowWidget() const;
+	void HideWidget() const;
+	bool IsShowingWidget() const;
+
+private:
+
+	/** Used to show / hide the tab containing our widgets. */
+	const TSharedRef<FTabManager> TabManager;
+
+	/** Needed to construct widget hierarchy.  */
+	FControlRigEditMode& OwningEditMode;
+	/** Handles selection changing. Needed to construct widget hierarchy. */
+	const TSharedRef<FRigSelectionViewModel> SelectionViewModel;
+
+	/** The tab we spawned. Used to detect whether the tab is currently being shown. */
+	TWeakPtr<SDockTab> SpawnedTab;
+
+	// Tab spawning
+	void RegisterTabSpawner(const TSharedRef<FTabManager>& InTabManager);
+	void UnregisterTabSpawner(const TSharedRef<FTabManager>& InTabManager);
+	TSharedRef<SDockTab> SpawnSnapperTab(const FSpawnTabArgs& Args);
+};
+}
+
diff --git a/Engine/Plugins/Animation/ControlRig/Source/ControlRigEditor/Private/EditMode/Constrain/SConstrainToolsRoot.cpp b/Engine/Plugins/Animation/ControlRig/Source/ControlRigEditor/Private/EditMode/Constrain/SConstrainToolsRoot.cpp
new file mode 100644
--- /dev/null
+++ b/Engine/Plugins/Animation/ControlRig/Source/ControlRigEditor/Private/EditMode/Constrain/SConstrainToolsRoot.cpp
@@ -0,0 +1,79 @@
+﻿// Copyright Epic Games, Inc. All Rights Reserved.
+
+#include "SConstrainToolsRoot.h"
+
+#include "ControlRigEditorStyle.h"
+#include "SConstraintsEditionWidget_ConstrainPanel.h"
+#include "EditMode/ControlRigEditMode.h"
+#include "EditMode/SControlRigSnapper.h"
+#include "EditMode/SControlRigSpacePicker.h"
+#include "Editor/Constraints/SConstraintsWidget.h"
+#include "Misc/RigEditModeUtils.h"
+#include "Misc/Tab/SInlineTabPanel.h"
+
+#define LOCTEXT_NAMESPACE "SConstrainToolsRoot"
+
+namespace UE::ControlRigEditor
+{
+void SConstrainToolsRoot::Construct(
+	const FArguments& InArgs, FControlRigEditMode& InOwningEditMode, const TSharedRef<FRigSelectionViewModel>& InSelectionViewModel
+	)
+{
+	ChildSlot
+	[
+		SNew(SInlineTabPanel)
+		.Tabs({
+			FInlineTabArgs()
+				.SetLabel(LOCTEXT("Spaces", "Spaces"))
+				.SetContent(CreateSpacesContent(InOwningEditMode, InSelectionViewModel))
+				.SetIcon(FSlateIcon(FControlRigEditorStyle::Get().GetStyleSetName(), "ControlRig.Constrain.Spaces")),
+			
+			FInlineTabArgs()
+				.SetLabel(LOCTEXT("Constraints", "Constraints"))
+				.SetContent(CreateConstraintsContent(InOwningEditMode, InSelectionViewModel))
+				.SetOnTabSelected(FSimpleDelegate::CreateSP(this, &SConstrainToolsRoot::OnTabSelected_ConstraintsEditionWidget))
+				.SetIcon(FSlateIcon(FControlRigEditorStyle::Get().GetStyleSetName(), "ControlRig.Constrain.Constraints")),
+			
+			FInlineTabArgs()
+				.SetLabel(LOCTEXT("Snapper", "Snapper"))
+				.SetContent(CreateSnapperContent())
+				.SetIcon(FSlateIcon(FControlRigEditorStyle::Get().GetStyleSetName(), "ControlRig.Constrain.Snapper")),
+		})
+	];
+}
+
+TSharedRef<SWidget> SConstrainToolsRoot::CreateSpacesContent(
+	FControlRigEditMode& InOwningEditMode, const TSharedRef<FRigSelectionViewModel>& InSelectionViewModel
+	)
+{
+	using namespace UE::ControlRigEditor;
+    const FInitialSpacePickerSelection InitialSelection = DetermineInitialSpacePickerSelection(InOwningEditMode);
+	
+	return SNew(SControlRigSpacePicker, InOwningEditMode, InSelectionViewModel)
+		.InitiallyCollapsed(false)
+		
+		// If this widget is constructed when there is already a control rig selected, show that in the spaces initially.
+		.Hierarchy(InitialSelection.Hierarchy)
+		.Controls(InitialSelection.SelectedControls);
+}
+
+TSharedRef<SWidget> SConstrainToolsRoot::CreateConstraintsContent(
+	FControlRigEditMode& InOwningMode, const TSharedRef<FRigSelectionViewModel>& InSelectionViewModel
+	)
+{
+	return SAssignNew(ConstrainsEditionWidget, SConstraintsEditionWidget_ConstrainPanel, InOwningMode, InSelectionViewModel)
+		.InitiallyCollapsed(false);
+}
+
+TSharedRef<SWidget> SConstrainToolsRoot::CreateSnapperContent()
+{
+	return SNew(SControlRigSnapper);
+}
+
+void SConstrainToolsRoot::OnTabSelected_ConstraintsEditionWidget() const
+{
+	ConstrainsEditionWidget->RefreshConstraintList();
+}
+}
+
+#undef LOCTEXT_NAMESPACE
\ No newline at end of file
diff --git a/Engine/Plugins/Animation/ControlRig/Source/ControlRigEditor/Private/EditMode/Constrain/SConstrainToolsRoot.h b/Engine/Plugins/Animation/ControlRig/Source/ControlRigEditor/Private/EditMode/Constrain/SConstrainToolsRoot.h
new file mode 100644
--- /dev/null
+++ b/Engine/Plugins/Animation/ControlRig/Source/ControlRigEditor/Private/EditMode/Constrain/SConstrainToolsRoot.h
@@ -0,0 +1,43 @@
+﻿// Copyright Epic Games, Inc. All Rights Reserved.
+
+#pragma once
+
+#include "Widgets/DeclarativeSyntaxSupport.h"
+#include "Widgets/SCompoundWidget.h"
+
+class FControlRigEditMode;
+
+namespace UE::ControlRigEditor
+{
+class FRigSelectionViewModel;
+class SConstraintsEditionWidget_ConstrainPanel;
+
+/**
+ * Coordinator widget that contains
+ * - SControlRigSnapper
+ * - SRigSpacePickerWidget
+ * - SConstraintsEditionWidget
+ * and ranges them under a tab-like structure.
+ */
+class SConstrainToolsRoot : public SCompoundWidget
+{
+public:
+	
+	SLATE_BEGIN_ARGS(SConstrainToolsRoot){}
+	SLATE_END_ARGS()
+
+	void Construct(const FArguments& InArgs, FControlRigEditMode& InOwningEditMode, const TSharedRef<FRigSelectionViewModel>& InSelectionViewModel);
+
+private:
+
+	/** The widget for editing constraints. */
+	TSharedPtr<SConstraintsEditionWidget_ConstrainPanel> ConstrainsEditionWidget;
+
+	TSharedRef<SWidget> CreateSpacesContent(FControlRigEditMode& InOwningEditMode, const TSharedRef<FRigSelectionViewModel>& InSelectionViewModel);
+	TSharedRef<SWidget> CreateConstraintsContent(FControlRigEditMode& InOwningMode, const TSharedRef<FRigSelectionViewModel>& InSelectionViewModel);
+	TSharedRef<SWidget> CreateSnapperContent();
+
+	/** Refreshes the list when we switch to the tab. */
+	void OnTabSelected_ConstraintsEditionWidget() const;
+};
+}
diff --git a/Engine/Plugins/Animation/ControlRig/Source/ControlRigEditor/Private/EditMode/Constrain/SConstraintsEditionWidget_ConstrainPanel.cpp b/Engine/Plugins/Animation/ControlRig/Source/ControlRigEditor/Private/EditMode/Constrain/SConstraintsEditionWidget_ConstrainPanel.cpp
new file mode 100644
--- /dev/null
+++ b/Engine/Plugins/Animation/ControlRig/Source/ControlRigEditor/Private/EditMode/Constrain/SConstraintsEditionWidget_ConstrainPanel.cpp
@@ -0,0 +1,217 @@
+﻿// Copyright Epic Games, Inc. All Rights Reserved.
+
+#include "SConstraintsEditionWidget_ConstrainPanel.h"
+
+#include "EditMode/ControlRigEditMode.h"
+#include "EditMode/SComponentPickerPopup.h"
+#include "EditMode/Models/RigSelectionViewModel.h"
+#include "Widgets/Input/SButton.h"
+#include "Widgets/Input/SComboButton.h"
+#include "Widgets/Layout/SExpandableArea.h"
+#include "Widgets/Layout/SSpacer.h"
+
+#define LOCTEXT_NAMESPACE "SExpandableConstraintsEditionWidget"
+
+namespace UE::ControlRigEditor
+{
+void SConstraintsEditionWidget_ConstrainPanel::Construct(
+	const FArguments& InArgs, FControlRigEditMode& InEditMode, const TSharedRef<FRigSelectionViewModel>& InSelectionViewModel
+	)
+{
+	OwningEditMode = &InEditMode;
+	SelectionViewModel = InSelectionViewModel;
+
+	InSelectionViewModel->OnControlSelected().AddSP(this, &SConstraintsEditionWidget_ConstrainPanel::OnControlSelected);
+	
+	ChildSlot
+	[
+		SNew(SExpandableArea)
+		.InitiallyCollapsed(InArgs._InitiallyCollapsed)
+		.AreaTitle(LOCTEXT("ConstraintsWidget", "Constraints"))
+		.AreaTitleFont(FAppStyle::GetFontStyle("DetailsView.CategoryFontStyle"))
+		.BorderBackgroundColor(FLinearColor(.6f, .6f, .6f))
+		.HeaderContent()
+		[
+			SNew(SHorizontalBox)
+
+			// "Constraints" label
+			+SHorizontalBox::Slot()
+			.AutoWidth()
+			.HAlign(HAlign_Left)
+			.VAlign(VAlign_Center)
+			.Padding(0.f, 0.f, 0.f, 0.f)
+			[
+				SNew(STextBlock)
+				.Text(LOCTEXT("ConstraintsWidget", "Constraints"))
+				.Font(FCoreStyle::Get().GetFontStyle("ExpandableArea.TitleFont"))
+			]
+
+			// Spacer
+			+SHorizontalBox::Slot()
+			.FillWidth(1.f)
+			[
+				SNew(SSpacer)
+			]
+			
+			// "Selected" button
+			+SHorizontalBox::Slot()
+			.AutoWidth()
+			.HAlign(HAlign_Right)
+			.VAlign(VAlign_Center)
+			.Padding(0.f, 2.f, 8.f, 2.f)
+			[
+				// Combo Button to swap what constraints we show  
+				CreateSelectionComboButton()
+			]
+			
+			// "Plus" icon
+			+SHorizontalBox::Slot()
+			.AutoWidth()
+			.HAlign(HAlign_Right)
+			.VAlign(VAlign_Center)
+			.Padding(0.f, 2.f, 8.f, 2.f)
+			[
+				CreatePlusButton()
+			]
+		]
+		.BodyContent() [ SAssignNew(ConstraintsEditionWidget, SConstraintsEditionWidget) ]
+		.OnAreaExpansionChanged_Lambda( [this](bool bIsExpanded){ ConstraintsEditionWidget->RefreshConstraintList(); })
+	];
+}
+
+SConstraintsEditionWidget_ConstrainPanel::~SConstraintsEditionWidget_ConstrainPanel()
+{
+	SelectionViewModel->OnControlSelected().RemoveAll(this);
+}
+
+TSharedRef<SWidget> SConstraintsEditionWidget_ConstrainPanel::CreateSelectionComboButton()
+{
+	return SNew(SComboButton)
+		.OnGetMenuContent_Lambda([this]()
+		{
+			FMenuBuilder MenuBuilder(true, nullptr);
+				
+			MenuBuilder.BeginSection("Constraints");
+			for (int32 Index = 0; Index < 4; ++Index)
+			{
+				const FBaseConstraintListWidget::EShowConstraints Constraints = static_cast<FBaseConstraintListWidget::EShowConstraints>(Index);
+				FUIAction ItemAction(FExecuteAction::CreateSP(this, &SConstraintsEditionWidget_ConstrainPanel::OnSelectShowConstraints, Index));
+				const TAttribute<FText> Text = ConstraintsEditionWidget->GetShowConstraintsText(Constraints);
+				const TAttribute<FText> Tooltip = ConstraintsEditionWidget->GetShowConstraintsTooltip(Constraints);
+				MenuBuilder.AddMenuEntry(Text, Tooltip, FSlateIcon(), ItemAction);
+			}
+			MenuBuilder.EndSection();
+
+			return MenuBuilder.MakeWidget();
+		})
+		.ButtonContent()
+		[
+			SNew(SHorizontalBox)
+
+			+SHorizontalBox::Slot()
+			[
+				SNew(STextBlock)
+				.Text_Lambda([this](){ return GetShowConstraintsName(); })
+				.ToolTipText_Lambda([this]() { return GetShowConstraintsTooltip(); })
+			]
+		];
+}
+
+TSharedRef<SWidget> SConstraintsEditionWidget_ConstrainPanel::CreatePlusButton()
+{
+	return SNew(SButton)
+		.ContentPadding(0.0f)
+		.ButtonStyle(FAppStyle::Get(), "NoBorder")
+		.IsEnabled_Lambda([this]()
+		{
+			const ULevel* CurrentLevel = OwningEditMode->GetWorld()->GetCurrentLevel();
+			const TArray<AActor*> SelectedActors = ObjectPtrDecay(CurrentLevel->Actors).FilterByPredicate([](const AActor* Actor)
+			{
+				return Actor && Actor->IsSelected();
+			});
+			return !SelectedActors.IsEmpty();
+		})
+		.OnClicked(this, &SConstraintsEditionWidget_ConstrainPanel::HandleAddConstraintClicked)
+		.Cursor(EMouseCursor::Default)
+		.ToolTipText(LOCTEXT("AddConstraint", "Add Constraint"))
+		[
+			SNew(SImage)
+			.Image(FAppStyle::GetBrush(TEXT("Icons.PlusCircle")))
+		];
+}
+
+void SConstraintsEditionWidget_ConstrainPanel::OnSelectShowConstraints(int32 Index)
+{
+	if (ConstraintsEditionWidget.IsValid())
+	{
+		const SConstraintsEditionWidget::EShowConstraints ShowConstraint = static_cast<SConstraintsEditionWidget::EShowConstraints>(Index);
+		ConstraintsEditionWidget->SetShowConstraints(ShowConstraint);
+	}
+}
+
+FReply SConstraintsEditionWidget_ConstrainPanel::HandleAddConstraintClicked()
+{
+	FMenuBuilder MenuBuilder(true, nullptr);
+
+	const auto AddConstraintWidget = [&](ETransformConstraintType InConstraintType)
+	{
+		const TSharedRef<SConstraintMenuEntry> Entry =
+			SNew(SConstraintMenuEntry, InConstraintType)
+			.OnConstraintCreated_Lambda([this]()
+			{
+				ConstraintsEditionWidget->RefreshConstraintList();
+				ConstraintsEditionWidget->ResetParent();
+			})
+			.OnGetParent_Lambda([this]()
+			{
+				const FConstrainable InvalidParent;
+				return ConstraintsEditionWidget ? ConstraintsEditionWidget->GetParent() : InvalidParent;
+			});
+		MenuBuilder.AddWidget(Entry, FText::GetEmpty(), true);
+	};
+	
+	MenuBuilder.BeginSection("CreateConstraint", LOCTEXT("CreateConstraintHeader", "Create New..."));
+	{
+		AddConstraintWidget(ETransformConstraintType::Translation);
+		AddConstraintWidget(ETransformConstraintType::Rotation);
+		AddConstraintWidget(ETransformConstraintType::Scale);
+		AddConstraintWidget(ETransformConstraintType::Parent);
+		AddConstraintWidget(ETransformConstraintType::LookAt);
+	}
+	MenuBuilder.EndSection();
+	
+	FSlateApplication::Get().PushMenu(
+		AsShared(),
+		FWidgetPath(),
+		MenuBuilder.MakeWidget(),
+		FSlateApplication::Get().GetCursorPos(),
+		FPopupTransitionEffect(FPopupTransitionEffect::ContextMenu)
+		);
+	
+	return FReply::Handled();
+}
+
+FText SConstraintsEditionWidget_ConstrainPanel::GetShowConstraintsName() const
+{
+	return ConstraintsEditionWidget.IsValid()
+		? ConstraintsEditionWidget->GetShowConstraintsText(FBaseConstraintListWidget::ShowConstraints)
+		: FText::GetEmpty();
+}
+
+FText SConstraintsEditionWidget_ConstrainPanel::GetShowConstraintsTooltip() const
+{
+	return ConstraintsEditionWidget.IsValid()
+		? ConstraintsEditionWidget->GetShowConstraintsTooltip(FBaseConstraintListWidget::ShowConstraints)
+		: FText::GetEmpty();
+}
+
+void SConstraintsEditionWidget_ConstrainPanel::OnControlSelected(UControlRig* Subject, FRigControlElement* RigControlElement, bool bIsSelected) const
+{
+	if (Subject && Subject->GetHierarchy())
+	{
+		ConstraintsEditionWidget->InvalidateConstraintList();
+	}
+}
+}
+
+#undef LOCTEXT_NAMESPACE
\ No newline at end of file
diff --git a/Engine/Plugins/Animation/ControlRig/Source/ControlRigEditor/Private/EditMode/Constrain/SConstraintsEditionWidget_ConstrainPanel.h b/Engine/Plugins/Animation/ControlRig/Source/ControlRigEditor/Private/EditMode/Constrain/SConstraintsEditionWidget_ConstrainPanel.h
new file mode 100644
--- /dev/null
+++ b/Engine/Plugins/Animation/ControlRig/Source/ControlRigEditor/Private/EditMode/Constrain/SConstraintsEditionWidget_ConstrainPanel.h
@@ -0,0 +1,65 @@
+﻿// Copyright Epic Games, Inc. All Rights Reserved.
+
+#pragma once
+
+#include "ControlRig.h"
+#include "Editor/Constraints/SConstraintsWidget.h"
+#include "Widgets/DeclarativeSyntaxSupport.h"
+#include "Widgets/SCompoundWidget.h"
+
+class FControlRigEditMode;
+
+namespace UE::ControlRigEditor
+{
+class FRigSelectionViewModel;
+	
+/**
+ * Constraints edition widget that is shown in the constraints window.
+ *
+ * Places a SConstraintsEditionWidget in an SExpandableArea.
+ * The SExpandableArea's header contains a combo button and + button in the header.
+ */
+class SConstraintsEditionWidget_ConstrainPanel : public SCompoundWidget
+{
+public:
+
+	SLATE_BEGIN_ARGS(SConstraintsEditionWidget_ConstrainPanel)
+		: _InitiallyCollapsed(true)
+	{}
+		/** Whether or not the area is initially collapsed */
+		SLATE_ARGUMENT(bool, InitiallyCollapsed)
+	SLATE_END_ARGS()
+
+	void Construct(const FArguments& InArgs, FControlRigEditMode& InEditMode, const TSharedRef<FRigSelectionViewModel>& InSelectionViewModel);
+	virtual ~SConstraintsEditionWidget_ConstrainPanel() override;
+
+	void RefreshConstraintList() const { ConstraintsEditionWidget->RefreshConstraintList(); }
+
+private:
+
+	/** Control rig edit mode we were constructed with */
+	FControlRigEditMode* OwningEditMode = nullptr;
+
+	/** View model for when control rig selection changes. */
+	TSharedPtr<FRigSelectionViewModel> SelectionViewModel;
+	
+	/** The constraints edition widget we wrap. */
+	TSharedPtr<SConstraintsEditionWidget> ConstraintsEditionWidget;
+	
+	TSharedRef<SWidget> CreateSelectionComboButton();
+	TSharedRef<SWidget> CreatePlusButton();
+
+	/**  */
+	void OnSelectShowConstraints(int32 Index);
+	/** Spawns the constraints menu when the plus button is pressed. */
+	FReply HandleAddConstraintClicked();
+	
+	FText GetShowConstraintsName() const;
+	FText GetShowConstraintsTooltip() const;
+
+	/** Refresh the constrain widget when the selected control rig changes. */
+	void OnControlSelected(UControlRig* Subject, FRigControlElement* RigControlElement, bool bIsSelected) const;
+	
+};
+}
+
diff --git a/Engine/Plugins/Animation/ControlRig/Source/ControlRigEditor/Private/EditMode/ControlRigEditMode.cpp b/Engine/Plugins/Animation/ControlRig/Source/ControlRigEditor/Private/EditMode/ControlRigEditMode.cpp
--- a/Engine/Plugins/Animation/ControlRig/Source/ControlRigEditor/Private/EditMode/ControlRigEditMode.cpp
+++ b/Engine/Plugins/Animation/ControlRig/Source/ControlRigEditor/Private/EditMode/ControlRigEditMode.cpp
@@ -85,6 +85,7 @@
 #include "Editor/Sequencer/Private/SSequencer.h"
 #include "Slate/SceneViewport.h"
 #include "EditorGizmos/EditorTransformGizmoUtil.h"
+#include "Misc/RigEditModeUtils.h"
 #include "Sequencer/AnimLayers/AnimLayers.h"
 #include "Tools/MotionTrailOptions.h"
 
@@ -430,7 +431,7 @@ void FControlRigEditMode::SetObjects_Internal()
 
 	if (UsesToolkits() && Toolkit.IsValid())
 	{
-		StaticCastSharedPtr<SControlRigEditModeTools>(Toolkit->GetInlineContent())->SetControlRigs(RuntimeControlRigs);
+		StaticCastSharedPtr<FControlRigEditModeToolkit>(Toolkit)->OnControlsChanged(RuntimeControlRigs);
 	}
 
 	if (!bHasValidRuntimeControlRig)
@@ -3394,30 +3395,19 @@ void FControlRigEditMode::SetAnimLayerPassthroughKey()
 
 void FControlRigEditMode::OpenSpacePickerWidget()
 {
-	TMap<UControlRig*, TArray<FRigElementKey>> SelectedControlRigsAndControls;
-	GetAllSelectedControls(SelectedControlRigsAndControls);
-
-	if (SelectedControlRigsAndControls.Num() < 1)
+	using namespace UE::ControlRigEditor;
+	const FInitialSpacePickerSelection InitialSelection = DetermineInitialSpacePickerSelection(*this);
+	if (!InitialSelection.IsValid())
 	{
 		return;
 	}
-
-	TArray<UControlRig*> ControlRigs;
-	TArray<TArray<FRigElementKey>> AllSelectedControls;
-	SelectedControlRigsAndControls.GenerateKeyArray(ControlRigs);
-	SelectedControlRigsAndControls.GenerateValueArray(AllSelectedControls);
-
-
-	//mz todo handle multiple control rigs with space picker
-	UControlRig* RuntimeRig = ControlRigs[0];
-	TArray<FRigElementKey>& SelectedControls = AllSelectedControls[0];
-
-	URigHierarchy* Hierarchy = RuntimeRig->GetHierarchy();
-
+	const TArray<FRigElementKey>& SelectedControls = InitialSelection.SelectedControls;
+	UControlRig* RuntimeRig = InitialSelection.RuntimeRig;
+	
 	TSharedRef<SRigSpacePickerWidget> PickerWidget =
 	SNew(SRigSpacePickerWidget)
-	.Hierarchy(Hierarchy)
-	.Controls(SelectedControls)
+	.Hierarchy(InitialSelection.Hierarchy)
+	.Controls(InitialSelection.SelectedControls)
 	.Title(LOCTEXT("PickSpace", "Pick Space"))
 	.AllowDelete(AreEditingControlRigDirectly())
 	.AllowReorder(AreEditingControlRigDirectly())
diff --git a/Engine/Plugins/Animation/ControlRig/Source/ControlRigEditor/Private/EditMode/ControlRigEditModeToolkit.cpp b/Engine/Plugins/Animation/ControlRig/Source/ControlRigEditor/Private/EditMode/ControlRigEditModeToolkit.cpp
--- a/Engine/Plugins/Animation/ControlRig/Source/ControlRigEditor/Private/EditMode/ControlRigEditModeToolkit.cpp
+++ b/Engine/Plugins/Animation/ControlRig/Source/ControlRigEditor/Private/EditMode/ControlRigEditModeToolkit.cpp
@@ -26,6 +26,7 @@
 #include "LevelEditor.h"
 #include "PropertyEditorModule.h"
 #include "SLevelViewport.h"
+#include "Models/RigSelectionViewModel.h"
 #include "Sequencer/AnimLayers/SAnimLayers.h"
 #include "Sequencer/AnimLayers/AnimLayers.h"
 #include "Tools/MotionTrailOptions.h"
@@ -50,12 +51,14 @@ bool FControlRigEditModeToolkit::bAnimLayerTabOpen = false;
 bool FControlRigEditModeToolkit::bPoseTabOpen = false;
 bool FControlRigEditModeToolkit::bSnapperTabOpen = false;
 bool FControlRigEditModeToolkit::bTweenOpen = false;
+bool FControlRigEditModeToolkit::bConstraintsOpen = false;
 
 const FName FControlRigEditModeToolkit::PoseTabName = FName(TEXT("PoseTab"));
 const FName FControlRigEditModeToolkit::MotionTrailTabName = FName(TEXT("MotionTrailTab"));
 const FName FControlRigEditModeToolkit::SnapperTabName = FName(TEXT("SnapperTab"));
 const FName FControlRigEditModeToolkit::AnimLayerTabName = FName(TEXT("AnimLayerTab"));
 const FName FControlRigEditModeToolkit::TweenOverlayName = FName(TEXT("TweenOverlay"));
+const FName FControlRigEditModeToolkit::ConstrainingTabName = FName(TEXT("ConstrainingTab"));
 const FName FControlRigEditModeToolkit::OutlinerTabName = FName(TEXT("ControlRigOutlinerTab"));
 const FName FControlRigEditModeToolkit::DetailsTabName = FName(TEXT("ControlRigDetailsTab"));
 const FName FControlRigEditModeToolkit::SpacePickerTabName = FName(TEXT("ControlRigSpacePicker"));
@@ -160,6 +163,10 @@ void FControlRigEditModeToolkit::TryInvokeToolkitUI(const FName InName)
 	{
 		TweenOverlayManager->ToggleVisibility();
 	}
+	else if (InName == ConstrainingTabName)
+	{
+		ConstrainToolsManager->ToggleVisibility();
+	}
 }
 
 bool FControlRigEditModeToolkit::IsToolkitUIActive(const FName InName) const
@@ -169,10 +176,23 @@ bool FControlRigEditModeToolkit::IsToolkitUIActive(const FName InName) const
 		return TweenOverlayManager->IsShowingWidget();
 	}
 
+	if (InName == ConstrainingTabName)
+	{
+		return ConstrainToolsManager->IsShowingWidget();
+	}
+
 	TSharedPtr<FAssetEditorModeUILayer> ModeUILayerPtr = ModeUILayer.Pin();
 	return ModeUILayerPtr->GetTabManager()->FindExistingLiveTab(FTabId(InName)).IsValid();
 }
 
+void FControlRigEditModeToolkit::OnControlsChanged(TConstArrayView<TWeakObjectPtr<UControlRig>> InControlRigs) const
+{
+	if (SelectionViewModel)
+	{
+		SelectionViewModel->RefreshSubscriptions(InControlRigs);
+	}
+}
+
 FText FControlRigEditModeToolkit::GetActiveToolDisplayName() const
 {
 	return ModeTools->GetActiveToolName();
@@ -225,14 +245,6 @@ TSharedRef<SDockTab> SpawnOutlinerTab(const FSpawnTabArgs& Args, FControlRigEdit
 		];
 }
 
-TSharedRef<SDockTab> SpawnSpacePickerTab(const FSpawnTabArgs& Args, FControlRigEditMode* InEditorMode)
-{
-	return SNew(SDockTab)
-		[
-			SNew(SControlRigSpacePicker, *InEditorMode)
-		];
-}
-
 TSharedRef<SDockTab> SpawnDetailsTab(const FSpawnTabArgs& Args, FControlRigEditMode* InEditorMode)
 {
 	UAnimDetailsProxyManager* ProxyManager = InEditorMode ? InEditorMode->GetAnimDetailsProxyManager() : nullptr;
@@ -341,11 +353,21 @@ void FControlRigEditModeToolkit::InvokeUI()
 		// doesn't work as expected todo ModeUILayerPtr->GetTabManager()->TryInvokeTab(UAssetEditorUISubsystem::TopLeftTabID);
 		ModeUILayerPtr->GetTabManager()->TryInvokeTab(UAssetEditorUISubsystem::BottomRightTabID);
 
+		SelectionViewModel = MakeShared<UE::ControlRigEditor::FRigSelectionViewModel>();
+		
 		TweenOverlayManager = MakeUnique<UE::ControlRigEditor::FTweenOverlayManager>(GetToolkitHost(), GetToolkitCommands(), SharedThis(&EditMode));
 		if (bTweenOpen)
 		{
 			TweenOverlayManager->ShowWidget();
 		}
+
+		ConstrainToolsManager = MakeUnique<UE::ControlRigEditor::FConstrainToolsManager>(
+			ModeUILayerPtr->GetTabManager().ToSharedRef(), EditMode, SelectionViewModel.ToSharedRef()
+			);
+		if (bConstraintsOpen)
+		{
+			ConstrainToolsManager->ShowWidget();
+		}
 		
 		if (bAnimLayerTabOpen)
 		{
@@ -398,6 +420,14 @@ void FControlRigEditModeToolkit::UnregisterAndRemoveFloatingTabs()
 			bTweenOpen = TweenOverlayManager->IsShowingWidget();
 			TweenOverlayManager.Reset();
 		}
+
+		if (ConstrainToolsManager)
+		{
+			bConstraintsOpen = ConstrainToolsManager->IsShowingWidget();
+			ConstrainToolsManager.Reset();
+		}
+		
+		SelectionViewModel.Reset();
 		
 		if (ModeUILayer.IsValid())
 		{
diff --git a/Engine/Plugins/Animation/ControlRig/Source/ControlRigEditor/Private/EditMode/ControlRigEditModeToolkit.h b/Engine/Plugins/Animation/ControlRig/Source/ControlRigEditor/Private/EditMode/ControlRigEditModeToolkit.h
--- a/Engine/Plugins/Animation/ControlRig/Source/ControlRigEditor/Private/EditMode/ControlRigEditModeToolkit.h
+++ b/Engine/Plugins/Animation/ControlRig/Source/ControlRigEditor/Private/EditMode/ControlRigEditModeToolkit.h
@@ -8,6 +8,7 @@
 #include "EditorModes.h"
 #include "Toolkits/BaseToolkit.h"
 #include "EditorModeManager.h"
+#include "Constrain/ConstrainToolsManager.h"
 #include "EditMode/SControlRigEditModeTools.h"
 #include "EditMode/ControlRigEditMode.h"
 #include "Tween/ControlRigTweenModel.h"
@@ -59,12 +60,19 @@ public:
 	void TryInvokeToolkitUI(const FName InName);
 	bool IsToolkitUIActive(const FName InName) const;
 
+	/** Called when the global array of control rigs changed the user is working on has changed. */
+	void OnControlsChanged(TConstArrayView<TWeakObjectPtr<UControlRig>> InControlRigs) const;
+
+	/** @return The constraining UI tools, which is a window containing tabs for spaces, constraining, and snapping. */
+	UE::ControlRigEditor::FConstrainToolsManager* GetConstrainToolsManager() const { return ConstrainToolsManager.Get(); }
+
 public:
 	
 	static const FName PoseTabName;
 	static const FName MotionTrailTabName;
 	static const FName AnimLayerTabName;
 	static const FName TweenOverlayName;
+	static const FName ConstrainingTabName;
 	static const FName SnapperTabName;
 	static const FName DetailsTabName;
 	static const FName OutlinerTabName;
@@ -88,12 +96,20 @@ protected:
 	static bool bPoseTabOpen;
 	static bool bSnapperTabOpen;
 	static bool bTweenOpen;
+	static bool bConstraintsOpen;
 	
 private:
 	
 	/** The edit mode we are bound to */
 	FControlRigEditMode& EditMode;
+
+	/** View model for selection changing. Used by UI that depends on selected control rig. */
+	TSharedPtr<UE::ControlRigEditor::FRigSelectionViewModel> SelectionViewModel;
+	
+	/** Manages the tween viewport UI overlay. */
 	TUniquePtr<UE::ControlRigEditor::FTweenOverlayManager> TweenOverlayManager;
+	/** Manages the constraining UI tools, which is a window containing tabs for spaces, constraining, and snapping. */
+	TUniquePtr<UE::ControlRigEditor::FConstrainToolsManager> ConstrainToolsManager;
 
 	/** The tools widget */
 	TSharedPtr<SControlRigEditModeTools> ModeTools;
diff --git a/Engine/Plugins/Animation/ControlRig/Source/ControlRigEditor/Private/EditMode/Models/RigSelectionViewModel.cpp b/Engine/Plugins/Animation/ControlRig/Source/ControlRigEditor/Private/EditMode/Models/RigSelectionViewModel.cpp
new file mode 100644
--- /dev/null
+++ b/Engine/Plugins/Animation/ControlRig/Source/ControlRigEditor/Private/EditMode/Models/RigSelectionViewModel.cpp
@@ -0,0 +1,47 @@
+﻿// Copyright Epic Games, Inc. All Rights Reserved.
+
+#include "RigSelectionViewModel.h"
+
+#include "ControlRig.h"
+
+namespace UE::ControlRigEditor
+{
+FRigSelectionViewModel::~FRigSelectionViewModel()
+{
+	CleanupSubscriptions();
+}
+
+void FRigSelectionViewModel::RefreshSubscriptions(TConstArrayView<TWeakObjectPtr<UControlRig>> InControlRigs)
+{
+	CleanupSubscriptions();
+
+	for (const TWeakObjectPtr<UControlRig>& WeakControlRig : InControlRigs)
+	{
+		if (UControlRig* ControlRig = WeakControlRig.Get())
+		{
+			SubscribedToRigs.AddUnique(ControlRig);
+			ControlRig->ControlSelected().AddRaw(this, &FRigSelectionViewModel::BroadcastOnRigElementSelected);
+		}
+	}
+}
+
+void FRigSelectionViewModel::CleanupSubscriptions()
+{
+	for (const TWeakObjectPtr<UControlRig>& WeakControlRig : SubscribedToRigs)
+	{
+		if (UControlRig* ControlRig = WeakControlRig.Get())
+		{
+			ControlRig->ControlSelected().RemoveAll(this);
+		}
+	}
+		
+	SubscribedToRigs.Reset();
+}
+
+void FRigSelectionViewModel::BroadcastOnRigElementSelected(
+	UControlRig* InSubject, FRigControlElement* InElement, bool bIsSelected
+	)
+{
+	OnControlSelectedDelegate.Broadcast(InSubject, InElement, bIsSelected);
+}
+}
diff --git a/Engine/Plugins/Animation/ControlRig/Source/ControlRigEditor/Private/EditMode/Models/RigSelectionViewModel.h b/Engine/Plugins/Animation/ControlRig/Source/ControlRigEditor/Private/EditMode/Models/RigSelectionViewModel.h
new file mode 100644
--- /dev/null
+++ b/Engine/Plugins/Animation/ControlRig/Source/ControlRigEditor/Private/EditMode/Models/RigSelectionViewModel.h
@@ -0,0 +1,43 @@
+﻿// Copyright Epic Games, Inc. All Rights Reserved.
+
+#pragma once
+
+#include "ControlRig.h"
+#include "Delegates/Delegate.h"
+
+namespace UE::ControlRigEditor
+{
+
+DECLARE_MULTICAST_DELEGATE_ThreeParams(FOnControlSelected, UControlRig* /*Subject*/, FRigControlElement* /*ControlElement*/, bool /*bSelected*/);
+	
+/**
+ * Responsible for invoking an event when control rig selection changes.
+ * The intention is for this to be a view model for UI.
+ */
+class FRigSelectionViewModel : public FNoncopyable
+{
+public:
+
+	~FRigSelectionViewModel();
+
+	/** Called when the global array of control rigs that the user is working on changes. */
+	void RefreshSubscriptions(TConstArrayView<TWeakObjectPtr<UControlRig>> InControlRigs);
+	
+	/** Invoked when a control rig is selected. */
+	FOnControlSelected& OnControlSelected() { return OnControlSelectedDelegate; }
+
+private:
+
+	/** The control rigs we subscribed to in the past. */
+	TArray<TWeakObjectPtr<UControlRig>> SubscribedToRigs;
+
+	/** Invoked when a control rig is selected. */
+	FOnControlSelected OnControlSelectedDelegate;
+
+	/** Unbinds all subscriptions. */
+	void CleanupSubscriptions();
+
+	void BroadcastOnRigElementSelected(UControlRig* InSubject, FRigControlElement* InElement, bool bIsSelected);
+};
+}
+
diff --git a/Engine/Plugins/Animation/ControlRig/Source/ControlRigEditor/Private/EditMode/SControlRigEditModeTools.cpp b/Engine/Plugins/Animation/ControlRig/Source/ControlRigEditor/Private/EditMode/SControlRigEditModeTools.cpp
--- a/Engine/Plugins/Animation/ControlRig/Source/ControlRigEditor/Private/EditMode/SControlRigEditModeTools.cpp
+++ b/Engine/Plugins/Animation/ControlRig/Source/ControlRigEditor/Private/EditMode/SControlRigEditModeTools.cpp
@@ -57,51 +57,13 @@ void SControlRigEditModeTools::Cleanup()
 			SettingsDetailsView->SetKeyframeHandler(nullptr);
 		}
 	}
-	
-	for (TPair<FDelegateHandle, TWeakObjectPtr<UControlRig>>& Handles : HandlesToClear)
-	{
-		if (Handles.Value.IsValid())
-		{
-			Handles.Value->ControlSelected().RemoveAll(this);
-		}
-		if (Handles.Key.IsValid())
-		{
-			Handles.Key.Reset();
-		}
-	}
-	
-	HandlesToClear.Reset();
 }
 
 void SControlRigEditModeTools::SetControlRigs(const TArrayView<TWeakObjectPtr<UControlRig>>& InControlRigs)
 {
-	for (TPair<FDelegateHandle, TWeakObjectPtr<UControlRig>>& Handles : HandlesToClear)
-	{
-		if (Handles.Value.IsValid())
-		{
-			Handles.Value->ControlSelected().RemoveAll(this);
-		}
-		if (Handles.Key.IsValid())
-		{
-			Handles.Key.Reset();
-		}
-	}
-	HandlesToClear.Reset();
 	ControlRigs = InControlRigs;
-	for (TWeakObjectPtr<UControlRig>& InControlRig : InControlRigs)
-	{
-		if (InControlRig.IsValid())
-		{
-			TPair<FDelegateHandle, TWeakObjectPtr<UControlRig>> Handles;
-			Handles.Key = InControlRig->ControlSelected().AddRaw(this, &SControlRigEditModeTools::OnRigElementSelected);
-			Handles.Value = InControlRig;
-			HandlesToClear.Add(Handles);
-		}
-	}
-
 	UpdateOverridesDetailsView();
 
-
 #if USE_LOCAL_DETAILS
 	HierarchyTreeView->RefreshTreeView(true);
 #endif
@@ -296,194 +258,6 @@ void SControlRigEditModeTools::Construct(const FArguments& InArgs, TSharedPtr<FC
 				ControlFloatDetailsView.ToSharedRef()
 			]
 #endif
-
-			+SVerticalBox::Slot()
-			.AutoHeight()
-			[
-				SAssignNew(PickerExpander, SExpandableArea)
-				.InitiallyCollapsed(true)
-				.AreaTitle(LOCTEXT("Picker_SpaceWidget", "Spaces"))
-				.AreaTitleFont(FAppStyle::GetFontStyle("DetailsView.CategoryFontStyle"))
-				.BorderBackgroundColor(FLinearColor(.6f, .6f, .6f))
-				.Padding(FMargin(8.f))
-				.HeaderContent()
-				[
-					SNew(SHorizontalBox)
-
-					+SHorizontalBox::Slot()
-					.AutoWidth()
-					.HAlign(HAlign_Left)
-					.VAlign(VAlign_Center)
-					.Padding(0.f, 0.f, 0.f, 0.f)
-					[
-						SNew(STextBlock)
-						.Text(LOCTEXT("Picker_SpaceWidget", "Spaces"))
-						.Font(FCoreStyle::Get().GetFontStyle("ExpandableArea.TitleFont"))
-					]
-					
-					+SHorizontalBox::Slot()
-					.FillWidth(1.f)
-					[
-						SNew(SSpacer)
-					]
-
-					+SHorizontalBox::Slot()
-					.AutoWidth()
-					.HAlign(HAlign_Right)
-					.VAlign(VAlign_Center)
-					.Padding(0.f, 2.f, 8.f, 2.f)
-					[
-						SNew(SButton)
-						.ContentPadding(0.0f)
-						.ButtonStyle(FAppStyle::Get(), "NoBorder")
-						.OnClicked(this, &SControlRigEditModeTools::HandleAddSpaceClicked)
-						.Cursor(EMouseCursor::Default)
-						.ToolTipText(LOCTEXT("AddSpace", "Add Space"))
-						[
-							SNew(SImage)
-							.Image(FAppStyle::GetBrush(TEXT("Icons.PlusCircle")))
-						]
-						.Visibility_Lambda([this]()->EVisibility
-						{
-							return GetAddSpaceButtonVisibility();
-						})
-					]
-				]
-				.BodyContent()
-				[
-					SAssignNew(SpacePickerWidget, SRigSpacePickerWidget)
-					.AllowDelete(true)
-					.AllowReorder(true)
-					.AllowAdd(false)
-					.ShowBakeAndCompensateButton(true)
-					.GetControlCustomization(this, &SControlRigEditModeTools::HandleGetControlElementCustomization)
-					.OnActiveSpaceChanged(this, &SControlRigEditModeTools::HandleActiveSpaceChanged)
-					.OnSpaceListChanged(this, &SControlRigEditModeTools::HandleSpaceListChanged)
-					.OnCompensateKeyButtonClicked(this, &SControlRigEditModeTools::OnCompensateKeyClicked)
-					.OnCompensateAllButtonClicked(this, &SControlRigEditModeTools::OnCompensateAllClicked)
-					.OnBakeButtonClicked(this, &SControlRigEditModeTools::OnBakeControlsToNewSpaceButtonClicked)					// todo: implement GetAdditionalSpacesDelegate to pull spaces from sequencer
-				]
-			]
-
-			+SVerticalBox::Slot()
-			.AutoHeight()
-			[
-				SAssignNew(ConstraintPickerExpander, SExpandableArea)
-				.InitiallyCollapsed(true)
-				.AreaTitle(LOCTEXT("ConstraintsWidget", "Constraints"))
-				.AreaTitleFont(FAppStyle::GetFontStyle("DetailsView.CategoryFontStyle"))
-				.BorderBackgroundColor(FLinearColor(.6f, .6f, .6f))
-				.HeaderContent()
-				[
-					SNew(SHorizontalBox)
-
-					// "Constraints" label
-					+SHorizontalBox::Slot()
-					.AutoWidth()
-					.HAlign(HAlign_Left)
-					.VAlign(VAlign_Center)
-					.Padding(0.f, 0.f, 0.f, 0.f)
-					[
-						SNew(STextBlock)
-						.Text(LOCTEXT("ConstraintsWidget", "Constraints"))
-						.Font(FCoreStyle::Get().GetFontStyle("ExpandableArea.TitleFont"))
-					]
-
-					// Spacer
-					+SHorizontalBox::Slot()
-					.FillWidth(1.f)
-					[
-						SNew(SSpacer)
-					]
-					// "Selected" button
-					+SHorizontalBox::Slot()
-					.AutoWidth()
-					.HAlign(HAlign_Right)
-					.VAlign(VAlign_Center)
-					.Padding(0.f, 2.f, 8.f, 2.f)
-					[
-						// Combo Button to swap what constraints we show  
-						SNew(SComboButton)
-						.OnGetMenuContent_Lambda([this]()
-						{
-							FMenuBuilder MenuBuilder(true, NULL);
-							MenuBuilder.BeginSection("Constraints");
-							if (ConstraintsEditionWidget.IsValid())
-							{
-								for (int32 Index = 0; Index < 4; ++Index)
-								{
-									FUIAction ItemAction(FExecuteAction::CreateSP(this, &SControlRigEditModeTools::OnSelectShowConstraints, Index));
-									const TAttribute<FText> Text = ConstraintsEditionWidget->GetShowConstraintsText((FBaseConstraintListWidget::EShowConstraints)(Index));
-									const TAttribute<FText> Tooltip = ConstraintsEditionWidget->GetShowConstraintsTooltip((FBaseConstraintListWidget::EShowConstraints)(Index));
-									MenuBuilder.AddMenuEntry(Text, Tooltip, FSlateIcon(), ItemAction);
-								}
-							}
-							MenuBuilder.EndSection();
-
-							return MenuBuilder.MakeWidget();
-						})
-						.ButtonContent()
-						[
-							SNew(SHorizontalBox)
-			
-							+SHorizontalBox::Slot()
-							[
-								SNew(STextBlock)
-								.Text_Lambda([this]()
-									{
-										return GetShowConstraintsName();
-									})
-								.ToolTipText_Lambda([this]()
-								{
-									return GetShowConstraintsTooltip();
-								})
-							]
-						]
-					]
-					// "Plus" icon
-					+SHorizontalBox::Slot()
-					.AutoWidth()
-					.HAlign(HAlign_Right)
-					.VAlign(VAlign_Center)
-					.Padding(0.f, 2.f, 8.f, 2.f)
-					[
-						SNew(SButton)
-						.ContentPadding(0.0f)
-						.ButtonStyle(FAppStyle::Get(), "NoBorder")
-						.IsEnabled_Lambda([this]()
-						{
-							TArray<AActor*> SelectedActors;
-							if (FControlRigEditMode* EditMode = static_cast<FControlRigEditMode*>(ModeTools->GetActiveMode(FControlRigEditMode::ModeName)))
-							{
-								const ULevel* CurrentLevel = EditMode->GetWorld()->GetCurrentLevel();
-								SelectedActors = ObjectPtrDecay(CurrentLevel->Actors).FilterByPredicate([](const AActor* Actor)
-								{
-									return Actor && Actor->IsSelected();
-								});
-							}
-							return !SelectedActors.IsEmpty();
-						})
-						.OnClicked(this, &SControlRigEditModeTools::HandleAddConstraintClicked)
-						.Cursor(EMouseCursor::Default)
-						.ToolTipText(LOCTEXT("AddConstraint", "Add Constraint"))
-						[
-							SNew(SImage)
-							.Image(FAppStyle::GetBrush(TEXT("Icons.PlusCircle")))
-						]
-					]
-				]
-				.BodyContent()
-				[
-					SAssignNew(ConstraintsEditionWidget, SConstraintsEditionWidget)
-				]
-				.OnAreaExpansionChanged_Lambda( [this](bool bIsExpanded)
-				{
-					if (ConstraintsEditionWidget)
-					{
-						ConstraintsEditionWidget->RefreshConstraintList();
-					}
-				})
-			]
 			
 			+SVerticalBox::Slot()
 			.AutoHeight()
@@ -553,35 +327,6 @@ void SControlRigEditModeTools::SetSettingsDetailsObject(const TWeakObjectPtr<>&
 	}
 }
 
-void SControlRigEditModeTools::OnSelectShowConstraints(int32 Index)
-{
-	if (ConstraintsEditionWidget.IsValid())
-	{
-		SConstraintsEditionWidget::EShowConstraints ShowConstraint = (SConstraintsEditionWidget::EShowConstraints)(Index);
-		ConstraintsEditionWidget->SetShowConstraints(ShowConstraint);
-	}
-}
-
-FText SControlRigEditModeTools::GetShowConstraintsName() const
-{
-	FText Text = FText::GetEmpty();
-	if (ConstraintsEditionWidget.IsValid())
-	{
-		Text = ConstraintsEditionWidget->GetShowConstraintsText(FBaseConstraintListWidget::ShowConstraints);
-	}
-	return Text;
-}
-
-FText SControlRigEditModeTools::GetShowConstraintsTooltip() const
-{
-	FText Text = FText::GetEmpty();
-	if (ConstraintsEditionWidget.IsValid())
-	{
-		Text = ConstraintsEditionWidget->GetShowConstraintsTooltip(FBaseConstraintListWidget::ShowConstraints);
-	}
-	return Text;
-}
-
 #if USE_LOCAL_DETAILS
 
 void SControlRigEditModeTools::SetEulerTransformDetailsObjects(const TArray<TWeakObjectPtr<>>& InObjects)
@@ -658,10 +403,6 @@ void SControlRigEditModeTools::SetVector2DDetailsObjects(const TArray<TWeakObjec
 void SControlRigEditModeTools::SetSequencer(TWeakPtr<ISequencer> InSequencer)
 {
 	WeakSequencer = InSequencer.Pin();
-	if (ConstraintsEditionWidget)
-	{
-		ConstraintsEditionWidget->SequencerChanged(InSequencer);
-	}
 	UpdateOverridesDetailsView();
 }
 
@@ -868,373 +609,6 @@ void SControlRigEditModeTools::HandleSelectionChanged(TSharedPtr<FRigTreeElement
 }
 #endif
 
-void SControlRigEditModeTools::OnRigElementSelected(UControlRig* Subject, FRigControlElement* ControlElement, bool bSelected)
-{
-#if USE_LOCAL_DETAILS
-	const FRigElementKey Key = ControlElement->GetKey();
-	for (int32 RootIndex = 0; RootIndex < HierarchyTreeView->GetRootElements().Num(); ++RootIndex)
-	{
-		TSharedPtr<FRigTreeElement> Found = HierarchyTreeView->FindElement(Key, HierarchyTreeView->GetRootElements()[RootIndex]);
-		if (Found.IsValid())
-		{
-			HierarchyTreeView->SetItemSelection(Found, bSelected, ESelectInfo::OnNavigation);
-			
-			TArray<TSharedPtr<FRigTreeElement>> SelectedItems = HierarchyTreeView->GetSelectedItems();
-			for (TSharedPtr<FRigTreeElement> SelectedItem : SelectedItems)
-			{
-				HierarchyTreeView->SetExpansionRecursive(SelectedItem, false, true);
-			}
-
-			if (SelectedItems.Num() > 0)
-			{
-				HierarchyTreeView->RequestScrollIntoView(SelectedItems.Last());
-			}
-		}
-	}
-#endif
-
-	if (Subject && Subject->GetHierarchy())
-	{
-		// get the selected controls
-		TArray<FRigElementKey> SelectedControls = Subject->GetHierarchy()->GetSelectedKeys(ERigElementType::Control);
-		if (SpacePickerWidget)
-		{
-			SpacePickerWidget->SetControls(Subject->GetHierarchy(), SelectedControls);
-		}
-		if (ConstraintsEditionWidget)
-		{
-			ConstraintsEditionWidget->InvalidateConstraintList();
-		}
-	}
-}
-
-const FRigControlElementCustomization* SControlRigEditModeTools::HandleGetControlElementCustomization(URigHierarchy* InHierarchy, const FRigElementKey& InControlKey)
-{
-	UControlRig* Rig = ControlRigs.Num() > 0 ? ControlRigs[0].Get() : nullptr;
-	for (TWeakObjectPtr<UControlRig>& ControlRig : ControlRigs)
-	{
-		if (ControlRig.IsValid() && ControlRig->GetHierarchy() == InHierarchy)
-		{
-			return ControlRig->GetControlCustomization(InControlKey);
-		}
-	}
-	return nullptr;
-}
-
-void SControlRigEditModeTools::HandleActiveSpaceChanged(URigHierarchy* InHierarchy, const FRigElementKey& InControlKey,
-	const FRigElementKey& InSpaceKey)
-{
-
-	if (WeakSequencer.IsValid())
-	{
-		for (TWeakObjectPtr<UControlRig>& ControlRig : ControlRigs)
-		{
-			if (ControlRig.IsValid() && ControlRig->GetHierarchy() == InHierarchy)
-			{
-				FString FailureReason;
-				FRigDependenciesProviderForControlRig DependencyProvider(ControlRig.Get());
-				if(!InHierarchy->CanSwitchToParent(InControlKey, InSpaceKey, DependencyProvider, &FailureReason))
-				{
-					// notification
-					FNotificationInfo Info(FText::FromString(FailureReason));
-					Info.bFireAndForget = true;
-					Info.FadeOutDuration = 2.0f;
-					Info.ExpireDuration = 8.0f;
-
-					const TSharedPtr<SNotificationItem> NotificationPtr = FSlateNotificationManager::Get().AddNotification(Info);
-					NotificationPtr->SetCompletionState(SNotificationItem::CS_Fail);
-					return;
-				}
-			
-				if (const FRigControlElement* ControlElement = InHierarchy->Find<FRigControlElement>(InControlKey))
-				{
-					ISequencer* Sequencer = WeakSequencer.Pin().Get();
-					if (Sequencer)
-					{
-						constexpr bool bDoNotCreate = false;
-						FSpaceChannelAndSection SpaceChannelAndSection = FControlRigSpaceChannelHelpers::FindSpaceChannelAndSectionForControl(ControlRig.Get(), InControlKey.Name, Sequencer, bDoNotCreate);
-						if (!SpaceChannelAndSection.SpaceChannel)
-						{
-							// no need to create a space channel if InSpaceKey is the current parent 
-							if (InHierarchy->GetActiveParent(InControlKey) == InSpaceKey)
-							{
-								return;
-							}
-						}
-						
-						FScopedTransaction Transaction(LOCTEXT("KeyControlRigSpace", "Key Control Rig Space"));
-						ControlRig->Modify();
-
-						SpaceChannelAndSection = FControlRigSpaceChannelHelpers::FindSpaceChannelAndSectionForControl(ControlRig.Get(), InControlKey.Name, Sequencer, true /*bCreateIfNeeded*/);
-						if (SpaceChannelAndSection.SpaceChannel)
-						{
-							const FFrameRate TickResolution = Sequencer->GetFocusedTickResolution();
-							const FFrameTime FrameTime = Sequencer->GetLocalTime().ConvertTo(TickResolution);
-							FFrameNumber CurrentTime = FrameTime.GetFrame();
-							FControlRigSpaceChannelHelpers::SequencerKeyControlRigSpaceChannel(ControlRig.Get(), Sequencer, SpaceChannelAndSection.SpaceChannel, SpaceChannelAndSection.SectionToKey, CurrentTime, InHierarchy, InControlKey, InSpaceKey);
-						}
-					}
-				}
-			}
-		}
-	}
-}
-
-void SControlRigEditModeTools::HandleSpaceListChanged(URigHierarchy* InHierarchy, const FRigElementKey& InControlKey,
-	const TArray<FRigElementKeyWithLabel>& InSpaceList)
-{
-	FScopedTransaction Transaction(LOCTEXT("ChangeControlRigSpace", "Change Control Rig Space"));
-
-	for (TWeakObjectPtr<UControlRig>& ControlRig : ControlRigs)
-	{
-		if (ControlRig.IsValid() && ControlRig->GetHierarchy() == InHierarchy)
-		{
-			ControlRig->Modify();
-
-			if (const FRigControlElement* ControlElement = InHierarchy->Find<FRigControlElement>(InControlKey))
-			{
-				FRigControlElementCustomization ControlCustomization = *ControlRig->GetControlCustomization(InControlKey);
-				ControlCustomization.AvailableSpaces = InSpaceList;
-				ControlCustomization.RemovedSpaces.Reset();
-
-				// remember  the elements which are in the asset's available list but removed by the user
-				for (const FRigElementKeyWithLabel& AvailableSpace : ControlElement->Settings.Customization.AvailableSpaces)
-				{
-					if (ControlCustomization.AvailableSpaces.FindByKey(AvailableSpace.Key) == nullptr)
-					{
-						ControlCustomization.RemovedSpaces.Add(AvailableSpace.Key);
-					}
-				}
-
-				ControlRig->SetControlCustomization(InControlKey, ControlCustomization);
-
-				if (FControlRigEditMode* EditMode = static_cast<FControlRigEditMode*>(ModeTools->GetActiveMode(FControlRigEditMode::ModeName)))
-				{
-					const TGuardValue<bool> SuspendGuard(EditMode->bSuspendHierarchyNotifs, true);
-					InHierarchy->Notify(ERigHierarchyNotification::ControlSettingChanged, ControlElement);
-				}
-				else
-				{
-					InHierarchy->Notify(ERigHierarchyNotification::ControlSettingChanged, ControlElement);
-				}
-
-				SpacePickerWidget->RefreshContents();
-			}
-		}
-	}
-}
-
-FReply SControlRigEditModeTools::HandleAddSpaceClicked()
-{
-	return SpacePickerWidget->HandleAddElementClicked();
-}
-
-EVisibility SControlRigEditModeTools::GetAddSpaceButtonVisibility() const
-{
-	return IsSpaceSwitchingRestricted() ? EVisibility::Hidden : EVisibility::Visible;
-}
-
-bool SControlRigEditModeTools::IsSpaceSwitchingRestricted() const
-{
-	return SpacePickerWidget->IsRestricted();
-}
-
-bool SControlRigEditModeTools::ReadyForBakeOrCompensation() const
-{
-	if (SpacePickerWidget->GetHierarchy() == nullptr)
-	{
-		return false;
-	}
-	if (SpacePickerWidget->GetControls().Num() == 0)
-	{
-		return false;
-	}
-
-	bool bNoValidControlRig = true;
-	for (const TWeakObjectPtr<UControlRig>& ControlRig : ControlRigs)
-	{
-		if (ControlRig.IsValid() && SpacePickerWidget->GetHierarchy() == ControlRig->GetHierarchy())
-		{
-			bNoValidControlRig = false;
-			break;
-		}
-	}
-
-	if (bNoValidControlRig)
-	{
-		return false;
-	}
-	ISequencer* Sequencer = WeakSequencer.Pin().Get();
-	if (Sequencer == nullptr || Sequencer->GetFocusedMovieSceneSequence() == nullptr || Sequencer->GetFocusedMovieSceneSequence()->GetMovieScene() == nullptr)
-	{
-		return false;
-	}
-	return true;
-}
-
-FReply SControlRigEditModeTools::OnCompensateKeyClicked()
-{
-	if (ReadyForBakeOrCompensation() == false)
-	{
-		return FReply::Unhandled();
-	}
-	ISequencer* Sequencer = WeakSequencer.Pin().Get();
-	const FFrameRate TickResolution = Sequencer->GetFocusedTickResolution();
-	const FFrameTime FrameTime = Sequencer->GetLocalTime().ConvertTo(TickResolution);
-	const TOptional<FFrameNumber> OptionalKeyTime = FrameTime.GetFrame();
-	const bool bSetPreviousKey = true;
-	Compensate(OptionalKeyTime, bSetPreviousKey);
-	return FReply::Handled();
-}
-
-FReply SControlRigEditModeTools::OnCompensateAllClicked()
-{
-	if (ReadyForBakeOrCompensation() == false)
-	{
-		return FReply::Unhandled();
-	}
-	const TOptional<FFrameNumber> OptionalKeyTime;
-	ISequencer* Sequencer = WeakSequencer.Pin().Get();
-	const bool bSetPreviousKey = true;
-	Compensate(OptionalKeyTime, bSetPreviousKey);
-	return FReply::Handled();
-}
-
-void SControlRigEditModeTools::Compensate(TOptional<FFrameNumber> OptionalKeyTime, bool bSetPreviousTick)
-{
-	if (ReadyForBakeOrCompensation() == false)
-	{
-		return;
-	}
-	ISequencer* Sequencer = WeakSequencer.Pin().Get();
-	for (TWeakObjectPtr<UControlRig>& ControlRig : ControlRigs)
-	{
-		if (ControlRig.IsValid() && SpacePickerWidget->GetHierarchy() == ControlRig->GetHierarchy())
-		{
-			// compensate spaces
-			if(UMovieSceneControlRigParameterSection* CRSection = FControlRigSpaceChannelHelpers::GetControlRigSection(Sequencer, ControlRig.Get()))
-			{	
-				// compensate spaces
-				FControlRigSpaceChannelHelpers::CompensateIfNeeded(
-					ControlRig.Get(), Sequencer, CRSection,
-				OptionalKeyTime, bSetPreviousTick);
-			}
-		}
-	}
-}
-
-FReply SControlRigEditModeTools::OnBakeControlsToNewSpaceButtonClicked()
-{
-	using namespace UE::MovieScene;
-	
-	if (ReadyForBakeOrCompensation() == false)
-	{
-		return FReply::Unhandled();
-	}
-	ISequencer* Sequencer = WeakSequencer.Pin().Get();
-
-	for (TWeakObjectPtr<UControlRig>& ControlRig : ControlRigs)
-	{
-		if (ControlRig.IsValid() && SpacePickerWidget->GetHierarchy() == ControlRig->GetHierarchy())
-		{
-
-			TArray<FRigElementKey> ControlKeys = SpacePickerWidget->GetControls();
-			FSpaceChannelAndSection SpaceChannelAndSection = FControlRigSpaceChannelHelpers::FindSpaceChannelAndSectionForControl(ControlRig.Get(), ControlKeys[0].Name, Sequencer, false /*bCreateIfNeeded*/);
-			if (SpaceChannelAndSection.SpaceChannel != nullptr)
-			{
-				//set up settings if not setup
-				if (BakeSpaceSettings.TargetSpace == FRigElementKey())
-				{
-					BakeSpaceSettings.TargetSpace = URigHierarchy::GetDefaultParentKey();
-					TRange<FFrameNumber> Range = Sequencer->GetFocusedMovieSceneSequence()->GetMovieScene()->GetPlaybackRange();
-					BakeSpaceSettings.Settings.StartFrame = Range.GetLowerBoundValue();
-					BakeSpaceSettings.Settings.EndFrame = Range.GetUpperBoundValue();
-				}
-
-				TSharedRef<SRigSpacePickerBakeWidget> BakeWidget =
-					SNew(SRigSpacePickerBakeWidget)
-					.Settings(BakeSpaceSettings)
-					.Hierarchy(SpacePickerWidget->GetHierarchy())
-					.Controls(ControlKeys) // use the cached controls here since the selection is not recovered until next tick.
-					.Sequencer(Sequencer)
-					.GetControlCustomization(this, &SControlRigEditModeTools::HandleGetControlElementCustomization)
-					.OnBake_Lambda([Sequencer, ControlRig](URigHierarchy* InHierarchy, TArray<FRigElementKey> InControls, const FRigSpacePickerBakeSettings& InSettings)
-						{		
-							FScopedTransaction Transaction(LOCTEXT("BakeControlToSpace", "Bake Control In Space"));
-							for (const FRigElementKey& ControlKey : InControls)
-							{
-								//when baking we will now create a channel if one doesn't exist, was causing confusion
-								FSpaceChannelAndSection SpaceChannelAndSection = FControlRigSpaceChannelHelpers::FindSpaceChannelAndSectionForControl(ControlRig.Get(), ControlKey.Name, Sequencer, true /*bCreateIfNeeded*/);
-								if (SpaceChannelAndSection.SpaceChannel)
-								{
-									FControlRigSpaceChannelHelpers::SequencerBakeControlInSpace(ControlRig.Get(), Sequencer, SpaceChannelAndSection.SpaceChannel, SpaceChannelAndSection.SectionToKey,
-										InHierarchy, ControlKey, InSettings);
-								}
-								SControlRigEditModeTools::BakeSpaceSettings = InSettings;
-							}
-							return FReply::Handled();
-						});
-
-				return BakeWidget->OpenDialog(true);
-			}
-			break; //mz todo need baketo handle more than one
-		}
-
-	}
-	return FReply::Unhandled();
-}
-
-FReply SControlRigEditModeTools::HandleAddConstraintClicked()
-{
-	FMenuBuilder MenuBuilder(true, nullptr);
-
-	auto AddConstraintWidget = [&](ETransformConstraintType InConstraintType)
-	{
-		const TSharedRef<SConstraintMenuEntry> Entry =
-			SNew(SConstraintMenuEntry, InConstraintType)
-			.OnConstraintCreated_Lambda([this]()
-			{
-				// magic number to auto expand the widget when creating a new constraint. We keep that number below a reasonable
-				// threshold to avoid automatically creating a large number of items (this can be style done by the user) 
-				static constexpr int32 NumAutoExpand = 20;
-				const int32 NumItems = ConstraintsEditionWidget ? ConstraintsEditionWidget->RefreshConstraintList() : 0;	
-				if (ConstraintPickerExpander && NumItems < NumAutoExpand)
-				{
-					ConstraintPickerExpander->SetExpanded(true);
-				}
-				if (ConstraintsEditionWidget)
-				{
-					ConstraintsEditionWidget->ResetParent(); 
-				}
-			})
-			.OnGetParent_Lambda([this]()
-			{
-				static const FConstrainable InvalidParent;
-				return ConstraintsEditionWidget ? ConstraintsEditionWidget->GetParent() : InvalidParent;
-			});
-			MenuBuilder.AddWidget(Entry, FText::GetEmpty(), true);
-	};
-	
-	MenuBuilder.BeginSection("CreateConstraint", LOCTEXT("CreateConstraintHeader", "Create New..."));
-	{
-		AddConstraintWidget(ETransformConstraintType::Translation);
-		AddConstraintWidget(ETransformConstraintType::Rotation);
-		AddConstraintWidget(ETransformConstraintType::Scale);
-		AddConstraintWidget(ETransformConstraintType::Parent);
-		AddConstraintWidget(ETransformConstraintType::LookAt);
-	}
-	MenuBuilder.EndSection();
-	
-	FSlateApplication::Get().PushMenu(
-		AsShared(),
-		FWidgetPath(),
-		MenuBuilder.MakeWidget(),
-		FSlateApplication::Get().GetCursorPos(),
-		FPopupTransitionEffect(FPopupTransitionEffect::ContextMenu));
-	
-	return FReply::Handled();
-}
-
 void SControlRigEditModeTools::CustomizeToolBarPalette(FToolBarBuilder& ToolBarBuilder)
 {
 	//TOGGLE SELECTED RIG CONTROLS
@@ -1291,15 +665,18 @@ void SControlRigEditModeTools::CustomizeToolBarPalette(FToolBarBuilder& ToolBarB
 		FSlateIcon(TEXT("ControlRigEditorStyle"), TEXT("ControlRig.TweenTool")),
 		EUserInterfaceActionType::ToggleButton
 	);
-
-	// Snap
+	
 	ToolBarBuilder.AddToolBarButton(
-		FExecuteAction::CreateRaw(OwningToolkit.Pin().Get(), &FControlRigEditModeToolkit::TryInvokeToolkitUI, FControlRigEditModeToolkit::SnapperTabName),
+		FUIAction(
+			FExecuteAction::CreateRaw(OwningToolkit.Pin().Get(), &FControlRigEditModeToolkit::TryInvokeToolkitUI, FControlRigEditModeToolkit::ConstrainingTabName),
+			FCanExecuteAction(),
+			FIsActionChecked::CreateRaw(OwningToolkit.Pin().Get(), &FControlRigEditModeToolkit::IsToolkitUIActive, FControlRigEditModeToolkit::ConstrainingTabName)
+		),
 		NAME_None,
-		LOCTEXT("Snapper", "Snapper"),
-		LOCTEXT("SnapperTooltip", "Snap child objects to a parent object over a set of frames"),
+		LOCTEXT("Constrain", "Constrain"),
+		LOCTEXT("SnapperTooltip", "Tools for constraining, such as snapping, space picker, and constraints edition."),
 		FSlateIcon(TEXT("ControlRigEditorStyle"), TEXT("ControlRig.SnapperTool")),
-		EUserInterfaceActionType::Button
+		EUserInterfaceActionType::ToggleButton
 	);
 
 	/*  leaving this around for a bit to get animator feedback on the new menu system
diff --git a/Engine/Plugins/Animation/ControlRig/Source/ControlRigEditor/Private/EditMode/SControlRigEditModeTools.h b/Engine/Plugins/Animation/ControlRig/Source/ControlRigEditor/Private/EditMode/SControlRigEditModeTools.h
--- a/Engine/Plugins/Animation/ControlRig/Source/ControlRigEditor/Private/EditMode/SControlRigEditModeTools.h
+++ b/Engine/Plugins/Animation/ControlRig/Source/ControlRigEditor/Private/EditMode/SControlRigEditModeTools.h
@@ -51,9 +51,9 @@ public:
 	/** Set the sequencer we are bound to */
 	void SetSequencer(TWeakPtr<ISequencer> InSequencer);
 
-	/** Set The Control Rig we are using*/
+	/** Set The Control Rig we are using */
 	void SetControlRigs(const TArrayView<TWeakObjectPtr<UControlRig>>& InControlRigs);
-
+	
 	/** Returns the hierarchy currently being used */
 	const URigHierarchy* GetHierarchy() const;
 	// IDetailKeyframeHandler interface
@@ -85,44 +85,16 @@ private:
 	TSharedPtr<IDetailsView> ControlVectorDetailsView;
 	TSharedPtr<SRigHierarchyTreeView> HierarchyTreeView;
 #endif
-
-	/** Space Picker controls*/
-	TSharedPtr<SExpandableArea> PickerExpander;
-	TSharedPtr<SRigSpacePickerWidget> SpacePickerWidget;
-
+	
 	/** Storage for control rigs */
 	TArray<TWeakObjectPtr<UControlRig>> ControlRigs;
 
-	/** Constraint edition widget. */
-	TSharedPtr<SExpandableArea> ConstraintPickerExpander = nullptr;
-	TSharedPtr<SConstraintsEditionWidget> ConstraintsEditionWidget = nullptr; 
-
 	TSharedPtr<IDetailsView> OverridesDetailsView;
 
 	/** Display or edit set up for property */
 	static bool ShouldShowPropertyOnDetailCustomization(const struct FPropertyAndParent& InPropertyAndParent);
 	static bool IsReadOnlyPropertyOnDetailCustomization(const struct FPropertyAndParent& InPropertyAndParent);
 
-	void OnRigElementSelected(UControlRig* Subject, FRigControlElement* ControlElement, bool bSelected);
-
-	const FRigControlElementCustomization* HandleGetControlElementCustomization(URigHierarchy* InHierarchy, const FRigElementKey& InControlKey);
-	void HandleActiveSpaceChanged(URigHierarchy* InHierarchy, const FRigElementKey& InControlKey, const FRigElementKey& InSpaceKey);
-	void HandleSpaceListChanged(URigHierarchy* InHierarchy, const FRigElementKey& InControlKey, const TArray<FRigElementKeyWithLabel>& InSpaceList);
-	FReply HandleAddSpaceClicked();
-	EVisibility GetAddSpaceButtonVisibility() const;
-	bool IsSpaceSwitchingRestricted() const;
-	FReply OnBakeControlsToNewSpaceButtonClicked();
-	FReply OnCompensateKeyClicked();
-	FReply OnCompensateAllClicked();
-	void Compensate(TOptional<FFrameNumber> OptionalKeyTime, bool bSetPreviousTick);
-	bool ReadyForBakeOrCompensation() const;
-	FReply HandleAddConstraintClicked();
-	
-	/** constraint type to show selection*/
-	void OnSelectShowConstraints(int32 Index);
-	FText GetShowConstraintsName() const;
-	FText GetShowConstraintsTooltip() const;
-
 	EVisibility GetOverridesExpanderVisibility() const;
 	void OnOverrideOptionFinishedChange(const FPropertyChangedEvent& PropertyChangedEvent);
 	bool ShouldShowOverrideProperty(const FPropertyAndParent& InPropertyAndParent) const;
@@ -143,9 +115,6 @@ private:
 	// The toolkit that created this UI
 	TWeakPtr<FControlRigEditModeToolkit> OwningToolkit;
 
-	//array of handles to clear when getting new control rigs
-	TArray<TPair<FDelegateHandle, TWeakObjectPtr<UControlRig>>> HandlesToClear;
-
 public:
 	/** Modes Panel Header Information **/
 	void CustomizeToolBarPalette(FToolBarBuilder& ToolBarBuilder);
diff --git a/Engine/Plugins/Animation/ControlRig/Source/ControlRigEditor/Private/EditMode/SControlRigSpacePicker.cpp b/Engine/Plugins/Animation/ControlRig/Source/ControlRigEditor/Private/EditMode/SControlRigSpacePicker.cpp
--- a/Engine/Plugins/Animation/ControlRig/Source/ControlRigEditor/Private/EditMode/SControlRigSpacePicker.cpp
+++ b/Engine/Plugins/Animation/ControlRig/Source/ControlRigEditor/Private/EditMode/SControlRigSpacePicker.cpp
@@ -18,96 +18,103 @@
 #include "MovieSceneSequence.h"
 #include "MovieScene.h"
 #include "ControlRigSpaceChannelEditors.h"
+#include "Models/RigSelectionViewModel.h"
 
 #define LOCTEXT_NAMESPACE "ControlRigSpacePicker"
 
-void SControlRigSpacePicker::Construct(const FArguments& InArgs, FControlRigEditMode& InEditMode)
+void SControlRigSpacePicker::Construct(
+	const FArguments& InArgs, FControlRigEditMode& InEditMode, const TSharedRef<UE::ControlRigEditor::FRigSelectionViewModel>& InSelectionViewModel
+	)
 {
+	SelectionViewModel = InSelectionViewModel;
+	SelectionViewModel->OnControlSelected().AddSP(this, &SControlRigSpacePicker::HandleControlSelected);
 	
 	ChildSlot
+	[
+		SNew(SScrollBox)
+		+ SScrollBox::Slot()
 		[
-			SNew(SScrollBox)
-			+ SScrollBox::Slot()
+			SNew(SVerticalBox)
+			+SVerticalBox::Slot()
+			.AutoHeight()
 			[
-				SNew(SVerticalBox)
-				+SVerticalBox::Slot()
-				.AutoHeight()
+				SAssignNew(PickerExpander, SExpandableArea)
+				.InitiallyCollapsed(InArgs._InitiallyCollapsed)
+				.AreaTitle(LOCTEXT("Picker_SpaceWidget", "Spaces"))
+				.AreaTitleFont(FAppStyle::GetFontStyle("DetailsView.CategoryFontStyle"))
+				.BorderBackgroundColor(FLinearColor(.6f, .6f, .6f))
+				.Padding(FMargin(8.f))
+				.HeaderContent()
 				[
-					SAssignNew(PickerExpander, SExpandableArea)
-					.InitiallyCollapsed(true)
-					.AreaTitle(LOCTEXT("Picker_SpaceWidget", "Spaces"))
-					.AreaTitleFont(FAppStyle::GetFontStyle("DetailsView.CategoryFontStyle"))
-					.BorderBackgroundColor(FLinearColor(.6f, .6f, .6f))
-					.Padding(FMargin(8.f))
-					.HeaderContent()
-					[
-						SNew(SHorizontalBox)
+					SNew(SHorizontalBox)
 
-						+SHorizontalBox::Slot()
-						.AutoWidth()
-						.HAlign(HAlign_Left)
-						.VAlign(VAlign_Center)
-						.Padding(0.f, 0.f, 0.f, 0.f)
-						[
-							SNew(STextBlock)
-							.Text(LOCTEXT("Picker_SpaceWidget", "Spaces"))
-							.Font(FCoreStyle::Get().GetFontStyle("ExpandableArea.TitleFont"))
-						]
-					
-						+SHorizontalBox::Slot()
-						.FillWidth(1.f)
-						[
-							SNew(SSpacer)
-						]
+					+SHorizontalBox::Slot()
+					.AutoWidth()
+					.HAlign(HAlign_Left)
+					.VAlign(VAlign_Center)
+					.Padding(0.f, 0.f, 0.f, 0.f)
+					[
+						SNew(STextBlock)
+						.Text(LOCTEXT("Picker_SpaceWidget", "Spaces"))
+						.Font(FCoreStyle::Get().GetFontStyle("ExpandableArea.TitleFont"))
+					]
+				
+					+SHorizontalBox::Slot()
+					.FillWidth(1.f)
+					[
+						SNew(SSpacer)
+					]
 
-						+SHorizontalBox::Slot()
-						.AutoWidth()
-						.HAlign(HAlign_Right)
-						.VAlign(VAlign_Center)
-						.Padding(0.f, 2.f, 8.f, 2.f)
+					+SHorizontalBox::Slot()
+					.AutoWidth()
+					.HAlign(HAlign_Right)
+					.VAlign(VAlign_Center)
+					.Padding(0.f, 2.f, 8.f, 2.f)
+					[
+						SNew(SButton)
+						.ContentPadding(0.0f)
+						.ButtonStyle(FAppStyle::Get(), "NoBorder")
+						.OnClicked(this, &SControlRigSpacePicker::HandleAddSpaceClicked)
+						.Cursor(EMouseCursor::Default)
+						.ToolTipText(LOCTEXT("AddSpace", "Add Space"))
 						[
-							SNew(SButton)
-							.ContentPadding(0.0f)
-							.ButtonStyle(FAppStyle::Get(), "NoBorder")
-							.OnClicked(this, &SControlRigSpacePicker::HandleAddSpaceClicked)
-							.Cursor(EMouseCursor::Default)
-							.ToolTipText(LOCTEXT("AddSpace", "Add Space"))
-							[
-								SNew(SImage)
-								.Image(FAppStyle::GetBrush(TEXT("Icons.PlusCircle")))
-							]
-							.Visibility_Lambda([this]()
-							{
-								return GetAddSpaceButtonVisibility();
-							})
+							SNew(SImage)
+							.Image(FAppStyle::GetBrush(TEXT("Icons.PlusCircle")))
 						]
-					]
-					.BodyContent()
-					[
-						SAssignNew(SpacePickerWidget, SRigSpacePickerWidget)
-						.AllowDelete(true)
-						.AllowReorder(true)
-						.AllowAdd(false)
-						.ShowBakeAndCompensateButton(true)
-						.GetControlCustomization(this, &SControlRigSpacePicker::HandleGetControlElementCustomization)
-						.OnActiveSpaceChanged(this, &SControlRigSpacePicker::HandleActiveSpaceChanged)
-						.OnSpaceListChanged(this, &SControlRigSpacePicker::HandleSpaceListChanged)
-						.OnCompensateKeyButtonClicked(this, &SControlRigSpacePicker::OnCompensateKeyClicked)
-						.OnCompensateAllButtonClicked(this, &SControlRigSpacePicker::OnCompensateAllClicked)
-						.OnBakeButtonClicked(this, &SControlRigSpacePicker::OnBakeControlsToNewSpaceButtonClicked)
-						// todo: implement GetAdditionalSpacesDelegate to pull spaces from sequencer
+						.Visibility_Lambda([this]()
+						{
+							return GetAddSpaceButtonVisibility();
+						})
 					]
 				]
-
+				.BodyContent()
+				[
+					SAssignNew(SpacePickerWidget, SRigSpacePickerWidget)
+					.AllowDelete(true)
+					.AllowReorder(true)
+					.AllowAdd(false)
+					.Hierarchy(InArgs._Hierarchy)
+					.Controls(InArgs._Controls)
+					.ShowBakeAndCompensateButton(true)
+					.GetControlCustomization(this, &SControlRigSpacePicker::HandleGetControlElementCustomization)
+					.OnActiveSpaceChanged(this, &SControlRigSpacePicker::HandleActiveSpaceChanged)
+					.OnSpaceListChanged(this, &SControlRigSpacePicker::HandleSpaceListChanged)
+					.OnCompensateKeyButtonClicked(this, &SControlRigSpacePicker::OnCompensateKeyClicked)
+					.OnCompensateAllButtonClicked(this, &SControlRigSpacePicker::OnCompensateAllClicked)
+					.OnBakeButtonClicked(this, &SControlRigSpacePicker::OnBakeControlsToNewSpaceButtonClicked)
+					// todo: implement GetAdditionalSpacesDelegate to pull spaces from sequencer
+				]
 			]
-		];
+
+		]
+	];
 
 	SetEditMode(InEditMode);
 }
 
 SControlRigSpacePicker::~SControlRigSpacePicker()
 {
-	//base class handles control rig related cleanup
+	SelectionViewModel->OnControlSelected().RemoveAll(this);
 }
 
 UControlRig* SControlRigSpacePicker::GetControlRig() const
diff --git a/Engine/Plugins/Animation/ControlRig/Source/ControlRigEditor/Private/EditMode/SControlRigSpacePicker.h b/Engine/Plugins/Animation/ControlRig/Source/ControlRigEditor/Private/EditMode/SControlRigSpacePicker.h
--- a/Engine/Plugins/Animation/ControlRig/Source/ControlRigEditor/Private/EditMode/SControlRigSpacePicker.h
+++ b/Engine/Plugins/Animation/ControlRig/Source/ControlRigEditor/Private/EditMode/SControlRigSpacePicker.h
@@ -12,23 +12,38 @@
 #include "Misc/FrameNumber.h"
 #include "Editor/SRigSpacePickerWidget.h"
 
-
 class ISequencer;
 class SExpandableArea;
 class SSearchableRigHierarchyTreeView;
 class UControlRig;
+namespace UE::ControlRigEditor { class FRigSelectionViewModel; }
 
 class SControlRigSpacePicker : public SCompoundWidget, public FControlRigBaseDockableView
 {
 
 	SLATE_BEGIN_ARGS(SControlRigSpacePicker)
+		: _InitiallyCollapsed(true)
 	{}
+		/** Whether or not the area is initially collapsed */
+		SLATE_ARGUMENT(bool, InitiallyCollapsed)
+		/** The initially selected hierarchy. */
+		SLATE_ARGUMENT(URigHierarchy*, Hierarchy)
+		/** The initially selected controls. */
+		SLATE_ARGUMENT(TArray<FRigElementKey>, Controls)
 	SLATE_END_ARGS()
-		~SControlRigSpacePicker();
 
-	void Construct(const FArguments& InArgs, FControlRigEditMode& InEditMode);
+	virtual ~SControlRigSpacePicker() override;
+
+	void Construct(
+		const FArguments& InArgs, 
+		FControlRigEditMode& InEditMode, const TSharedRef<UE::ControlRigEditor::FRigSelectionViewModel>& InSelectionViewModel
+		);
 
 private:
+	
+	/** View model for when control rig selection changes. */
+	TSharedPtr<UE::ControlRigEditor::FRigSelectionViewModel> SelectionViewModel;
+	
 	virtual void HandleControlSelected(UControlRig* Subject, FRigControlElement* InControl, bool bSelected) override;
 
 	/** Space picker widget*/
diff --git a/Engine/Plugins/Animation/ControlRig/Source/ControlRigEditor/Private/Editor/SRigSpacePickerWidget.cpp b/Engine/Plugins/Animation/ControlRig/Source/ControlRigEditor/Private/Editor/SRigSpacePickerWidget.cpp
--- a/Engine/Plugins/Animation/ControlRig/Source/ControlRigEditor/Private/Editor/SRigSpacePickerWidget.cpp
+++ b/Engine/Plugins/Animation/ControlRig/Source/ControlRigEditor/Private/Editor/SRigSpacePickerWidget.cpp
@@ -255,10 +255,12 @@ void SRigSpacePickerWidget::SetControls(URigHierarchy* InHierarchy, const TArray
 	PendingSelectionHandle = RegisterActiveTimer(0.f,
 		FWidgetActiveTimerDelegate::CreateLambda([this, WeakHierarchy = TWeakObjectPtr<URigHierarchy>(InHierarchy), InControls](double, float)
 		{
-			if(URigHierarchy* HierarchyToUpdate = WeakHierarchy.Get())
-			{
-				UpdateSelection(HierarchyToUpdate, InControls);
-			}
+			// No null check here:
+			// InHiearchy is nullptr when you deselect, e.g. when you click the floor. Passing nullptr to UpdateSelection deselects.
+			// In the unlikely case that InHiearchy was valid but TWeakObjectPtr becomes stale, there's nothing to display... so also deselect.
+			URigHierarchy* HierarchyToUpdate = WeakHierarchy.Get();
+				
+			UpdateSelection(HierarchyToUpdate, InControls);
 			return EActiveTimerReturnType::Stop;
 		})
 	);
@@ -1235,7 +1237,6 @@ void SRigSpacePickerBakeWidget::Construct(const FArguments& InArgs)
 	*Settings = InArgs._Settings;
 	//always setting space to be parent as default, since stored space may not be available.
 	Settings->Get()->TargetSpace = URigHierarchy::GetDefaultParentKey();
-	Sequencer = InArgs._Sequencer;
 
 	FStructureDetailsViewArgs StructureViewArgs;
 	StructureViewArgs.bShowObjects = true;
@@ -1253,7 +1254,7 @@ void SRigSpacePickerBakeWidget::Construct(const FArguments& InArgs)
 	DetailsView = PropertyEditor.CreateStructureDetailView(ViewArgs, StructureViewArgs, TSharedPtr<FStructOnScope>());
 
 	DetailsView->GetDetailsView()->RegisterInstancedCustomPropertyTypeLayout("FrameNumber",
-		FOnGetPropertyTypeCustomizationInstance::CreateSP(Sequencer, &ISequencer::MakeFrameNumberDetailsCustomization));
+		FOnGetPropertyTypeCustomizationInstance::CreateSP(InArgs._Sequencer, &ISequencer::MakeFrameNumberDetailsCustomization));
 	DetailsView->SetStructureData(Settings);
 
 	ChildSlot
diff --git a/Engine/Plugins/Animation/ControlRig/Source/ControlRigEditor/Private/Editor/SRigSpacePickerWidget.h b/Engine/Plugins/Animation/ControlRig/Source/ControlRigEditor/Private/Editor/SRigSpacePickerWidget.h
--- a/Engine/Plugins/Animation/ControlRig/Source/ControlRigEditor/Private/Editor/SRigSpacePickerWidget.h
+++ b/Engine/Plugins/Animation/ControlRig/Source/ControlRigEditor/Private/Editor/SRigSpacePickerWidget.h
@@ -223,7 +223,6 @@ private:
 	//used for setting up the details
 	TSharedPtr<TStructOnScope<FRigSpacePickerBakeSettings>> Settings;
 
-	ISequencer* Sequencer;
 	FRigControlElementCustomization Customization;
 	
 	TWeakPtr<SWindow> DialogWindow;
diff --git a/Engine/Plugins/Animation/ControlRig/Source/ControlRigEditor/Private/Misc/RigEditModeUtils.cpp b/Engine/Plugins/Animation/ControlRig/Source/ControlRigEditor/Private/Misc/RigEditModeUtils.cpp
new file mode 100644
--- /dev/null
+++ b/Engine/Plugins/Animation/ControlRig/Source/ControlRigEditor/Private/Misc/RigEditModeUtils.cpp
@@ -0,0 +1,27 @@
+﻿// Copyright Epic Games, Inc. All Rights Reserved.
+
+#include "RigEditModeUtils.h"
+
+#include "ControlRig.h"
+#include "EditMode/ControlRigEditMode.h"
+
+namespace UE::ControlRigEditor
+{
+FInitialSpacePickerSelection DetermineInitialSpacePickerSelection(FControlRigEditMode& InEditMode)
+{
+	TMap<UControlRig*, TArray<FRigElementKey>> SelectedControlRigsAndControls;
+	InEditMode.GetAllSelectedControls(SelectedControlRigsAndControls);
+	
+	TArray<UControlRig*> ControlRigs;
+	TArray<TArray<FRigElementKey>> AllSelectedControls;
+	SelectedControlRigsAndControls.GenerateKeyArray(ControlRigs);
+	SelectedControlRigsAndControls.GenerateValueArray(AllSelectedControls);
+	
+	//mz todo handle multiple control rigs with space picker
+	UControlRig* RuntimeRig = ControlRigs.IsEmpty() ? nullptr : ControlRigs[0];
+	TArray<FRigElementKey> SelectedControls = AllSelectedControls.IsEmpty() ? TArray<FRigElementKey>() : MoveTemp(AllSelectedControls[0]);
+	URigHierarchy* Hierarchy = RuntimeRig ? RuntimeRig->GetHierarchy() : nullptr;
+
+	return FInitialSpacePickerSelection{ MoveTemp(SelectedControls), RuntimeRig, Hierarchy };
+}
+}
diff --git a/Engine/Plugins/Animation/ControlRig/Source/ControlRigEditor/Private/Misc/RigEditModeUtils.h b/Engine/Plugins/Animation/ControlRig/Source/ControlRigEditor/Private/Misc/RigEditModeUtils.h
new file mode 100644
--- /dev/null
+++ b/Engine/Plugins/Animation/ControlRig/Source/ControlRigEditor/Private/Misc/RigEditModeUtils.h
@@ -0,0 +1,28 @@
+﻿// Copyright Epic Games, Inc. All Rights Reserved.
+
+#pragma once
+
+#include "Containers/Array.h"
+
+class FControlRigEditMode;
+class UControlRig;
+class URigHierarchy;
+struct FRigElementKey;
+
+namespace UE::ControlRigEditor
+{
+struct FInitialSpacePickerSelection
+{
+	/** The controls selected on the rig. */
+	TArray<FRigElementKey> SelectedControls;
+	/** The selected rig */
+	UControlRig* RuntimeRig = nullptr;
+	/** Hierarchy of RuntimeRig */
+	URigHierarchy* Hierarchy = nullptr;
+	
+	bool IsValid() const { return Hierarchy && !SelectedControls.IsEmpty(); }  
+};
+
+/** @return The initial hierarchy and controls to select when constructing SRigSpacePickerWidget. */
+FInitialSpacePickerSelection DetermineInitialSpacePickerSelection(FControlRigEditMode& InEditMode);
+}
diff --git a/Engine/Plugins/Animation/ControlRig/Source/ControlRigEditor/Private/Misc/Tab/SInlineTabPanel.cpp b/Engine/Plugins/Animation/ControlRig/Source/ControlRigEditor/Private/Misc/Tab/SInlineTabPanel.cpp
new file mode 100644
--- /dev/null
+++ b/Engine/Plugins/Animation/ControlRig/Source/ControlRigEditor/Private/Misc/Tab/SInlineTabPanel.cpp
@@ -0,0 +1,120 @@
+﻿// Copyright Epic Games, Inc. All Rights Reserved.
+
+#include "SInlineTabPanel.h"
+
+#include "STabArea.h"
+#include "Components/VerticalBox.h"
+#include "Styling/AppStyle.h"
+#include "Widgets/Images/SImage.h"
+#include "Widgets/Layout/SBorder.h"
+#include "Widgets/Layout/SSeparator.h"
+#include "Widgets/Layout/SWidgetSwitcher.h"
+#include "Widgets/Text/STextBlock.h"
+
+namespace UE::ControlRigEditor
+{
+namespace InlineTabPanel
+{
+static TArray<FTabEntry> Transform(TConstArrayView<FInlineTabArgs> InArgs, const TSharedRef<SInlineTabPanel>& This)
+{
+	const int32 Num = InArgs.Num();
+	
+	TArray<FTabEntry> Result;
+	Result.Reserve(Num);
+	
+	for (int32 Index = 0; Index < Num; ++Index)
+	{
+		const FInlineTabArgs& CurrentArgs = InArgs[Index];
+
+		Result.Emplace();
+		FTabEntry& Entry = Result[Index];
+		
+		Entry.ButtonContent.Widget = SNew(SHorizontalBox)
+			.ToolTipText(CurrentArgs.ToolTipText)
+			
+			+SHorizontalBox::Slot()
+			.AutoWidth()
+			[
+				SNew(SImage)
+				.Visibility(CurrentArgs.Icon.IsSet() ? EVisibility::HitTestInvisible : EVisibility::Collapsed)
+				.Image(CurrentArgs.Icon.GetIcon())
+			]
+			
+			+SHorizontalBox::Slot()
+			.AutoWidth()
+			.Padding(4.f, 0.f, 0.f, 0.f)
+			[
+				SNew(STextBlock)
+				.Text(CurrentArgs.Label)
+			];
+		Entry.OnTabSelected.BindLambda([Index, WeakOwnerPanel = This.ToWeakPtr(), Delegate = Entry.OnTabSelected]()
+		{
+			if (const TSharedPtr<SInlineTabPanel> OwnerPin = WeakOwnerPanel.Pin())
+			{
+				OwnerPin->SwitchToTab(Index);
+				Delegate.ExecuteIfBound();
+			}
+		});
+	}
+
+	return Result;
+}
+}
+	
+void SInlineTabPanel::Construct(const FArguments& InArgs)
+{
+	ChildSlot
+	[
+		SNew(SVerticalBox)
+
+		// Toolbar
+		+SVerticalBox::Slot()
+		.AutoHeight()
+		[
+			SNew(SBorder)
+			.BorderImage(FAppStyle::Get().GetBrush("ToolPanel.GroupBorder"))
+			.Padding(0.0f)
+			[
+				SNew(STabArea)
+					.Tabs(InlineTabPanel::Transform(InArgs._Tabs, SharedThis(this)))
+					.ActiveTabIndex(InArgs._ActiveTabIndex)
+					.Padding(InArgs._Padding)
+			]
+		]
+
+		// Separate the tabs & content
+		+SVerticalBox::Slot()
+		.AutoHeight()
+		.Padding(2.f)
+		[
+			SNew(SSeparator)
+		]
+
+		// Tabs
+		+SVerticalBox::Slot()
+		.FillHeight(1.f)
+		.HAlign(HAlign_Fill)
+		.VAlign(VAlign_Fill)
+		[
+			SAssignNew(TabContentSwitcher, SWidgetSwitcher)
+		]
+	];
+
+	for (int32 Index = 0; Index < InArgs._Tabs.Num(); ++Index)
+	{
+		const FInlineTabArgs& TabArgs = InArgs._Tabs[Index];
+		TabContentSwitcher->AddSlot()
+		[
+			TabArgs.Content.Widget
+		];
+	}
+}
+
+void SInlineTabPanel::SwitchToTab(int32 InIndex) const
+{
+	if (ensure(InIndex >= 0 && InIndex < TabContentSwitcher->GetNumWidgets()))
+	{
+		TabContentSwitcher->SetActiveWidgetIndex(InIndex);
+	}
+}
+}
diff --git a/Engine/Plugins/Animation/ControlRig/Source/ControlRigEditor/Private/Misc/Tab/SInlineTabPanel.h b/Engine/Plugins/Animation/ControlRig/Source/ControlRigEditor/Private/Misc/Tab/SInlineTabPanel.h
new file mode 100644
--- /dev/null
+++ b/Engine/Plugins/Animation/ControlRig/Source/ControlRigEditor/Private/Misc/Tab/SInlineTabPanel.h
@@ -0,0 +1,76 @@
+﻿// Copyright Epic Games, Inc. All Rights Reserved.
+
+#pragma once
+
+#include "Textures/SlateIcon.h"
+#include "Widgets/DeclarativeSyntaxSupport.h"
+#include "Widgets/SCompoundWidget.h"
+
+class SWidgetSwitcher;
+
+namespace UE::ControlRigEditor
+{
+class STabArea;
+	
+struct FInlineTabArgs
+{
+	/** The name of the tab */
+	TAttribute<FText> Label;
+	
+	/** The tooltip of the tab */
+	TAttribute<FText> ToolTipText;
+
+	/** Optional icon displayed in front of the label. */
+	FSlateIcon Icon;
+
+	/** Content to display as tab content */
+	TAlwaysValidWidget Content;
+
+	/** Invoked when this tab's content is shown (except for when SInlineTabPanel is constructed). */
+	FSimpleDelegate OnTabSelected;
+
+	FInlineTabArgs& SetLabel(TAttribute<FText> InLabel) { Label = MoveTemp(InLabel); return *this; }
+	FInlineTabArgs& SetToolTipText(TAttribute<FText> InToolTipText) { ToolTipText = MoveTemp(InToolTipText); return *this; }
+	FInlineTabArgs& SetIcon(const FSlateIcon& InIcon) { Icon = InIcon; return *this; }
+	FInlineTabArgs& SetContent(TAlwaysValidWidget InWidget) { Content = MoveTemp(InWidget); return *this; }
+	FInlineTabArgs& SetContent(const TSharedRef<SWidget>& InWidget) { Content.Widget = InWidget; return *this; }
+	FInlineTabArgs& SetOnTabSelected(FSimpleDelegate InOnTabSelected) { OnTabSelected = MoveTemp(InOnTabSelected); return *this; }
+};
+
+/**
+ * Coordinator widget. A vertical box with tab buttons on the top and the selected tab content at the bottom.
+ * 
+ * This is supposed to implement easy to use fence pattern.
+ * If you need more flexibility for styling, considering building your widget hierarchy directly by using STabArea, SWidgetSwitcher, and STabButton
+ * instead of complicating this widget's interface.
+ */
+class SInlineTabPanel : public SCompoundWidget
+{
+public:
+
+	SLATE_BEGIN_ARGS(SInlineTabPanel)
+		: _ActiveTabIndex(0)
+		, _Padding(3.f)
+		{}
+		
+		/** The tabs to create */
+		SLATE_ARGUMENT(TArray<FInlineTabArgs>, Tabs)
+		/** The tab that should be active by default. */
+		SLATE_ARGUMENT(int32, ActiveTabIndex)
+		
+		/** Padding between buttons */
+		SLATE_ARGUMENT(FMargin, Padding)
+	SLATE_END_ARGS()
+
+	void Construct(const FArguments& InArgs);
+
+	/** Shows the content of the specified tab index. */
+	void SwitchToTab(int32 InIndex) const;
+
+private:
+
+	/** Switches the content for the tabs */
+	TSharedPtr<SWidgetSwitcher> TabContentSwitcher;
+};
+}
+
diff --git a/Engine/Plugins/Animation/ControlRig/Source/ControlRigEditor/Private/Misc/Tab/STabArea.cpp b/Engine/Plugins/Animation/ControlRig/Source/ControlRigEditor/Private/Misc/Tab/STabArea.cpp
new file mode 100644
--- /dev/null
+++ b/Engine/Plugins/Animation/ControlRig/Source/ControlRigEditor/Private/Misc/Tab/STabArea.cpp
@@ -0,0 +1,69 @@
+﻿// Copyright Epic Games, Inc. All Rights Reserved.
+
+#include "STabArea.h"
+
+#include "STabButton.h"
+
+#include "Widgets/SBoxPanel.h"
+
+namespace UE::ControlRigEditor
+{
+void STabArea::Construct(const FArguments& InArgs)
+{
+	const TSharedRef<SHorizontalBox> Content = SNew(SHorizontalBox);
+	for (const FTabEntry& Entry : InArgs._Tabs)
+	{
+		const int32 ButtonIndex = TabButtons.Num();
+		
+		TSharedPtr<STabButton> Button;
+		Content->AddSlot()
+			.Padding(InArgs._Padding)
+			.AutoWidth()
+			[
+				SAssignNew(Button, STabButton)
+				.OnActivated_Lambda([this, OnTabSelected = Entry.OnTabSelected, ButtonIndex]()
+				{
+					OnButtonActivated(ButtonIndex);
+					OnTabSelected.ExecuteIfBound();
+				})
+				[
+					Entry.ButtonContent.Widget
+				]
+			];
+
+		TabButtons.Emplace(Button.ToSharedRef());
+	}
+
+	if (ensure(TabButtons.IsValidIndex(InArgs._ActiveTabIndex)))
+	{
+		TabButtons[InArgs._ActiveTabIndex]->Activate();
+	}
+	
+	ChildSlot
+		// The buttons should be centered (i.e. padd equally on left and right if there's space left over)
+		.HAlign(HAlign_Center)
+	[
+		Content
+	];
+}
+
+void STabArea::SetButtonActivated(int32 ButtonIndex)
+{
+	if (ensure(TabButtons.IsValidIndex(ButtonIndex)))
+	{
+		TabButtons[ButtonIndex]->Activate();
+		OnButtonActivated(ButtonIndex);
+	}
+}
+
+void STabArea::OnButtonActivated(int32 ButtonIndex)
+{
+	for (int32 i = 0; i < TabButtons.Num(); ++i)
+	{
+		if (ButtonIndex != i)
+		{
+			TabButtons[i]->Deactivate();
+		}
+	}
+}
+}
diff --git a/Engine/Plugins/Animation/ControlRig/Source/ControlRigEditor/Private/Misc/Tab/STabArea.h b/Engine/Plugins/Animation/ControlRig/Source/ControlRigEditor/Private/Misc/Tab/STabArea.h
new file mode 100644
--- /dev/null
+++ b/Engine/Plugins/Animation/ControlRig/Source/ControlRigEditor/Private/Misc/Tab/STabArea.h
@@ -0,0 +1,50 @@
+﻿// Copyright Epic Games, Inc. All Rights Reserved.
+
+#pragma once
+
+#include "Widgets/DeclarativeSyntaxSupport.h"
+#include "Widgets/SCompoundWidget.h"
+
+namespace UE::ControlRigEditor
+{
+class STabButton;
+	
+struct FTabEntry
+{
+	TAlwaysValidWidget ButtonContent;
+	FSimpleDelegate OnTabSelected;
+};
+
+/** Manages multiple STabButtons, making sure exactly 1 is active at a time. */
+class STabArea : public SCompoundWidget
+{
+public:
+
+	SLATE_BEGIN_ARGS(STabArea)
+		: _ActiveTabIndex(0)
+		, _Padding(3.f)
+		{}
+		/** The tabs to create */
+		SLATE_ARGUMENT(TArray<FTabEntry>, Tabs)
+		/** The tab that should be active by default. */
+		SLATE_ARGUMENT(int32, ActiveTabIndex)
+		
+		/** Padding between buttons */
+		SLATE_ARGUMENT(FMargin, Padding)
+	SLATE_END_ARGS()
+
+	void Construct(const FArguments& InArgs);
+
+	/** Informs the owning view has manually changed to the content displayed by ButtonIndex. Make the button appear selected. */
+	void SetButtonActivated(int32 ButtonIndex);
+
+private:
+
+	/** The buttons in this area. */
+	TArray<TSharedRef<STabButton>> TabButtons;
+
+	/** Called when the button at TabButtons[ButtonIndex] is activated. */
+	void OnButtonActivated(int32 ButtonIndex);
+};
+}
+
diff --git a/Engine/Plugins/Animation/ControlRig/Source/ControlRigEditor/Private/Misc/Tab/STabButton.cpp b/Engine/Plugins/Animation/ControlRig/Source/ControlRigEditor/Private/Misc/Tab/STabButton.cpp
new file mode 100644
--- /dev/null
+++ b/Engine/Plugins/Animation/ControlRig/Source/ControlRigEditor/Private/Misc/Tab/STabButton.cpp
@@ -0,0 +1,39 @@
+﻿// Copyright Epic Games, Inc. All Rights Reserved.
+
+#include "STabButton.h"
+
+#include "ControlRigEditorStyle.h"
+#include "Widgets/Input/SCheckBox.h"
+
+namespace UE::ControlRigEditor
+{
+void STabButton::Construct(const FArguments& InArgs)
+{
+	OnActivatedDelegate = InArgs._OnActivated;
+	
+	ChildSlot
+	[
+		SNew(SCheckBox)
+		.Style(FControlRigEditorStyle::Get(), TEXT("ControlRig.TabButton"))
+		.IsChecked(this, &STabButton::IsChecked)
+        .OnCheckStateChanged(this, &STabButton::OnButtonClicked)
+		[
+			InArgs._ButtonContent.Widget
+		]
+	];
+}
+
+ECheckBoxState STabButton::IsChecked() const
+{
+	return bIsActive ? ECheckBoxState::Checked : ECheckBoxState::Unchecked;
+}
+
+void STabButton::OnButtonClicked(ECheckBoxState NewState)
+{
+	if (!bIsActive)
+	{
+		bIsActive = true;
+		OnActivatedDelegate.ExecuteIfBound();
+	}
+}
+}
diff --git a/Engine/Plugins/Animation/ControlRig/Source/ControlRigEditor/Private/Misc/Tab/STabButton.h b/Engine/Plugins/Animation/ControlRig/Source/ControlRigEditor/Private/Misc/Tab/STabButton.h
new file mode 100644
--- /dev/null
+++ b/Engine/Plugins/Animation/ControlRig/Source/ControlRigEditor/Private/Misc/Tab/STabButton.h
@@ -0,0 +1,44 @@
+﻿// Copyright Epic Games, Inc. All Rights Reserved.
+
+#pragma once
+
+#include "HAL/Platform.h"
+#include "Widgets/DeclarativeSyntaxSupport.h"
+#include "Widgets/SCompoundWidget.h"
+
+enum class ECheckBoxState : uint8;
+
+namespace UE::ControlRigEditor
+{
+/** Used to switch "tabs". It is a blue button when active and gray when inactive. */
+class STabButton : public SCompoundWidget
+{
+public:
+
+	SLATE_BEGIN_ARGS(STabButton){}
+		/** Called when this button is clicked and becomes active. Not trigger if it was already active. */
+		SLATE_EVENT(FSimpleDelegate, OnActivated)
+		
+		SLATE_DEFAULT_SLOT(FArguments, ButtonContent)
+	SLATE_END_ARGS()
+
+	void Construct(const FArguments& InArgs);
+
+	/** Makes the button appear active. */
+	void Activate() { bIsActive = true; }
+	/** Makes the button appear inactive. */
+	void Deactivate() { bIsActive = false; }
+	
+private:
+
+	/** Whether the button is active (the tab content is supposed to be shown). */
+	bool bIsActive = false;
+
+	/** Called when this button is clicked and becomes active. Not trigger if it was already active. */
+	FSimpleDelegate OnActivatedDelegate;
+	
+	ECheckBoxState IsChecked() const;
+	void OnButtonClicked(ECheckBoxState NewState);
+};
+}
+
diff --git a/Engine/Plugins/Animation/ControlRig/Source/ControlRigEditor/Private/Sequencer/ControlRigSpaceChannelEditors.cpp b/Engine/Plugins/Animation/ControlRig/Source/ControlRigEditor/Private/Sequencer/ControlRigSpaceChannelEditors.cpp
--- a/Engine/Plugins/Animation/ControlRig/Source/ControlRigEditor/Private/Sequencer/ControlRigSpaceChannelEditors.cpp
+++ b/Engine/Plugins/Animation/ControlRig/Source/ControlRigEditor/Private/Sequencer/ControlRigSpaceChannelEditors.cpp
@@ -599,7 +599,6 @@ static EControlRigContextChannelToKey GetCurrentTransformKeysAtThisTime(UControl
 
 				if (FRigControlElement* ControlElement = ControlRig->FindControl(ControlName))
 				{
-					FMovieSceneControlRigSpaceBaseKey Value;
 					switch (ControlElement->Settings.ControlType)
 					{
 					case ERigControlType::Position:
@@ -797,7 +796,6 @@ void FControlRigSpaceChannelHelpers::DeleteTransformKeysAtThisTime(UControlRig*
 
 			if (FRigControlElement* ControlElement = ControlRig->FindControl(ControlName))
 			{
-				FMovieSceneControlRigSpaceBaseKey Value;
 				switch (ControlElement->Settings.ControlType)
 				{
 					case ERigControlType::Position:
@@ -1259,7 +1257,6 @@ void FControlRigSpaceChannelHelpers::HandleSpaceKeyTimeChanged(UControlRig* Cont
 				FFrameNumber Delta = NextFrame - CurrentFrame;
 				if (FRigControlElement* ControlElement = ControlRig->FindControl(ControlName))
 				{
-					FMovieSceneControlRigSpaceBaseKey Value;
 					switch (ControlElement->Settings.ControlType)
 					{
 						case ERigControlType::Position:
@@ -1358,7 +1355,7 @@ void FControlRigSpaceChannelHelpers::CompensateIfNeeded(
 	UControlRig* ControlRig,
 	ISequencer* Sequencer,
 	UMovieSceneControlRigParameterSection* Section,
-	TOptional<FFrameNumber>& OptionalTime, 
+	const TOptional<FFrameNumber>& OptionalTime, 
 	bool bCompPreviousTick)
 {
 	if (bDoNotCompensate == true)
@@ -1535,7 +1532,6 @@ FReply FControlRigSpaceChannelHelpers::OpenBakeDialog(ISequencer* Sequencer, FMo
 			FRigSpacePickerBakeSettings Settings;
 
 			const FFrameRate TickResolution = Sequencer->GetFocusedTickResolution();
-			FMovieSceneControlRigSpaceBaseKey Value;
 
 			URigHierarchy* RigHierarchy = ControlRig->GetHierarchy();
 			Settings.TargetSpace = URigHierarchy::GetDefaultParentKey();
diff --git a/Engine/Plugins/Animation/ControlRig/Source/ControlRigEditor/Public/ControlRigSpaceChannelEditors.h b/Engine/Plugins/Animation/ControlRig/Source/ControlRigEditor/Public/ControlRigSpaceChannelEditors.h
--- a/Engine/Plugins/Animation/ControlRig/Source/ControlRigEditor/Public/ControlRigSpaceChannelEditors.h
+++ b/Engine/Plugins/Animation/ControlRig/Source/ControlRigEditor/Public/ControlRigSpaceChannelEditors.h
@@ -47,7 +47,7 @@ struct FControlRigSpaceChannelHelpers
 	static FKeyHandle SequencerKeyControlRigSpaceChannel(UControlRig* ControlRig, ISequencer* Sequencer, FMovieSceneControlRigSpaceChannel* Channel, UMovieSceneSection* SectionToKey, FFrameNumber Time, URigHierarchy* RigHierarchy, const FRigElementKey& ControlKey, const FRigElementKey& SpaceKey);
 	static void SequencerSpaceChannelKeyDeleted(UControlRig* ControlRig, ISequencer* Sequencer, FName ControlName, FMovieSceneControlRigSpaceChannel* Channel, UMovieSceneControlRigParameterSection* SectionToKey, FFrameNumber TimeOfDeletion);
 	static UMovieSceneControlRigParameterSection* GetControlRigSection(ISequencer* Sequencer, const UControlRig* ControlRig);
-	static void CompensateIfNeeded(UControlRig* ControlRig, ISequencer* Sequencer, UMovieSceneControlRigParameterSection* Section, TOptional<FFrameNumber>& Time, bool bCompPreviousTick);
+	static void CompensateIfNeeded(UControlRig* ControlRig, ISequencer* Sequencer, UMovieSceneControlRigParameterSection* Section, const TOptional<FFrameNumber>& Time, bool bCompPreviousTick);
 	static FSpaceChannelAndSection FindSpaceChannelAndSectionForControl(UControlRig* ControlRig, FName ControlName, ISequencer* Sequencer, bool bCreateIfNeeded);
 	static void SequencerBakeControlInSpace(UControlRig* ControlRig, ISequencer* Sequencer, FMovieSceneControlRigSpaceChannel* Channel, UMovieSceneSection* SectionToKey,
 		URigHierarchy* RigHierarchy, const FRigElementKey& ControlKey, FRigSpacePickerBakeSettings InSettings);
