Control Rig: Spaces, Constraints & Snapper now use SActionButton instead so the style is consistent  with UE5 (before they looked UE4).

diff --git a/Engine/Plugins/Animation/ControlRig/Source/ControlRigEditor/Private/EditMode/SControlRigSnapper.cpp b/Engine/Plugins/Animation/ControlRig/Source/ControlRigEditor/Private/EditMode/SControlRigSnapper.cpp
--- a/Engine/Plugins/Animation/ControlRig/Source/ControlRigEditor/Private/EditMode/SControlRigSnapper.cpp
+++ b/Engine/Plugins/Animation/ControlRig/Source/ControlRigEditor/Private/EditMode/SControlRigSnapper.cpp
@@ -31,6 +31,7 @@
 #include "Widgets/Layout/SScrollBox.h"
 #include "Tools/ControlRigSnapSettings.h"
 #include "MovieScene.h"
+#include "SActionButton.h"
 
 #define LOCTEXT_NAMESPACE "ControlRigSnapper"
 
@@ -265,14 +266,10 @@ void SControlRigSnapper::Construct(const FArguments& InArgs)
 					.HAlign(HAlign_Center)
 					.VAlign(VAlign_Center)
 					[
-						SNew(SButton)
-						.HAlign(HAlign_Fill)
-						.ContentPadding(FMargin(10.0f, 2.0f, 10.0f, 2.0f))
+						SNew(SActionButton)
+						.Text(LOCTEXT("SnapAnimation", "Snap Animation"))
+						.ButtonContentPadding(FMargin(10.0f, 2.0f, 10.0f, 2.0f))
 						.OnClicked(this, &SControlRigSnapper::OnSnapAnimationClicked)
-						[
-							SNew(STextBlock)
-							.Text(LOCTEXT("SnapAnimation", "Snap Animation"))
-						]
 					]
 				]
 			]
diff --git a/Engine/Plugins/Animation/ControlRig/Source/ControlRigEditor/Private/Editor/Constraints/SConstraintsWidget.cpp b/Engine/Plugins/Animation/ControlRig/Source/ControlRigEditor/Private/Editor/Constraints/SConstraintsWidget.cpp
--- a/Engine/Plugins/Animation/ControlRig/Source/ControlRigEditor/Private/Editor/Constraints/SConstraintsWidget.cpp
+++ b/Engine/Plugins/Animation/ControlRig/Source/ControlRigEditor/Private/Editor/Constraints/SConstraintsWidget.cpp
@@ -24,6 +24,7 @@
 #include "Styling/SlateIconFinder.h"
 #include "Widgets/Input/NumericTypeInterface.h"
 #include "FrameNumberDetailsCustomization.h"
+#include "SActionButton.h"
 #include "SAssetDropTarget.h"
 #include "Sequencer/ControlRigParameterTrackEditor.h"
 #include "Sequencer/EditModeAnimationUtil.h"
@@ -1435,14 +1436,10 @@ void SConstraintsEditionWidget::Construct(const FArguments& InArgs)
 		.HAlign(HAlign_Right)
 		.Padding(1.f)
 		[
-			SNew(SButton)
-			.ButtonStyle(FAppStyle::Get(), "FlatButton.Default")
+			SNew(SActionButton)
 			.Text(LOCTEXT("BakeButton", "Bake..."))
 			.OnClicked(this, &SConstraintsEditionWidget::OnBakeClicked)
-			.IsEnabled_Lambda([this]()
-			{
-				return (ListItems.Num() > 0);
-			})
+			.IsEnabled_Lambda([this]() { return ListItems.Num() > 0; })
 			.ToolTipText(LOCTEXT("BakeConstraintButtonToolTip", "Bake selected constraints"))
 		];
 	};
diff --git a/Engine/Plugins/Animation/ControlRig/Source/ControlRigEditor/Private/Editor/SRigSpacePickerWidget.cpp b/Engine/Plugins/Animation/ControlRig/Source/ControlRigEditor/Private/Editor/SRigSpacePickerWidget.cpp
--- a/Engine/Plugins/Animation/ControlRig/Source/ControlRigEditor/Private/Editor/SRigSpacePickerWidget.cpp
+++ b/Engine/Plugins/Animation/ControlRig/Source/ControlRigEditor/Private/Editor/SRigSpacePickerWidget.cpp
@@ -21,6 +21,7 @@
 #include "ControlRig.h"
 #include "ControlRigBlueprintLegacy.h"
 #include "RigVMBlueprintGeneratedClass.h"
+#include "SActionButton.h"
 
 #define LOCTEXT_NAMESPACE "SRigSpacePickerWidget"
 
@@ -180,49 +181,36 @@ void SRigSpacePickerWidget::Construct(const FArguments& InArgs)
 			BottomButtonsListBox->AddSlot()
 			.AutoWidth()
 			.VAlign(VAlign_Center)
-			.HAlign(HAlign_Left)
-			.Padding(0.f)
+			.Padding(4.f, 0.f)
 			[
-				SNew(SButton)
-				.ButtonStyle(FAppStyle::Get(), "FlatButton.Default")
+				SNew(SActionButton)
 				.Text(LOCTEXT("CompensateKeyButton", "Comp Key"))
 				.OnClicked(InArgs._OnCompensateKeyButtonClicked)
-				.IsEnabled_Lambda([this]()
-				{
-					return (ControlKeys.Num() > 0);
-				})
+				.IsEnabled_Lambda([this]() { return ControlKeys.Num() > 0; })
 				.ToolTipText(LOCTEXT("CompensateKeyTooltip", "Compensate key at the current time."))
 			];
+			
 			BottomButtonsListBox->AddSlot()
 			.AutoWidth()
 			.VAlign(VAlign_Center)
-			.HAlign(HAlign_Left)
-			.Padding(0.f)
+			.Padding(4.f, 0.f)
 			[
-				SNew(SButton)
-				.ButtonStyle(FAppStyle::Get(), "FlatButton.Default")
+				SNew(SActionButton)
 				.Text(LOCTEXT("CompensateAllButton", "Comp All"))
 				.OnClicked(InArgs._OnCompensateAllButtonClicked)
-				.IsEnabled_Lambda([this]()
-				{
-					return (ControlKeys.Num() > 0);
-				})
+				.IsEnabled_Lambda([this](){ return ControlKeys.Num() > 0; })
 				.ToolTipText(LOCTEXT("CompensateAllTooltip", "Compensate all space switch keys."))
 			];
+			
 			BottomButtonsListBox->AddSlot()
 			.AutoWidth()
 			.VAlign(VAlign_Center)
-			.HAlign(HAlign_Right)
-			.Padding(0.f)
+			.Padding(4.f, 0.f)
 			[
-				SNew(SButton)
-				.ButtonStyle(FAppStyle::Get(), "FlatButton.Default")
+				SNew(SActionButton)
 				.Text(LOCTEXT("BakeButton", "Bake..."))
 				.OnClicked(InArgs._OnBakeButtonClicked)
-				.IsEnabled_Lambda([this]()
-				{
-					return (ControlKeys.Num() > 0);
-				})
+				.IsEnabled_Lambda([this]() { return ControlKeys.Num() > 0; })
 				.ToolTipText(LOCTEXT("BakeButtonToolTip", "Allows to bake the animation of one or more controls to a single space."))
 			];
 		}
