Control Rig: Constraints show Level Sequence by default

diff --git a/Engine/Plugins/Animation/ControlRig/Source/ControlRigEditor/Private/EditMode/Constrain/SConstraintsEditionWidget_ConstrainPanel.cpp b/Engine/Plugins/Animation/ControlRig/Source/ControlRigEditor/Private/EditMode/Constrain/SConstraintsEditionWidget_ConstrainPanel.cpp
--- a/Engine/Plugins/Animation/ControlRig/Source/ControlRigEditor/Private/EditMode/Constrain/SConstraintsEditionWidget_ConstrainPanel.cpp
+++ b/Engine/Plugins/Animation/ControlRig/Source/ControlRigEditor/Private/EditMode/Constrain/SConstraintsEditionWidget_ConstrainPanel.cpp
@@ -77,6 +77,9 @@ void SConstraintsEditionWidget_ConstrainPanel::Construct(
 		.BodyContent() [ SAssignNew(ConstraintsEditionWidget, SConstraintsEditionWidget) ]
 		.OnAreaExpansionChanged_Lambda( [this](bool bIsExpanded){ ConstraintsEditionWidget->RefreshConstraintList(); })
 	];
+
+	// LevelSequence is the most sensible default. It should be set back to this even if the user had the panel open earlier, closed it, and reopend it.
+	ConstraintsEditionWidget->SetShowConstraints(FBaseConstraintListWidget::EShowConstraints::ShowLevelSequence);
 }
 
 SConstraintsEditionWidget_ConstrainPanel::~SConstraintsEditionWidget_ConstrainPanel()
