Control Rig: Remove unused USE_LOCAL_DETAILS

diff --git a/Engine/Plugins/Animation/ControlRig/Source/ControlRigEditor/Private/EditMode/SControlRigEditModeTools.cpp b/Engine/Plugins/Animation/ControlRig/Source/ControlRigEditor/Private/EditMode/SControlRigEditModeTools.cpp
--- a/Engine/Plugins/Animation/ControlRig/Source/ControlRigEditor/Private/EditMode/SControlRigEditModeTools.cpp
+++ b/Engine/Plugins/Animation/ControlRig/Source/ControlRigEditor/Private/EditMode/SControlRigEditModeTools.cpp
@@ -63,10 +63,6 @@ void SControlRigEditModeTools::SetControlRigs(const TArrayView<TWeakObjectPtr<UC
 {
 	ControlRigs = InControlRigs;
 	UpdateOverridesDetailsView();
-
-#if USE_LOCAL_DETAILS
-	HierarchyTreeView->RefreshTreeView(true);
-#endif
 }
 
 const URigHierarchy* SControlRigEditModeTools::GetHierarchy() const
@@ -104,63 +100,6 @@ void SControlRigEditModeTools::Construct(const FArguments& InArgs, TSharedPtr<FC
 	ModeTools = InEditMode.GetModeManager();
 
 	SettingsDetailsView = FModuleManager::GetModuleChecked<FPropertyEditorModule>("PropertyEditor").CreateDetailView(DetailsViewArgs);
-
-#if USE_LOCAL_DETAILS
-	ControlEulerTransformDetailsView = FModuleManager::GetModuleChecked<FPropertyEditorModule>("PropertyEditor").CreateDetailView(DetailsViewArgs);
-	ControlEulerTransformDetailsView->SetKeyframeHandler(SharedThis(this));
-	ControlEulerTransformDetailsView->SetIsPropertyVisibleDelegate(FIsPropertyVisible::CreateSP(this, &SControlRigEditModeTools::ShouldShowPropertyOnDetailCustomization));
-	ControlEulerTransformDetailsView->SetIsPropertyReadOnlyDelegate(FIsPropertyReadOnly::CreateSP(this, &SControlRigEditModeTools::IsReadOnlyPropertyOnDetailCustomization));
-	ControlEulerTransformDetailsView->SetGenericLayoutDetailsDelegate(FOnGetDetailCustomizationInstance::CreateStatic(&FControlRigEditModeGenericDetails::MakeInstance, ModeTools));
-
-	ControlTransformDetailsView = FModuleManager::GetModuleChecked<FPropertyEditorModule>("PropertyEditor").CreateDetailView(DetailsViewArgs);
-	ControlTransformDetailsView->SetKeyframeHandler(SharedThis(this));
-	ControlTransformDetailsView->SetIsPropertyVisibleDelegate(FIsPropertyVisible::CreateSP(this, &SControlRigEditModeTools::ShouldShowPropertyOnDetailCustomization));
-	ControlTransformDetailsView->SetIsPropertyReadOnlyDelegate(FIsPropertyReadOnly::CreateSP(this, &SControlRigEditModeTools::IsReadOnlyPropertyOnDetailCustomization));
-	ControlTransformDetailsView->SetGenericLayoutDetailsDelegate(FOnGetDetailCustomizationInstance::CreateStatic(&FControlRigEditModeGenericDetails::MakeInstance, ModeTools));
-
-	ControlTransformNoScaleDetailsView = FModuleManager::GetModuleChecked<FPropertyEditorModule>("PropertyEditor").CreateDetailView(DetailsViewArgs);
-	ControlTransformNoScaleDetailsView->SetKeyframeHandler(SharedThis(this));
-	ControlTransformNoScaleDetailsView->SetIsPropertyVisibleDelegate(FIsPropertyVisible::CreateSP(this, &SControlRigEditModeTools::ShouldShowPropertyOnDetailCustomization));
-	ControlTransformNoScaleDetailsView->SetIsPropertyReadOnlyDelegate(FIsPropertyReadOnly::CreateSP(this, &SControlRigEditModeTools::IsReadOnlyPropertyOnDetailCustomization));
-	ControlTransformNoScaleDetailsView->SetGenericLayoutDetailsDelegate(FOnGetDetailCustomizationInstance::CreateStatic(&FControlRigEditModeGenericDetails::MakeInstance, ModeTools));
-
-	ControlFloatDetailsView = FModuleManager::GetModuleChecked<FPropertyEditorModule>("PropertyEditor").CreateDetailView(DetailsViewArgs);
-	ControlFloatDetailsView->SetKeyframeHandler(SharedThis(this));
-	ControlFloatDetailsView->SetIsPropertyVisibleDelegate(FIsPropertyVisible::CreateSP(this, &SControlRigEditModeTools::ShouldShowPropertyOnDetailCustomization));
-	ControlFloatDetailsView->SetIsPropertyReadOnlyDelegate(FIsPropertyReadOnly::CreateSP(this, &SControlRigEditModeTools::IsReadOnlyPropertyOnDetailCustomization));
-	ControlFloatDetailsView->SetGenericLayoutDetailsDelegate(FOnGetDetailCustomizationInstance::CreateStatic(&FControlRigEditModeGenericDetails::MakeInstance, ModeTools));
-
-	ControlEnumDetailsView = FModuleManager::GetModuleChecked<FPropertyEditorModule>("PropertyEditor").CreateDetailView(DetailsViewArgs);
-	ControlEnumDetailsView->SetKeyframeHandler(SharedThis(this));
-	ControlEnumDetailsView->SetIsPropertyVisibleDelegate(FIsPropertyVisible::CreateSP(this, &SControlRigEditModeTools::ShouldShowPropertyOnDetailCustomization));
-	ControlEnumDetailsView->SetIsPropertyReadOnlyDelegate(FIsPropertyReadOnly::CreateSP(this, &SControlRigEditModeTools::IsReadOnlyPropertyOnDetailCustomization));
-	ControlEnumDetailsView->SetGenericLayoutDetailsDelegate(FOnGetDetailCustomizationInstance::CreateStatic(&FControlRigEditModeGenericDetails::MakeInstance, ModeTools));
-
-	ControlIntegerDetailsView = FModuleManager::GetModuleChecked<FPropertyEditorModule>("PropertyEditor").CreateDetailView(DetailsViewArgs);
-	ControlIntegerDetailsView->SetKeyframeHandler(SharedThis(this));
-	ControlIntegerDetailsView->SetIsPropertyVisibleDelegate(FIsPropertyVisible::CreateSP(this, &SControlRigEditModeTools::ShouldShowPropertyOnDetailCustomization));
-	ControlIntegerDetailsView->SetIsPropertyReadOnlyDelegate(FIsPropertyReadOnly::CreateSP(this, &SControlRigEditModeTools::IsReadOnlyPropertyOnDetailCustomization));
-	ControlIntegerDetailsView->SetGenericLayoutDetailsDelegate(FOnGetDetailCustomizationInstance::CreateStatic(&FControlRigEditModeGenericDetails::MakeInstance, ModeTools));
-
-	ControlBoolDetailsView = FModuleManager::GetModuleChecked<FPropertyEditorModule>("PropertyEditor").CreateDetailView(DetailsViewArgs);
-	ControlBoolDetailsView->SetKeyframeHandler(SharedThis(this));
-	ControlBoolDetailsView->SetIsPropertyVisibleDelegate(FIsPropertyVisible::CreateSP(this, &SControlRigEditModeTools::ShouldShowPropertyOnDetailCustomization));
-	ControlBoolDetailsView->SetIsPropertyReadOnlyDelegate(FIsPropertyReadOnly::CreateSP(this, &SControlRigEditModeTools::IsReadOnlyPropertyOnDetailCustomization));
-	ControlBoolDetailsView->SetGenericLayoutDetailsDelegate(FOnGetDetailCustomizationInstance::CreateStatic(&FControlRigEditModeGenericDetails::MakeInstance, ModeTools));
-
-	ControlVectorDetailsView = FModuleManager::GetModuleChecked<FPropertyEditorModule>("PropertyEditor").CreateDetailView(DetailsViewArgs);
-	ControlVectorDetailsView->SetKeyframeHandler(SharedThis(this));
-	ControlVectorDetailsView->SetIsPropertyVisibleDelegate(FIsPropertyVisible::CreateSP(this, &SControlRigEditModeTools::ShouldShowPropertyOnDetailCustomization));
-	ControlVectorDetailsView->SetIsPropertyReadOnlyDelegate(FIsPropertyReadOnly::CreateSP(this, &SControlRigEditModeTools::IsReadOnlyPropertyOnDetailCustomization));
-	ControlVectorDetailsView->SetGenericLayoutDetailsDelegate(FOnGetDetailCustomizationInstance::CreateStatic(&FControlRigEditModeGenericDetails::MakeInstance, ModeTools));
-
-	ControlVector2DDetailsView = FModuleManager::GetModuleChecked<FPropertyEditorModule>("PropertyEditor").CreateDetailView(DetailsViewArgs);
-	ControlVector2DDetailsView->SetKeyframeHandler(SharedThis(this));
-	ControlVector2DDetailsView->SetIsPropertyVisibleDelegate(FIsPropertyVisible::CreateSP(this, &SControlRigEditModeTools::ShouldShowPropertyOnDetailCustomization));
-	ControlVector2DDetailsView->SetIsPropertyReadOnlyDelegate(FIsPropertyReadOnly::CreateSP(this, &SControlRigEditModeTools::IsReadOnlyPropertyOnDetailCustomization));
-	ControlVector2DDetailsView->SetGenericLayoutDetailsDelegate(FOnGetDetailCustomizationInstance::CreateStatic(&FControlRigEditModeGenericDetails::MakeInstance, ModeTools));
-#endif
-
 	OverridesDetailsView = FModuleManager::GetModuleChecked<FPropertyEditorModule>("PropertyEditor").CreateDetailView(DetailsViewArgs);
 	OverridesDetailsView->OnFinishedChangingProperties().AddSP(this, &SControlRigEditModeTools::OnOverrideOptionFinishedChange);
 	OverridesDetailsView->SetIsPropertyVisibleDelegate(FIsPropertyVisible::CreateSP(this, &SControlRigEditModeTools::ShouldShowOverrideProperty));
@@ -176,88 +115,18 @@ void SControlRigEditModeTools::Construct(const FArguments& InArgs, TSharedPtr<FC
 	DisplaySettings.bArrangeByModules = false;
 	DisplaySettings.bFlattenModules = false;
 	DisplaySettings.NameDisplayMode = EElementNameDisplayMode::AssetDefault;
-#if USE_LOCAL_DETAILS
-	FRigTreeDelegates RigTreeDelegates;
-	RigTreeDelegates.OnGetHierarchy = FOnGetRigTreeHierarchy::CreateSP(this, &SControlRigEditModeTools::GetHierarchy);
-	RigTreeDelegates.OnGetDisplaySettings = FOnGetRigTreeDisplaySettings::CreateSP(this, &SControlRigEditModeTools::GetDisplaySettings);
-	RigTreeDelegates.OnSelectionChanged = FOnRigTreeSelectionChanged::CreateSP(this, &SControlRigEditModeTools::HandleSelectionChanged);
-#endif
-
-
+	
 	ChildSlot
 	[
 		SNew(SScrollBox)
 		+ SScrollBox::Slot()
 		[
 			SNew(SVerticalBox)
-#if USE_LOCAL_DETAILS
-			+SVerticalBox::Slot()
-			.AutoHeight()
-			[
-				SAssignNew(PickerExpander, SExpandableArea)
-				.InitiallyCollapsed(true)
-				.AreaTitle(LOCTEXT("Picker_Header", "Controls"))
-				.AreaTitleFont(FAppStyle::GetFontStyle("DetailsView.CategoryFontStyle"))
-				.BorderBackgroundColor(FLinearColor(.6f, .6f, .6f))
-				.BodyContent()
-				[
-					SAssignNew(HierarchyTreeView, SRigHierarchyTreeView)
-					.RigTreeDelegates(RigTreeDelegates)
-				]
-			]
-#endif
 			+SVerticalBox::Slot()
 			.AutoHeight()
 			[
 				SettingsDetailsView.ToSharedRef()
 			]
-#if USE_LOCAL_DETAILS
-			+ SVerticalBox::Slot()
-			.AutoHeight()
-			[
-				ControlEulerTransformDetailsView.ToSharedRef()
-			]
-			+ SVerticalBox::Slot()
-			.AutoHeight()
-			[
-				ControlTransformDetailsView.ToSharedRef()
-			]
-			+ SVerticalBox::Slot()
-			.AutoHeight()
-			[
-				ControlTransformNoScaleDetailsView.ToSharedRef()
-			]
-			+ SVerticalBox::Slot()
-			.AutoHeight()
-			[
-				ControlBoolDetailsView.ToSharedRef()
-			]
-			+ SVerticalBox::Slot()
-			.AutoHeight()
-			[
-				ControlIntegerDetailsView.ToSharedRef()
-			]
-			+ SVerticalBox::Slot()
-			.AutoHeight()
-			[
-				ControlEnumDetailsView.ToSharedRef()
-			]
-			+ SVerticalBox::Slot()
-			.AutoHeight()
-			[
-				ControlVectorDetailsView.ToSharedRef()
-			]
-			+ SVerticalBox::Slot()
-			.AutoHeight()
-			[
-				ControlVector2DDetailsView.ToSharedRef()
-			]
-			+ SVerticalBox::Slot()
-			.AutoHeight()
-			[
-				ControlFloatDetailsView.ToSharedRef()
-			]
-#endif
 			
 			+SVerticalBox::Slot()
 			.AutoHeight()
@@ -291,9 +160,6 @@ void SControlRigEditModeTools::Construct(const FArguments& InArgs, TSharedPtr<FC
 			]
 		]
 	];
-#if USE_LOCAL_DETAILS
-	HierarchyTreeView->RefreshTreeView(true);
-#endif
 }
 
 void SControlRigEditModeTools::SetSettingsDetailsObject(const TWeakObjectPtr<>& InObject)
@@ -327,79 +193,6 @@ void SControlRigEditModeTools::SetSettingsDetailsObject(const TWeakObjectPtr<>&
 	}
 }
 
-#if USE_LOCAL_DETAILS
-
-void SControlRigEditModeTools::SetEulerTransformDetailsObjects(const TArray<TWeakObjectPtr<>>& InObjects)
-{
-	if (ControlEulerTransformDetailsView)
-	{
-		ControlEulerTransformDetailsView->SetObjects(InObjects);
-	}
-};
-
-void SControlRigEditModeTools::SetTransformDetailsObjects(const TArray<TWeakObjectPtr<>>& InObjects)
-{
-	if (ControlTransformDetailsView)
-	{
-		ControlTransformDetailsView->SetObjects(InObjects);
-	}
-}
-
-void SControlRigEditModeTools::SetTransformNoScaleDetailsObjects(const TArray<TWeakObjectPtr<>>& InObjects)
-{
-	if (ControlTransformNoScaleDetailsView)
-	{
-		ControlTransformNoScaleDetailsView->SetObjects(InObjects);
-	}
-}
-
-void SControlRigEditModeTools::SetFloatDetailsObjects(const TArray<TWeakObjectPtr<>>& InObjects)
-{
-	if (ControlFloatDetailsView)
-	{
-		ControlFloatDetailsView->SetObjects(InObjects);
-	}
-}
-
-void SControlRigEditModeTools::SetBoolDetailsObjects(const TArray<TWeakObjectPtr<>>& InObjects)
-{
-	if (ControlBoolDetailsView)
-	{
-		ControlBoolDetailsView->SetObjects(InObjects);
-	}
-}
-
-void SControlRigEditModeTools::SetIntegerDetailsObjects(const TArray<TWeakObjectPtr<>>& InObjects)
-{
-	if (ControlIntegerDetailsView)
-	{
-		ControlIntegerDetailsView->SetObjects(InObjects);
-	}
-}
-void SControlRigEditModeTools::SetEnumDetailsObjects(const TArray<TWeakObjectPtr<>>& InObjects)
-{
-	if (ControlVectorDetailsView)
-	{
-		ControlVectorDetailsView->SetObjects(InObjects);
-	}
-}
-
-void SControlRigEditModeTools::SetVectorDetailsObjects(const TArray<TWeakObjectPtr<>>& InObjects)
-{
-	if (ControlEnumDetailsView)
-	{
-		ControlEnumDetailsView->SetObjects(InObjects);
-	}
-}
-
-void SControlRigEditModeTools::SetVector2DDetailsObjects(const TArray<TWeakObjectPtr<>>& InObjects)
-{
-	if (ControlVector2DDetailsView)
-	{
-		ControlVector2DDetailsView->SetObjects(InObjects);
-	}
-}
-#endif
 void SControlRigEditModeTools::SetSequencer(TWeakPtr<ISequencer> InSequencer)
 {
 	WeakSequencer = InSequencer.Pin();
@@ -540,75 +333,6 @@ bool SControlRigEditModeTools::IsReadOnlyPropertyOnDetailCustomization(const FPr
 		(InPropertyAndParent.ParentProperties.Num() > 0 && ShouldPropertyBeEnabled(*InPropertyAndParent.ParentProperties[0])));
 }
 
-#if USE_LOCAL_DETAILS
-static bool bPickerChangingSelection = false;
-
-void SControlRigEditModeTools::OnManipulatorsPicked(const TArray<FName>& Manipulators)
-{
-	FControlRigEditMode* ControlRigEditMode = static_cast<FControlRigEditMode*>(ModeTools->GetActiveMode(FControlRigEditMode::ModeName));
-	if (ControlRigEditMode)
-	{
-		if (!bPickerChangingSelection)
-		{
-			TGuardValue<bool> SelectGuard(bPickerChangingSelection, true);
-			ControlRigEditMode->ClearRigElementSelection((uint32)ERigElementType::Control);
-			ControlRigEditMode->SetRigElementSelection(ERigElementType::Control, Manipulators, true);
-		}
-	}
-}
-
-void SControlRigEditModeTools::HandleModifiedEvent(ERigVMGraphNotifType InNotifType, URigVMGraph* InGraph, UObject* InSubject)
-{
-	if (bPickerChangingSelection)
-	{
-		return;
-	}
-
-	TGuardValue<bool> SelectGuard(bPickerChangingSelection, true);
-	switch (InNotifType)
-	{
-		case ERigVMGraphNotifType::NodeSelected:
-		case ERigVMGraphNotifType::NodeDeselected:
-		{
-			URigVMNode* Node = Cast<URigVMNode>(InSubject);
-			if (Node)
-			{
-				// those are not yet implemented yet
-				// ControlPicker->SelectManipulator(Node->Name, InType == EControlRigModelNotifType::NodeSelected);
-			}
-			break;
-		}
-		default:
-		{
-			break;
-		}
-	}
-}
-
-void SControlRigEditModeTools::HandleSelectionChanged(TSharedPtr<FRigTreeElement> Selection, ESelectInfo::Type SelectInfo)
-{
-	if(bIsChangingRigHierarchy)
-	{
-		return;
-	}
-	
-	const URigHierarchy* Hierarchy = GetHierarchy();
-	if (Hierarchy)
-	{
-		URigHierarchyController* Controller = ((URigHierarchy*)Hierarchy)->GetController(true);
-		check(Controller);
-		
-		TGuardValue<bool> GuardRigHierarchyChanges(bIsChangingRigHierarchy, true);
-
-		const TArray<FRigElementKey> NewSelection = HierarchyTreeView->GetSelectedKeys();
-		if(!Controller->SetSelection(NewSelection))
-		{
-			return;
-		}
-	}
-}
-#endif
-
 void SControlRigEditModeTools::CustomizeToolBarPalette(FToolBarBuilder& ToolBarBuilder)
 {
 	//TOGGLE SELECTED RIG CONTROLS
diff --git a/Engine/Plugins/Animation/ControlRig/Source/ControlRigEditor/Private/EditMode/SControlRigEditModeTools.h b/Engine/Plugins/Animation/ControlRig/Source/ControlRigEditor/Private/EditMode/SControlRigEditModeTools.h
--- a/Engine/Plugins/Animation/ControlRig/Source/ControlRigEditor/Private/EditMode/SControlRigEditModeTools.h
+++ b/Engine/Plugins/Animation/ControlRig/Source/ControlRigEditor/Private/EditMode/SControlRigEditModeTools.h
@@ -23,8 +23,6 @@ class FToolBarBuilder;
 class FEditorModeTools;
 class FControlRigEditModeToolkit;
 
-#define USE_LOCAL_DETAILS 0
-
 class SControlRigEditModeTools : public SCompoundWidget, public IDetailKeyframeHandler
 {
 public:
@@ -36,17 +34,6 @@ public:
 
 	/** Set the objects to be displayed in the details panel */
 	void SetSettingsDetailsObject(const TWeakObjectPtr<>& InObject);
-#if USE_LOCAL_DETAILS
-	void SetEulerTransformDetailsObjects(const TArray<TWeakObjectPtr<>>& InObjects);
-	void SetTransformDetailsObjects(const TArray<TWeakObjectPtr<>>& InObjects);
-	void SetTransformNoScaleDetailsObjects(const TArray<TWeakObjectPtr<>>& InObjects);
-	void SetFloatDetailsObjects(const TArray<TWeakObjectPtr<>>& InObjects);
-	void SetBoolDetailsObjects(const TArray<TWeakObjectPtr<>>& InObjects);
-	void SetIntegerDetailsObjects(const TArray<TWeakObjectPtr<>>& InObjects);
-	void SetEnumDetailsObjects(const TArray<TWeakObjectPtr<>>& InObjects);
-	void SetVectorDetailsObjects(const TArray<TWeakObjectPtr<>>& InObjects);
-	void SetVector2DDetailsObjects(const TArray<TWeakObjectPtr<>>& InObjects);
-#endif
 
 	/** Set the sequencer we are bound to */
 	void SetSequencer(TWeakPtr<ISequencer> InSequencer);
@@ -68,23 +55,8 @@ public:
 private:
 
 	/** Sequencer we are currently bound to */
-
 	TWeakPtr<ISequencer> WeakSequencer;
 	TSharedPtr<IDetailsView> SettingsDetailsView;
-
-#if USE_LOCAL_DETAILS
-	/** The details views we do most of our work within */
-	TSharedPtr<IDetailsView> ControlEulerTransformDetailsView;
-	TSharedPtr<IDetailsView> ControlTransformDetailsView;
-	TSharedPtr<IDetailsView> ControlTransformNoScaleDetailsView;
-	TSharedPtr<IDetailsView> ControlFloatDetailsView;
-	TSharedPtr<IDetailsView> ControlBoolDetailsView;
-	TSharedPtr<IDetailsView> ControlIntegerDetailsView;
-	TSharedPtr<IDetailsView> ControlEnumDetailsView;
-	TSharedPtr<IDetailsView> ControlVector2DDetailsView;
-	TSharedPtr<IDetailsView> ControlVectorDetailsView;
-	TSharedPtr<SRigHierarchyTreeView> HierarchyTreeView;
-#endif
 	
 	/** Storage for control rigs */
 	TArray<TWeakObjectPtr<UControlRig>> ControlRigs;
