Control Rig: You can now only spawn one space picker using tab at a time. Spawning it again replaces the old one.

diff --git a/Engine/Plugins/Animation/ControlRig/Source/ControlRigEditor/Private/EditMode/ControlRigEditMode.cpp b/Engine/Plugins/Animation/ControlRig/Source/ControlRigEditor/Private/EditMode/ControlRigEditMode.cpp
--- a/Engine/Plugins/Animation/ControlRig/Source/ControlRigEditor/Private/EditMode/ControlRigEditMode.cpp
+++ b/Engine/Plugins/Animation/ControlRig/Source/ControlRigEditor/Private/EditMode/ControlRigEditMode.cpp
@@ -86,6 +86,7 @@
 #include "Slate/SceneViewport.h"
 #include "EditorGizmos/EditorTransformGizmoUtil.h"
 #include "Misc/RigEditModeUtils.h"
+#include "Picker/FloatingSpacePickerManager.h"
 #include "Sequencer/AnimLayers/AnimLayers.h"
 #include "Tools/MotionTrailOptions.h"
 
@@ -3396,169 +3397,16 @@ void FControlRigEditMode::SetAnimLayerPassthroughKey()
 void FControlRigEditMode::OpenSpacePickerWidget()
 {
 	using namespace UE::ControlRigEditor;
-	const FInitialSpacePickerSelection InitialSelection = DetermineInitialSpacePickerSelection(*this);
-	if (!InitialSelection.IsValid())
+
+	// Lazily create it. If you need access to it outside OpenSpacePickerWidget, consider refactoring creation to be when this mode is created.
+	if (!FloatingSpacePickerManager)
 	{
-		return;
+		FloatingSpacePickerManager = MakePimpl<FFloatingSpacePickerManager>(*this);
 	}
-	const TArray<FRigElementKey>& SelectedControls = InitialSelection.SelectedControls;
-	UControlRig* RuntimeRig = InitialSelection.RuntimeRig;
 	
-	TSharedRef<SRigSpacePickerWidget> PickerWidget =
-	SNew(SRigSpacePickerWidget)
-	.Hierarchy(InitialSelection.Hierarchy)
-	.Controls(InitialSelection.SelectedControls)
-	.Title(LOCTEXT("PickSpace", "Pick Space"))
-	.AllowDelete(AreEditingControlRigDirectly())
-	.AllowReorder(AreEditingControlRigDirectly())
-	.AllowAdd(AreEditingControlRigDirectly())
-	.GetControlCustomization_Lambda([this, RuntimeRig](URigHierarchy*, const FRigElementKey& InControlKey)
-	{
-		return RuntimeRig->GetControlCustomization(InControlKey);
-	})
-	.OnActiveSpaceChanged_Lambda([this, SelectedControls, RuntimeRig](URigHierarchy* InHierarchy, const FRigElementKey& InControlKey, const FRigElementKey& InSpaceKey)
-	{
-		check(SelectedControls.Contains(InControlKey));
-		if (!AreEditingControlRigDirectly())
-		{
-			if (WeakSequencer.IsValid())
-			{
-				if (const FRigControlElement* ControlElement = InHierarchy->Find<FRigControlElement>(InControlKey))
-				{
-					ISequencer* Sequencer = WeakSequencer.Pin().Get();
-					if (Sequencer)
-					{
-						constexpr bool bDoNotCreate = false;
-						FSpaceChannelAndSection SpaceChannelAndSection = FControlRigSpaceChannelHelpers::FindSpaceChannelAndSectionForControl(RuntimeRig, InControlKey.Name, Sequencer, bDoNotCreate);
-						if (!SpaceChannelAndSection.SpaceChannel)
-						{
-							// no need to create a space channel if InSpaceKey is the current parent 
-							if (InHierarchy->GetActiveParent(InControlKey) == InSpaceKey)
-							{
-								return;
-							}
-						}
-						
-						FScopedTransaction Transaction(LOCTEXT("KeyControlRigSpace", "Key Control Rig Space"));
-						SpaceChannelAndSection = FControlRigSpaceChannelHelpers::FindSpaceChannelAndSectionForControl(RuntimeRig, InControlKey.Name, Sequencer, true /*bCreateIfNeeded*/);
-						if (SpaceChannelAndSection.SpaceChannel)
-						{
-							const FFrameRate TickResolution = Sequencer->GetFocusedTickResolution();
-							const FFrameTime FrameTime = Sequencer->GetLocalTime().ConvertTo(TickResolution);
-							FFrameNumber CurrentTime = FrameTime.GetFrame();
-							FControlRigSpaceChannelHelpers::SequencerKeyControlRigSpaceChannel(RuntimeRig, Sequencer, SpaceChannelAndSection.SpaceChannel, SpaceChannelAndSection.SectionToKey, CurrentTime, InHierarchy, InControlKey, InSpaceKey);
-						}
-					}
-				}
-			}
-		}
-		else
-		{
-			if (RuntimeRig->IsAdditive())
-			{
-				const FTransform Transform = RuntimeRig->GetControlGlobalTransform(InControlKey.Name);
-				RuntimeRig->SwitchToParent(InControlKey, InSpaceKey, false, true);
-				{
-					const TGuardValue<float> AbsoluteTimeGuard(RuntimeRig->AbsoluteTime, RuntimeRig->AbsoluteTime);
-				    RuntimeRig->Evaluate_AnyThread();
-				}
-				FRigControlValue ControlValue = RuntimeRig->GetControlValueFromGlobalTransform(InControlKey.Name, Transform, ERigTransformType::CurrentGlobal);
-				RuntimeRig->SetControlValue(InControlKey.Name, ControlValue);
-				{
-					const TGuardValue<float> AbsoluteTimeGuard(RuntimeRig->AbsoluteTime, RuntimeRig->AbsoluteTime);
-					RuntimeRig->Evaluate_AnyThread();
-				}
-			}
-			else
-			{
-				constexpr bool bInitial = false, bAffectChildren = true;
-				const FTransform Transform = InHierarchy->GetGlobalTransform(InControlKey);
-				FRigDependenciesProviderForControlRig DependenciesProvider(RuntimeRig);
-				FString OutFailureReason;
-				if (InHierarchy->SwitchToParent(InControlKey, InSpaceKey, bInitial, bAffectChildren, DependenciesProvider, &OutFailureReason))
-				{
-					InHierarchy->SetGlobalTransform(InControlKey, Transform);
-				}
-				else
-				{
-					if(URigHierarchyController* Controller = InHierarchy->GetController())
-					{
-						static constexpr TCHAR MessageFormat[] = TEXT("Could not switch %s to parent %s: %s");
-						Controller->ReportAndNotifyErrorf(MessageFormat,
-							*InControlKey.Name.ToString(),
-							*InSpaceKey.Name.ToString(),
-							*OutFailureReason);
-					}
-				}
-			}
-		}
-		
-	})
-	.OnSpaceListChanged_Lambda([this, SelectedControls, RuntimeRig](URigHierarchy* InHierarchy, const FRigElementKey& InControlKey, const TArray<FRigElementKeyWithLabel>& InSpaceList)
-	{
-		check(SelectedControls.Contains(InControlKey));
-
-		// check if we are in the control rig editor or in the level
-		if(AreEditingControlRigDirectly())
-		{
-			if (FControlRigAssetInterfacePtr Blueprint = RuntimeRig->GetClass()->ClassGeneratedBy)
-			{
-				if(URigHierarchy* Hierarchy = Blueprint->GetHierarchy())
-				{
-					// update the settings in the control element
-					if(FRigControlElement* ControlElement = Hierarchy->Find<FRigControlElement>(InControlKey))
-					{
-						Blueprint->Modify();
-						FScopedTransaction Transaction(LOCTEXT("ControlChangeAvailableSpaces", "Edit Available Spaces"));
-
-						ControlElement->Settings.Customization.AvailableSpaces = InSpaceList;
-						Hierarchy->Notify(ERigHierarchyNotification::ControlSettingChanged, ControlElement);
-					}
-
-					// also update the debugged instance
-					if(Hierarchy != InHierarchy)
-					{
-						if(FRigControlElement* ControlElement = InHierarchy->Find<FRigControlElement>(InControlKey))
-						{
-							ControlElement->Settings.Customization.AvailableSpaces = InSpaceList;
-							InHierarchy->Notify(ERigHierarchyNotification::ControlSettingChanged, ControlElement);
-						}
-					}
-				}
-			}
-		}
-		else
-		{
-			// update the settings in the control element
-			if(FRigControlElement* ControlElement = InHierarchy->Find<FRigControlElement>(InControlKey))
-			{
-				FScopedTransaction Transaction(LOCTEXT("ControlChangeAvailableSpaces", "Edit Available Spaces"));
-
-				InHierarchy->Modify();
-
-				FRigControlElementCustomization ControlCustomization = *RuntimeRig->GetControlCustomization(InControlKey);	
-				ControlCustomization.AvailableSpaces = InSpaceList;
-				ControlCustomization.RemovedSpaces.Reset();
-
-				// remember  the elements which are in the asset's available list but removed by the user
-				for(const FRigElementKeyWithLabel& AvailableSpace : ControlElement->Settings.Customization.AvailableSpaces)
-				{
-					if(ControlCustomization.AvailableSpaces.FindByKey(AvailableSpace.Key) == nullptr)
-					{
-						ControlCustomization.RemovedSpaces.Add(AvailableSpace.Key);
-					}
-				}
-
-				RuntimeRig->SetControlCustomization(InControlKey, ControlCustomization);
-				InHierarchy->Notify(ERigHierarchyNotification::ControlSettingChanged, ControlElement);
-			}
-		}
-
-	});
-	// todo: implement GetAdditionalSpacesDelegate to pull spaces from sequencer
-
-	PickerWidget->OpenDialog(false);
+	FloatingSpacePickerManager->SummonSpacePickerAtCursor();
 }
+
 static void ToggleTool(const FString& InToolName)
 {
 	if (FLevelEditorModule* LevelEditorModule = FModuleManager::GetModulePtr<FLevelEditorModule>(TEXT("LevelEditor")))
diff --git a/Engine/Plugins/Animation/ControlRig/Source/ControlRigEditor/Private/EditMode/Picker/FloatingSpacePickerManager.cpp b/Engine/Plugins/Animation/ControlRig/Source/ControlRigEditor/Private/EditMode/Picker/FloatingSpacePickerManager.cpp
new file mode 100644
--- /dev/null
+++ b/Engine/Plugins/Animation/ControlRig/Source/ControlRigEditor/Private/EditMode/Picker/FloatingSpacePickerManager.cpp
@@ -0,0 +1,68 @@
+﻿// Copyright Epic Games, Inc. All Rights Reserved.
+
+#include "FloatingSpacePickerManager.h"
+
+#include "SFloatingSpacePicker_Base.h"
+#include "SFloatingSpacePicker_AssetEditor.h"
+#include "SFloatingSpacePicker_LevelEditor.h"
+#include "Misc/RigEditModeUtils.h"
+
+namespace UE::ControlRigEditor
+{
+FFloatingSpacePickerManager::FFloatingSpacePickerManager(FControlRigEditMode& InOwningMode)
+	: OwningMode(InOwningMode)
+{}
+
+FFloatingSpacePickerManager::~FFloatingSpacePickerManager()
+{
+	CloseSpacePicker();
+}
+
+void FFloatingSpacePickerManager::SummonSpacePickerAtCursor()
+{
+	// We could reuse the window but it's easier to just recreate the widget from scratch. The overhead is fine.
+	CloseSpacePicker();
+
+	if (const TSharedPtr<SFloatingSpacePicker_Base> SpacePicker = CreateSpacePicker())
+	{
+		WeakWindow = SpacePicker->ShowWindow();
+	}
+}
+
+void FFloatingSpacePickerManager::CloseSpacePicker()
+{
+	if (const TSharedPtr<SWindow> WindowPin = WeakWindow.Pin())
+	{
+		WindowPin->RequestDestroyWindow();
+		WeakWindow.Reset();
+	}
+}
+
+TSharedPtr<SFloatingSpacePicker_Base> FFloatingSpacePickerManager::CreateSpacePicker()
+{
+	using namespace UE::ControlRigEditor;
+	const FInitialSpacePickerSelection InitialSelection = DetermineInitialSpacePickerSelection(OwningMode);
+	if (!InitialSelection.IsValid())
+	{
+		return nullptr;
+	}
+
+	if (OwningMode.AreEditingControlRigDirectly())
+	{
+		return SNew(SFloatingSpacePicker_AssetEditor)
+			.SelectedControls(InitialSelection.SelectedControls)
+			.DisplayedRig(InitialSelection.RuntimeRig);
+	}
+	else
+	{
+		TAttribute<TSharedPtr<ISequencer>> SequencerAttr = TAttribute<TSharedPtr<ISequencer>>::CreateLambda([this]
+		{
+			// OwningMode should be valid here because it is responsible for destroying us when it is destroyed.
+			return OwningMode.GetWeakSequencer().Pin();
+		});
+		return SNew(SFloatingSpacePicker_LevelEditor, MoveTemp(SequencerAttr))
+			.SelectedControls(InitialSelection.SelectedControls)
+			.DisplayedRig(InitialSelection.RuntimeRig);
+	}
+}
+}
\ No newline at end of file
diff --git a/Engine/Plugins/Animation/ControlRig/Source/ControlRigEditor/Private/EditMode/Picker/FloatingSpacePickerManager.h b/Engine/Plugins/Animation/ControlRig/Source/ControlRigEditor/Private/EditMode/Picker/FloatingSpacePickerManager.h
new file mode 100644
--- /dev/null
+++ b/Engine/Plugins/Animation/ControlRig/Source/ControlRigEditor/Private/EditMode/Picker/FloatingSpacePickerManager.h
@@ -0,0 +1,42 @@
+﻿// Copyright Epic Games, Inc. All Rights Reserved.
+
+#pragma once
+
+#include "EditMode/ControlRigEditMode.h"
+#include "Templates/UnrealTemplate.h"
+
+namespace UE::ControlRigEditor
+{
+}
+
+namespace UE::ControlRigEditor
+{
+class SFloatingSpacePicker_Base;
+	
+/** Manages the floating space picker, which by default is spawned by pressing tab. */
+class FFloatingSpacePickerManager : public FNoncopyable
+{
+public:
+
+	explicit FFloatingSpacePickerManager(FControlRigEditMode& InOwningMode UE_LIFETIMEBOUND);
+	~FFloatingSpacePickerManager();
+
+	/** Spawns a new space picker at the cursor. Replaces the old one if one exists. */
+	void SummonSpacePickerAtCursor();
+
+	/** Closes the space picker if it currently is being shown. */
+	void CloseSpacePicker();
+
+private:
+
+	/** The edit mode that owns us. Needed to display some information about the space picker. */
+	FControlRigEditMode& OwningMode;
+
+	/** The window that contained the last created space picker. If still valid, we replace the window. */
+	TWeakPtr<SWindow> WeakWindow;
+
+	/** Creates the space picker depending on the FControlRigEditMode::AreEditingControlRigDirectly. */
+	TSharedPtr<SFloatingSpacePicker_Base> CreateSpacePicker();
+};
+}
+
diff --git a/Engine/Plugins/Animation/ControlRig/Source/ControlRigEditor/Private/EditMode/Picker/SFloatingSpacePicker_AssetEditor.cpp b/Engine/Plugins/Animation/ControlRig/Source/ControlRigEditor/Private/EditMode/Picker/SFloatingSpacePicker_AssetEditor.cpp
new file mode 100644
--- /dev/null
+++ b/Engine/Plugins/Animation/ControlRig/Source/ControlRigEditor/Private/EditMode/Picker/SFloatingSpacePicker_AssetEditor.cpp
@@ -0,0 +1,124 @@
+﻿// Copyright Epic Games, Inc. All Rights Reserved.
+
+#include "SFloatingSpacePicker_AssetEditor.h"
+
+#include "ControlRigSpaceChannelEditors.h"
+#include "ScopedTransaction.h"
+#include "Editor/SRigSpacePickerWidget.h"
+#include "Rigs/RigHierarchyController.h"
+
+#define LOCTEXT_NAMESPACE "SFloatingSpacePicker"
+
+namespace UE::ControlRigEditor
+{
+namespace DirectEditDetails
+{
+struct FGuardAbsoluteTime
+{
+	UControlRig* const ControlRig;
+	const float AbsoluteTime;
+	
+	explicit FGuardAbsoluteTime(UControlRig* InControlRig) : ControlRig(InControlRig), AbsoluteTime(ControlRig->GetAbsoluteTime()) {}
+	~FGuardAbsoluteTime(){ ControlRig->SetAbsoluteTime(AbsoluteTime); }
+};
+}
+	
+void SFloatingSpacePicker_AssetEditor::Construct(const FArguments& InArgs)
+{
+	constexpr bool bIsEditingDirectly = true;
+	SFloatingSpacePicker_Base::Construct(InArgs, bIsEditingDirectly);
+}
+
+void SFloatingSpacePicker_AssetEditor::OnActiveSpaceChanged(
+	URigHierarchy* InHierarchy, const FRigElementKey& InControlKey, const FRigElementKey& InSpaceKey
+	)
+{
+	check(SelectedControls.Contains(InControlKey));
+
+	UControlRig* RuntimeRig = WeakDisplayedRig.Get();
+	if (!ensure(RuntimeRig))
+	{
+		return;
+	}
+	
+	if (RuntimeRig->IsAdditive())
+	{
+		const FTransform Transform = RuntimeRig->GetControlGlobalTransform(InControlKey.Name);
+		RuntimeRig->SwitchToParent(InControlKey, InSpaceKey, false, true);
+		{
+			DirectEditDetails::FGuardAbsoluteTime Guard(RuntimeRig);
+			RuntimeRig->Evaluate_AnyThread();
+		}
+		FRigControlValue ControlValue = RuntimeRig->GetControlValueFromGlobalTransform(InControlKey.Name, Transform, ERigTransformType::CurrentGlobal);
+		RuntimeRig->SetControlValue(InControlKey.Name, ControlValue);
+		{
+			DirectEditDetails::FGuardAbsoluteTime Guard(RuntimeRig);
+			RuntimeRig->Evaluate_AnyThread();
+		}
+	}
+	else
+	{
+		constexpr bool bInitial = false, bAffectChildren = true;
+		const FTransform Transform = InHierarchy->GetGlobalTransform(InControlKey);
+		FRigDependenciesProviderForControlRig DependenciesProvider(RuntimeRig);
+		FString OutFailureReason;
+		if (InHierarchy->SwitchToParent(InControlKey, InSpaceKey, bInitial, bAffectChildren, DependenciesProvider, &OutFailureReason))
+		{
+			InHierarchy->SetGlobalTransform(InControlKey, Transform);
+		}
+		else
+		{
+			if( URigHierarchyController* Controller = InHierarchy->GetController())
+			{
+				static constexpr TCHAR MessageFormat[] = TEXT("Could not switch %s to parent %s: %s");
+				Controller->ReportAndNotifyErrorf(MessageFormat,
+					*InControlKey.Name.ToString(),
+					*InSpaceKey.Name.ToString(),
+					*OutFailureReason);
+			}
+		}
+	}
+}
+
+void SFloatingSpacePicker_AssetEditor::OnSpaceListChanged(
+	URigHierarchy* InHierarchy, const FRigElementKey& InControlKey, const TArray<FRigElementKeyWithLabel>& InSpaceList
+	)
+{
+	check(SelectedControls.Contains(InControlKey));
+
+	UControlRig* RuntimeRig = WeakDisplayedRig.Get();
+	if (!ensure(RuntimeRig))
+	{
+		return;
+	}
+
+	FControlRigAssetInterfacePtr Blueprint = RuntimeRig->GetClass()->ClassGeneratedBy;
+	URigHierarchy* BlueprintHierarchy = Blueprint ? Blueprint->GetHierarchy() : nullptr;
+	if (!BlueprintHierarchy)
+	{
+		return;
+	}
+	
+	// update the settings in the control element
+	if (FRigControlElement* ControlElement = BlueprintHierarchy->Find<FRigControlElement>(InControlKey))
+	{
+		Blueprint->Modify();
+		FScopedTransaction Transaction(LOCTEXT("ControlChangeAvailableSpaces", "Edit Available Spaces"));
+
+		ControlElement->Settings.Customization.AvailableSpaces = InSpaceList;
+		BlueprintHierarchy->Notify(ERigHierarchyNotification::ControlSettingChanged, ControlElement);
+	}
+
+	// also update the debugged instance
+	if (BlueprintHierarchy != InHierarchy)
+	{
+		if(FRigControlElement* ControlElement = InHierarchy->Find<FRigControlElement>(InControlKey))
+		{
+			ControlElement->Settings.Customization.AvailableSpaces = InSpaceList;
+			InHierarchy->Notify(ERigHierarchyNotification::ControlSettingChanged, ControlElement);
+		}
+	}
+}
+}
+
+#undef LOCTEXT_NAMESPACE
\ No newline at end of file
diff --git a/Engine/Plugins/Animation/ControlRig/Source/ControlRigEditor/Private/EditMode/Picker/SFloatingSpacePicker_AssetEditor.h b/Engine/Plugins/Animation/ControlRig/Source/ControlRigEditor/Private/EditMode/Picker/SFloatingSpacePicker_AssetEditor.h
new file mode 100644
--- /dev/null
+++ b/Engine/Plugins/Animation/ControlRig/Source/ControlRigEditor/Private/EditMode/Picker/SFloatingSpacePicker_AssetEditor.h
@@ -0,0 +1,34 @@
+﻿// Copyright Epic Games, Inc. All Rights Reserved.
+
+#pragma once
+
+#include "ControlRig.h"
+#include "SFloatingSpacePicker_Base.h"
+#include "Widgets/DeclarativeSyntaxSupport.h"
+#include "Widgets/SCompoundWidget.h"
+
+class FControlRigEditMode;
+class SRigSpacePickerWidget;
+
+namespace UE::ControlRigEditor
+{
+/**
+ * Handles displaying the floating space picker when directly editing control rigs,
+ * i.e. when FControlRigEditMode::IsEditingControlRigDirectly() == true.
+ */
+class SFloatingSpacePicker_AssetEditor : public SFloatingSpacePicker_Base
+{
+public:
+	
+	void Construct(const FArguments& InArgs);
+
+private:
+
+	//~ Begin SFloatingSpacePicker_Base Interface
+	virtual void OnActiveSpaceChanged(URigHierarchy* InHierarchy, const FRigElementKey& InControlKey, const FRigElementKey& InSpaceKey) override;
+	virtual void OnSpaceListChanged(URigHierarchy* InHierarchy, const FRigElementKey& InControlKey, const TArray<FRigElementKeyWithLabel>& InSpaceList) override;
+	//~ End SFloatingSpacePicker_Base Interface
+};
+}
+
+
diff --git a/Engine/Plugins/Animation/ControlRig/Source/ControlRigEditor/Private/EditMode/Picker/SFloatingSpacePicker_Base.cpp b/Engine/Plugins/Animation/ControlRig/Source/ControlRigEditor/Private/EditMode/Picker/SFloatingSpacePicker_Base.cpp
new file mode 100644
--- /dev/null
+++ b/Engine/Plugins/Animation/ControlRig/Source/ControlRigEditor/Private/EditMode/Picker/SFloatingSpacePicker_Base.cpp
@@ -0,0 +1,46 @@
+﻿// Copyright Epic Games, Inc. All Rights Reserved.
+
+#include "SFloatingSpacePicker_Base.h"
+
+#include "ControlRigSpaceChannelEditors.h"
+#include "ScopedTransaction.h"
+#include "Editor/SRigSpacePickerWidget.h"
+#include "Rigs/RigHierarchyController.h"
+
+#define LOCTEXT_NAMESPACE "SFloatingSpacePicker"
+
+namespace UE::ControlRigEditor
+{
+void SFloatingSpacePicker_Base::Construct(const FArguments& InArgs, bool bIsEditingDirectly)
+{
+	check(InArgs._DisplayedRig && InArgs._DisplayedRig->GetHierarchy() && !InArgs._SelectedControls.IsEmpty());
+	
+	WeakDisplayedRig = InArgs._DisplayedRig;
+	SelectedControls = InArgs._SelectedControls;
+	
+	SRigSpacePickerWidget::Construct(
+		SRigSpacePickerWidget::FArguments()
+		.Hierarchy(InArgs._DisplayedRig->GetHierarchy())
+		.Controls(SelectedControls)
+		.Title(LOCTEXT("PickSpace", "Pick Space"))
+		.AllowDelete(bIsEditingDirectly)
+		.AllowReorder(bIsEditingDirectly)
+		.AllowAdd(bIsEditingDirectly)
+		.GetControlCustomization_Lambda([this](URigHierarchy*, const FRigElementKey& InControlKey)
+		{
+			const UControlRig* RuntimeRig = WeakDisplayedRig.Get();
+			return ensure(RuntimeRig) ? RuntimeRig->GetControlCustomization(InControlKey) : nullptr;
+		})
+		.OnActiveSpaceChanged(this, &SFloatingSpacePicker_Base::OnActiveSpaceChanged)
+		.OnSpaceListChanged(this, &SFloatingSpacePicker_Base::OnSpaceListChanged)
+	);
+}
+
+TSharedPtr<SWindow> SFloatingSpacePicker_Base::ShowWindow()
+{
+	constexpr bool bModal = false;
+	return OpenDialog(bModal);
+}
+}
+
+#undef LOCTEXT_NAMESPACE
\ No newline at end of file
diff --git a/Engine/Plugins/Animation/ControlRig/Source/ControlRigEditor/Private/EditMode/Picker/SFloatingSpacePicker_Base.h b/Engine/Plugins/Animation/ControlRig/Source/ControlRigEditor/Private/EditMode/Picker/SFloatingSpacePicker_Base.h
new file mode 100644
--- /dev/null
+++ b/Engine/Plugins/Animation/ControlRig/Source/ControlRigEditor/Private/EditMode/Picker/SFloatingSpacePicker_Base.h
@@ -0,0 +1,53 @@
+﻿// Copyright Epic Games, Inc. All Rights Reserved.
+
+#pragma once
+
+#include "ControlRig.h"
+#include "Editor/SRigSpacePickerWidget.h"
+#include "Widgets/DeclarativeSyntaxSupport.h"
+#include "Widgets/SCompoundWidget.h"
+
+class FControlRigEditMode;
+class SRigSpacePickerWidget;
+
+namespace UE::ControlRigEditor
+{
+/**
+ * Handles displaying the floating space picker when directly editing control rigs,
+ * i.e. when FControlRigEditMode::IsEditingControlRigDirectly() == true.
+ */
+class SFloatingSpacePicker_Base
+	// Initially this was SCompoundWidget and created SRigSpacePickerWidget as child widget. Won't work.
+	// We need to inherit because SRigSpacePickerWidget::OpenDialog uses AsShared to put itself into the SWindow; but then the parent widget, which is
+	// this SFloatingSpacePicker_Base, would wind up unreferenced and destroyed.
+	// SRigSpacePickerWidget::OpenDialog should be adjusted to work with composed views by extracting it from SRigSpacePickerWidget... that would allow
+	// us to pass in the widget that should be in the window.
+	: public SRigSpacePickerWidget
+{
+public:
+
+	SLATE_BEGIN_ARGS(SFloatingSpacePicker_Base){}
+		/** Required. The controls to affect. */
+		SLATE_ARGUMENT(TArray<FRigElementKey>, SelectedControls)
+		/** Required. The rig that we're picking for. */
+		SLATE_ARGUMENT(UControlRig*, DisplayedRig)
+	SLATE_END_ARGS()
+	
+	void Construct(const FArguments& InArgs, bool bIsEditingDirectly);
+
+	/** Places this widget in a window and shows it. */
+	TSharedPtr<SWindow> ShowWindow();
+
+protected:
+
+	/** The rig being displayed */
+	TWeakObjectPtr<UControlRig> WeakDisplayedRig = nullptr;
+	/** The controls being displayed. */
+	TArray<FRigElementKey> SelectedControls;
+
+	virtual void OnActiveSpaceChanged(URigHierarchy* InHierarchy, const FRigElementKey& InControlKey, const FRigElementKey& InSpaceKey) = 0;
+	virtual void OnSpaceListChanged(URigHierarchy* InHierarchy, const FRigElementKey& InControlKey, const TArray<FRigElementKeyWithLabel>& InSpaceList) = 0;
+};
+}
+
+
diff --git a/Engine/Plugins/Animation/ControlRig/Source/ControlRigEditor/Private/EditMode/Picker/SFloatingSpacePicker_LevelEditor.cpp b/Engine/Plugins/Animation/ControlRig/Source/ControlRigEditor/Private/EditMode/Picker/SFloatingSpacePicker_LevelEditor.cpp
new file mode 100644
--- /dev/null
+++ b/Engine/Plugins/Animation/ControlRig/Source/ControlRigEditor/Private/EditMode/Picker/SFloatingSpacePicker_LevelEditor.cpp
@@ -0,0 +1,109 @@
+﻿// Copyright Epic Games, Inc. All Rights Reserved.
+
+#include "SFloatingSpacePicker_LevelEditor.h"
+
+#include "ControlRigSpaceChannelEditors.h"
+#include "ScopedTransaction.h"
+#include "Editor/SRigSpacePickerWidget.h"
+#include "Rigs/RigHierarchyController.h"
+
+#define LOCTEXT_NAMESPACE "SFloatingSpacePicker"
+
+namespace UE::ControlRigEditor
+{
+void SFloatingSpacePicker_LevelEditor::Construct(const FArguments& InArgs, TAttribute<TSharedPtr<ISequencer>> InSequencerAttr)
+{
+	SequencerAttr = MoveTemp(InSequencerAttr);
+	
+	constexpr bool bIsEditingDirectly = false;
+	SFloatingSpacePicker_Base::Construct(InArgs, bIsEditingDirectly);
+}
+
+void SFloatingSpacePicker_LevelEditor::OnActiveSpaceChanged(
+	URigHierarchy* InHierarchy, const FRigElementKey& InControlKey, const FRigElementKey& InSpaceKey
+	)
+{
+	check(SelectedControls.Contains(InControlKey));
+
+	UControlRig* RuntimeRig = WeakDisplayedRig.Get();
+	if (!ensure(RuntimeRig))
+	{
+		return;
+	}
+
+	const TSharedPtr<ISequencer> Sequencer = SequencerAttr.Get();
+	const FRigControlElement* ControlElement = InHierarchy->Find<FRigControlElement>(InControlKey);
+	if (!Sequencer || !ControlElement)
+	{
+		return;
+	}
+	
+	constexpr bool bDoNotCreate = false;
+	FSpaceChannelAndSection SpaceChannelAndSection = FControlRigSpaceChannelHelpers::FindSpaceChannelAndSectionForControl(
+		RuntimeRig, InControlKey.Name, Sequencer.Get(), bDoNotCreate
+		);
+	if (!SpaceChannelAndSection.SpaceChannel)
+	{
+		// no need to create a space channel if InSpaceKey is the current parent 
+		if (InHierarchy->GetActiveParent(InControlKey) == InSpaceKey)
+		{
+			return;
+		}
+	}
+					
+	const FScopedTransaction Transaction(LOCTEXT("KeyControlRigSpace", "Key Control Rig Space"));
+	constexpr bool bCreateIfNeeded = true;
+	SpaceChannelAndSection = FControlRigSpaceChannelHelpers::FindSpaceChannelAndSectionForControl(
+		RuntimeRig, InControlKey.Name, Sequencer.Get(), bCreateIfNeeded
+		);
+	if (SpaceChannelAndSection.SpaceChannel)
+	{
+		const FFrameRate TickResolution = Sequencer->GetFocusedTickResolution();
+		const FFrameTime FrameTime = Sequencer->GetLocalTime().ConvertTo(TickResolution);
+		FFrameNumber CurrentTime = FrameTime.GetFrame();
+		FControlRigSpaceChannelHelpers::SequencerKeyControlRigSpaceChannel(
+			RuntimeRig, Sequencer.Get(), SpaceChannelAndSection.SpaceChannel, SpaceChannelAndSection.SectionToKey, CurrentTime,
+			InHierarchy, InControlKey, InSpaceKey
+			);
+	}
+}
+
+void SFloatingSpacePicker_LevelEditor::OnSpaceListChanged(
+	URigHierarchy* InHierarchy, const FRigElementKey& InControlKey, const TArray<FRigElementKeyWithLabel>& InSpaceList
+	)
+{
+	check(SelectedControls.Contains(InControlKey));
+
+	UControlRig* RuntimeRig = WeakDisplayedRig.Get();
+	if (!ensure(RuntimeRig))
+	{
+		return;
+	}
+		
+	// update the settings in the control element
+	if (FRigControlElement* ControlElement = InHierarchy->Find<FRigControlElement>(InControlKey))
+	{
+		FScopedTransaction Transaction(LOCTEXT("ControlChangeAvailableSpaces", "Edit Available Spaces"));
+
+		InHierarchy->Modify();
+
+		FRigControlElementCustomization ControlCustomization = *RuntimeRig->GetControlCustomization(InControlKey);	
+		ControlCustomization.AvailableSpaces = InSpaceList;
+		ControlCustomization.RemovedSpaces.Reset();
+
+		// remember  the elements which are in the asset's available list but removed by the user
+		for(const FRigElementKeyWithLabel& AvailableSpace : ControlElement->Settings.Customization.AvailableSpaces)
+		{
+			if(ControlCustomization.AvailableSpaces.FindByKey(AvailableSpace.Key) == nullptr)
+			{
+				ControlCustomization.RemovedSpaces.Add(AvailableSpace.Key);
+			}
+		}
+
+		RuntimeRig->SetControlCustomization(InControlKey, ControlCustomization);
+		InHierarchy->Notify(ERigHierarchyNotification::ControlSettingChanged, ControlElement);
+	}
+}
+}
+
+#undef LOCTEXT_NAMESPACE
\ No newline at end of file
diff --git a/Engine/Plugins/Animation/ControlRig/Source/ControlRigEditor/Private/EditMode/Picker/SFloatingSpacePicker_LevelEditor.h b/Engine/Plugins/Animation/ControlRig/Source/ControlRigEditor/Private/EditMode/Picker/SFloatingSpacePicker_LevelEditor.h
new file mode 100644
--- /dev/null
+++ b/Engine/Plugins/Animation/ControlRig/Source/ControlRigEditor/Private/EditMode/Picker/SFloatingSpacePicker_LevelEditor.h
@@ -0,0 +1,36 @@
+﻿// Copyright Epic Games, Inc. All Rights Reserved.
+
+#pragma once
+
+#include "ControlRig.h"
+#include "SFloatingSpacePicker_Base.h"
+#include "Widgets/DeclarativeSyntaxSupport.h"
+#include "Widgets/SCompoundWidget.h"
+
+class FControlRigEditMode;
+class ISequencer;
+class SRigSpacePickerWidget;
+
+namespace UE::ControlRigEditor
+{
+/**
+ * Handles displaying the floating space picker when indirectly editing control rigs,
+ * i.e. when FControlRigEditMode::AreEditingControlRigDirectly() == false.
+ */
+class SFloatingSpacePicker_LevelEditor : public SFloatingSpacePicker_Base
+{
+public:
+	
+	void Construct(const FArguments& InArgs, TAttribute<TSharedPtr<ISequencer>> InSequencerAttr);
+
+private:
+
+	/** The level editor's Sequencer. */
+	TAttribute<TSharedPtr<ISequencer>> SequencerAttr;
+
+	virtual void OnActiveSpaceChanged(URigHierarchy* InHierarchy, const FRigElementKey& InControlKey, const FRigElementKey& InSpaceKey) override;
+	virtual void OnSpaceListChanged(URigHierarchy* InHierarchy, const FRigElementKey& InControlKey, const TArray<FRigElementKeyWithLabel>& InSpaceList) override;
+};
+}
+
+
diff --git a/Engine/Plugins/Animation/ControlRig/Source/ControlRigEditor/Private/Editor/SRigSpacePickerWidget.cpp b/Engine/Plugins/Animation/ControlRig/Source/ControlRigEditor/Private/Editor/SRigSpacePickerWidget.cpp
--- a/Engine/Plugins/Animation/ControlRig/Source/ControlRigEditor/Private/Editor/SRigSpacePickerWidget.cpp
+++ b/Engine/Plugins/Animation/ControlRig/Source/ControlRigEditor/Private/Editor/SRigSpacePickerWidget.cpp
@@ -321,7 +321,7 @@ class SRigSpaceDialogWindow : public SWindow
 {
 }; 
 
-FReply SRigSpacePickerWidget::OpenDialog(bool bModal)
+TSharedPtr<SWindow> SRigSpacePickerWidget::OpenDialog(bool bModal)
 {
 	check(!DialogWindow.IsValid());
 		
@@ -366,7 +366,7 @@ FReply SRigSpacePickerWidget::OpenDialog(bool bModal)
 		FSlateApplication::Get().AddWindow( Window );
 	}
 
-	return FReply::Handled();
+	return Window;
 }
 
 void SRigSpacePickerWidget::CloseDialog()
diff --git a/Engine/Plugins/Animation/ControlRig/Source/ControlRigEditor/Private/Editor/SRigSpacePickerWidget.h b/Engine/Plugins/Animation/ControlRig/Source/ControlRigEditor/Private/Editor/SRigSpacePickerWidget.h
--- a/Engine/Plugins/Animation/ControlRig/Source/ControlRigEditor/Private/Editor/SRigSpacePickerWidget.h
+++ b/Engine/Plugins/Animation/ControlRig/Source/ControlRigEditor/Private/Editor/SRigSpacePickerWidget.h
@@ -68,7 +68,7 @@ public:
 
 	UE_API void SetControls(URigHierarchy* InHierarchy, const TArray<FRigElementKey>& InControls);
 
-	UE_API FReply OpenDialog(bool bModal = true);
+	UE_API TSharedPtr<SWindow> OpenDialog(bool bModal = true);
 	UE_API void CloseDialog();
 	
 	UE_API virtual FReply OnKeyDown( const FGeometry& MyGeometry, const FKeyEvent& InKeyEvent ) override;
diff --git a/Engine/Plugins/Animation/ControlRig/Source/ControlRigEditor/Private/Sequencer/ControlRigSpaceChannelEditors.cpp b/Engine/Plugins/Animation/ControlRig/Source/ControlRigEditor/Private/Sequencer/ControlRigSpaceChannelEditors.cpp
--- a/Engine/Plugins/Animation/ControlRig/Source/ControlRigEditor/Private/Sequencer/ControlRigSpaceChannelEditors.cpp
+++ b/Engine/Plugins/Animation/ControlRig/Source/ControlRigEditor/Private/Sequencer/ControlRigSpaceChannelEditors.cpp
@@ -126,8 +126,7 @@ static FKeyHandle SequencerOpenSpaceSwitchDialog(UControlRig* ControlRig, TArray
 	});
 	// todo: implement GetAdditionalSpacesDelegate to pull spaces from sequencer
 
-	FReply Reply = PickerWidget->OpenDialog(true);
-	if (Reply.IsEventHandled())
+	if (const TSharedPtr<SWindow> Window = PickerWidget->OpenDialog(true))
 	{
 		return Handle;
 	}
diff --git a/Engine/Plugins/Animation/ControlRig/Source/ControlRigEditor/Public/EditMode/ControlRigEditMode.h b/Engine/Plugins/Animation/ControlRig/Source/ControlRigEditor/Public/EditMode/ControlRigEditMode.h
--- a/Engine/Plugins/Animation/ControlRig/Source/ControlRigEditor/Public/EditMode/ControlRigEditMode.h
+++ b/Engine/Plugins/Animation/ControlRig/Source/ControlRigEditor/Public/EditMode/ControlRigEditMode.h
@@ -59,6 +59,7 @@ class UMovieSceneSection;
 class UControlRigEditModeSettings;
 class UEditorTransformGizmoContextObject;
 struct FRotationContext;
+namespace UE::ControlRigEditor { class FFloatingSpacePickerManager; }
 
 DECLARE_DELEGATE_RetVal_ThreeParams(FTransform, FOnGetRigElementTransform, const FRigElementKey& /*RigElementKey*/, bool /*bLocal*/, bool /*bOnDebugInstance*/);
 DECLARE_DELEGATE_ThreeParams(FOnSetRigElementTransform, const FRigElementKey& /*RigElementKey*/, const FTransform& /*Transform*/, bool /*bLocal*/);
@@ -436,6 +437,13 @@ protected:
 	TWeakPtr<ISequencer> WeakSequencer;
 	FGuid LastMovieSceneSig;
 
+	/**
+	 * Handles spawning the floating space picker. Ensures only one can be created at the time. Removes the widget when destroyed.
+	 * 
+	 * We create this lazily because it's only used in OpenSpacePickerWidget.
+	 * If in the future this is needed in more locations, we should consider creating it when this mode is created.
+	 */
+	TPimplPtr<UE::ControlRigEditor::FFloatingSpacePickerManager> FloatingSpacePickerManager;
 
 	/** The scope for the interaction, one per manipulated Control rig */
 	TMap<UControlRig*, TUniquePtr<FControlRigInteractionScope>> InteractionScopes;
