Control Rig: Fix scroll bar not working for space picker summoned with tab.

diff --git a/Engine/Plugins/Animation/ControlRig/Source/ControlRigEditor/Private/Editor/SRigHierarchyTreeView.cpp b/Engine/Plugins/Animation/ControlRig/Source/ControlRigEditor/Private/Editor/SRigHierarchyTreeView.cpp
--- a/Engine/Plugins/Animation/ControlRig/Source/ControlRigEditor/Private/Editor/SRigHierarchyTreeView.cpp
+++ b/Engine/Plugins/Animation/ControlRig/Source/ControlRigEditor/Private/Editor/SRigHierarchyTreeView.cpp
@@ -931,7 +931,7 @@ void SRigHierarchyTreeView::RefreshTreeView(bool bRebuildContent)
 	if(bRebuildContent)
 	{
 		const URigHierarchy* Hierarchy = Delegates.GetHierarchy();
-		if(Hierarchy)
+		if (ensure(Hierarchy))
 		{
 			TArray<const FRigSocketElement*> Sockets;
 			TArray<const FRigConnectorElement*> Connectors;
@@ -1454,36 +1454,27 @@ void SSearchableRigHierarchyTreeView::Construct(const FArguments& InArgs)
 	FRigTreeDelegates TreeDelegates = InArgs._RigTreeDelegates;
 	SuperGetRigTreeDisplaySettings = TreeDelegates.OnGetDisplaySettings;
 
-	MaxHeight = InArgs._MaxHeight;
-
 	TreeDelegates.OnGetDisplaySettings.BindSP(this, &SSearchableRigHierarchyTreeView::GetDisplaySettings);
 
 	TSharedPtr<SVerticalBox> VerticalBox;
 	ChildSlot
 	[
-		SAssignNew(VerticalBox, SVerticalBox)
-		+SVerticalBox::Slot()
-		.AutoHeight()
-		.VAlign(VAlign_Top)
-		.HAlign(HAlign_Fill)
-		.Padding(0.0f)
+		SNew(SBorder)
+		.BorderImage(FAppStyle::GetBrush("SCSEditor.TreePanel"))
 		[
-			SAssignNew(SearchBox, SSearchBox)
-			.InitialText(InArgs._InitialFilterText)
-			.OnTextChanged(this, &SSearchableRigHierarchyTreeView::OnFilterTextChanged)
-		]
+			SAssignNew(VerticalBox, SVerticalBox)
+			+SVerticalBox::Slot()
+			.AutoHeight()
+			[
+				SAssignNew(SearchBox, SSearchBox)
+				.InitialText(InArgs._InitialFilterText)
+				.OnTextChanged(this, &SSearchableRigHierarchyTreeView::OnFilterTextChanged)
+			]
 
-		+SVerticalBox::Slot()
-		.VAlign(VAlign_Top)
-		.HAlign(HAlign_Fill)
-		.Padding(0.0f, 0.0f)
-		[
-			SNew(SScrollBox)
-			+ SScrollBox::Slot()
+			+SVerticalBox::Slot()
 			[
-				SNew(SBorder)
-				.Padding(2.0f)
-				.BorderImage(FAppStyle::GetBrush("SCSEditor.TreePanel"))
+				SNew(SScrollBox)
+				+SScrollBox::Slot()
 				[
 					SAssignNew(TreeView, SRigHierarchyTreeView)
 					.RigTreeDelegates(TreeDelegates)
@@ -1491,15 +1482,6 @@ void SSearchableRigHierarchyTreeView::Construct(const FArguments& InArgs)
 			]
 		]
 	];
-
-	if (MaxHeight > SMALL_NUMBER)
-	{
-		VerticalBox->GetSlot(1).SetMaxHeight(MaxHeight);
-	}
-	else
-	{
-		VerticalBox->GetSlot(1).SetAutoHeight();
-	}
 }
 
 const FRigTreeDisplaySettings& SSearchableRigHierarchyTreeView::GetDisplaySettings()
diff --git a/Engine/Plugins/Animation/ControlRig/Source/ControlRigEditor/Private/Editor/SRigHierarchyTreeView.h b/Engine/Plugins/Animation/ControlRig/Source/ControlRigEditor/Private/Editor/SRigHierarchyTreeView.h
--- a/Engine/Plugins/Animation/ControlRig/Source/ControlRigEditor/Private/Editor/SRigHierarchyTreeView.h
+++ b/Engine/Plugins/Animation/ControlRig/Source/ControlRigEditor/Private/Editor/SRigHierarchyTreeView.h
@@ -417,15 +417,13 @@ class SSearchableRigHierarchyTreeView : public SCompoundWidget
 public:
 
 	SLATE_BEGIN_ARGS(SSearchableRigHierarchyTreeView)
-		:_MaxHeight(0.f)
 	{}
 		SLATE_ARGUMENT(FRigTreeDelegates, RigTreeDelegates)
 		SLATE_ARGUMENT(FText, InitialFilterText)
-		SLATE_ARGUMENT(float, MaxHeight)
 	SLATE_END_ARGS()
 
 	void Construct(const FArguments& InArgs);
-	virtual ~SSearchableRigHierarchyTreeView() {}
+	
 	TSharedRef<SSearchBox> GetSearchBox() const { return SearchBox.ToSharedRef(); }
 	TSharedRef<SRigHierarchyTreeView> GetTreeView() const { return TreeView.ToSharedRef(); }
 	const FRigTreeDisplaySettings& GetDisplaySettings();
@@ -439,7 +437,6 @@ private:
 	FRigTreeDisplaySettings Settings;
 	TSharedPtr<SSearchBox> SearchBox;
 	TSharedPtr<SRigHierarchyTreeView> TreeView;
-	float MaxHeight = 0.f;
 };
 
 #undef UE_API
