Control Rig: Fix right-clicking constraint opening menu offscreen

diff --git a/Engine/Plugins/Animation/ControlRig/Source/ControlRigEditor/Private/Editor/Constraints/SConstraintsWidget.cpp b/Engine/Plugins/Animation/ControlRig/Source/ControlRigEditor/Private/Editor/Constraints/SConstraintsWidget.cpp
--- a/Engine/Plugins/Animation/ControlRig/Source/ControlRigEditor/Private/Editor/Constraints/SConstraintsWidget.cpp
+++ b/Engine/Plugins/Animation/ControlRig/Source/ControlRigEditor/Private/Editor/Constraints/SConstraintsWidget.cpp
@@ -1823,7 +1823,7 @@ TSharedPtr<SWidget> SConstraintsEditionWidget::CreateContextMenu()
 		DetailsViewArgs.bUpdatesFromSelection = false;
 		DetailsViewArgs.bShowOptions = false;
 		DetailsViewArgs.bShowModifiedPropertiesOption = false;
-		DetailsViewArgs.ColumnWidth = 0.45f;
+		DetailsViewArgs.ColumnWidth = 0.35f;
 	}
 	
 	TSharedRef<IDetailsView> DetailsView = FModuleManager::GetModuleChecked<FPropertyEditorModule>("PropertyEditor").CreateDetailView(DetailsViewArgs);
@@ -1887,12 +1887,14 @@ TSharedPtr<SWidget> SConstraintsEditionWidget::CreateContextMenu()
 	// constraint details
 	MenuBuilder.BeginSection("EditConstraint", LOCTEXT("EditConstraintHeader", "Edit Constraint"));
 	{
-		// this spacer is used to set a minimum size to the details builder so that the menu won't shrink when
-		// collapsing the transform offset for instance
-		static const FVector2D MinimumDetailsSize(300.0f, 0.0f);
-		MenuBuilder.AddWidget(SNew(SSpacer).Size(MinimumDetailsSize), FText::GetEmpty(), false, false);
-		
-		MenuBuilder.AddWidget(DetailsView, FText::GetEmpty(), true);
+		MenuBuilder.AddWidget(
+			// Summary: This makes sure the menu anchor position is realigned when you open the menu near the edge of your screen, so the menu does
+			// not spill over to the next screen.
+			// Slate calls Prepass and then uses the widget's desired size. However, the details view always defers until next tick to compute it.
+			// The width override forces a desired size.
+			SNew(SBox).WidthOverride(480.f) [ DetailsView ],
+			FText::GetEmpty(), true
+			);
 	}
 	MenuBuilder.EndSection();
 
