Anim In Engine: Don't add anim layers to level sequenc unless needed, added bool bAddIfDoesNotExist flag to GetAnimLayers
#prelfight

diff --git a/Engine/Plugins/Animation/ControlRig/Source/ControlRigEditor/Private/ControlRigSequencerEditorLibrary.cpp b/Engine/Plugins/Animation/ControlRig/Source/ControlRigEditor/Private/ControlRigSequencerEditorLibrary.cpp
--- a/Engine/Plugins/Animation/ControlRig/Source/ControlRigEditor/Private/ControlRigSequencerEditorLibrary.cpp
+++ b/Engine/Plugins/Animation/ControlRig/Source/ControlRigEditor/Private/ControlRigSequencerEditorLibrary.cpp
@@ -3574,7 +3574,8 @@ int32 UControlRigSequencerEditorLibrary::AddAnimLayerFromSelection()
 		return INDEX_NONE;
 	}
 	TSharedPtr<ISequencer> SequencerPtr = WeakSequencer.Pin();
-	if (UAnimLayers* AnimLayers = UAnimLayers::GetAnimLayers(SequencerPtr.Get()))
+	constexpr bool bAddIfDoesNotExist = true;
+	if (UAnimLayers* AnimLayers = UAnimLayers::GetAnimLayers(SequencerPtr.Get(), bAddIfDoesNotExist))
 	{
 		return AnimLayers->AddAnimLayerFromSelection(SequencerPtr.Get());
 	}
@@ -3627,7 +3628,7 @@ bool UControlRigSequencerEditorLibrary::MergeAnimLayersWithSettings(const TArray
 	return false;
 }
 
-TArray<UAnimLayer*> UControlRigSequencerEditorLibrary::GetAnimLayers()
+TArray<UAnimLayer*> UControlRigSequencerEditorLibrary::GetAnimLayers(bool bAddIfDoesNotExist)
 {
 	TArray<UAnimLayer*> AnimLayersArray;
 	TWeakPtr<ISequencer> WeakSequencer = GetSequencerFromAsset();
@@ -3637,7 +3638,7 @@ TArray<UAnimLayer*> UControlRigSequencerEditorLibrary::GetAnimLayers()
 		return AnimLayersArray;
 	}
 	TSharedPtr<ISequencer> SequencerPtr = WeakSequencer.Pin();
-	if (UAnimLayers* AnimLayers = UAnimLayers::GetAnimLayers(SequencerPtr.Get()))
+	if (UAnimLayers* AnimLayers = UAnimLayers::GetAnimLayers(SequencerPtr.Get(), bAddIfDoesNotExist))
 	{
 		return AnimLayers->AnimLayers;
 	}
diff --git a/Engine/Plugins/Animation/ControlRig/Source/ControlRigEditor/Private/Sequencer/AnimLayers/AnimLayers.cpp b/Engine/Plugins/Animation/ControlRig/Source/ControlRigEditor/Private/Sequencer/AnimLayers/AnimLayers.cpp
--- a/Engine/Plugins/Animation/ControlRig/Source/ControlRigEditor/Private/Sequencer/AnimLayers/AnimLayers.cpp
+++ b/Engine/Plugins/Animation/ControlRig/Source/ControlRigEditor/Private/Sequencer/AnimLayers/AnimLayers.cpp
@@ -1591,17 +1591,17 @@ bool UAnimLayers::HasAnimLayers(ISequencer* InSequencer)
 	return false;
 }
 
-UAnimLayers* UAnimLayers::GetAnimLayers(ISequencer* SequencerPtr)
+UAnimLayers* UAnimLayers::GetAnimLayers(ISequencer* SequencerPtr, bool bAddIfDoesNotExist)
 {
 	if (!SequencerPtr)
 	{
 		return nullptr;
 	}
 	ULevelSequence* LevelSequence = Cast<ULevelSequence>(SequencerPtr->GetFocusedMovieSceneSequence());
-	return UAnimLayers::GetAnimLayers(LevelSequence);
+	return UAnimLayers::GetAnimLayers(LevelSequence,bAddIfDoesNotExist);
 }
 
-UAnimLayers* UAnimLayers::GetAnimLayers(ULevelSequence* LevelSequence)
+UAnimLayers* UAnimLayers::GetAnimLayers(ULevelSequence* LevelSequence, bool bAddIfDoesNotExist)
 {
 	if(!LevelSequence)
 	{
@@ -3921,3 +3921,4 @@ bool UAnimLayers::AdjustmentBlendLayers(ISequencer* InSequencer, int32 LayerInde
 #undef LOCTEXT_NAMESPACE
 
 
+
diff --git a/Engine/Plugins/Animation/ControlRig/Source/ControlRigEditor/Private/Sequencer/AnimLayers/AnimLayers.h b/Engine/Plugins/Animation/ControlRig/Source/ControlRigEditor/Private/Sequencer/AnimLayers/AnimLayers.h
--- a/Engine/Plugins/Animation/ControlRig/Source/ControlRigEditor/Private/Sequencer/AnimLayers/AnimLayers.h
+++ b/Engine/Plugins/Animation/ControlRig/Source/ControlRigEditor/Private/Sequencer/AnimLayers/AnimLayers.h
@@ -348,8 +348,8 @@ public:
 	TArray<UMovieSceneSection*> GetSelectedLayerSections() const;
 
 	//Get
-	static UAnimLayers* GetAnimLayers(ISequencer* InSequencer);
-	static UAnimLayers* GetAnimLayers(ULevelSequence* InLevelSequence);
+	static UAnimLayers* GetAnimLayers(ISequencer* InSequencer, bool bAddIfDoesNotExist = false);
+	static UAnimLayers* GetAnimLayers(ULevelSequence* InLevelSequence, bool bAddIfDoesNotExist = false);
 	static bool HasAnimLayers(ISequencer* InSequencer);
 
 	static 	TSharedPtr<ISequencer> GetSequencerFromAsset();
diff --git a/Engine/Plugins/Animation/ControlRig/Source/ControlRigEditor/Private/Sequencer/AnimLayers/SAnimLayers.cpp b/Engine/Plugins/Animation/ControlRig/Source/ControlRigEditor/Private/Sequencer/AnimLayers/SAnimLayers.cpp
--- a/Engine/Plugins/Animation/ControlRig/Source/ControlRigEditor/Private/Sequencer/AnimLayers/SAnimLayers.cpp
+++ b/Engine/Plugins/Animation/ControlRig/Source/ControlRigEditor/Private/Sequencer/AnimLayers/SAnimLayers.cpp
@@ -1967,7 +1967,7 @@ FReply SAnimLayers::OnAddClicked()
 {
 	if (ISequencer* SequencerPtr = GetSequencer())
 	{
-		if (UAnimLayers* AnimLayersPtr = UAnimLayers::GetAnimLayers(SequencerPtr))
+		if (UAnimLayers* AnimLayersPtr = UAnimLayers::GetAnimLayers(SequencerPtr, true/*bAddIfDoesNotExist */ ))
 		{
 			const FAnimLayersScopedSelection ScopedSelection(*AnimLayersPtr);
 
diff --git a/Engine/Plugins/Animation/ControlRig/Source/ControlRigEditor/Public/ControlRigSequencerEditorLibrary.h b/Engine/Plugins/Animation/ControlRig/Source/ControlRigEditor/Public/ControlRigSequencerEditorLibrary.h
--- a/Engine/Plugins/Animation/ControlRig/Source/ControlRigEditor/Public/ControlRigSequencerEditorLibrary.h
+++ b/Engine/Plugins/Animation/ControlRig/Source/ControlRigEditor/Public/ControlRigSequencerEditorLibrary.h
@@ -1212,11 +1212,12 @@ public:
 	static UE_API bool MergeAnimLayersWithSettings(const TArray<int32>& Indices, const FMergeAnimLayerSettings& MergeAnimLayerSettings);
 
 	/**
-	* Get the animation layer objects
+	* Get the animation layer objects, will create it by default if it doesn't exist
+	* @param bAddIfDoesNotExist If it doesn't exist create it
 	* @return Returns array of anim layer objects if they exist on active Sequencer
 	*/
 	UFUNCTION(BlueprintCallable, Category = "Editor Scripting | Sequencer Tools | Control Rig | Animation Layers")
-	static UE_API TArray<UAnimLayer*> GetAnimLayers();
+	static UE_API TArray<UAnimLayer*> GetAnimLayers(bool bAddIfDoesNotExist = true);
 	
 	/**
 	* Helper function to get the index in the anim layer array from the anim layer
