Control Rig: Space Picker's pop-up menu anchor point is anchored properly near edge of the screen. It no longer goes off-screen.
Also replaced the "+" button with a positive action combo button.

diff --git a/Engine/Plugins/Animation/ControlRig/Source/ControlRigEditor/Private/EditMode/SControlRigSpacePicker.cpp b/Engine/Plugins/Animation/ControlRig/Source/ControlRigEditor/Private/EditMode/SControlRigSpacePicker.cpp
--- a/Engine/Plugins/Animation/ControlRig/Source/ControlRigEditor/Private/EditMode/SControlRigSpacePicker.cpp
+++ b/Engine/Plugins/Animation/ControlRig/Source/ControlRigEditor/Private/EditMode/SControlRigSpacePicker.cpp
@@ -33,79 +33,22 @@ void SControlRigSpacePicker::Construct(
 	[
 		SNew(SScrollBox)
 		+ SScrollBox::Slot()
+		.AutoSize()
 		[
-			SNew(SVerticalBox)
-			+SVerticalBox::Slot()
-			.AutoHeight()
-			[
-				SAssignNew(PickerExpander, SExpandableArea)
-				.InitiallyCollapsed(InArgs._InitiallyCollapsed)
-				.AreaTitle(LOCTEXT("Picker_SpaceWidget", "Spaces"))
-				.AreaTitleFont(FAppStyle::GetFontStyle("DetailsView.CategoryFontStyle"))
-				.BorderBackgroundColor(FLinearColor(.6f, .6f, .6f))
-				.Padding(FMargin(8.f))
-				.HeaderContent()
-				[
-					SNew(SHorizontalBox)
-
-					+SHorizontalBox::Slot()
-					.AutoWidth()
-					.HAlign(HAlign_Left)
-					.VAlign(VAlign_Center)
-					.Padding(0.f, 0.f, 0.f, 0.f)
-					[
-						SNew(STextBlock)
-						.Text(LOCTEXT("Picker_SpaceWidget", "Spaces"))
-						.Font(FCoreStyle::Get().GetFontStyle("ExpandableArea.TitleFont"))
-					]
-				
-					+SHorizontalBox::Slot()
-					.FillWidth(1.f)
-					[
-						SNew(SSpacer)
-					]
-
-					+SHorizontalBox::Slot()
-					.AutoWidth()
-					.HAlign(HAlign_Right)
-					.VAlign(VAlign_Center)
-					.Padding(0.f, 2.f, 8.f, 2.f)
-					[
-						SNew(SButton)
-						.ContentPadding(0.0f)
-						.ButtonStyle(FAppStyle::Get(), "NoBorder")
-						.OnClicked(this, &SControlRigSpacePicker::HandleAddSpaceClicked)
-						.Cursor(EMouseCursor::Default)
-						.ToolTipText(LOCTEXT("AddSpace", "Add Space"))
-						[
-							SNew(SImage)
-							.Image(FAppStyle::GetBrush(TEXT("Icons.PlusCircle")))
-						]
-						.Visibility_Lambda([this]()
-						{
-							return GetAddSpaceButtonVisibility();
-						})
-					]
-				]
-				.BodyContent()
-				[
-					SAssignNew(SpacePickerWidget, SRigSpacePickerWidget)
-					.AllowDelete(true)
-					.AllowReorder(true)
-					.AllowAdd(false)
-					.Hierarchy(InArgs._Hierarchy)
-					.Controls(InArgs._Controls)
-					.ShowBakeAndCompensateButton(true)
-					.GetControlCustomization(this, &SControlRigSpacePicker::HandleGetControlElementCustomization)
-					.OnActiveSpaceChanged(this, &SControlRigSpacePicker::HandleActiveSpaceChanged)
-					.OnSpaceListChanged(this, &SControlRigSpacePicker::HandleSpaceListChanged)
-					.OnCompensateKeyButtonClicked(this, &SControlRigSpacePicker::OnCompensateKeyClicked)
-					.OnCompensateAllButtonClicked(this, &SControlRigSpacePicker::OnCompensateAllClicked)
-					.OnBakeButtonClicked(this, &SControlRigSpacePicker::OnBakeControlsToNewSpaceButtonClicked)
-					// todo: implement GetAdditionalSpacesDelegate to pull spaces from sequencer
-				]
-			]
-
+			SAssignNew(SpacePickerWidget, SRigSpacePickerWidget)
+			.AllowDelete(true)
+			.AllowReorder(true)
+			.AllowAdd(true)
+			.Hierarchy(InArgs._Hierarchy)
+			.Controls(InArgs._Controls)
+			.ShowBakeAndCompensateButton(true)
+			.GetControlCustomization(this, &SControlRigSpacePicker::HandleGetControlElementCustomization)
+			.OnActiveSpaceChanged(this, &SControlRigSpacePicker::HandleActiveSpaceChanged)
+			.OnSpaceListChanged(this, &SControlRigSpacePicker::HandleSpaceListChanged)
+			.OnCompensateKeyButtonClicked(this, &SControlRigSpacePicker::OnCompensateKeyClicked)
+			.OnCompensateAllButtonClicked(this, &SControlRigSpacePicker::OnCompensateAllClicked)
+			.OnBakeButtonClicked(this, &SControlRigSpacePicker::OnBakeControlsToNewSpaceButtonClicked)
+			// todo: implement GetAdditionalSpacesDelegate to pull spaces from sequencer
 		]
 	];
 
@@ -249,11 +192,6 @@ void SControlRigSpacePicker::HandleSpaceListChanged(URigHierarchy* InHierarchy,
 	}
 }
 
-FReply SControlRigSpacePicker::HandleAddSpaceClicked()
-{
-	return SpacePickerWidget->HandleAddElementClicked();
-}
-
 bool SControlRigSpacePicker::ReadyForBakeOrCompensation() const
 {
 	if (SpacePickerWidget->GetHierarchy() == nullptr)
diff --git a/Engine/Plugins/Animation/ControlRig/Source/ControlRigEditor/Private/EditMode/SControlRigSpacePicker.h b/Engine/Plugins/Animation/ControlRig/Source/ControlRigEditor/Private/EditMode/SControlRigSpacePicker.h
--- a/Engine/Plugins/Animation/ControlRig/Source/ControlRigEditor/Private/EditMode/SControlRigSpacePicker.h
+++ b/Engine/Plugins/Animation/ControlRig/Source/ControlRigEditor/Private/EditMode/SControlRigSpacePicker.h
@@ -53,7 +53,6 @@ private:
 	const FRigControlElementCustomization* HandleGetControlElementCustomization(URigHierarchy* InHierarchy, const FRigElementKey& InControlKey);
 	void HandleActiveSpaceChanged(URigHierarchy* InHierarchy, const FRigElementKey& InControlKey, const FRigElementKey& InSpaceKey);
 	void HandleSpaceListChanged(URigHierarchy* InHierarchy, const FRigElementKey& InControlKey, const TArray<FRigElementKeyWithLabel>& InSpaceList);
-	FReply HandleAddSpaceClicked();
 	FReply OnBakeControlsToNewSpaceButtonClicked();
 	FReply OnCompensateKeyClicked();
 	FReply OnCompensateAllClicked();
diff --git a/Engine/Plugins/Animation/ControlRig/Source/ControlRigEditor/Private/Editor/SRigSpacePickerWidget.cpp b/Engine/Plugins/Animation/ControlRig/Source/ControlRigEditor/Private/Editor/SRigSpacePickerWidget.cpp
--- a/Engine/Plugins/Animation/ControlRig/Source/ControlRigEditor/Private/Editor/SRigSpacePickerWidget.cpp
+++ b/Engine/Plugins/Animation/ControlRig/Source/ControlRigEditor/Private/Editor/SRigSpacePickerWidget.cpp
@@ -6,13 +6,12 @@
 #include "Widgets/Text/STextBlock.h"
 #include "Widgets/Layout/SUniformGridPanel.h"
 #include "Widgets/Input/SButton.h"
+#include "Widgets/Input/SComboButton.h"
 #include "Widgets/Input/SSearchBox.h"
 #include "Editor/SRigHierarchyTreeView.h"
 #include "ControlRigEditorStyle.h"
 #include "PropertyCustomizationHelpers.h"
 #include "ISequencer.h"
-#include "Widgets/Input/NumericTypeInterface.h"
-#include "FrameNumberDetailsCustomization.h"
 #include "Widgets/Notifications/SNotificationList.h"
 #include "Framework/Application/SlateApplication.h"
 #include "Framework/Application/SlateUser.h"
@@ -22,6 +21,7 @@
 #include "ControlRigBlueprintLegacy.h"
 #include "RigVMBlueprintGeneratedClass.h"
 #include "SActionButton.h"
+#include "SPositiveActionButton.h"
 
 #define LOCTEXT_NAMESPACE "SRigSpacePickerWidget"
 
@@ -46,7 +46,6 @@ void SRigSpacePickerWidget::Construct(const FArguments& InArgs)
 	GetControlCustomizationDelegate = InArgs._GetControlCustomization;
 	GetAdditionalSpacesDelegate = InArgs._GetAdditionalSpaces;
 	bRepopulateRequired = false;
-	bLaunchingContextMenu = false;
 
 	if(!GetActiveSpaceDelegate.IsBound())
 	{
@@ -152,20 +151,7 @@ void SRigSpacePickerWidget::Construct(const FArguments& InArgs)
 			.HAlign(HAlign_Left)
 			.Padding(0.f)
 			[
-				SNew(SButton)
-				.ContentPadding(0.0f)
-				.ButtonStyle(FAppStyle::Get(), "NoBorder")
-				.OnClicked(this, &SRigSpacePickerWidget::HandleAddElementClicked)
-				.Cursor(EMouseCursor::Default)
-				.ToolTipText(LOCTEXT("AddSpace", "Add Space"))
-				[
-					SNew(SImage)
-					.Image(FAppStyle::GetBrush(TEXT("Icons.PlusCircle")))
-				]
-				.Visibility_Lambda([this]()
-				{
-					return IsRestricted() ? EVisibility::Collapsed : EVisibility::Visible;
-				})
+				CreateAddButton()
 			];
 		}
 
@@ -341,17 +327,15 @@ TSharedPtr<SWindow> SRigSpacePickerWidget::OpenDialog(bool bModal)
 	];
 	
 	Window->SetWidgetToFocusOnActivate(AsShared());
-	if (!Window->GetOnWindowDeactivatedEvent().IsBoundToObject(this))
+	Window->GetOnWindowDeactivatedEvent().AddLambda([this]()
 	{
-		Window->GetOnWindowDeactivatedEvent().AddLambda([this]()
+		// Do not reset if we lost focus because of opening the context menu
+		if (!bIsAddMenuOpen)
 		{
-			// Do not reset if we lost focus because of opening the context menu
-			if (!ContextMenu.IsValid())
-			{
-				SetControls(nullptr, {});
-			}
-		});
-	}
+			CloseDialog();
+			SetControls(nullptr, {});
+		}
+	});
 	
 	DialogWindow = Window;
 
@@ -371,15 +355,6 @@ TSharedPtr<SWindow> SRigSpacePickerWidget::OpenDialog(bool bModal)
 
 void SRigSpacePickerWidget::CloseDialog()
 {
-	if(bLaunchingContextMenu)
-	{
-		return;
-	}
-	
-	if(ContextMenu.IsValid())
-	{
-		return;
-	}
 	
 	if ( DialogWindow.IsValid() )
 	{
@@ -721,7 +696,50 @@ void SRigSpacePickerWidget::HandleSpaceDelete(FRigElementKey InKey)
 	}
 }
 
-FReply SRigSpacePickerWidget::HandleAddElementClicked()
+TSharedRef<SWidget> SRigSpacePickerWidget::CreateAddButton()
+{
+	return SAssignNew(AddComboButton, SPositiveActionButton)
+		.OnGetMenuContent(this, &SRigSpacePickerWidget::OnGetAddButtonContent)
+		.OnMenuOpenChanged(this, &SRigSpacePickerWidget::OnMenuOpenChanged)
+		.Cursor(EMouseCursor::Default)
+		.ToolTipText(LOCTEXT("AddSpace.ToolTip", "Add Space"))
+		.Text(LOCTEXT("AddSpace.Label", "Add"))
+		.Visibility_Lambda([this]()
+		{
+			return IsRestricted() ? EVisibility::Collapsed : EVisibility::Visible;
+		});
+}
+
+namespace UE::ControlRigEditor::RigSpacePickerDetail
+{
+static bool CompareKeys(const FRigHierarchyKey& A, const FRigHierarchyKey& B)
+{
+	if(A.IsElement() && B.IsComponent())
+	{
+		return true;
+	}
+	if(B.IsElement() && A.IsComponent())
+	{
+		return false;
+	}
+
+	// controls should always show up first - so we'll sort them to the start of the list
+	if(A.IsElement() && B.IsElement())
+	{
+		if(A.GetElement().Type == ERigElementType::Control && B.GetElement().Type != ERigElementType::Control)
+		{
+			return true;
+		}
+		if(B.GetElement().Type == ERigElementType::Control && A.GetElement().Type != ERigElementType::Control)
+		{
+			return false;
+		}
+	}
+	return A < B;
+}
+}
+
+TSharedRef<SWidget> SRigSpacePickerWidget::OnGetAddButtonContent()
 {
 	HierarchyDisplaySettings.bShowConnectors = false;
 	HierarchyDisplaySettings.bShowSockets = false;
@@ -730,138 +748,94 @@ FReply SRigSpacePickerWidget::HandleAddElementClicked()
 	FRigTreeDelegates TreeDelegates;
 	TreeDelegates.OnGetDisplaySettings = FOnGetRigTreeDisplaySettings::CreateSP(this, &SRigSpacePickerWidget::GetHierarchyDisplaySettings); 
 	TreeDelegates.OnGetHierarchy = FOnGetRigTreeHierarchy::CreateSP(this, &SRigSpacePickerWidget::GetHierarchyConst);
-	TreeDelegates.OnMouseButtonClick = FOnRigTreeMouseButtonClick::CreateLambda([this](TSharedPtr<FRigTreeElement> InItem)
+	TreeDelegates.OnMouseButtonClick = FOnRigTreeMouseButtonClick::CreateSP(this, &SRigSpacePickerWidget::HandleClickTreeItemInAddMenu);
+	TreeDelegates.OnCompareKeys = FOnRigTreeCompareKeys::CreateStatic(&UE::ControlRigEditor::RigSpacePickerDetail::CompareKeys);
+
+	const TSharedRef<SSearchableRigHierarchyTreeView> SearchableTreeView = SNew(SSearchableRigHierarchyTreeView).RigTreeDelegates(TreeDelegates);
+	SearchableTreeView->GetTreeView()->RefreshTreeView(true);
+	AddComboButton->SetMenuContentWidgetToFocus(SearchableTreeView->GetSearchBox());
+	
+	return SearchableTreeView;
+}
+
+void SRigSpacePickerWidget::OnMenuOpenChanged(bool bIsOpen)
+{
+	bIsAddMenuOpen = bIsOpen;
+	
+	if (!bIsOpen && DialogWindow.IsValid())
 	{
-		if(InItem.IsValid())
+		DialogWindow.Pin()->BringToFront(true);
+
+		TSharedRef<SWidget> ThisRef = AsShared();
+		FSlateApplication::Get().ForEachUser([&ThisRef](FSlateUser& User) {
+			User.SetFocus(ThisRef, EFocusCause::SetDirectly);
+		});
+	}
+}
+
+void SRigSpacePickerWidget::HandleClickTreeItemInAddMenu(TSharedPtr<FRigTreeElement> InItem)
+{
+	if (InItem.IsValid())
+	{
+		const FRigElementKey Key = InItem->Key.GetElement();
+		if (!IsDefaultSpace(Key) && IsValidKey(Key))
 		{
-			const FRigElementKey Key = InItem->Key.GetElement();
-			if(!IsDefaultSpace(Key) && IsValidKey(Key))
+			if (Hierarchy.IsValid())
 			{
-				if (Hierarchy.IsValid())
+				URigHierarchy* StrongHierarchy = Hierarchy.Get();
+				for(const FRigElementKey& ControlKey : ControlKeys)
 				{
-					URigHierarchy* StrongHierarchy = Hierarchy.Get();
-					for(const FRigElementKey& ControlKey : ControlKeys)
-					{
-						FRigDependenciesProviderForControlRig DependencyProvider;
-						FString FailureReason;
+					FRigDependenciesProviderForControlRig DependencyProvider;
+					FString FailureReason;
 
-						if(UControlRig* ControlRig = StrongHierarchy->GetTypedOuter<UControlRig>())
-						{
-							DependencyProvider = FRigDependenciesProviderForControlRig(ControlRig); 
-						}
-						else if(UControlRigBlueprint* RigBlueprint = StrongHierarchy->GetTypedOuter<UControlRigBlueprint>())
-						{
-							if(UControlRig* CDO = Cast<UControlRig>(RigBlueprint->GetRigVMBlueprintGeneratedClass()->GetDefaultObject()))
-							{
-								DependencyProvider = FRigDependenciesProviderForControlRig(CDO); 
-							}
-						}
-						
-						if(!StrongHierarchy->CanSwitchToParent(ControlKey, Key, DependencyProvider, &FailureReason))
+					if(UControlRig* ControlRig = StrongHierarchy->GetTypedOuter<UControlRig>())
+					{
+						DependencyProvider = FRigDependenciesProviderForControlRig(ControlRig); 
+					}
+					else if(UControlRigBlueprint* RigBlueprint = StrongHierarchy->GetTypedOuter<UControlRigBlueprint>())
+					{
+						if(UControlRig* CDO = Cast<UControlRig>(RigBlueprint->GetRigVMBlueprintGeneratedClass()->GetDefaultObject()))
 						{
-							// notification
-							FNotificationInfo Info(FText::FromString(FailureReason));
-							Info.bFireAndForget = true;
-							Info.FadeOutDuration = 2.0f;
-							Info.ExpireDuration = 8.0f;
-
-							const TSharedPtr<SNotificationItem> NotificationPtr = FSlateNotificationManager::Get().AddNotification(Info);
-							NotificationPtr->SetCompletionState(SNotificationItem::CS_Fail);
-							return;
+							DependencyProvider = FRigDependenciesProviderForControlRig(CDO); 
 						}
 					}
 					
-					TArray<FRigElementKeyWithLabel> ChangedSpaceKeys = CurrentSpaceKeys;
-					if(ChangedSpaceKeys.FindByKey(Key) == nullptr)
+					if (!StrongHierarchy->CanSwitchToParent(ControlKey, Key, DependencyProvider, &FailureReason))
 					{
-						if(ControlKeys.IsEmpty())
-						{
-							ChangedSpaceKeys.Emplace(Key);
-						}
-						else
-						{
-							ChangedSpaceKeys.Emplace(Key, StrongHierarchy->GetDisplayLabelForParent(ControlKeys[0], Key));
-						}
+						// notification
+						FNotificationInfo Info(FText::FromString(FailureReason));
+						Info.bFireAndForget = true;
+						Info.FadeOutDuration = 2.0f;
+						Info.ExpireDuration = 8.0f;
+
+						const TSharedPtr<SNotificationItem> NotificationPtr = FSlateNotificationManager::Get().AddNotification(Info);
+						NotificationPtr->SetCompletionState(SNotificationItem::CS_Fail);
+						return;
 					}
-
-					for(const FRigElementKey& ControlKey : ControlKeys)
+				}
+				
+				TArray<FRigElementKeyWithLabel> ChangedSpaceKeys = CurrentSpaceKeys;
+				if(ChangedSpaceKeys.FindByKey(Key) == nullptr)
+				{
+					if(ControlKeys.IsEmpty())
 					{
-						SpaceListChangedEvent.Broadcast(StrongHierarchy, ControlKey, ChangedSpaceKeys);
+						ChangedSpaceKeys.Emplace(Key);
+					}
+					else
+					{
+						ChangedSpaceKeys.Emplace(Key, StrongHierarchy->GetDisplayLabelForParent(ControlKeys[0], Key));
 					}
 				}
-			}
-		}
-
-		if(ContextMenu.IsValid())
-		{
-			ContextMenu.Pin()->Dismiss();
-			ContextMenu.Reset();
-		}
-	});
 
-	TreeDelegates.OnCompareKeys = FOnRigTreeCompareKeys::CreateLambda([](const FRigHierarchyKey& A, const FRigHierarchyKey& B) -> bool
-	{
-		if(A.IsElement() && B.IsComponent())
-		{
-			return true;
-		}
-		if(B.IsElement() && A.IsComponent())
-		{
-			return false;
-		}
-
-		// controls should always show up first - so we'll sort them to the start of the list
-		if(A.IsElement() && B.IsElement())
-		{
-			if(A.GetElement().Type == ERigElementType::Control && B.GetElement().Type != ERigElementType::Control)
-			{
-				return true;
-			}
-			if(B.GetElement().Type == ERigElementType::Control && A.GetElement().Type != ERigElementType::Control)
-			{
-				return false;
+				for(const FRigElementKey& ControlKey : ControlKeys)
+				{
+					SpaceListChangedEvent.Broadcast(StrongHierarchy, ControlKey, ChangedSpaceKeys);
+				}
 			}
 		}
-		return A < B;
-	});
-
-	TSharedPtr<SSearchableRigHierarchyTreeView> SearchableTreeView = SNew(SSearchableRigHierarchyTreeView)
-	.RigTreeDelegates(TreeDelegates);
-	SearchableTreeView->GetTreeView()->RefreshTreeView(true);
-
-	const bool bFocusImmediately = false;
-	// Create as context menu
-	TGuardValue<bool> AboutToShowMenu(bLaunchingContextMenu, true);
-	ContextMenu = FSlateApplication::Get().PushMenu(
-		AsShared(),
-		FWidgetPath(),
-		SearchableTreeView.ToSharedRef(),
-		FSlateApplication::Get().GetCursorPos(),
-		FPopupTransitionEffect(FPopupTransitionEffect::ContextMenu),
-		bFocusImmediately
-	);
-
-	if(!ContextMenu.IsValid())
-	{
-		return FReply::Unhandled();
 	}
-	
-	ContextMenu.Pin()->GetOnMenuDismissed().AddLambda([this](TSharedRef<IMenu> InMenu)
-	{
-		ContextMenu.Reset();
-
-		if(DialogWindow.IsValid())
-		{
-			DialogWindow.Pin()->BringToFront(true);
-
-			TSharedRef<SWidget> ThisRef = AsShared();
-			FSlateApplication::Get().ForEachUser([&ThisRef](FSlateUser& User) {
-				User.SetFocus(ThisRef, EFocusCause::SetDirectly);
-			});
-		}
-	});
 
-	return FReply::Handled().SetUserFocus(SearchableTreeView->GetSearchBox(), EFocusCause::SetDirectly);
+	AddComboButton->SetIsMenuOpen(false, false);
 }
 
 bool SRigSpacePickerWidget::IsRestricted() const
diff --git a/Engine/Plugins/Animation/ControlRig/Source/ControlRigEditor/Private/Editor/SRigSpacePickerWidget.h b/Engine/Plugins/Animation/ControlRig/Source/ControlRigEditor/Private/Editor/SRigSpacePickerWidget.h
--- a/Engine/Plugins/Animation/ControlRig/Source/ControlRigEditor/Private/Editor/SRigSpacePickerWidget.h
+++ b/Engine/Plugins/Animation/ControlRig/Source/ControlRigEditor/Private/Editor/SRigSpacePickerWidget.h
@@ -11,9 +11,11 @@
 #include "Misc/FrameNumber.h"
 #include "EditorUndoClient.h"
 #include "RigSpacePickerBakeSettings.h"
+#include "SPositiveActionButton.h"
 
 #define UE_API CONTROLRIGEDITOR_API
 
+class SComboButton;
 class IMenu;
 
 DECLARE_DELEGATE_RetVal_TwoParams(FRigElementKey, FRigSpacePickerGetActiveSpace, URigHierarchy*, const FRigElementKey&);
@@ -131,9 +133,18 @@ private:
 	UE_API FReply HandleSpaceMoveDown(FRigElementKey InKey);
 	UE_API void HandleSpaceDelete(FRigElementKey InKey);
 
+	/** The combo button for adding items. Can be nullptr. */
+	TSharedPtr<SPositiveActionButton> AddComboButton;
+
+	bool bIsAddMenuOpen = false;
+	
+	TSharedRef<SWidget> CreateAddButton();
+	TSharedRef<SWidget> OnGetAddButtonContent();
+	void OnMenuOpenChanged(bool bIsOpen);
+	void HandleClickTreeItemInAddMenu(TSharedPtr<FRigTreeElement> InItem);
+
 public:
 	
-	UE_API FReply HandleAddElementClicked();
 	UE_API bool IsRestricted() const;
 
 private:
@@ -163,7 +174,6 @@ private:
 	bool bAllowDelete;
 	bool bAllowAdd;
 	bool bShowBakeAndCompensateButton;
-	bool bLaunchingContextMenu;
 
 	FRigSpacePickerGetControlCustomization GetControlCustomizationDelegate;
 	FRigSpacePickerGetActiveSpace GetActiveSpaceDelegate;
@@ -174,7 +184,6 @@ private:
 	TSharedPtr<SVerticalBox> ItemSpacesListBox;
 	TSharedPtr<SHorizontalBox> BottomButtonsListBox;
 	TWeakPtr<SWindow> DialogWindow;
-	TWeakPtr<IMenu> ContextMenu;
 	FDelegateHandle HierarchyModifiedHandle;
 	FDelegateHandle ActiveSpaceChangedWindowHandle;
 	FRigTreeDisplaySettings HierarchyDisplaySettings;
