Control Rig:  Constrains edition widget in the new constrain window now has SBorder as backtground.

The SExpandableArea header was removed. This area should not be collapsible because it's the content of an inline tab. The only reason it was collapsible was because that's how the old UI did it.

diff --git a/Engine/Plugins/Animation/ControlRig/Source/ControlRigEditor/Private/EditMode/Constrain/SConstraintsEditionWidget_ConstrainPanel.cpp b/Engine/Plugins/Animation/ControlRig/Source/ControlRigEditor/Private/EditMode/Constrain/SConstraintsEditionWidget_ConstrainPanel.cpp
--- a/Engine/Plugins/Animation/ControlRig/Source/ControlRigEditor/Private/EditMode/Constrain/SConstraintsEditionWidget_ConstrainPanel.cpp
+++ b/Engine/Plugins/Animation/ControlRig/Source/ControlRigEditor/Private/EditMode/Constrain/SConstraintsEditionWidget_ConstrainPanel.cpp
@@ -3,12 +3,12 @@
 #include "SConstraintsEditionWidget_ConstrainPanel.h"
 
 #include "EditMode/ControlRigEditMode.h"
-#include "EditMode/SComponentPickerPopup.h"
 #include "EditMode/Models/RigSelectionViewModel.h"
+#include "EditMode/SComponentPickerPopup.h"
+#include "SPositiveActionButton.h"
 #include "Widgets/Input/SButton.h"
 #include "Widgets/Input/SComboButton.h"
-#include "Widgets/Layout/SExpandableArea.h"
-#include "Widgets/Layout/SSpacer.h"
+#include "Widgets/Layout/SBox.h"
 
 #define LOCTEXT_NAMESPACE "SExpandableConstraintsEditionWidget"
 
@@ -22,60 +22,37 @@ void SConstraintsEditionWidget_ConstrainPanel::Construct(
 	SelectionViewModel = InSelectionViewModel;
 
 	InSelectionViewModel->OnControlSelected().AddSP(this, &SConstraintsEditionWidget_ConstrainPanel::OnControlSelected);
-	
+
+	constexpr float TopAreaHeight = 22;
 	ChildSlot
 	[
-		SNew(SExpandableArea)
-		.InitiallyCollapsed(InArgs._InitiallyCollapsed)
-		.AreaTitle(LOCTEXT("ConstraintsWidget", "Constraints"))
-		.AreaTitleFont(FAppStyle::GetFontStyle("DetailsView.CategoryFontStyle"))
-		.BorderBackgroundColor(FLinearColor(.6f, .6f, .6f))
-		.HeaderContent()
-		[
-			SNew(SHorizontalBox)
+		SNew(SVerticalBox)
 
-			// "Constraints" label
-			+SHorizontalBox::Slot()
-			.AutoWidth()
-			.HAlign(HAlign_Left)
-			.VAlign(VAlign_Center)
-			.Padding(0.f, 0.f, 0.f, 0.f)
-			[
-				SNew(STextBlock)
-				.Text(LOCTEXT("ConstraintsWidget", "Constraints"))
-				.Font(FCoreStyle::Get().GetFontStyle("ExpandableArea.TitleFont"))
-			]
-
-			// Spacer
-			+SHorizontalBox::Slot()
-			.FillWidth(1.f)
-			[
-				SNew(SSpacer)
-			]
-			
-			// "Selected" button
-			+SHorizontalBox::Slot()
-			.AutoWidth()
-			.HAlign(HAlign_Right)
-			.VAlign(VAlign_Center)
-			.Padding(0.f, 2.f, 8.f, 2.f)
-			[
-				// Combo Button to swap what constraints we show  
-				CreateSelectionComboButton()
-			]
-			
-			// "Plus" icon
-			+SHorizontalBox::Slot()
-			.AutoWidth()
-			.HAlign(HAlign_Right)
-			.VAlign(VAlign_Center)
-			.Padding(0.f, 2.f, 8.f, 2.f)
+		+SVerticalBox::Slot()
+		.AutoHeight()
+		[
+			SNew(SBorder)
+			.BorderImage(FAppStyle::GetBrush("Menu.Background"))
+			.Padding(2.f)
 			[
-				CreatePlusButton()
+				SAssignNew(ConstraintsEditionWidget, SConstraintsEditionWidget)
+
+				// We need to make sure the splitter "row" we're adding is the same height
+				.LeftSplitterContent()
+				[
+					SNew(SBox).HeightOverride(TopAreaHeight) [ SNullWidget::NullWidget ]
+				]
+				.RightSplitterContent()
+				[
+					SNew(SBox).HeightOverride(TopAreaHeight) [ CreateSelectionComboButton() ]
+				]
+				
+				.LeftButtonArea()
+				[
+					CreateAddConstraintButton()
+				]
 			]
 		]
-		.BodyContent() [ SAssignNew(ConstraintsEditionWidget, SConstraintsEditionWidget) ]
-		.OnAreaExpansionChanged_Lambda( [this](bool bIsExpanded){ ConstraintsEditionWidget->RefreshConstraintList(); })
 	];
 
 	// LevelSequence is the most sensible default. It should be set back to this even if the user had the panel open earlier, closed it, and reopend it.
@@ -120,27 +97,29 @@ TSharedRef<SWidget> SConstraintsEditionWidget_ConstrainPanel::CreateSelectionCom
 		];
 }
 
-TSharedRef<SWidget> SConstraintsEditionWidget_ConstrainPanel::CreatePlusButton()
+TSharedRef<SWidget> SConstraintsEditionWidget_ConstrainPanel::CreateAddConstraintButton()
 {
-	return SNew(SButton)
-		.ContentPadding(0.0f)
-		.ButtonStyle(FAppStyle::Get(), "NoBorder")
-		.IsEnabled_Lambda([this]()
+	const auto IsEnabled = [this]
+	{
+		const ULevel* CurrentLevel = OwningEditMode->GetWorld()->GetCurrentLevel();
+		const TArray<AActor*> SelectedActors = ObjectPtrDecay(CurrentLevel->Actors).FilterByPredicate([](const AActor* Actor)
 		{
-			const ULevel* CurrentLevel = OwningEditMode->GetWorld()->GetCurrentLevel();
-			const TArray<AActor*> SelectedActors = ObjectPtrDecay(CurrentLevel->Actors).FilterByPredicate([](const AActor* Actor)
-			{
-				return Actor && Actor->IsSelected();
-			});
-			return !SelectedActors.IsEmpty();
-		})
+			return Actor && Actor->IsSelected();
+		});
+		return !SelectedActors.IsEmpty();
+	};
+	
+	return SNew(SPositiveActionButton)
+		.IsEnabled_Lambda(IsEnabled)
 		.OnClicked(this, &SConstraintsEditionWidget_ConstrainPanel::HandleAddConstraintClicked)
 		.Cursor(EMouseCursor::Default)
-		.ToolTipText(LOCTEXT("AddConstraint", "Add Constraint"))
-		[
-			SNew(SImage)
-			.Image(FAppStyle::GetBrush(TEXT("Icons.PlusCircle")))
-		];
+		.Text(LOCTEXT("AddConstraint.Label", "Add"))
+		.ToolTipText_Lambda([IsEnabled]
+		{
+			return IsEnabled()
+				? LOCTEXT("AddConstraint.ToolTip.Enabled", "Add constraint")
+				: LOCTEXT("AddConstraint.ToolTip.Disabled", "Select a control rig to add constraint");
+		});
 }
 
 void SConstraintsEditionWidget_ConstrainPanel::OnSelectShowConstraints(int32 Index)
diff --git a/Engine/Plugins/Animation/ControlRig/Source/ControlRigEditor/Private/EditMode/Constrain/SConstraintsEditionWidget_ConstrainPanel.h b/Engine/Plugins/Animation/ControlRig/Source/ControlRigEditor/Private/EditMode/Constrain/SConstraintsEditionWidget_ConstrainPanel.h
--- a/Engine/Plugins/Animation/ControlRig/Source/ControlRigEditor/Private/EditMode/Constrain/SConstraintsEditionWidget_ConstrainPanel.h
+++ b/Engine/Plugins/Animation/ControlRig/Source/ControlRigEditor/Private/EditMode/Constrain/SConstraintsEditionWidget_ConstrainPanel.h
@@ -45,9 +45,10 @@ private:
 	
 	/** The constraints edition widget we wrap. */
 	TSharedPtr<SConstraintsEditionWidget> ConstraintsEditionWidget;
-	
+
+	// Header area (similar to expandable area border we used to have)
 	TSharedRef<SWidget> CreateSelectionComboButton();
-	TSharedRef<SWidget> CreatePlusButton();
+	TSharedRef<SWidget> CreateAddConstraintButton();
 
 	/**  */
 	void OnSelectShowConstraints(int32 Index);
diff --git a/Engine/Plugins/Animation/ControlRig/Source/ControlRigEditor/Private/Editor/Constraints/SConstraintsWidget.cpp b/Engine/Plugins/Animation/ControlRig/Source/ControlRigEditor/Private/Editor/Constraints/SConstraintsWidget.cpp
--- a/Engine/Plugins/Animation/ControlRig/Source/ControlRigEditor/Private/Editor/Constraints/SConstraintsWidget.cpp
+++ b/Engine/Plugins/Animation/ControlRig/Source/ControlRigEditor/Private/Editor/Constraints/SConstraintsWidget.cpp
@@ -1270,7 +1270,7 @@ void SConstraintsEditionWidget::Construct(const FArguments& InArgs)
 	
 	WeakSequencer = GetSequencerChecked();
 
-	auto GetParentWidget = [this]()
+	auto GetParentWidget = [this, &InArgs]()
 	{
 		TSharedRef<SSplitter> Splitter = SNew(SSplitter)
 		.Style(FAppStyle::Get(), "DetailsView.Splitter")
@@ -1281,16 +1281,30 @@ void SConstraintsEditionWidget::Construct(const FArguments& InArgs)
 		+ SSplitter::Slot()
 		.Value(0.4f)
 		[
-			SNew(SHorizontalBox)
-			.Clipping(EWidgetClipping::OnDemand)
-			+ SHorizontalBox::Slot()
-			.VAlign(VAlign_Center)
-			.HAlign(HAlign_Left)
-			.Padding(20.0f, 0.0f, 0.0f, 0.0f)
+			SNew(SVerticalBox)
+
+			+SVerticalBox::Slot()
+			.Padding(4.0f, 2.0f)
+			.AutoHeight()
 			[
-				SNew(STextBlock)
-				.Font(FAppStyle::Get().GetFontStyle("PropertyWindow.NormalFont"))
-				.Text(LOCTEXT("PreSelectedParent_Label", "Pre Selected Parent"))
+				InArgs._LeftSplitterContent.Widget
+			]
+
+			+SVerticalBox::Slot()
+			.Padding(4.0f, 2.0f)
+			.AutoHeight()
+			[
+				SNew(SHorizontalBox)
+				.Clipping(EWidgetClipping::OnDemand)
+				+ SHorizontalBox::Slot()
+				.VAlign(VAlign_Center)
+				.HAlign(HAlign_Left)
+				.Padding(20.0f, 0.0f, 0.0f, 0.0f)
+				[
+					SNew(STextBlock)
+					.Font(FAppStyle::Get().GetFontStyle("PropertyWindow.NormalFont"))
+					.Text(LOCTEXT("PreSelectedParent_Label", "Pre Selected Parent"))
+				]
 			]
 		]
 
@@ -1298,39 +1312,52 @@ void SConstraintsEditionWidget::Construct(const FArguments& InArgs)
 		+ SSplitter::Slot()
 		.Value(0.6f)
 		[
-			SNew(SAssetDropTarget)
-			.bSupportsMultiDrop(true)
-			.OnAreAssetsAcceptableForDrop_Lambda([](TArrayView<FAssetData> InAssets)
-			{
-				const bool bHasAnyNotConstrainable = InAssets.ContainsByPredicate([](const FAssetData& Data)
-				{
-					return !IsActorConstrainable(Cast<AActor>(Data.GetAsset()));	
-				});
-				return !bHasAnyNotConstrainable;
-			})
-			.OnAssetsDropped_Lambda( [this](const FDragDropEvent&, TArrayView<FAssetData> InAssets)
-			{
-				UpdateParent(InAssets.IsEmpty() ? nullptr : Cast<AActor>(InAssets[0].GetAsset()));
-			})
+			SNew(SVerticalBox)
+
+			+SVerticalBox::Slot()
+			.Padding(4.0f, 2.0f)
+			.AutoHeight()
 			[
-				SNew(SHorizontalBox)
-											
-				// actor chooser widget
-				+ SHorizontalBox::Slot()
-				.Padding(FMargin(4.0f,2.0f))
-				.HAlign(HAlign_Right)
-				[
-					GetActorChooserWidget().ToSharedRef()
-				]
-			
-				// use selected widget
-				+ SHorizontalBox::Slot()
-				.AutoWidth()
+				InArgs._MiddleSplitterContent.Widget
+			]
+
+			+SVerticalBox::Slot()
+			.Padding(4.0f, 2.0f)
+			.AutoHeight()
+			[
+				SNew(SAssetDropTarget)
+				.bSupportsMultiDrop(true)
+				.OnAreAssetsAcceptableForDrop_Lambda([](TArrayView<FAssetData> InAssets)
+				{
+					const bool bHasAnyNotConstrainable = InAssets.ContainsByPredicate([](const FAssetData& Data)
+					{
+						return !IsActorConstrainable(Cast<AActor>(Data.GetAsset()));	
+					});
+					return !bHasAnyNotConstrainable;
+				})
+				.OnAssetsDropped_Lambda( [this](const FDragDropEvent&, TArrayView<FAssetData> InAssets)
+				{
+					UpdateParent(InAssets.IsEmpty() ? nullptr : Cast<AActor>(InAssets[0].GetAsset()));
+				})
 				[
-					SNew(SBox)
-					.Padding(FMargin(4.0f,0.0f))
+					SNew(SHorizontalBox)
+												
+					// actor chooser widget
+					+ SHorizontalBox::Slot()
+					.HAlign(HAlign_Left)
+					[
+						GetActorChooserWidget().ToSharedRef()
+					]
+				
+					// use selected widget
+					+ SHorizontalBox::Slot()
+					.AutoWidth()
 					[
-						GetUseSelectedActorWidget()
+						SNew(SBox)
+						.Padding(FMargin(4.0f,0.0f))
+						[
+							GetUseSelectedActorWidget()
+						]
 					]
 				]
 			]
@@ -1343,36 +1370,50 @@ void SConstraintsEditionWidget::Construct(const FArguments& InArgs)
 		.MinSize(26.f)
 		.Resizable(false)
 		[
-			SNew(SBorder)
-			.BorderImage(FAppStyle::Get().GetBrush("DetailsView.CategoryMiddle"))
-			.BorderBackgroundColor(FAppStyle::Get().GetSlateColor("Colors.Panel"))
-			.HAlign(HAlign_Center)
-			.VAlign(VAlign_Center)
+			SNew(SVerticalBox)
+
+			+SVerticalBox::Slot()
+			.Padding(4.0f, 2.0f)
+			.AutoHeight()
+			[
+				InArgs._RightSplitterContent.Widget
+			]
+
+			+SVerticalBox::Slot()
+			.Padding(4.0f, 2.0f)
+			.AutoHeight()
 			[
-				SNew(SBox)
+				SNew(SBorder)
+				.BorderImage(FAppStyle::Get().GetBrush("DetailsView.CategoryMiddle"))
+				.BorderBackgroundColor(FAppStyle::Get().GetSlateColor("Colors.Panel"))
 				.HAlign(HAlign_Center)
 				.VAlign(VAlign_Center)
-				.WidthOverride(22.0f)
-				.HeightOverride(22.0f)
-				.ToolTipText(LOCTEXT("ResetParent_Tooltip", "Reset parent."))
 				[
-					SNew(SButton)
-					.ButtonStyle( FAppStyle::Get(), "SimpleButton" )
-					.OnClicked_Lambda([this]()
-					{
-						ResetParent();
-						return FReply::Handled();
-					})
-					.Visibility_Lambda([this]()
-					{
-						return IsParentValid() ? EVisibility::Visible : EVisibility::Hidden;
-					})
-					.ContentPadding(0.0f)
-					.IsFocusable(false)
-					[ 
-						SNew( SImage )
-						.Image(FAppStyle::GetBrush("PropertyWindow.DiffersFromDefault"))
-						.ColorAndOpacity( FSlateColor::UseForeground() )
+					SNew(SBox)
+					.HAlign(HAlign_Center)
+					.VAlign(VAlign_Center)
+					.WidthOverride(22.0f)
+					.HeightOverride(22.0f)
+					.ToolTipText(LOCTEXT("ResetParent_Tooltip", "Reset parent."))
+					[
+						SNew(SButton)
+						.ButtonStyle( FAppStyle::Get(), "SimpleButton" )
+						.OnClicked_Lambda([this]()
+						{
+							ResetParent();
+							return FReply::Handled();
+						})
+						.Visibility_Lambda([this]()
+						{
+							return IsParentValid() ? EVisibility::Visible : EVisibility::Hidden;
+						})
+						.ContentPadding(0.0f)
+						.IsFocusable(false)
+						[ 
+							SNew( SImage )
+							.Image(FAppStyle::GetBrush("PropertyWindow.DiffersFromDefault"))
+							.ColorAndOpacity( FSlateColor::UseForeground() )
+						]
 					]
 				]
 			]
@@ -1419,9 +1460,15 @@ void SConstraintsEditionWidget::Construct(const FArguments& InArgs)
 		];
 	};
 
-	auto GetBakingWidget = [this]()
+	auto GetBakingWidget = [this, &InArgs]()
 	{
 		return SNew(SHorizontalBox)
+
+		+SHorizontalBox::Slot()
+		.AutoWidth()
+		[
+			InArgs._LeftButtonArea.Widget
+		]
 		
 		+ SHorizontalBox::Slot()
 		.FillWidth(1.f)
diff --git a/Engine/Plugins/Animation/ControlRig/Source/ControlRigEditor/Private/Editor/Constraints/SConstraintsWidget.h b/Engine/Plugins/Animation/ControlRig/Source/ControlRigEditor/Private/Editor/Constraints/SConstraintsWidget.h
--- a/Engine/Plugins/Animation/ControlRig/Source/ControlRigEditor/Private/Editor/Constraints/SConstraintsWidget.h
+++ b/Engine/Plugins/Animation/ControlRig/Source/ControlRigEditor/Private/Editor/Constraints/SConstraintsWidget.h
@@ -297,7 +297,16 @@ class SConstraintsEditionWidget : public SCompoundWidget, public FBaseConstraint
 {
 public:
 	SLATE_BEGIN_ARGS(SConstraintsEditionWidget)	{}
+		
+		// This is a temporary way to allow you to add "rows" above the details for this panel. This will be revisited in 5.8.
+		SLATE_NAMED_SLOT(FArguments, LeftSplitterContent)
+		SLATE_NAMED_SLOT(FArguments, MiddleSplitterContent)
+		SLATE_NAMED_SLOT(FArguments, RightSplitterContent)
+		
+		/** Allows adding extra buttons on the left of the bake button. */
+		SLATE_NAMED_SLOT(FArguments, LeftButtonArea)
 	SLATE_END_ARGS()
+	
 	/** Constructs this widget with InArgs */
 	UE_API void Construct(const FArguments& InArgs);
 
