Control Rig: Clean up stale control rigs created by PIE playback

#rnx

diff --git a/Engine/Plugins/Animation/ControlRig/Source/ControlRig/Private/Sequencer/MovieSceneControlRigParameterTrack.cpp b/Engine/Plugins/Animation/ControlRig/Source/ControlRig/Private/Sequencer/MovieSceneControlRigParameterTrack.cpp
--- a/Engine/Plugins/Animation/ControlRig/Source/ControlRig/Private/Sequencer/MovieSceneControlRigParameterTrack.cpp
+++ b/Engine/Plugins/Animation/ControlRig/Source/ControlRig/Private/Sequencer/MovieSceneControlRigParameterTrack.cpp
@@ -17,6 +17,7 @@
 
 #if WITH_EDITOR
 #include "ScopedTransaction.h"
+#include "Editor.h"
 #endif//WITH_EDITOR
 
 #include UE_INLINE_GENERATED_CPP_BY_NAME(MovieSceneControlRigParameterTrack)
@@ -77,6 +78,10 @@ UMovieSceneControlRigParameterTrack::UMovieSceneControlRigParameterTrack(const F
 	SupportedBlendTypes.Add(EMovieSceneBlendType::Override);
 
 	UE::MovieScene::ShouldUseLegacyControlRigTemplate = &UMovieSceneControlRigParameterTrack::ShouldUseLegacyTemplate;
+
+#if WITH_EDITOR
+	FEditorDelegates::ShutdownPIE.AddUObject(this, &UMovieSceneControlRigParameterTrack::OnShutdownPIE);
+#endif
 }
 
 void UMovieSceneControlRigParameterTrack::BeginDestroy()
@@ -130,6 +135,33 @@ void UMovieSceneControlRigParameterTrack::CheckForNameChanges()
 	}
 }
 
+#if WITH_EDITOR
+
+void UMovieSceneControlRigParameterTrack::OnShutdownPIE(bool bIsSimulating)
+{
+	TArray<TWeakObjectPtr<UWorld>> KeysToRemove;
+	for (const TPair<TWeakObjectPtr<UWorld>,TObjectPtr<UControlRig>>& Pair : GameWorldControlRigs)
+	{
+		if (!Pair.Key.IsValid() || Pair.Key.IsStale())
+		{
+			KeysToRemove.Add(Pair.Key);
+			if (UControlRig* GameWorldControlRig = Pair.Value)
+			{
+				if (IsValid(GameWorldControlRig))
+				{
+					GameWorldControlRig->MarkAsGarbage();
+				}
+			}
+		}
+	}
+	for (const TWeakObjectPtr<UWorld>& KeyToRemove : KeysToRemove)
+	{
+		GameWorldControlRigs.Remove(KeyToRemove);
+	}
+}
+
+#endif
+
 UMovieSceneSection* UMovieSceneControlRigParameterTrack::CreateNewSection()
 {
 	UMovieSceneControlRigParameterSection* NewSection = NewObject<UMovieSceneControlRigParameterSection>(this, NAME_None, RF_Transactional);
diff --git a/Engine/Plugins/Animation/ControlRig/Source/ControlRig/Public/Sequencer/MovieSceneControlRigParameterTrack.h b/Engine/Plugins/Animation/ControlRig/Source/ControlRig/Public/Sequencer/MovieSceneControlRigParameterTrack.h
--- a/Engine/Plugins/Animation/ControlRig/Source/ControlRig/Public/Sequencer/MovieSceneControlRigParameterTrack.h
+++ b/Engine/Plugins/Animation/ControlRig/Source/ControlRig/Public/Sequencer/MovieSceneControlRigParameterTrack.h
@@ -209,6 +209,10 @@ private:
 	//look for name changes in particularly with modular rigs so we can update the rotation order map
 	void CheckForNameChanges();
 
+#if WITH_EDITOR
+	void OnShutdownPIE(bool bIsSimulating);
+#endif
+
 private:
 
 
