Control Rig: Avoid applying the control rig component during teardown

#rnx

diff --git a/Engine/Plugins/Animation/ControlRig/Source/ControlRig/Private/ControlRigComponent.cpp b/Engine/Plugins/Animation/ControlRig/Source/ControlRig/Private/ControlRigComponent.cpp
--- a/Engine/Plugins/Animation/ControlRig/Source/ControlRig/Private/ControlRigComponent.cpp
+++ b/Engine/Plugins/Animation/ControlRig/Source/ControlRig/Private/ControlRigComponent.cpp
@@ -176,7 +176,7 @@ void UControlRigComponent::OnUnregister()
 	{
 		if (AActor* Actor = GetOwner())
 		{
-			bBeginDestroyed = URigVMHost::IsGarbageOrDestroyed(Actor);
+			bBeginDestroyed = URigVMHost::IsGarbageOrDestroyed(Actor) || Actor->IsActorBeingDestroyed();
 		}
 	}
 
