Control Rig: Dependency Traverse to ignore variables by default

To make the dependency tracking a bit more forgiving we'll use the variable dependencies only for the dependency viewer

#rnx

diff --git a/Engine/Plugins/Animation/ControlRig/Source/ControlRig/Private/Rigs/RigDependency.cpp b/Engine/Plugins/Animation/ControlRig/Source/ControlRig/Private/Rigs/RigDependency.cpp
--- a/Engine/Plugins/Animation/ControlRig/Source/ControlRig/Private/Rigs/RigDependency.cpp
+++ b/Engine/Plugins/Animation/ControlRig/Source/ControlRig/Private/Rigs/RigDependency.cpp
@@ -16,10 +16,12 @@ const TRigHierarchyDependencyMap& FEmptyRigDependenciesProvider::GetReverseRigHi
 	return Empty;
 }
 
-FRigDependenciesProviderForVM::FRigDependenciesProviderForVM(UControlRig* InControlRig, const FName& InEventName)
+FRigDependenciesProviderForVM::FRigDependenciesProviderForVM(UControlRig* InControlRig, const FName& InEventName, bool InFollowVariables)
 	: WeakControlRig(InControlRig)
 	, EventName(InEventName)
-{}
+	, bFollowVariables(InFollowVariables)
+{
+}
 
 // UE_DISABLE_OPTIMIZATION
 
@@ -157,7 +159,7 @@ void FRigDependenciesProviderForVM::ComputeDependencies(const FControlRigExecute
 
 	// for each written transform / metadata on an instruction
 	TArray<TTuple<int32,FRigHierarchyRecord>> InstructionsToVisit;
-	TArray<FRigHierarchyRecord> ReadVariables, WrittenVariables;
+	TArray<FRigHierarchyRecord> WrittenVariables;
 	TMap<FRigVMOperand, TArray<int32>> InputOperandToInstructions;
 	TMap<FRigVMOperand, TArray<int32>> OutputOperandToInstructions;
 
@@ -166,6 +168,10 @@ void FRigDependenciesProviderForVM::ComputeDependencies(const FControlRigExecute
 		FRigVMOperandArray InputOperands = InVM->GetByteCode().GetInputOperands(InstructionIndex);
 		for (const FRigVMOperand& Operand : InputOperands)
 		{
+			if (!bFollowVariables && Operand.GetMemoryType() == ERigVMMemoryType::External)
+			{
+				continue;
+			}
 			const FRigVMOperand OperandNoRegisterOffset(Operand.GetMemoryType(), Operand.GetRegisterIndex());
 			InputOperandToInstructions.FindOrAdd(OperandNoRegisterOffset).Add(InstructionIndex);
 		}
@@ -173,6 +179,10 @@ void FRigDependenciesProviderForVM::ComputeDependencies(const FControlRigExecute
 		FRigVMOperandArray OutputOperands = InVM->GetByteCode().GetOutputOperands(InstructionIndex);
 		for (const FRigVMOperand& Operand : OutputOperands)
 		{
+			if (!bFollowVariables && Operand.GetMemoryType() == ERigVMMemoryType::External)
+			{
+				continue;
+			}
 			const FRigVMOperand OperandNoRegisterOffset(Operand.GetMemoryType(), Operand.GetRegisterIndex());
 			OutputOperandToInstructions.FindOrAdd(OperandNoRegisterOffset).Add(InstructionIndex);
 		}
@@ -180,7 +190,6 @@ void FRigDependenciesProviderForVM::ComputeDependencies(const FControlRigExecute
 
 	for(int32 InstructionIndex = 0; InstructionIndex < Instructions.Num(); InstructionIndex++)
 	{
-		ReadVariables.Reset();
 		WrittenVariables.Reset();
 
 		if (ReadRecordsPerInstruction.IsValidIndex(InstructionIndex))
@@ -210,6 +219,10 @@ void FRigDependenciesProviderForVM::ComputeDependencies(const FControlRigExecute
 
 						if (OutputOperand.GetMemoryType() == ERigVMMemoryType::External)
 						{
+							if (!bFollowVariables)
+							{
+								continue;
+							}
 							const FRigHierarchyRecord VariableRecord(FRigHierarchyRecord::EType_Variable, OutputOperand.GetRegisterIndex(), ReadRecord.VM);
 							WrittenVariables.AddUnique(VariableRecord);
 							LastRecord = VariableRecord;
@@ -288,6 +301,10 @@ void FRigDependenciesProviderForVM::ComputeDependencies(const FControlRigExecute
 
 						if (InputOperand.GetMemoryType() == ERigVMMemoryType::External)
 						{
+							if (!bFollowVariables)
+							{
+								continue;
+							}
 							const FRigHierarchyRecord VariableRecord(FRigHierarchyRecord::EType_Variable, InputOperand.GetRegisterIndex(), WrittenRecord.VM);
 							OutDependencies.FindOrAdd(InstructionsToVisit[InstructionToVisitIndex].Get<1>()).Add(VariableRecord);
 							LastRecord = VariableRecord;
@@ -311,8 +328,8 @@ void FRigDependenciesProviderForVM::ComputeDependencies(const FControlRigExecute
 #endif
 }
 
-FRigDependenciesProviderForControlRig::FRigDependenciesProviderForControlRig(UControlRig* InControlRig, const FName& InEventName)
-	: FRigDependenciesProviderForVM(InControlRig, InEventName)
+FRigDependenciesProviderForControlRig::FRigDependenciesProviderForControlRig(UControlRig* InControlRig, const FName& InEventName, bool InFollowVariables)
+	: FRigDependenciesProviderForVM(InControlRig, InEventName, InFollowVariables)
 {
 }
 
diff --git a/Engine/Plugins/Animation/ControlRig/Source/ControlRig/Public/Rigs/RigDependency.h b/Engine/Plugins/Animation/ControlRig/Source/ControlRig/Public/Rigs/RigDependency.h
--- a/Engine/Plugins/Animation/ControlRig/Source/ControlRig/Public/Rigs/RigDependency.h
+++ b/Engine/Plugins/Animation/ControlRig/Source/ControlRig/Public/Rigs/RigDependency.h
@@ -24,6 +24,7 @@ struct IRigDependenciesProvider
 	virtual const TRigHierarchyDependencyMap& GetRigHierarchyDependencies() const = 0;
 	virtual const TRigHierarchyDependencyMap& GetReverseRigHierarchyDependencies() const = 0;
 	virtual void InvalidateCache() = 0;
+	virtual uint32 GetHash() const = 0;
 };
 
 /**
@@ -36,6 +37,7 @@ struct FEmptyRigDependenciesProvider : public IRigDependenciesProvider
 	UE_API virtual const TRigHierarchyDependencyMap& GetRigHierarchyDependencies() const override;
 	UE_API virtual const TRigHierarchyDependencyMap& GetReverseRigHierarchyDependencies() const;
 	virtual void InvalidateCache() {}
+	virtual uint32 GetHash() const { return 0; }
 };
 
 /**
@@ -45,7 +47,7 @@ struct FEmptyRigDependenciesProvider : public IRigDependenciesProvider
 struct FRigDependenciesProviderForVM : public IRigDependenciesProvider
 {
 	FRigDependenciesProviderForVM() = default;
-	UE_API FRigDependenciesProviderForVM(UControlRig* InControlRig, const FName& InEventName = NAME_None);
+	UE_API FRigDependenciesProviderForVM(UControlRig* InControlRig, const FName& InEventName = NAME_None, bool InFollowVariables = false);
 	virtual ~FRigDependenciesProviderForVM() override = default;
 	
 	UE_API virtual const TRigHierarchyDependencyMap& GetRigHierarchyDependencies() const override;
@@ -53,6 +55,7 @@ struct FRigDependenciesProviderForVM : public IRigDependenciesProvider
 	UE_API virtual void InvalidateCache() override;
 
 	UE_API UControlRig* GetControlRig() const;
+	virtual uint32 GetHash() const { return RecordsHash; }
 
 protected:
 
@@ -60,6 +63,7 @@ protected:
 
 	TWeakObjectPtr<UControlRig> WeakControlRig = nullptr;
 	FName EventName = NAME_None;
+	bool bFollowVariables;
 	
 	mutable TRigHierarchyDependencyMap CachedDependencies;
 	mutable TRigHierarchyDependencyMap CachedReverseDependencies;
@@ -75,7 +79,7 @@ protected:
 struct FRigDependenciesProviderForControlRig : public FRigDependenciesProviderForVM
 {
 	FRigDependenciesProviderForControlRig() = default;
-	UE_API FRigDependenciesProviderForControlRig(UControlRig* InControlRig, const FName& InEventName = NAME_None);
+	UE_API FRigDependenciesProviderForControlRig(UControlRig* InControlRig, const FName& InEventName = NAME_None, bool InFollowVariables = false);
 	virtual ~FRigDependenciesProviderForControlRig() override = default;
 	
 	UE_API virtual const TRigHierarchyDependencyMap& GetRigHierarchyDependencies() const override;
diff --git a/Engine/Plugins/Animation/ControlRig/Source/ControlRigEditor/Private/RigDependencyGraph/RigDependencyGraph.cpp b/Engine/Plugins/Animation/ControlRig/Source/ControlRigEditor/Private/RigDependencyGraph/RigDependencyGraph.cpp
--- a/Engine/Plugins/Animation/ControlRig/Source/ControlRigEditor/Private/RigDependencyGraph/RigDependencyGraph.cpp
+++ b/Engine/Plugins/Animation/ControlRig/Source/ControlRigEditor/Private/RigDependencyGraph/RigDependencyGraph.cpp
@@ -55,7 +55,7 @@ void URigDependencyGraph::Initialize(const TSharedRef<IControlRigBaseEditor>& In
 
 	if(WeakControlRig.IsValid() && WeakControlRig.Get())
 	{
-		RigDependencyProvider = MakeShareable(new FRigDependenciesProviderForControlRig(WeakControlRig.Get()));
+		RigDependencyProvider = MakeShareable(new FRigDependenciesProviderForControlRig(WeakControlRig.Get(), NAME_None, true));
 	}
 	else
 	{
@@ -545,7 +545,7 @@ TArray<URigDependencyGraph::FNodeId> URigDependencyGraph::GetNodesToExpand(const
 					
 					if(DependencyRecord.IsInstruction() || DependencyRecord.IsVariable() || DependencyRecord.IsMetadata())
 					{
-						NodesToFollow.Add(DependencyRecord);
+						NodesToFollow.AddUnique(DependencyRecord);
 					}
 				}
 			}
