Control Rig: Fix misalignment issue in new constraint window

preflight

diff --git a/Engine/Plugins/Animation/ControlRig/Source/ControlRigEditor/Private/EditMode/Constrain/SConstraintsEditionWidget_ConstrainPanel.cpp b/Engine/Plugins/Animation/ControlRig/Source/ControlRigEditor/Private/EditMode/Constrain/SConstraintsEditionWidget_ConstrainPanel.cpp
--- a/Engine/Plugins/Animation/ControlRig/Source/ControlRigEditor/Private/EditMode/Constrain/SConstraintsEditionWidget_ConstrainPanel.cpp
+++ b/Engine/Plugins/Animation/ControlRig/Source/ControlRigEditor/Private/EditMode/Constrain/SConstraintsEditionWidget_ConstrainPanel.cpp
@@ -9,6 +9,7 @@
 #include "Widgets/Input/SButton.h"
 #include "Widgets/Input/SComboButton.h"
 #include "Widgets/Layout/SBox.h"
+#include "Widgets/SNullWidget.h"
 
 #define LOCTEXT_NAMESPACE "SExpandableConstraintsEditionWidget"
 
@@ -37,14 +38,18 @@ void SConstraintsEditionWidget_ConstrainPanel::Construct(
 			[
 				SAssignNew(ConstraintsEditionWidget, SConstraintsEditionWidget)
 
-				// We need to make sure the splitter "row" we're adding is the same height
+				// We need to make sure the splitter "row" we're adding is the same height. 
 				.LeftSplitterContent()
 				[
 					SNew(SBox).HeightOverride(TopAreaHeight) [ SNullWidget::NullWidget ]
 				]
+				.MiddleSplitterContent()
+				[
+					SNew(SBox).HAlign(HAlign_Left).HeightOverride(TopAreaHeight) [ CreateSelectionComboButton() ]
+				]
 				.RightSplitterContent()
 				[
-					SNew(SBox).HeightOverride(TopAreaHeight) [ CreateSelectionComboButton() ]
+					SNew(SBox).HeightOverride(TopAreaHeight) [ SNullWidget::NullWidget ]
 				]
 				
 				.LeftButtonArea()
