Control Rig: Dependency Graph Viewer, double click workflow

#rnx

diff --git a/Engine/Plugins/Animation/ControlRig/Source/ControlRigEditor/Private/RigDependencyGraph/RigDependencyGraph.cpp b/Engine/Plugins/Animation/ControlRig/Source/ControlRigEditor/Private/RigDependencyGraph/RigDependencyGraph.cpp
--- a/Engine/Plugins/Animation/ControlRig/Source/ControlRigEditor/Private/RigDependencyGraph/RigDependencyGraph.cpp
+++ b/Engine/Plugins/Animation/ControlRig/Source/ControlRigEditor/Private/RigDependencyGraph/RigDependencyGraph.cpp
@@ -18,8 +18,9 @@ URigDependencyGraph::URigDependencyGraph(const FObjectInitializer& ObjectInitial
 	, LayoutIterationsTotal(0)
 	, LayoutIterationsMax(DefaultLayoutIterationsPerNode * 1000)
 	, bFollowParentRelationShips(true)
-	, bFollowVMRelationShips(true)
+	, bFollowVMRelationShips(false)
 	, bLockContent(false)
+	, bIsPerformingGridLayout(true)
 	, BlockSelectionCounter(2)
 {
 }
@@ -130,6 +131,22 @@ void URigDependencyGraph::RequestRefreshLayout()
 	UpdateFadeOutStates();
 }
 
+void URigDependencyGraph::RequestRefreshLayout(const TArray<FNodeId>& InNodes)
+{
+	for (const FNodeId& NodeId : InNodes)
+	{
+		if (URigDependencyGraphNode* const* NodePtr = NodeIdLookup.Find(NodeId))
+		{
+			URigDependencyGraphNode* Node = *NodePtr;
+			if (Node)
+			{
+				Node->FollowLayout.Reset();
+			}
+		}
+	}
+	RequestRefreshLayout();
+}
+
 bool URigDependencyGraph::NeedsRefreshLayout() const
 {
 	return LayoutIterationsLeft > 0;
@@ -523,10 +540,15 @@ void URigDependencyGraph::SetContentLocked(bool InLockContent)
 	bLockContent = InLockContent
 ;}
 
-TArray<URigDependencyGraph::FNodeId> URigDependencyGraph::GetNodesToExpand(const TArray<FNodeId>& InNodeIds) const
+TArray<URigDependencyGraph::FNodeId> URigDependencyGraph::GetNodesToExpand(const TArray<FNodeId>& InNodeIds, bool bExpandInputs, bool bExpandOutputs) const
 {
-	TArray<FNodeId> NodesToFollow = InNodeIds;
 	TArray<FNodeId> NodesToExpand;
+	if (!bExpandInputs && !bExpandOutputs)
+	{
+		return NodesToExpand;
+	}
+	
+	TArray<FNodeId> NodesToFollow = InNodeIds;
 	
 	// expand the nodes to construct one level using the dependency provider
 	if (bFollowVMRelationShips && RigDependencyProvider.IsValid())
@@ -537,23 +559,29 @@ TArray<URigDependencyGraph::FNodeId> URigDependencyGraph::GetNodesToExpand(const
 		for (int32 NodeIndex = 0; NodeIndex < NodesToFollow.Num(); ++NodeIndex)
 		{
 			const FRigHierarchyRecord Record = NodesToFollow[NodeIndex];
-			if (const TSet<FRigHierarchyRecord>* Records = Dependencies.Find(Record))
+			if (bExpandInputs)
 			{
-				for (const FRigHierarchyRecord& DependencyRecord : *Records)
+				if (const TSet<FRigHierarchyRecord>* Records = Dependencies.Find(Record))
 				{
-					NodesToExpand.AddUnique(DependencyRecord);
-					
-					if(DependencyRecord.IsInstruction() || DependencyRecord.IsVariable() || DependencyRecord.IsMetadata())
+					for (const FRigHierarchyRecord& DependencyRecord : *Records)
 					{
-						NodesToFollow.AddUnique(DependencyRecord);
+						NodesToExpand.AddUnique(DependencyRecord);
+					
+						if(DependencyRecord.IsInstruction() || DependencyRecord.IsVariable() || DependencyRecord.IsMetadata())
+						{
+							NodesToFollow.AddUnique(DependencyRecord);
+						}
 					}
 				}
 			}
-			if (const TSet<FRigHierarchyRecord>* ReverseRecords = ReverseDependencies.Find(Record))
+			if (bExpandOutputs)
 			{
-				for (const FRigHierarchyRecord& DependentRecord : *ReverseRecords)
+				if (const TSet<FRigHierarchyRecord>* ReverseRecords = ReverseDependencies.Find(Record))
 				{
-					NodesToExpand.AddUnique(DependentRecord);
+					for (const FRigHierarchyRecord& DependentRecord : *ReverseRecords)
+					{
+						NodesToExpand.AddUnique(DependentRecord);
+					}
 				}
 			}
 		}
@@ -573,8 +601,15 @@ TArray<URigDependencyGraph::FNodeId> URigDependencyGraph::GetNodesToExpand(const
 
 				if (Hierarchy->IsValidIndex(InNodeIds[NodeIndex].Index))
 				{
-					TArray<int32> ParentsAndChildren = Hierarchy->GetParents(InNodeIds[NodeIndex].Index);
-					ParentsAndChildren.Append(Hierarchy->GetChildren(InNodeIds[NodeIndex].Index));
+					TArray<int32> ParentsAndChildren;
+					if (bExpandInputs)
+					{
+						ParentsAndChildren.Append(Hierarchy->GetParents(InNodeIds[NodeIndex].Index));
+					}
+					if (bExpandOutputs)
+					{
+						ParentsAndChildren.Append(Hierarchy->GetChildren(InNodeIds[NodeIndex].Index));
+					}
 
 					for (const int32& ParentOrChildIndex : ParentsAndChildren)
 					{
@@ -796,13 +831,7 @@ void URigDependencyGraph::Tick(float DeltaTime)
 
 void URigDependencyGraph::SelectAllNodes()
 {
-	TArray<FNodeId> NodeIds;
-	NodeIds.Reserve(DependencyGraphNodes.Num());
-	for (const TObjectPtr<URigDependencyGraphNode>& Node : DependencyGraphNodes)
-	{
-		NodeIds.Add(Node->GetNodeId());
-	}
-	SelectNodes(NodeIds);
+	SelectNodes(GetAllNodeIds());
 }
 
 void URigDependencyGraph::ConstructNodes(const TArray<FNodeId>& InNodesToConstruct, bool bExpand)
@@ -886,11 +915,23 @@ void URigDependencyGraph::RemoveAllNodes()
 	
 	DependencyGraphNodes.Reset();
 	LastRecordsHash.Reset();
+	bIsPerformingGridLayout = true;
 
 	RequestRefreshLayout();
 	NotifyGraphChanged();
 }
 
+TArray<URigDependencyGraph::FNodeId> URigDependencyGraph::GetAllNodeIds() const
+{
+	TArray<FNodeId> NodeIds;
+	NodeIds.Reserve(DependencyGraphNodes.Num());
+	for (const TObjectPtr<URigDependencyGraphNode>& Node : DependencyGraphNodes)
+	{
+		NodeIds.Add(Node->GetNodeId());
+	}
+	return NodeIds;
+}
+
 void URigDependencyGraph::RemoveSelectedNodes()
 {
 	if (DependencyGraphNodes.IsEmpty() || SelectedNodes.IsEmpty())
@@ -1074,6 +1115,7 @@ FBox2D URigDependencyGraph::GetSelectedNodesBounds() const
 void URigDependencyGraph::CancelLayout()
 {
 	LayoutIterationsLeft = 0;
+	bIsPerformingGridLayout = false;
 }
 
 void URigDependencyGraph::OnHierarchyModified(ERigHierarchyNotification InNotification, URigHierarchy* InHierarchy, const FRigNotificationSubject& InSubject)
diff --git a/Engine/Plugins/Animation/ControlRig/Source/ControlRigEditor/Private/RigDependencyGraph/RigDependencyGraph.h b/Engine/Plugins/Animation/ControlRig/Source/ControlRigEditor/Private/RigDependencyGraph/RigDependencyGraph.h
--- a/Engine/Plugins/Animation/ControlRig/Source/ControlRigEditor/Private/RigDependencyGraph/RigDependencyGraph.h
+++ b/Engine/Plugins/Animation/ControlRig/Source/ControlRigEditor/Private/RigDependencyGraph/RigDependencyGraph.h
@@ -69,6 +69,7 @@ public:
 
 	/** Requests a layout refresh */
 	void RequestRefreshLayout();
+	void RequestRefreshLayout(const TArray<FNodeId>& InNodes);
 
 	/** Get whether a layout refresh was requested */
 	bool NeedsRefreshLayout() const;
@@ -87,7 +88,7 @@ public:
 	void SetFollowVMRelationShips(bool InFollowVMRelationShips);
 	void SetContentLocked(bool InLockContent);
 
-	TArray<FNodeId> GetNodesToExpand(const TArray<FNodeId>& InNodeIds) const;
+	TArray<FNodeId> GetNodesToExpand(const TArray<FNodeId>& InNodeIds, bool bExpandInputs = true, bool bExpandOutputs = true) const;
 
 	const URigVMNode* GetRigVMNodeForInstruction(const FNodeId& InNodeId) const;
 	FRigHierarchyRecord GetRecordForRigVMNode(const URigVMNode* InNode) const;
@@ -98,6 +99,8 @@ public:
 
 	void SelectAllNodes();
 	void RemoveAllNodes();
+	TArray<FNodeId> GetAllNodeIds() const;
+	const TArray<FNodeId>& GetSelectedNodeIds() const { return SelectedNodes; }
 	void RemoveSelectedNodes();
 	void RemoveUnselectedNodes();
 	void RemoveUnrelatedNodes();
@@ -170,6 +173,7 @@ private:
 	bool bFollowParentRelationShips;
 	bool bFollowVMRelationShips;
 	bool bLockContent;
+	bool bIsPerformingGridLayout;
 
 	/** As selection broadcasting is deferred in the graph, we need to block re-broadcasting a for a few frames on refresh */
 	int32 BlockSelectionCounter;
diff --git a/Engine/Plugins/Animation/ControlRig/Source/ControlRigEditor/Private/RigDependencyGraph/RigDependencyGraphCommands.cpp b/Engine/Plugins/Animation/ControlRig/Source/ControlRigEditor/Private/RigDependencyGraph/RigDependencyGraphCommands.cpp
--- a/Engine/Plugins/Animation/ControlRig/Source/ControlRigEditor/Private/RigDependencyGraph/RigDependencyGraphCommands.cpp
+++ b/Engine/Plugins/Animation/ControlRig/Source/ControlRigEditor/Private/RigDependencyGraph/RigDependencyGraphCommands.cpp
@@ -19,6 +19,7 @@ void FRigDependencyGraphCommands::RegisterCommands()
 	UI_COMMAND(RemoveUnrelated, "Hide Unrelated", "Hides nodes without any dependencies", EUserInterfaceActionType::Button, FInputChord());
 	UI_COMMAND(IsolateSelected, "Isolate Selection", "Hides all non-selected nodes", EUserInterfaceActionType::Button, FInputChord(EKeys::I));
 	UI_COMMAND(SyncWithHierarchy, "Sync Hierarchy", "Syncs the content of the graph with the hierarchy selection", EUserInterfaceActionType::Button, FInputChord(EKeys::F5));
+	UI_COMMAND(RunLayoutSimulation, "Auto Layout", "Starts a layout simulation to relax the nodes", EUserInterfaceActionType::Button, FInputChord(EKeys::R));
 	UI_COMMAND(CancelLayoutSimulation, "Cancel Layout", "Cancels the currently running layout simulation", EUserInterfaceActionType::Button, FInputChord(EKeys::Escape));
 	UI_COMMAND(FindNodesByName, "Find Nodes By Name", "Finds nodes by name", EUserInterfaceActionType::Button, FInputChord(EModifierKey::Control, EKeys::F));
 }
diff --git a/Engine/Plugins/Animation/ControlRig/Source/ControlRigEditor/Private/RigDependencyGraph/RigDependencyGraphCommands.h b/Engine/Plugins/Animation/ControlRig/Source/ControlRigEditor/Private/RigDependencyGraph/RigDependencyGraphCommands.h
--- a/Engine/Plugins/Animation/ControlRig/Source/ControlRigEditor/Private/RigDependencyGraph/RigDependencyGraphCommands.h
+++ b/Engine/Plugins/Animation/ControlRig/Source/ControlRigEditor/Private/RigDependencyGraph/RigDependencyGraphCommands.h
@@ -57,6 +57,9 @@ public:
 	/* sync the view with the hierarchy */
 	TSharedPtr<FUICommandInfo> SyncWithHierarchy;
 
+	/* starts a layout simulation to relax the nodes */
+	TSharedPtr<FUICommandInfo> RunLayoutSimulation;
+
 	/* cancel an ongoing layout simulation */
 	TSharedPtr<FUICommandInfo> CancelLayoutSimulation;
 
diff --git a/Engine/Plugins/Animation/ControlRig/Source/ControlRigEditor/Private/RigDependencyGraph/RigDependencyGraphSchema.cpp b/Engine/Plugins/Animation/ControlRig/Source/ControlRigEditor/Private/RigDependencyGraph/RigDependencyGraphSchema.cpp
--- a/Engine/Plugins/Animation/ControlRig/Source/ControlRigEditor/Private/RigDependencyGraph/RigDependencyGraphSchema.cpp
+++ b/Engine/Plugins/Animation/ControlRig/Source/ControlRigEditor/Private/RigDependencyGraph/RigDependencyGraphSchema.cpp
@@ -83,7 +83,7 @@ void URigDependencyGraphSchema::LayoutNodes(URigDependencyGraph* InGraph, int32
 	const int32 NumNodes = Nodes.Num();
 	if (NumNodes == 0)
 	{
-		InGraph->LayoutIterationsLeft = 0;
+		InGraph->CancelLayout();
 		return;
 	}
 	
@@ -127,43 +127,50 @@ void URigDependencyGraphSchema::LayoutNodes(URigDependencyGraph* InGraph, int32
 		// a grid.
 		if (NumNodesRequiringLayout == NumNodes)
 		{
-			TArray<bool> VisitedFlag;
-			int32 MaxDepth = 0;
-			for (const URigDependencyGraphNode* Node : Nodes)
+			if (InGraph->bIsPerformingGridLayout)
 			{
-				MaxDepth = FMath::Max<int32>(MaxDepth, Node->GetDependencyDepth_Impl(VisitedFlag));
-			}
+				TArray<bool> VisitedFlag;
+				int32 MaxDepth = 0;
+				for (const URigDependencyGraphNode* Node : Nodes)
+				{
+					MaxDepth = FMath::Max<int32>(MaxDepth, Node->GetDependencyDepth_Impl(VisitedFlag));
+				}
 
-			TArray<TArray<URigDependencyGraphNode*>> NodesSortedByDepth;
-			NodesSortedByDepth.SetNum(MaxDepth + 1);
+				TArray<TArray<URigDependencyGraphNode*>> NodesSortedByDepth;
+				NodesSortedByDepth.SetNum(MaxDepth + 1);
 
-			for (URigDependencyGraphNode* Node : Nodes)
-			{
-				NodesSortedByDepth[Node->GetDependencyDepth()].Add(Node);
-			}
-
-			int32 CurrentColumnX = 0;
-			for (const TArray<URigDependencyGraphNode*>& Column : NodesSortedByDepth)
-			{
-				// Lay out root nodes
-				float MaxWidth = 0.0f;
-				int32 CurrentYOffset = 0;
-				int32 TotalRootY = 0;
-				for (const URigDependencyGraphNode* Node : Column)
+				for (URigDependencyGraphNode* Node : Nodes)
 				{
-					TotalRootY += (int32)Node->GetDimensions().Y + NodeMarginY;
+					NodesSortedByDepth[Node->GetDependencyDepth()].Add(Node);
 				}
 
-				for (URigDependencyGraphNode* Node : Column)
+				int32 CurrentColumnX = 0;
+				for (const TArray<URigDependencyGraphNode*>& Column : NodesSortedByDepth)
 				{
-					Node->NodePosX = CurrentColumnX;
-					Node->NodePosY = CurrentYOffset - (TotalRootY / 2);
+					// Lay out root nodes
+					float MaxWidth = 0.0f;
+					int32 CurrentYOffset = 0;
+					int32 TotalRootY = 0;
+					for (const URigDependencyGraphNode* Node : Column)
+					{
+						TotalRootY += (int32)Node->GetDimensions().Y + NodeMarginY;
+					}
 
-					CurrentYOffset += (int32)Node->GetDimensions().Y + NodeMarginY;
-					MaxWidth = FMath::Max(MaxWidth, Node->GetDimensions().X);
-				}
+					for (URigDependencyGraphNode* Node : Column)
+					{
+						Node->NodePosX = CurrentColumnX;
+						Node->NodePosY = CurrentYOffset - (TotalRootY / 2);
+
+						CurrentYOffset += (int32)Node->GetDimensions().Y + NodeMarginY;
+						MaxWidth = FMath::Max(MaxWidth, Node->GetDimensions().X);
+					}
 
-				CurrentColumnX += (MaxWidth + NodeMarginX);
+					CurrentColumnX += (MaxWidth + NodeMarginX);
+				}
+			}
+			else
+			{
+				bRandomize = false;
 			}
 		}
 		else
@@ -447,7 +454,7 @@ void URigDependencyGraphSchema::LayoutNodes(URigDependencyGraph* InGraph, int32
 		// the simulation has converged.
 		if (NumSimulatedNodes == 0)
 		{
-			InGraph->LayoutIterationsLeft = 0;
+			InGraph->CancelLayout();
 			break;
 		}
 	}
diff --git a/Engine/Plugins/Animation/ControlRig/Source/ControlRigEditor/Private/RigDependencyGraph/SRigDependencyGraph.cpp b/Engine/Plugins/Animation/ControlRig/Source/ControlRigEditor/Private/RigDependencyGraph/SRigDependencyGraph.cpp
--- a/Engine/Plugins/Animation/ControlRig/Source/ControlRigEditor/Private/RigDependencyGraph/SRigDependencyGraph.cpp
+++ b/Engine/Plugins/Animation/ControlRig/Source/ControlRigEditor/Private/RigDependencyGraph/SRigDependencyGraph.cpp
@@ -373,7 +373,35 @@ void SRigDependencyGraph::HandleNodeDoubleClicked(UEdGraphNode* InNode)
 		return;
 	}
 
-	const TArray<FNodeId> NodesToExpand = GraphObj->GetNodesToExpand({RigDependencyGraphNode->GetNodeId()});
+	bool bExpandInputs = true;
+	bool bExpandOutputs = true;
+
+	const FGeometry PaintGeometry = GraphEditor->GetPaintSpaceGeometry();
+	FVector2D MousePosition = FSlateApplication::Get().GetCursorPos();
+	if (TSharedPtr<SWindow> Window = FSlateApplication::Get().FindWidgetWindow(GraphEditor->AsShared()))
+	{
+		MousePosition -= FVector2D(Window->GetPositionInScreen());
+	}
+
+	if(PaintGeometry.IsUnderLocation(MousePosition))
+	{
+		FVector2f GraphEditorLocation = FVector2f::ZeroVector;
+		float GraphEditorZoomAmount = 0;
+		GraphEditor->GetViewLocation(GraphEditorLocation, GraphEditorZoomAmount);
+
+		const FVector2D WidgetPosition = FVector2D(PaintGeometry.AbsoluteToLocal(MousePosition) / GraphEditorZoomAmount + GraphEditorLocation);
+		const FVector2D Position = FVector2D(RigDependencyGraphNode->GetPosition());
+		const FVector2D Dimensions = RigDependencyGraphNode->GetDimensions();
+		const double Left = Position.X;
+		const double Right = Left + Dimensions.X;
+		const double BorderLeft = Left + (Dimensions.X) * 0.3;
+		const double BorderRight = Right - (Dimensions.X) * 0.3;
+		const double Cursor = WidgetPosition.X;
+		bExpandInputs = Cursor < BorderRight;
+		bExpandOutputs = Cursor > BorderLeft;
+	}
+
+	const TArray<FNodeId> NodesToExpand = GraphObj->GetNodesToExpand({RigDependencyGraphNode->GetNodeId()}, bExpandInputs, bExpandOutputs);
 	if (!NodesToExpand.IsEmpty())
 	{
 		TArray<FNodeId> AllNodeIds;
@@ -454,6 +482,7 @@ TSharedRef<SWidget> SRigDependencyGraph::CreateFilterMenu()
 		MenuBuilder.AddMenuEntry(Actions.IsolateSelected);
 		MenuBuilder.AddMenuEntry(Actions.RemoveUnrelated);
 		MenuBuilder.AddMenuEntry(Actions.SyncWithHierarchy);
+		MenuBuilder.AddMenuEntry(Actions.RunLayoutSimulation);
 		MenuBuilder.AddMenuEntry(Actions.CancelLayoutSimulation);
 	}
 	MenuBuilder.EndSection();
@@ -503,7 +532,6 @@ void SRigDependencyGraph::BindCommands()
 
 	CommandList->MapAction(Commands.RemoveSelected, FExecuteAction::CreateLambda([this] ()
 	{
-		TArray<FNodeId> AllNodeIds;
 		GraphObj->RemoveSelectedNodes();
 	}));
 
@@ -514,7 +542,6 @@ void SRigDependencyGraph::BindCommands()
 
 	CommandList->MapAction(Commands.RemoveUnrelated, FExecuteAction::CreateLambda([this] ()
 	{
-		TArray<FNodeId> AllNodeIds;
 		GraphObj->RemoveUnrelatedNodes();
 	}));
 
@@ -539,6 +566,25 @@ void SRigDependencyGraph::BindCommands()
 		GraphObj->ConstructNodes(NodesToSelect, false);
 	}));
 
+	CommandList->MapAction(Commands.RunLayoutSimulation,
+		FExecuteAction::CreateLambda([this] ()
+		{
+			TArray<FNodeId> SelectedNodes = GraphObj->GetSelectedNodeIds();
+			if (SelectedNodes.IsEmpty())
+			{
+				SelectedNodes = GraphObj->GetAllNodeIds();
+			}
+			if (!SelectedNodes.IsEmpty())
+			{
+				GraphObj->RequestRefreshLayout(SelectedNodes);
+			}
+		}),
+		FCanExecuteAction::CreateLambda([this] () -> bool
+		{
+			return GraphObj->GetNodes().Num() > 0 && GraphObj->LayoutIterationsLeft <= 0;
+		})
+	);
+
 	CommandList->MapAction(Commands.CancelLayoutSimulation,
 		FExecuteAction::CreateLambda([this] ()
 		{
@@ -735,7 +781,7 @@ void SRigDependencyGraph::Tick(const FGeometry& AllottedGeometry, const double I
 	{
 		const FBox2D OldBounds = GraphObj->GetAllNodesBounds();
 		GraphObj->GetRigDependencyGraphSchema()->LayoutNodes(GraphObj, 32);
-		if (!GraphObj->GetNodes().IsEmpty() && GraphObj->ZoomAndFitDuringLayout.Get(false))
+		if (!GraphObj->GetNodes().IsEmpty() && GraphObj->ZoomAndFitDuringLayout.Get(false) && GraphObj->bIsPerformingGridLayout)
 		{
 			const FBox2D NewBounds = GraphObj->GetAllNodesBounds();
 			if (!NewBounds.GetCenter().Equals(OldBounds.GetCenter(), 10.f) ||
@@ -751,7 +797,12 @@ void SRigDependencyGraph::OnGraphEditorTick(const FGeometry& AllottedGeometry, c
 {
 	// flashlight to lighten up nodes in proximity of the mouse
 	const FGeometry PaintGeometry = GraphEditor->GetPaintSpaceGeometry();
-	const FVector2D MousePosition = FSlateApplication::Get().GetCursorPos();
+	FVector2D MousePosition = FSlateApplication::Get().GetCursorPos();
+	if (TSharedPtr<SWindow> Window = FSlateApplication::Get().FindWidgetWindow(GraphEditor->AsShared()))
+	{
+		MousePosition -= FVector2D(Window->GetPositionInScreen());
+	}
+
 	const float FlashLightRadius = FlashLightRadiusAttribute.Get();
 
 	if(PaintGeometry.IsUnderLocation(MousePosition) && FlashLightRadius > SMALL_NUMBER)
diff --git a/Engine/Plugins/Runtime/RigVM/Source/RigVMEditor/Private/Editor/RigVMEditor.cpp b/Engine/Plugins/Runtime/RigVM/Source/RigVMEditor/Private/Editor/RigVMEditor.cpp
--- a/Engine/Plugins/Runtime/RigVM/Source/RigVMEditor/Private/Editor/RigVMEditor.cpp
+++ b/Engine/Plugins/Runtime/RigVM/Source/RigVMEditor/Private/Editor/RigVMEditor.cpp
@@ -4321,7 +4321,12 @@ void FRigVMEditorBase::OnGraphEditorTick(const FGeometry& AllottedGeometry, cons
 {
 	// flashlight to lighten up nodes in proximity of the mouse
 	const FGeometry PaintGeometry = InGraphEditor->GetPaintSpaceGeometry();
-	const FVector2D MousePosition = FSlateApplication::Get().GetCursorPos();
+	FVector2D MousePosition = FSlateApplication::Get().GetCursorPos();
+	if (TSharedPtr<SWindow> Window = FSlateApplication::Get().FindWidgetWindow(InGraphEditor))
+	{
+		MousePosition -= FVector2D(Window->GetPositionInScreen());
+	}
+
 	const bool bUseFlashLight = GetDefault<URigVMEditorSettings>()->bUseFlashLight;
 	const float FlashLightRadius = 5.f * 64.f;
 
