Control Rig: Fix crash when icons are garbage collected
#rnx

diff --git a/Engine/Plugins/Animation/ControlRig/Source/ControlRigEditor/Private/Editor/SModularRigTreeView.cpp b/Engine/Plugins/Animation/ControlRig/Source/ControlRigEditor/Private/Editor/SModularRigTreeView.cpp
--- a/Engine/Plugins/Animation/ControlRig/Source/ControlRigEditor/Private/Editor/SModularRigTreeView.cpp
+++ b/Engine/Plugins/Animation/ControlRig/Source/ControlRigEditor/Private/Editor/SModularRigTreeView.cpp
@@ -36,6 +36,7 @@
 #define LOCTEXT_NAMESPACE "SModularRigTreeView"
 
 TMap<FSoftObjectPath, TSharedPtr<FSlateBrush>> FModularRigTreeElement::IconPathToBrush;
+TArray<TStrongObjectPtr<UTexture2D>> FModularRigTreeElement::Icons;
 
 //////////////////////////////////////////////////////////////
 /// FModularRigTreeElement
@@ -1062,6 +1063,7 @@ TPair<const FSlateBrush*, FSlateColor> FModularRigTreeElement::GetBrushAndColor(
 							{
 								const TSharedPtr<FSlateBrush> NewBrush = MakeShareable(new FSlateBrush(UWidgetBlueprintLibrary::MakeBrushFromTexture(Icon, 16.0f, 16.0f)));
 								IconPathToBrush.FindOrAdd(IconPath) = NewBrush;
+								Icons.Add(TStrongObjectPtr<UTexture2D>(Icon));
 								Brush = NewBrush.Get();
 							}
 						}
diff --git a/Engine/Plugins/Animation/ControlRig/Source/ControlRigEditor/Private/Editor/SModularRigTreeView.h b/Engine/Plugins/Animation/ControlRig/Source/ControlRigEditor/Private/Editor/SModularRigTreeView.h
--- a/Engine/Plugins/Animation/ControlRig/Source/ControlRigEditor/Private/Editor/SModularRigTreeView.h
+++ b/Engine/Plugins/Animation/ControlRig/Source/ControlRigEditor/Private/Editor/SModularRigTreeView.h
@@ -152,6 +152,7 @@ public:
 	FOnRenameRequested OnRenameRequested;
 
 	static TMap<FSoftObjectPath, TSharedPtr<FSlateBrush>> IconPathToBrush;
+	static TArray<TStrongObjectPtr<UTexture2D>> Icons;
 
 	/** The brush to use when rendering an icon */
 	const FSlateBrush* IconBrush;
