Control Rig: Add default initialize value for bFollowVariables as per bughawk feedback

#rnx

diff --git a/Engine/Plugins/Animation/ControlRig/Source/ControlRig/Public/Rigs/RigDependency.h b/Engine/Plugins/Animation/ControlRig/Source/ControlRig/Public/Rigs/RigDependency.h
--- a/Engine/Plugins/Animation/ControlRig/Source/ControlRig/Public/Rigs/RigDependency.h
+++ b/Engine/Plugins/Animation/ControlRig/Source/ControlRig/Public/Rigs/RigDependency.h
@@ -63,7 +63,7 @@ protected:
 
 	TWeakObjectPtr<UControlRig> WeakControlRig = nullptr;
 	FName EventName = NAME_None;
-	bool bFollowVariables;
+	bool bFollowVariables = false;
 	
 	mutable TRigHierarchyDependencyMap CachedDependencies;
 	mutable TRigHierarchyDependencyMap CachedReverseDependencies;
