Control Rig: Store schematic viewport visibility in an editor setting

#rnx

diff --git a/Engine/Plugins/Animation/ControlRig/Source/ControlRig/Private/Settings/ControlRigSettings.cpp b/Engine/Plugins/Animation/ControlRig/Source/ControlRig/Private/Settings/ControlRigSettings.cpp
--- a/Engine/Plugins/Animation/ControlRig/Source/ControlRig/Private/Settings/ControlRigSettings.cpp
+++ b/Engine/Plugins/Animation/ControlRig/Source/ControlRig/Private/Settings/ControlRigSettings.cpp
@@ -29,6 +29,7 @@ UControlRigEditorSettings::UControlRigEditorSettings(const FObjectInitializer& O
 	ConstructionEventBorderColor = FLinearColor::Red;
 	BackwardsSolveBorderColor = FLinearColor::Yellow;
 	BackwardsAndForwardsBorderColor = FLinearColor::Blue;
+	bShowSchematicViewInModularRig = true;
 	bShowStackedHierarchy = false;
 	MaxStackSize = 16;
 	bLeftMouseDragDoesMarquee = true;
diff --git a/Engine/Plugins/Animation/ControlRig/Source/ControlRig/Public/Settings/ControlRigSettings.h b/Engine/Plugins/Animation/ControlRig/Source/ControlRig/Public/Settings/ControlRigSettings.h
--- a/Engine/Plugins/Animation/ControlRig/Source/ControlRig/Public/Settings/ControlRigSettings.h
+++ b/Engine/Plugins/Animation/ControlRig/Source/ControlRig/Public/Settings/ControlRigSettings.h
@@ -121,6 +121,12 @@ public:
 	UPROPERTY(EditAnywhere, config, Category = Viewport)
 	FLinearColor BackwardsAndForwardsBorderColor;
 
+	/**
+	 * Show or hide the schematic view icons in the modular rig viewport
+	 */
+	UPROPERTY(EditAnywhere, config, Category = Viewport)
+	bool bShowSchematicViewInModularRig;
+
 	/**
 	 * Option to toggle displaying the stacked hierarchy items.
 	 * Note that changing this option potentially requires to re-open the editors in question. 
diff --git a/Engine/Plugins/Animation/ControlRig/Source/ControlRigEditor/Private/Editor/ControlRigEditor.cpp b/Engine/Plugins/Animation/ControlRig/Source/ControlRigEditor/Private/Editor/ControlRigEditor.cpp
--- a/Engine/Plugins/Animation/ControlRig/Source/ControlRigEditor/Private/Editor/ControlRigEditor.cpp
+++ b/Engine/Plugins/Animation/ControlRig/Source/ControlRigEditor/Private/Editor/ControlRigEditor.cpp
@@ -266,7 +266,6 @@ FControlRigBaseEditor::FControlRigBaseEditor()
 	, bIsConstructionEventRunning(false)
 	, LastHierarchyHash(INDEX_NONE)
 	, bRefreshDirectionManipulationTargetsRequired(false)
-	, bSchematicViewPortIsHidden(false)
 {
 	SchematicModel = MakeShared<FControlRigSchematicModel>();
 }
@@ -3029,7 +3028,7 @@ void FControlRigBaseEditor::HandleToggleSchematicViewport()
 	if(SchematicViewport.IsValid())
 	{
 		SchematicModel->UpdateControlRigContent();
-		bSchematicViewPortIsHidden = !bSchematicViewPortIsHidden; 
+		UControlRigEditorSettings::Get()->bShowSchematicViewInModularRig = !UControlRigEditorSettings::Get()->bShowSchematicViewInModularRig;
 	}
 }
 
@@ -3044,7 +3043,7 @@ bool FControlRigBaseEditor::IsSchematicViewportActive() const
 
 EVisibility FControlRigBaseEditor::GetSchematicOverlayVisibility() const
 {
-	if(bSchematicViewPortIsHidden)
+	if(!UControlRigEditorSettings::Get()->bShowSchematicViewInModularRig)
 	{
 		return EVisibility::Hidden;
 	}
diff --git a/Engine/Plugins/Animation/ControlRig/Source/ControlRigEditor/Private/Editor/ControlRigEditor.h b/Engine/Plugins/Animation/ControlRig/Source/ControlRigEditor/Private/Editor/ControlRigEditor.h
--- a/Engine/Plugins/Animation/ControlRig/Source/ControlRigEditor/Private/Editor/ControlRigEditor.h
+++ b/Engine/Plugins/Animation/ControlRig/Source/ControlRigEditor/Private/Editor/ControlRigEditor.h
@@ -503,7 +503,6 @@ protected:
 	bool bRefreshDirectionManipulationTargetsRequired;
 	FSimpleMulticastDelegate RequestNavigateToConnectorWarningDelegate;
 	TSharedPtr<SSchematicGraphPanel> SchematicViewport;
-	bool bSchematicViewPortIsHidden;
 	FTimerHandle RecordReplayTimerHandle;
 
 public:
