Control Rig: Follow up to 45525053

Make sure to persist settings changes.

#rnx

diff --git a/Engine/Plugins/Animation/ControlRig/Source/ControlRigEditor/Private/Editor/ControlRigEditor.cpp b/Engine/Plugins/Animation/ControlRig/Source/ControlRigEditor/Private/Editor/ControlRigEditor.cpp
--- a/Engine/Plugins/Animation/ControlRig/Source/ControlRigEditor/Private/Editor/ControlRigEditor.cpp
+++ b/Engine/Plugins/Animation/ControlRig/Source/ControlRigEditor/Private/Editor/ControlRigEditor.cpp
@@ -3029,6 +3029,7 @@ void FControlRigBaseEditor::HandleToggleSchematicViewport()
 	{
 		SchematicModel->UpdateControlRigContent();
 		UControlRigEditorSettings::Get()->bShowSchematicViewInModularRig = !UControlRigEditorSettings::Get()->bShowSchematicViewInModularRig;
+		UControlRigEditorSettings::Get()->SaveConfig();
 	}
 }
 
