Control Rig: Avoid crash in FRigModuleInstanceDetails
#rnx

diff --git a/Engine/Plugins/Animation/ControlRig/Source/ControlRigEditor/Private/ControlRigModuleDetails.cpp b/Engine/Plugins/Animation/ControlRig/Source/ControlRigEditor/Private/ControlRigModuleDetails.cpp
--- a/Engine/Plugins/Animation/ControlRig/Source/ControlRigEditor/Private/ControlRigModuleDetails.cpp
+++ b/Engine/Plugins/Animation/ControlRig/Source/ControlRigEditor/Private/ControlRigModuleDetails.cpp
@@ -380,7 +380,13 @@ void FRigModuleInstanceDetails::CustomizeDetails(IDetailLayoutBuilder& DetailBui
 				{
 					if (TSharedPtr<FRigModuleInstanceDetails> StrongThis = WeakThisPtr.Pin())
 					{
-						return StrongThis->PerModuleInfos[0].GetModularRig()->GetHierarchy();
+						if (!StrongThis->PerModuleInfos.IsEmpty())
+						{
+							if (UModularRig* ModularRig = StrongThis->PerModuleInfos[0].GetModularRig())
+							{
+								return ModularRig->GetHierarchy();
+							}
+						}
 					}
 					return nullptr;
 				});
