Anim In Engine: use bAddIfNot exist flag

diff --git a/Engine/Plugins/Animation/ControlRig/Source/ControlRigEditor/Private/Sequencer/AnimLayers/AnimLayers.cpp b/Engine/Plugins/Animation/ControlRig/Source/ControlRigEditor/Private/Sequencer/AnimLayers/AnimLayers.cpp
--- a/Engine/Plugins/Animation/ControlRig/Source/ControlRigEditor/Private/Sequencer/AnimLayers/AnimLayers.cpp
+++ b/Engine/Plugins/Animation/ControlRig/Source/ControlRigEditor/Private/Sequencer/AnimLayers/AnimLayers.cpp
@@ -1612,7 +1612,7 @@ UAnimLayers* UAnimLayers::GetAnimLayers(ULevelSequence* LevelSequence, bool bAdd
 		if (IInterface_AssetUserData* AssetUserDataInterface = Cast< IInterface_AssetUserData >(LevelSequence))
 		{
 			UAnimLayers* AnimLayers = AssetUserDataInterface->GetAssetUserData< UAnimLayers >();
-			if (!AnimLayers)
+			if (!AnimLayers && bAddIfDoesNotExist)
 			{
 				AnimLayers = NewObject<UAnimLayers>(LevelSequence, NAME_None, RF_Public | RF_Transactional);
 				AssetUserDataInterface->AddAssetUserData(AnimLayers);
