Selection Sets: Main Changes are in SeletionSets and SSelection Sets. Like AnimLayers users AssetUserData to store selection sets which are basically a set of names what corresponds to an Actor Label or a Control rig Control. Due to many interations we have some settings to show default colors and some specifications for showing actors or doing a horizontal or vertical layout.
Selection Sets are currently using tabs but we have option to switch to viewport overlay via pragma
Exposed Mirrortable functions so we can do mirroring with seletion sets
Added control rig shape filtering so we only draw certain shapes so selection sets can control that
Exposed some commonly used functions to EvaluateSequencerTools

#DisabledValidations Preflight

diff --git a/Engine/Plugins/Animation/ControlRig/Source/ControlRig/Private/Tools/ControlRigPoseMirrorTable.cpp b/Engine/Plugins/Animation/ControlRig/Source/ControlRig/Private/Tools/ControlRigPoseMirrorTable.cpp
--- a/Engine/Plugins/Animation/ControlRig/Source/ControlRig/Private/Tools/ControlRigPoseMirrorTable.cpp
+++ b/Engine/Plugins/Animation/ControlRig/Source/ControlRig/Private/Tools/ControlRigPoseMirrorTable.cpp
@@ -58,9 +58,8 @@ FControlRigToMirrorItems::FControlRigToMirrorItems(const UControlRig* InControlR
 	}
 }
 
-namespace UE::ControlRigPoseMirrorTable
-{
-FString GetMirrorString(const FString& InNameString, const TArray<FMirrorFindReplaceExpression>& MirrorFindReplaceExpressions)
+
+FString FControlRigPoseMirrorTable::GetMirrorString(const FString& InNameString, const TArray<FMirrorFindReplaceExpression>& MirrorFindReplaceExpressions)
 {
 	
 	for (const FMirrorFindReplaceExpression& RegExStr : MirrorFindReplaceExpressions)
@@ -109,7 +108,6 @@ FString GetMirrorString(const FString& InNameString, const TArray<FMirrorFindRep
 	}
 	return FString();
 }
-}
 
 TMap<const UControlRig*, FControlRigPoseMirrorTable::FMatchedControls> FControlRigPoseMirrorTable::MatchedControls;
 
@@ -131,7 +129,7 @@ void FControlRigPoseMirrorTable::SetUpMirrorTable(const UControlRig* ControlRig)
 			for (FRigControlElement* ControlElement : CurrentControls)
 			{
 				FString CurrentString = ControlElement->GetName();
-				FString NewString = UE::ControlRigPoseMirrorTable::GetMirrorString(CurrentString, AnimSettings->MirrorFindReplaceExpressions);
+				FString NewString = FControlRigPoseMirrorTable::GetMirrorString(CurrentString, AnimSettings->MirrorFindReplaceExpressions);
 				if (NewString.IsEmpty() == false && CurrentString != NewString)
 				{
 					FName CurrentName(*CurrentString);
diff --git a/Engine/Plugins/Animation/ControlRig/Source/ControlRig/Public/Tools/ControlRigPoseMirrorTable.h b/Engine/Plugins/Animation/ControlRig/Source/ControlRig/Public/Tools/ControlRigPoseMirrorTable.h
--- a/Engine/Plugins/Animation/ControlRig/Source/ControlRig/Public/Tools/ControlRigPoseMirrorTable.h
+++ b/Engine/Plugins/Animation/ControlRig/Source/ControlRig/Public/Tools/ControlRigPoseMirrorTable.h
@@ -11,10 +11,13 @@
 #include "UObject/Object.h"
 #include "Tools/MirrorCalculator.h"
 
+#define UE_API CONTROLRIG_API
+
 class UControlRig;
 struct FRigControl;
 struct FRigControlCopy;
 struct FControlRigControlPose;
+struct FMirrorFindReplaceExpression;
 
 
 struct FControlRigPoseMirrorTable
@@ -24,15 +27,17 @@ public:
 	~FControlRigPoseMirrorTable() {};
 
 	/*Get the matched control with the given name*/
-	FRigControlCopy* GetControl(const UControlRig* ControlRig, FControlRigControlPose& Pose, FName ControlrigName, bool bDoMiror);
+	UE_API FRigControlCopy* GetControl(const UControlRig* ControlRig, FControlRigControlPose& Pose, FName ControlrigName, bool bDoMiror);
 
 	/*Whether or not the Control with this name is matched*/
-	bool IsMatched(const UControlRig* ControlRig,const FName& ControlName);
+	UE_API bool IsMatched(const UControlRig* ControlRig,const FName& ControlName);
 
 	/*Return the Mirrored Global(In Control Rig Space) Translation and Mirrored Local Rotation*/
-	void GetMirrorTransform(const FRigControlCopy& ControlCopy, bool bDoLocal, bool bIsMatched,FTransform& OutGlobalTransform, FTransform& OutLocalTransform) const;
+	UE_API void GetMirrorTransform(const FRigControlCopy& ControlCopy, bool bDoLocal, bool bIsMatched,FTransform& OutGlobalTransform, FTransform& OutLocalTransform) const;
+
+	UE_API void Reset();
 
-	void Reset();
+	UE_API static FString GetMirrorString(const FString& InNameString, const TArray<FMirrorFindReplaceExpression>& MirrorFindReplaceExpressions);
 private:
 	/*Set up the Mirror Table*/
 	void SetUpMirrorTable(const UControlRig* ControlRig);
@@ -44,4 +49,6 @@ private:
 	//cache of matches, note raw pointer is fine for UControlRig, just used for lookup not accessed.
 	static TMap<const UControlRig*,FMatchedControls>  MatchedControls;
 
-};
\ No newline at end of file
+};
+
+#undef UE_API
diff --git a/Engine/Plugins/Animation/ControlRig/Source/ControlRigEditor/Private/ControlRigEditorStyle.h b/Engine/Plugins/Animation/ControlRig/Source/ControlRigEditor/Private/ControlRigEditorStyle.h
--- a/Engine/Plugins/Animation/ControlRig/Source/ControlRigEditor/Private/ControlRigEditorStyle.h
+++ b/Engine/Plugins/Animation/ControlRig/Source/ControlRigEditor/Private/ControlRigEditorStyle.h
@@ -96,6 +96,10 @@ public:
 			Set("ControlRig.AnimLayers", new IMAGE_BRUSH_SVG("Slate/AnimLayers", Icon16x16));
 			Set("ControlRig.KeyAdd", new IMAGE_BRUSH_SVG("Slate/KeyAdd_14", Icon14x14));
 			Set("ControlRig.KeySpecial", new IMAGE_BRUSH_SVG("Slate/KeySpecial_14", Icon14x14));
+			Set("ControlRig.SelectionSet", new IMAGE_BRUSH_SVG("Slate/SectionSelection", Icon16x16));
+			Set("ControlRig.FilterSelected", new IMAGE_BRUSH_SVG("Slate/CurveEditorSettingsSelectionFilter", Icon16x16));
+
+
 
 		}
 		// Control Rig Editor styles
diff --git a/Engine/Plugins/Animation/ControlRig/Source/ControlRigEditor/Private/EditMode/ControlRigEditMode.cpp b/Engine/Plugins/Animation/ControlRig/Source/ControlRigEditor/Private/EditMode/ControlRigEditMode.cpp
--- a/Engine/Plugins/Animation/ControlRig/Source/ControlRigEditor/Private/EditMode/ControlRigEditMode.cpp
+++ b/Engine/Plugins/Animation/ControlRig/Source/ControlRigEditor/Private/EditMode/ControlRigEditMode.cpp
@@ -4277,6 +4277,16 @@ void FControlRigEditMode::PostUndo()
 
 }
 
+void FControlRigEditMode::ClearOutHidden(UControlRig* InControlRig)
+{
+	ControlRigEditMode::Shapes::ClearOutHidden(InControlRig);
+}
+
+void FControlRigEditMode::ShowControlRigControls(UControlRig* InControlRig, const TSet<FString>& InNames, bool bVal)
+{
+	ControlRigEditMode::Shapes::ShowControlRigControls(InControlRig, InNames, bVal);
+}
+
 void FControlRigEditMode::RequestToRecreateControlShapeActors(UControlRig* ControlRig)
 { 
 	if (ControlRig)
@@ -4322,9 +4332,11 @@ bool FControlRigEditMode::TryUpdatingControlsShapes(UControlRig* InControlRig)
 	{
 		// not control needing shape so clear the shape actors
 		DestroyShapesActors(InControlRig);
+		ClearOutHidden(InControlRig);
 		return true;
 	}
 	
+
 	const TArray<TObjectPtr<AControlRigShapeActor>>& Shapes = *ShapeActors;
 	const int32 NumShapes = Shapes.Num();
 
@@ -5903,6 +5915,7 @@ void FControlRigEditMode::RemoveControlRig(UControlRig* InControlRig)
 		DelegateHelpers.Remove(InControlRig);
 	}
 	DestroyShapesActors(InControlRig);
+	ClearOutHidden(InControlRig);
 	if (RuntimeControlRigs.IsValidIndex(Index))
 	{
 		RuntimeControlRigs.RemoveAt(Index);
@@ -6111,7 +6124,8 @@ void FControlRigEditMode::DestroyShapesActors(UControlRig* InControlRig)
 		{
 			DestroyShapesActorsFromWorld(ShapeActors.Value);
 		}
-		
+		ControlRigEditMode::Shapes::ClearOutHidden(nullptr);
+		ClearOutHidden(nullptr);
 		ControlRigShapeActors.Reset();
 		ControlRigsToRecreate.Reset();
 		
diff --git a/Engine/Plugins/Animation/ControlRig/Source/ControlRigEditor/Private/EditMode/ControlRigEditModeCommands.cpp b/Engine/Plugins/Animation/ControlRig/Source/ControlRigEditor/Private/EditMode/ControlRigEditModeCommands.cpp
--- a/Engine/Plugins/Animation/ControlRig/Source/ControlRigEditor/Private/EditMode/ControlRigEditModeCommands.cpp
+++ b/Engine/Plugins/Animation/ControlRig/Source/ControlRigEditor/Private/EditMode/ControlRigEditModeCommands.cpp
@@ -48,6 +48,10 @@ void FControlRigEditModeCommands::RegisterCommands()
 	UI_COMMAND(ToggleControlShapeTransformEdit, "Toggle Shape Transform Edit", "Toggle Editing Selected Control's Shape Transform", EUserInterfaceActionType::ToggleButton, FInputChord(EKeys::Period, EModifierKey::Control)); 
 	UI_COMMAND(OpenSpacePickerWidget, "Open the Space Picker", "Allows space switching on the control", EUserInterfaceActionType::Button, FInputChord(EKeys::Tab)); 
 
+
+	UI_COMMAND(SummonSelectionSetsWidget, "Move selection sets to cursor", "Places the selection sets overlay at your mouse position", EUserInterfaceActionType::Button, FInputChord());
+	UI_COMMAND(ToggleSelectionSetsWidget, "Show or hide selection sets", "Shows the selection sets overlay if it's hidden, and hides it if it's visible.", EUserInterfaceActionType::Button, FInputChord());
+
 	UI_COMMAND(SummonTweenWidget, "Move tween slider to cursor", "Places the tween slider at your mouse position", EUserInterfaceActionType::Button, FInputChord(EKeys::U)); 
 	UI_COMMAND(ToggleTweenWidget, "Show or hide tween slider", "Shows the tween slider if it's hidden, and hides it if it's visible.", EUserInterfaceActionType::Button, FInputChord(EKeys::U, EModifierKey::Alt)); 
 }
diff --git a/Engine/Plugins/Animation/ControlRig/Source/ControlRigEditor/Private/EditMode/ControlRigEditModeCommands.h b/Engine/Plugins/Animation/ControlRig/Source/ControlRigEditor/Private/EditMode/ControlRigEditModeCommands.h
--- a/Engine/Plugins/Animation/ControlRig/Source/ControlRigEditor/Private/EditMode/ControlRigEditModeCommands.h
+++ b/Engine/Plugins/Animation/ControlRig/Source/ControlRigEditor/Private/EditMode/ControlRigEditModeCommands.h
@@ -110,13 +110,18 @@ public:
 	/** Opens up the space picker widget */
 	TSharedPtr< FUICommandInfo > OpenSpacePickerWidget;
 
-	
 	/** Opens the tween slider widget and positions it at the mouse. */
 	TSharedPtr< FUICommandInfo > SummonTweenWidget;
 
 	/** Toggles the visibility of the tween widget. */
 	TSharedPtr< FUICommandInfo > ToggleTweenWidget;
 
+	/** Opens the selection sets slider widget and positions it at the mouse. */
+	TSharedPtr< FUICommandInfo > SummonSelectionSetsWidget;
+
+	/** Toggles the visibility of the selection sets widget. */
+	TSharedPtr< FUICommandInfo > ToggleSelectionSetsWidget;
+
 	/**
 	 * Initialize commands
 	 */
diff --git a/Engine/Plugins/Animation/ControlRig/Source/ControlRigEditor/Private/EditMode/ControlRigEditModeSettings.h b/Engine/Plugins/Animation/ControlRig/Source/ControlRigEditor/Private/EditMode/ControlRigEditModeSettings.h
--- a/Engine/Plugins/Animation/ControlRig/Source/ControlRigEditor/Private/EditMode/ControlRigEditModeSettings.h
+++ b/Engine/Plugins/Animation/ControlRig/Source/ControlRigEditor/Private/EditMode/ControlRigEditModeSettings.h
@@ -37,6 +37,12 @@ class UControlRigEditModeSettings : public UDeveloperSettings
 			FToolWidget_DragBoxPosition(FVector2f(20.f, 30.f), HAlign_Right, VAlign_Bottom),
 			false
 		};
+
+		LastSelectionSetOverlayState = FToolWidget_FlyoutSavedState
+		{
+			FToolWidget_DragBoxPosition(FVector2f(20.f, 30.f), HAlign_Right, VAlign_Top),
+			false
+		};
 	}
 
 	// UObject interface
@@ -102,6 +108,9 @@ public:
 	UPROPERTY(config)
 	FToolWidget_FlyoutSavedState LastTweenOverlayState;
 
+	UPROPERTY(config)
+	FToolWidget_FlyoutSavedState LastSelectionSetOverlayState;
+
 	/** The tint to apply when the tween overlay is out of focus. */
 	UPROPERTY(config, EditAnywhere, Category = "Animation Settings|Tween")
 	FLinearColor TweenOutOfFocusTint = FLinearColor(1.f, 1.f, 1.f, .35f);
diff --git a/Engine/Plugins/Animation/ControlRig/Source/ControlRigEditor/Private/EditMode/ControlRigEditModeToolkit.cpp b/Engine/Plugins/Animation/ControlRig/Source/ControlRigEditor/Private/EditMode/ControlRigEditModeToolkit.cpp
--- a/Engine/Plugins/Animation/ControlRig/Source/ControlRigEditor/Private/EditMode/ControlRigEditModeToolkit.cpp
+++ b/Engine/Plugins/Animation/ControlRig/Source/ControlRigEditor/Private/EditMode/ControlRigEditModeToolkit.cpp
@@ -29,9 +29,12 @@
 #include "Models/RigSelectionViewModel.h"
 #include "Sequencer/AnimLayers/SAnimLayers.h"
 #include "Sequencer/AnimLayers/AnimLayers.h"
+#include "Sequencer/SelectionSets/SSelectionSets.h"
+#include "Sequencer/SelectionSets/SelectionSets.h"
 #include "Tools/MotionTrailOptions.h"
 #include "UObject/WeakObjectPtr.h"
 
+
 class SAnimDetailsView;
 namespace UE::ControlRigEditor
 {
@@ -52,6 +55,8 @@ bool FControlRigEditModeToolkit::bPoseTabOpen = false;
 bool FControlRigEditModeToolkit::bSnapperTabOpen = false;
 bool FControlRigEditModeToolkit::bTweenOpen = false;
 bool FControlRigEditModeToolkit::bConstraintsOpen = false;
+bool FControlRigEditModeToolkit::bSelectionSetsOpen = false;
+
 
 const FName FControlRigEditModeToolkit::PoseTabName = FName(TEXT("PoseTab"));
 const FName FControlRigEditModeToolkit::MotionTrailTabName = FName(TEXT("MotionTrailTab"));
@@ -59,6 +64,16 @@ const FName FControlRigEditModeToolkit::SnapperTabName = FName(TEXT("SnapperTab"
 const FName FControlRigEditModeToolkit::AnimLayerTabName = FName(TEXT("AnimLayerTab"));
 const FName FControlRigEditModeToolkit::TweenOverlayName = FName(TEXT("TweenOverlay"));
 const FName FControlRigEditModeToolkit::ConstrainingTabName = FName(TEXT("ConstrainingTab"));
+
+#if SELECTION_SETS_AS_TAB
+
+const FName FControlRigEditModeToolkit::SelectionSetsTabName = FName(TEXT("SelectionSets"));
+
+#else
+
+const FName FControlRigEditModeToolkit::SelectionSetsOverlayName = FName(TEXT("SelectionSets"));
+
+#endif
 const FName FControlRigEditModeToolkit::OutlinerTabName = FName(TEXT("ControlRigOutlinerTab"));
 const FName FControlRigEditModeToolkit::DetailsTabName = FName(TEXT("ControlRigDetailsTab"));
 const FName FControlRigEditModeToolkit::SpacePickerTabName = FName(TEXT("ControlRigSpacePicker"));
@@ -167,6 +182,20 @@ void FControlRigEditModeToolkit::TryInvokeToolkitUI(const FName InName)
 	{
 		ConstrainToolsManager->ToggleVisibility();
 	}
+
+#if SELECTION_SETS_AS_TAB
+	else if (InName == SelectionSetsTabName)
+	{
+		FTabId TabID(SelectionSetsTabName);
+		ModeUILayerPtr->GetTabManager()->TryInvokeTab(TabID, false /*bIsActive*/);
+	}
+#else
+	else if (InName == SelectionSetsOverlayName)
+	{
+		SelectionSetsOverlayManager->ToggleVisibility();
+	}
+#endif
+	
 }
 
 bool FControlRigEditModeToolkit::IsToolkitUIActive(const FName InName) const
@@ -175,6 +204,12 @@ bool FControlRigEditModeToolkit::IsToolkitUIActive(const FName InName) const
 	{
 		return TweenOverlayManager->IsShowingWidget();
 	}
+#if SELECTION_SETS_AS_TAB == 0
+	else if (InName == SelectionSetsOverlayName)
+	{
+		return SelectionSetsOverlayManager->IsShowingWidget();
+	}
+#endif
 
 	if (InName == ConstrainingTabName)
 	{
@@ -229,6 +264,15 @@ TSharedRef<SDockTab> SpawnMotionTrailTab(const FSpawnTabArgs& Args)
 		];
 }
 
+#if SELECTION_SETS_AS_TAB
+	TSharedRef<SDockTab> SpawnSelectionSetsTab(const FSpawnTabArgs& Args, FControlRigEditMode* InEditorMode)
+	{
+		return SNew(SDockTab)
+			[
+				SNew(SSelectionSets, *InEditorMode)
+			];
+	}
+#endif
 TSharedRef<SDockTab> SpawnAnimLayerTab(const FSpawnTabArgs& Args, FControlRigEditMode* InEditorMode)
 {
 	return SNew(SDockTab)
@@ -339,6 +383,19 @@ void FControlRigEditModeToolkit::RequestModeUITabs()
 			.SetGroup(MenuGroup)
 			.SetIcon(FSlateIcon(TEXT("ControlRigEditorStyle"), TEXT("ControlRig.AnimLayers")));
 		ModeUILayerPtr->GetTabManager()->RegisterDefaultTabWindowSize(AnimLayerTabName, FVector2D(425, 200));
+
+#if SELECTION_SETS_AS_TAB
+		ModeUILayerPtr->GetTabManager()->UnregisterTabSpawner(SelectionSetsTabName);
+		ModeUILayerPtr->GetTabManager()->RegisterTabSpawner(SelectionSetsTabName, FOnSpawnTab::CreateStatic(&SpawnSelectionSetsTab, &EditMode))
+			.SetDisplayName(LOCTEXT("SelectionSetsTab", "Sets"))
+			.SetTooltipText(LOCTEXT("SelectionSetsTabTooltip", "Selection Sets"))
+			.SetGroup(MenuGroup)
+			.SetIcon(FSlateIcon(TEXT("ControlRigEditorStyle"), TEXT("ControlRig.SelectionSet")));
+		ModeUILayerPtr->GetTabManager()->RegisterDefaultTabWindowSize(SelectionSetsTabName, FVector2D(425, 200));
+
+#endif
+
+
 	}
 };
 
@@ -364,11 +421,36 @@ void FControlRigEditModeToolkit::InvokeUI()
 		ConstrainToolsManager = MakeUnique<UE::ControlRigEditor::FConstrainToolsManager>(
 			ModeUILayerPtr->GetTabManager().ToSharedRef(), EditMode, SelectionViewModel.ToSharedRef()
 			);
-		if (bConstraintsOpen)
+	if (bConstraintsOpen)
 		{
 			ConstrainToolsManager->ShowWidget();
 		}
-		
+#if SELECTION_SETS_AS_TAB
+		if (bSelectionSetsOpen)
+		{
+			TryInvokeToolkitUI(SelectionSetsTabName);
+		}
+		/*
+		else if (TSharedPtr<ISequencer> SequencerPtr = UAIESelectionSets::GetSequencerFromAsset())
+		{
+			//wasn't open but there are layers in open level sequence we try to open it
+			if (UAIESelectionSets* SelectionSets = UAIESelectionSets::GetSelectionSets(SequencerPtr))
+			{
+				if (SelectionSets && SelectionSets->bOpenUIOnOpen)
+				{
+					TryInvokeToolkitUI(SelectionSetsTabName);
+				}
+			}
+		}
+		*/
+
+#else
+		SelectionSetsOverlayManager = MakeUnique<UE::AIE::FSelectionSetsOverlayManager>(GetToolkitHost(), GetToolkitCommands(), SharedThis(&EditMode));
+		if (bSelectionSetsOpen)
+		{
+			SelectionSetsOverlayManager->ShowWidget();
+		}
+#endif		
 		if (bAnimLayerTabOpen)
 		{
 			TryInvokeToolkitUI(AnimLayerTabName);
@@ -428,6 +510,17 @@ void FControlRigEditModeToolkit::UnregisterAndRemoveFloatingTabs()
 		}
 		
 		SelectionViewModel.Reset();
+#if SELECTION_SETS_AS_TAB == 0
+		if (SelectionSetsOverlayManager)
+		{
+			bSelectionSetsOpen = SelectionSetsOverlayManager->IsShowingWidget();
+			if (bSelectionSetsOpen)
+			{
+				SelectionSetsOverlayManager->DestroyWidget();
+			}
+			SelectionSetsOverlayManager.Reset();
+		}
+#endif
 		
 		if (ModeUILayer.IsValid())
 		{
@@ -446,7 +539,33 @@ void FControlRigEditModeToolkit::UnregisterAndRemoveFloatingTabs()
 			}
 
 			ModeUILayerPtr->GetTabManager()->UnregisterTabSpawner(MotionTrailTabName);
+#if SELECTION_SETS_AS_TAB
+			TSharedPtr<SDockTab> SelectionSetTab = ModeUILayerPtr->GetTabManager()->FindExistingLiveTab(FTabId(SelectionSetsTabName));
+			if (SelectionSetTab)
+			{
+				bSelectionSetsOpen = true;
+				SelectionSetTab->RequestCloseTab();
+			}
+			else
+			{
+				bSelectionSetsOpen = false;
+			}
+			if (TSharedPtr<ISequencer> SequencerPtr = UAIESelectionSets::GetSequencerFromAsset())
+			{
+				//wasn't open but there are layers in open level sequence we try to open it
+				if (UAIESelectionSets* SelectionSets = UAIESelectionSets::GetSelectionSets(SequencerPtr))
+				{
+					/*
+					if (SelectionSets && SelectionSets->bOpenUIOnOpen)
+					{
+						AnimLayers->bOpenUIOnOpen = bAnimLayerTabOpen;
+					}
+					*/
+				}
+			}
+			ModeUILayerPtr->GetTabManager()->UnregisterTabSpawner(SelectionSetsTabName);
 
+#endif
 			TSharedPtr<SDockTab> AnimLayerTab = ModeUILayerPtr->GetTabManager()->FindExistingLiveTab(FTabId(AnimLayerTabName));
 			if (AnimLayerTab)
 			{
diff --git a/Engine/Plugins/Animation/ControlRig/Source/ControlRigEditor/Private/EditMode/ControlRigEditModeToolkit.h b/Engine/Plugins/Animation/ControlRig/Source/ControlRigEditor/Private/EditMode/ControlRigEditModeToolkit.h
--- a/Engine/Plugins/Animation/ControlRig/Source/ControlRigEditor/Private/EditMode/ControlRigEditModeToolkit.h
+++ b/Engine/Plugins/Animation/ControlRig/Source/ControlRigEditor/Private/EditMode/ControlRigEditModeToolkit.h
@@ -11,9 +11,17 @@
 #include "Constrain/ConstrainToolsManager.h"
 #include "EditMode/SControlRigEditModeTools.h"
 #include "EditMode/ControlRigEditMode.h"
+#include "Sequencer/SelectionSets/SelectionSetsOverlayManager.h"
+
 #include "Tween/ControlRigTweenModel.h"
 #include "Tween/TweenOverlayManager.h"
 
+//when true selection sets will show up as a tab, 0 will be an overlay
+
+#ifndef SELECTION_SETS_AS_TAB
+#define SELECTION_SETS_AS_TAB 1
+#endif
+
 class SAnimDetailsView;
 class SControlRigOutliner;
 namespace UE::ControlRigEditor 
@@ -73,6 +81,11 @@ public:
 	static const FName AnimLayerTabName;
 	static const FName TweenOverlayName;
 	static const FName ConstrainingTabName;
+#if SELECTION_SETS_AS_TAB
+	static const FName SelectionSetsTabName;
+#else
+	static const FName SelectionSetsOverlayName;
+#endif
 	static const FName SnapperTabName;
 	static const FName DetailsTabName;
 	static const FName OutlinerTabName;
@@ -96,8 +109,10 @@ protected:
 	static bool bPoseTabOpen;
 	static bool bSnapperTabOpen;
 	static bool bTweenOpen;
-	static bool bConstraintsOpen;
 	
+	static bool bConstraintsOpen;
+	static bool bSelectionSetsOpen;
+
 private:
 	
 	/** The edit mode we are bound to */
@@ -110,6 +125,7 @@ private:
 	TUniquePtr<UE::ControlRigEditor::FTweenOverlayManager> TweenOverlayManager;
 	/** Manages the constraining UI tools, which is a window containing tabs for spaces, constraining, and snapping. */
 	TUniquePtr<UE::ControlRigEditor::FConstrainToolsManager> ConstrainToolsManager;
+	TUniquePtr<UE::AIE::FSelectionSetsOverlayManager> SelectionSetsOverlayManager;
 
 	/** The tools widget */
 	TSharedPtr<SControlRigEditModeTools> ModeTools;
diff --git a/Engine/Plugins/Animation/ControlRig/Source/ControlRigEditor/Private/EditMode/ControlShapeActorHelper.cpp b/Engine/Plugins/Animation/ControlRig/Source/ControlRigEditor/Private/EditMode/ControlShapeActorHelper.cpp
--- a/Engine/Plugins/Animation/ControlRig/Source/ControlRigEditor/Private/EditMode/ControlShapeActorHelper.cpp
+++ b/Engine/Plugins/Animation/ControlRig/Source/ControlRigEditor/Private/EditMode/ControlShapeActorHelper.cpp
@@ -8,9 +8,51 @@
 #include "ModularRig.h"
 #include "Rigs/RigHierarchyElements.h"
 
+
 namespace ControlRigEditMode::Shapes
 {
 
+struct FControls
+{
+	TSet<FString> ControlNames;
+};
+static TMap <TWeakObjectPtr<UControlRig>, FControls> HiddenControls;
+
+
+void ClearOutHidden(UControlRig* InControlRig)
+{
+	if (InControlRig == nullptr)
+	{
+		HiddenControls.Reset();
+	}
+	else
+	{
+		HiddenControls.Remove(InControlRig);
+	}
+}
+
+void ShowControlRigControls(UControlRig* InControlRig, const TSet<FString>& InNames, bool bVal)
+{
+	if (InControlRig)
+	{
+		FControls& Names = HiddenControls.FindOrAdd(InControlRig);
+		if (bVal)
+		{
+			for (const FString& Name : InNames)
+			{
+				Names.ControlNames.Remove(Name);
+			}
+		}
+		else //hide so we add them
+		{
+			for (const FString& Name : InNames)
+			{
+				Names.ControlNames.Add(Name);
+			}
+		}
+	}
+}
+
 bool IsSupportedControlType(const ERigControlType& ControlType)
 {
 	switch (ControlType)
@@ -230,6 +272,14 @@ void UpdateControlShape(AControlRigShapeActor* InShapeActor, FRigControlElement*
 			}
 		}
 	}
+	//hidden possibly
+	if (bIsVisible)
+	{
+		if (FControls* HiddenControlNames = HiddenControls.Find(InShapeActor->ControlRig))
+		{
+			bIsVisible = HiddenControlNames->ControlNames.Contains(InControlElement->GetName()) == false;
+		}
+	}
 
 	InShapeActor->SetIsTemporarilyHiddenInEditor(!bIsVisible || InUpdateParams.bControlsHiddenInViewport);
 
diff --git a/Engine/Plugins/Animation/ControlRig/Source/ControlRigEditor/Private/EditMode/ControlShapeActorHelper.h b/Engine/Plugins/Animation/ControlRig/Source/ControlRigEditor/Private/EditMode/ControlShapeActorHelper.h
--- a/Engine/Plugins/Animation/ControlRig/Source/ControlRigEditor/Private/EditMode/ControlShapeActorHelper.h
+++ b/Engine/Plugins/Animation/ControlRig/Source/ControlRigEditor/Private/EditMode/ControlShapeActorHelper.h
@@ -20,13 +20,13 @@ namespace ControlRigEditMode::Shapes
 
 	// Returns true if the control type supports being modified using this widget mode
 	bool IsModeSupported(const ERigControlType& InControlType, const UE::Widget::EWidgetMode& InMode);
-	
+
 	// Returns the list of controls for which a shape is expected
 	void GetControlsEligibleForShapes(const UControlRig* InControlRig, TArray<FRigControlElement*>& OutControls);
 
 	// Destroys shape actors and removes them from their UWorld
 	void DestroyShapesActorsFromWorld(const TArray<TObjectPtr<AControlRigShapeActor>>& InShapeActorsToDestroy);
-	
+
 	// Parameters used to update shape actors (transform, visibility, etc.)
 	struct FShapeUpdateParams
 	{
@@ -44,4 +44,11 @@ namespace ControlRigEditMode::Shapes
 
 	// Updates shape actors transform, visibility, etc.
 	void UpdateControlShape(AControlRigShapeActor* InShapeActor, FRigControlElement* InControlElement, const FShapeUpdateParams& InUpdateParams);
+
+	//handle per control visibility by keeping list of hidden controls. We won't create these when requested
+	//clear out the hidden ones on this control rig or all if InControlRig is null
+	void ClearOutHidden(UControlRig* InControlRig);
+	// user is responsible to make sure we recreate them at some other point after the map is set up
+	void ShowControlRigControls(UControlRig* InControlRig, const TSet<FString>& Names, bool bVal);
+
 }
diff --git a/Engine/Plugins/Animation/ControlRig/Source/ControlRigEditor/Private/EditMode/SControlRigEditModeTools.cpp b/Engine/Plugins/Animation/ControlRig/Source/ControlRigEditor/Private/EditMode/SControlRigEditModeTools.cpp
--- a/Engine/Plugins/Animation/ControlRig/Source/ControlRigEditor/Private/EditMode/SControlRigEditModeTools.cpp
+++ b/Engine/Plugins/Animation/ControlRig/Source/ControlRigEditor/Private/EditMode/SControlRigEditModeTools.cpp
@@ -403,6 +403,34 @@ void SControlRigEditModeTools::CustomizeToolBarPalette(FToolBarBuilder& ToolBarB
 		EUserInterfaceActionType::ToggleButton
 	);
 
+	// Selection Sets
+#if SELECTION_SETS_AS_TAB
+
+	ToolBarBuilder.AddToolBarButton(
+		FExecuteAction::CreateRaw(OwningToolkit.Pin().Get(), &FControlRigEditModeToolkit::TryInvokeToolkitUI, FControlRigEditModeToolkit::SelectionSetsTabName),
+		NAME_None,
+		LOCTEXT("SelectionSets", "Selection Sets"),
+		LOCTEXT("SelectionSetsTooltip", "Open Selection Sets"),
+		FSlateIcon(TEXT("ControlRigEditorStyle"), TEXT("ControlRig.SelectionSet")),
+		EUserInterfaceActionType::ToggleButton
+	);
+#else
+
+	ToolBarBuilder.AddToolBarButton(
+		FUIAction(
+			FExecuteAction::CreateRaw(OwningToolkit.Pin().Get(), &FControlRigEditModeToolkit::TryInvokeToolkitUI, FControlRigEditModeToolkit::SelectionSetsOverlayName),
+			FCanExecuteAction(),
+			FIsActionChecked::CreateRaw(OwningToolkit.Pin().Get(), &FControlRigEditModeToolkit::IsToolkitUIActive, FControlRigEditModeToolkit::SelectionSetsOverlayName)
+		),
+		NAME_None,
+		LOCTEXT("SelectionSets", "Selection Sets"),
+		LOCTEXT("SelectionSetsTooltip", "Open Selection Sets"),
+		FSlateIcon(TEXT("ControlRigEditorStyle"), TEXT("ControlRig.SelectionSet")),
+		EUserInterfaceActionType::ToggleButton
+	);
+#endif
+	
+
 	/*  leaving this around for a bit to get animator feedback on the new menu system
 	// Motion Trail
 	ToolBarBuilder.AddToolBarButton(
diff --git a/Engine/Plugins/Animation/ControlRig/Source/ControlRigEditor/Private/Misc/Flyout/FlyoutOverlayManager.cpp b/Engine/Plugins/Animation/ControlRig/Source/ControlRigEditor/Private/Misc/Flyout/FlyoutOverlayManager.cpp
--- a/Engine/Plugins/Animation/ControlRig/Source/ControlRigEditor/Private/Misc/Flyout/FlyoutOverlayManager.cpp
+++ b/Engine/Plugins/Animation/ControlRig/Source/ControlRigEditor/Private/Misc/Flyout/FlyoutOverlayManager.cpp
@@ -210,6 +210,24 @@ void FFlyoutOverlayManager::HideWidget()
 	}
 }
 
+void FFlyoutOverlayManager::DestroyWidget()
+{
+	if (DragWidget)
+	{
+		if (bShouldShowWidget)
+		{
+			HideWidget(); //hide before destroying
+		}
+		Args.WidgetHost->RemoveWidgetFromHost();
+		ensureAlwaysMsgf(DragWidget.IsUnique(), TEXT("Widget not properly removed from hierarchy"));
+
+		// Defensive: if WidgetHost did not properly remove from the hierarchy, the widget life will be extended by our DragWidget referencing it.
+		// Clear it to give it a chance to die.
+		ContentWidget->ClearChildContent();
+		DragWidget.Reset();
+		ContentWidget.Reset();
+	}
+}
 TSharedPtr<FFlyoutTemporaryPositionOverride> FFlyoutOverlayManager::TryTemporarilyPositionWidgetAtCursor(ETemporaryFlyoutPositionFlags InFlags)
 {
 	// No double starting the operation
diff --git a/Engine/Plugins/Animation/ControlRig/Source/ControlRigEditor/Private/Misc/Flyout/FlyoutOverlayManager.h b/Engine/Plugins/Animation/ControlRig/Source/ControlRigEditor/Private/Misc/Flyout/FlyoutOverlayManager.h
--- a/Engine/Plugins/Animation/ControlRig/Source/ControlRigEditor/Private/Misc/Flyout/FlyoutOverlayManager.h
+++ b/Engine/Plugins/Animation/ControlRig/Source/ControlRigEditor/Private/Misc/Flyout/FlyoutOverlayManager.h
@@ -62,7 +62,7 @@ public:
 	void ToggleVisibility() { if (IsShowingWidget()) { HideWidget(); } else { ShowWidget(); } }
 	void ShowWidget();
 	void HideWidget();
-
+	void DestroyWidget();
 	/** @return Whether the widget is temporarily repositioned. */
 	bool IsTemporarilyRepositioned() const { return TemporaryMoveOperation.IsSet(); }
 	/**
diff --git a/Engine/Plugins/Animation/ControlRig/Source/ControlRigEditor/Private/Sequencer/SelectionSets/SSelectionSets.cpp b/Engine/Plugins/Animation/ControlRig/Source/ControlRigEditor/Private/Sequencer/SelectionSets/SSelectionSets.cpp
new file mode 100644
--- /dev/null
+++ b/Engine/Plugins/Animation/ControlRig/Source/ControlRigEditor/Private/Sequencer/SelectionSets/SSelectionSets.cpp
@@ -0,0 +1,1577 @@
+// Copyright Epic Games, Inc. All Rights Reserved.
+#include "SSelectionSets.h"
+#include "SelectionSets.h"
+#include "SelectionSetsSettings.h"
+
+#include "ScopedTransaction.h"
+#include "SPositiveActionButton.h"
+#include "Widgets/SWidget.h"
+#include "EditMode/ControlRigEditMode.h"
+#include "Toolkits/IToolkitHost.h"
+#include "DesktopPlatformModule.h"
+#include "Misc/ConfigCacheIni.h"
+#include "Misc/Paths.h"
+#include "EditorDirectories.h"
+#include "EditorModeManager.h"
+#include "ControlRig.h"
+#include "LevelSequencePlayer.h"
+#include "LevelSequenceEditorBlueprintLibrary.h"
+#include "MovieScene.h"
+#include "Editor.h"
+#include "EditorModeManager.h"
+#include "ILevelSequenceEditorToolkit.h"
+#include "ISequencer.h"
+#include "SceneOutlinerPublicTypes.h"
+#include "ControlRigEditorStyle.h"
+#include "LevelEditor.h"
+#include "ToolMenus.h"
+
+#include "Modules/ModuleManager.h"
+#include "MVVM/Selection/SequencerSelectionEventSuppressor.h"
+#include "MVVM/Selection/SequencerCoreSelection.h"
+#include "MVVM/ViewModels/SequencerEditorViewModel.h"
+#include "MVVM/Selection/Selection.h"
+#include "Styling/SlateIconFinder.h"
+
+#include "Widgets/Colors/SColorBlock.h"
+#include "Widgets/Colors/SColorPicker.h"
+#include "Widgets/Input/SSearchBox.h"
+#include "Widgets/Layout/SGridPanel.h"
+#include "Widgets/Layout/SWrapBox.h"
+#include "Widgets/Layout/SScrollBox.h"
+#include "Widgets/Layout/SSpacer.h"
+#include "Widgets/Text/SInlineEditableTextBlock.h"
+#include "Misc/TextFilterExpressionEvaluator.h"
+
+
+#define LOCTEXT_NAMESPACE "SelectionSets"
+
+///////////////////////////////////////////////////////////
+/// SSelectionSetButton
+////////////////////////////////////////////////SS///////////
+
+void SSelectionSetButton::Construct(const FArguments& InArgs)
+{
+	SelectionSetGuid = InArgs._Guid;
+	SelectionSetsWidget = InArgs._SetsWidget;
+	SButton::FArguments ButtonArgs;
+	ButtonArgs.ButtonStyle(FAppStyle::Get(), "NoBorder");
+	ButtonArgs._OnClicked = InArgs._OnClicked;
+	ButtonArgs._Content = InArgs._Content;
+	ButtonArgs.ContentPadding(1.f);
+
+	SButton::Construct(ButtonArgs);
+}
+
+void SSelectionSetButton::AddSelectionToSet() const
+{
+	if (TSharedPtr<ISequencer> SequencerPtr = UAIESelectionSets::GetSequencerFromAsset())
+	{
+		if (UAIESelectionSets* SelectionSets = UAIESelectionSets::GetSelectionSets(SequencerPtr))
+		{
+			SelectionSets->AddSelectionToSetItem(SelectionSetGuid, SequencerPtr);
+		}
+	}
+}
+
+void SSelectionSetButton::RemoveSelectionFromSet() const
+{
+	if (TSharedPtr<ISequencer> SequencerPtr = UAIESelectionSets::GetSequencerFromAsset())
+	{
+		if (UAIESelectionSets* SelectionSets = UAIESelectionSets::GetSelectionSets(SequencerPtr))
+		{
+			SelectionSets->RemoveSelectionFromSetItem(SelectionSetGuid, SequencerPtr);
+		}
+	}
+}
+
+void SSelectionSetButton::DeleteSet() const
+{
+	if (TSharedPtr<ISequencer> SequencerPtr = UAIESelectionSets::GetSequencerFromAsset())
+	{
+		if (UAIESelectionSets* SelectionSets = UAIESelectionSets::GetSelectionSets(SequencerPtr))
+		{
+			SelectionSets->DeleteSetItem(SelectionSetGuid,SequencerPtr);
+		}
+	}
+}
+
+void SSelectionSetButton::RenameSet() const
+{
+	if (SelectionSetsWidget)
+	{
+		SelectionSetsWidget->SetFocusOnTextBlock(SelectionSetGuid);
+	}
+}
+
+TSharedRef<SWidget> SSelectionSetButton::CreateSelectionSetColorWidget(const FGuid& InGuid, TSharedPtr<ISequencer>& SequencerPtr) const
+{
+	// clang-format off
+	return
+
+		SNew(SBox)
+		.HAlign(HAlign_Right)
+		[
+			SNew(SBox)
+				.Padding(FMargin(0.0f, 0.0f, 0.0f, 0.0f))
+				.WidthOverride(100.0f)
+				[
+					SNew(SBorder)
+						.BorderImage(FAppStyle::Get().GetBrush("Menu.WidgetBorder"))
+						.Padding(FMargin(1.0f))
+						[
+							SNew(SColorBlock).IsEnabled(true)
+								//.Size(FVector2D(6.0, 38.0))
+								.Color_Lambda([InGuid, WeakSequencer = SequencerPtr.ToWeakPtr()]()
+									{
+										FLinearColor Color(FLinearColor::White);
+										if (TSharedPtr<ISequencer> SequencerPtr = WeakSequencer.Pin())
+										{
+											if (UAIESelectionSets* SelectionSets = UAIESelectionSets::GetSelectionSets(SequencerPtr))
+											{
+												FLinearColor CurrentColor;
+												if (SelectionSets->GetItemColor(InGuid, CurrentColor))
+												{
+													Color = CurrentColor;
+												}
+											}
+										}
+										return Color;
+									}
+								)
+
+								.OnMouseButtonDown_Lambda([InGuid, WeakSequencer = SequencerPtr.ToWeakPtr()](const FGeometry&, const FPointerEvent&)
+									{
+										if (TSharedPtr<ISequencer> SequencerPtr = WeakSequencer.Pin())
+										{
+											if (UAIESelectionSets* SelectionSets = UAIESelectionSets::GetSelectionSets(SequencerPtr))
+											{
+												FLinearColor Color;
+												if (SelectionSets->GetItemColor(InGuid, Color))
+												{
+
+													FColorPickerArgs PickerArgs;
+													PickerArgs.bUseAlpha = false;
+													PickerArgs.InitialColor = Color;
+													PickerArgs.OnColorCommitted = FOnLinearColorValueChanged::CreateLambda([InGuid, WeakSequencer = SequencerPtr.ToWeakPtr()](FLinearColor Color)
+													{
+														if (TSharedPtr<ISequencer> SequencerPtr = WeakSequencer.Pin())
+														{
+															if (UAIESelectionSets* SelectionSets = UAIESelectionSets::GetSelectionSets(SequencerPtr))
+															{
+																SelectionSets->SetItemColor(InGuid, Color);
+															}
+														}
+													});
+													OpenColorPicker(PickerArgs);
+												}
+											}
+										}
+										return FReply::Handled();
+									})
+						]
+				]
+		];
+
+	// clang-format on
+}
+
+void SSelectionSetButton::CreateMirror() const
+{
+	if (TSharedPtr<ISequencer> SequencerPtr = UAIESelectionSets::GetSequencerFromAsset())
+	{
+		if (UAIESelectionSets* SelectionSets = UAIESelectionSets::GetSelectionSets(SequencerPtr))
+		{
+			FGuid NewGuid = SelectionSets->CreateMirror(SequencerPtr, SelectionSetGuid);
+			SelectionSetsWidget->SetFocusOnTextBlock(NewGuid);
+
+		}
+	}
+}
+
+void SSelectionSetButton::SelectMirror() const
+{
+	if (TSharedPtr<ISequencer> SequencerPtr = UAIESelectionSets::GetSequencerFromAsset())
+	{
+		if (UAIESelectionSets* SelectionSets = UAIESelectionSets::GetSelectionSets(SequencerPtr))
+		{
+			const bool bAdd = FSlateApplication::Get().GetModifierKeys().IsShiftDown();
+			const bool bToggle = FSlateApplication::Get().GetModifierKeys().IsControlDown();
+
+			SelectionSets->SelectItem(SelectionSetGuid, true, bAdd, bToggle );
+		}
+	}
+}
+
+void SSelectionSetButton::HideSelectionSet() const
+{
+	if (TSharedPtr<ISequencer> SequencerPtr = UAIESelectionSets::GetSequencerFromAsset())
+	{
+		if (UAIESelectionSets* SelectionSets = UAIESelectionSets::GetSelectionSets(SequencerPtr))
+		{
+			SelectionSets->ShowOrHideControls(SelectionSetGuid, false /*bShow*/, false /*bDoMirror*/);
+
+		}
+	}
+}
+
+void SSelectionSetButton::ShowSelectionSet() const
+{
+	if (TSharedPtr<ISequencer> SequencerPtr = UAIESelectionSets::GetSequencerFromAsset())
+	{
+		if (UAIESelectionSets* SelectionSets = UAIESelectionSets::GetSelectionSets(SequencerPtr))
+		{
+			SelectionSets->ShowOrHideControls(SelectionSetGuid, true /*bShow*/, false /*bDoMirror*/);
+		}
+	}
+}
+
+void SSelectionSetButton::IsolateSelectionSet() const
+{
+	if (TSharedPtr<ISequencer> SequencerPtr = UAIESelectionSets::GetSequencerFromAsset())
+	{
+		if (UAIESelectionSets* SelectionSets = UAIESelectionSets::GetSelectionSets(SequencerPtr))
+		{
+			SelectionSets->IsolateControls(SelectionSetGuid);
+		}
+	}
+}
+
+void SSelectionSetButton::ShowAllSelectionSet() const
+{
+	if (TSharedPtr<ISequencer> SequencerPtr = UAIESelectionSets::GetSequencerFromAsset())
+	{
+		if (UAIESelectionSets* SelectionSets = UAIESelectionSets::GetSelectionSets(SequencerPtr))
+		{
+			SelectionSets->ShowAllControls(SelectionSetGuid);
+		}
+	}
+}
+
+void SSelectionSetButton::KeyAll() const
+{
+	if (TSharedPtr<ISequencer> SequencerPtr = UAIESelectionSets::GetSequencerFromAsset())
+	{
+		if (UAIESelectionSets* SelectionSets = UAIESelectionSets::GetSelectionSets(SequencerPtr))
+		{
+			SelectionSets->KeyAll(SequencerPtr, SelectionSetGuid);
+		}
+	}
+}
+
+FReply SSelectionSetButton::OnMouseButtonDown(const FGeometry& InMyGeometry, const FPointerEvent& InMouseEvent)
+{
+
+	if (InMouseEvent.GetEffectingButton() == EKeys::RightMouseButton)
+	{
+		if (TSharedPtr<ISequencer> SequencerPtr = UAIESelectionSets::GetSequencerFromAsset())
+		{
+			FVector2f SummonLocation = InMouseEvent.GetScreenSpacePosition();
+
+			const bool bShouldCloseWindowAfterMenuSelection = true;
+			FMenuBuilder MenuBuilder(bShouldCloseWindowAfterMenuSelection, nullptr);
+
+			MenuBuilder.BeginSection("SSelectionSetButton", LOCTEXT("SelectionSet", "Selection Sets"));
+			{
+				FUIAction Action = FUIAction(FExecuteAction::CreateRaw((this), &SSelectionSetButton::AddSelectionToSet));
+				const FText Label = LOCTEXT("AddSelectionToSet", "Add Selection");
+				const FText ToolTipText = LOCTEXT("AddSelectionToSettooltip", "Add current selection to this selection set");
+				MenuBuilder.AddMenuEntry(Label, ToolTipText, FSlateIcon(), Action);
+			}
+			{
+				FUIAction Action = FUIAction(FExecuteAction::CreateRaw((this), &SSelectionSetButton::RemoveSelectionFromSet));
+				const FText Label = LOCTEXT("RemoveSelectionFromSet", "Remove Selection");
+				const FText ToolTipText = LOCTEXT("RemoveSelectionFromSettooltip", "Remove current selection from this selection set");
+				MenuBuilder.AddMenuEntry(Label, ToolTipText, FSlateIcon(), Action);
+			}
+			MenuBuilder.AddMenuSeparator();
+			{
+				FUIAction Action = FUIAction(FExecuteAction::CreateRaw((this), &SSelectionSetButton::DeleteSet));
+				const FText Label = LOCTEXT("DeleteSet", "Delete");
+				const FText ToolTipText = LOCTEXT("DeleteSettooltip", "Delete selection set");
+				MenuBuilder.AddMenuEntry(Label, ToolTipText, FSlateIcon(), Action);
+			}
+			{
+				FUIAction Action = FUIAction(FExecuteAction::CreateRaw((this), &SSelectionSetButton::RenameSet));
+				const FText Label = LOCTEXT("RenameSet", "Rename");
+				const FText ToolTipText = LOCTEXT("RenameSettooltip", "Rename selection set");
+				MenuBuilder.AddMenuEntry(Label, ToolTipText, FSlateIcon(), Action);
+			}
+			{
+				const FText Label = LOCTEXT("SetColor", "Set Color");
+				const FText ToolTipText = LOCTEXT("SetColortooltip", "Set Selection Set Color");
+				const TSharedRef<SWidget> ColorWidget = CreateSelectionSetColorWidget(SelectionSetGuid, SequencerPtr);
+				MenuBuilder.AddWidget(ColorWidget, Label, true, true, ToolTipText);
+			}
+			MenuBuilder.AddMenuSeparator();
+			{
+				FUIAction Action = FUIAction(FExecuteAction::CreateRaw((this), &SSelectionSetButton::CreateMirror));
+				const FText Label = LOCTEXT("CreateMirror", "Create Mirror");
+				const FText ToolTipText = LOCTEXT("CreateMirrortooltip", "Create selection set from mirrored");
+				MenuBuilder.AddMenuEntry(Label, ToolTipText, FSlateIcon(), Action);
+			}
+			{
+				FUIAction Action = FUIAction(FExecuteAction::CreateRaw((this), &SSelectionSetButton::SelectMirror));
+				const FText Label = LOCTEXT("SelectMirror", "Select Mirror");
+				const FText ToolTipText = LOCTEXT("SelectMirrortooltip", "Select mirror of selection set");
+				MenuBuilder.AddMenuEntry(Label, ToolTipText, FSlateIcon(), Action);
+			}
+			MenuBuilder.AddMenuSeparator();
+			{
+				FUIAction Action = FUIAction(FExecuteAction::CreateRaw((this), &SSelectionSetButton::HideSelectionSet));
+				const FText Label = LOCTEXT("HideSelectionSet", "Hide Selection Set");
+				const FText ToolTipText = LOCTEXT("HideSelectionSettooltip", "Hide the Controls and Actors in the Selection Set on selected objects");
+				MenuBuilder.AddMenuEntry(Label, ToolTipText, FSlateIcon(), Action);
+			}
+			{
+				FUIAction Action = FUIAction(FExecuteAction::CreateRaw((this), &SSelectionSetButton::ShowSelectionSet));
+				const FText Label = LOCTEXT("ShowSelectionSet", "Show Selection Set");
+				const FText ToolTipText = LOCTEXT("ShowSelectionSettooltip", "Show the Controls and the Actors in the Selection Set on selected objects");
+				MenuBuilder.AddMenuEntry(Label, ToolTipText, FSlateIcon(), Action);
+			}
+			{
+				FUIAction Action = FUIAction(FExecuteAction::CreateRaw((this), &SSelectionSetButton::IsolateSelectionSet));
+				const FText Label = LOCTEXT("IsolateSelectionSet", "Isolate Selection Set");
+				const FText ToolTipText = LOCTEXT("IsolateSelectionSettooltip", "Just show the Controls and the Actors in the Selection Set on selected objects");
+				MenuBuilder.AddMenuEntry(Label, ToolTipText, FSlateIcon(), Action);
+			}
+			{
+				FUIAction Action = FUIAction(FExecuteAction::CreateRaw((this), &SSelectionSetButton::ShowAllSelectionSet));
+				const FText Label = LOCTEXT("ShowAllSelection", "Show All Selection Set");
+				const FText ToolTipText = LOCTEXT("ShowAllselectionSettootip", "Controls and the Actors in the Selection Set on selected objects");
+				MenuBuilder.AddMenuEntry(Label, ToolTipText, FSlateIcon(), Action);
+			}
+			MenuBuilder.AddMenuSeparator();
+			{
+				FUIAction Action = FUIAction(FExecuteAction::CreateRaw((this), &SSelectionSetButton::KeyAll));
+				const FText Label = LOCTEXT("KeyAllSet", "Key Selection Set");
+				const FText ToolTipText = LOCTEXT("KeyAllSettooltip", "Key everything in the selection set at this time that's active");
+				MenuBuilder.AddMenuEntry(Label, ToolTipText, FSlateIcon(), Action);
+			}
+			MenuBuilder.AddMenuSeparator();
+			TSharedPtr<SWidget> MenuContent = MenuBuilder.MakeWidget();
+
+			if (MenuContent.IsValid())
+			{
+				FWidgetPath WidgetPath = InMouseEvent.GetEventPath() != nullptr ? *InMouseEvent.GetEventPath() : FWidgetPath();
+				FSlateApplication::Get().PushMenu(AsShared(), WidgetPath, MenuContent.ToSharedRef(), SummonLocation, FPopupTransitionEffect(FPopupTransitionEffect::ContextMenu));
+			}
+			return FReply::Handled();
+		}
+	}
+
+	// Call the base class implementation for other mouse buttons
+	return SButton::OnMouseButtonDown(InMyGeometry, InMouseEvent);
+}
+
+FReply SSelectionSetButton::OnMouseButtonDoubleClick(const FGeometry& InMyGeometry, const FPointerEvent& InMouseEvent)
+{
+	return SButton::OnMouseButtonDoubleClick(InMyGeometry, InMouseEvent);
+}
+
+//////////////////////////////////////////////////////////////
+/// SSelectionSets
+///////////////////////////////////////////////////////////
+
+void SSelectionSets::Construct(const FArguments& InArgs, FControlRigEditMode& InEditMode)
+{
+	SetTextFilter = MakeShareable(new FTextFilterExpressionEvaluator(ETextFilterExpressionEvaluatorMode::Complex));
+
+	LastMovieSceneSig = FGuid();
+	FocusOnTextBlock = FGuid();
+	bSetUpTimerIsSet = false;
+	TSharedPtr<ISequencer> SequencerPtr = UAIESelectionSets::GetSequencerFromAsset();
+	if (SequencerPtr.IsValid())
+	{
+		SequencerPtr->OnActivateSequence().AddRaw(this, &SSelectionSets::OnActivateSequence);
+		SequencerPtr->OnMovieSceneDataChanged().AddRaw(this, &SSelectionSets::OnMovieSceneDataChanged);
+		SequencerPtr->OnCloseEvent().AddRaw(this, &SSelectionSets::OnCloseEvent);
+		SequencerPtr->OnMovieSceneBindingsChanged().AddRaw(this, &SSelectionSets::OnMovieSceneBindingsChanged);
+		if (UAIESelectionSets* SelectionSets = UAIESelectionSets::GetSelectionSets(SequencerPtr))
+		{
+			NewSelectionSet(SelectionSets);
+		}
+	}
+	else
+	{
+		return;
+	}
+	const FButtonStyle& ButtonStyle = FAppStyle::Get().GetWidgetStyle<FButtonStyle>("Menu.Button");
+	const float MenuIconSize = FAppStyle::Get().GetFloat("Menu.MenuIconSize");
+
+	PostUndoRedoHandle = FEditorDelegates::PostUndoRedo.AddLambda(
+		[this]()
+		{
+			bPopulateUI = true;
+		}
+	);
+
+	SAssignNew(MainAddFilterSetSlot,SBorder)
+			.BorderImage(FAppStyle::GetBrush("ToolPanel.GroupBorder"))
+				.Padding(0.0f)
+				[
+					SNew(SHorizontalBox)
+						+ SHorizontalBox::Slot()
+						//.FillContentWidth(1.0)
+						.AutoWidth()
+						.HAlign(EHorizontalAlignment::HAlign_Left)
+						.Padding(5.0f)
+						[
+							SNew(SActionButton)
+								.OnClicked(this, &SSelectionSets::OnAddClicked)
+								.Icon(FAppStyle::Get().GetBrush("Icons.Plus"))
+								.Text(LOCTEXT("SelectionSet", "Add Set"))
+								.ToolTipText(LOCTEXT("AddSelectionSetTooltip", "Add a new Selection Set"))
+						]
+					+ SHorizontalBox::Slot()
+						.FillWidth(1.f)
+						[
+							SNew(SSpacer)
+						]
+						+ SHorizontalBox::Slot().AutoWidth()
+						[
+							SNew(SComboButton)
+								.ComboButtonStyle(FAppStyle::Get(), "SimpleComboButton")
+								.OnGetMenuContent(this, &SSelectionSets::MakeImportExportMenu)
+								.ButtonContent()
+								[
+									SNew(SImage)
+										.Image(FAppStyle::Get().GetBrush("Icons.Reimport"))
+								]
+						]
+					+ SHorizontalBox::Slot()
+						//.FillContentWidth(1.0)
+						.AutoWidth()
+						.HAlign(EHorizontalAlignment::HAlign_Right)
+						.Padding(4.0f)
+						.VAlign(VAlign_Center)
+						[
+
+							SNew(SCheckBox)
+								.Padding(4)
+								.Style(FAppStyle::Get(), "FilterBar.FilterButton")
+								.ToolTipText(LOCTEXT("OnlyShowSelected", "Only Show Sets on the Selected or last Selected object"))
+								.IsChecked_Lambda([WeakSequencer = SequencerPtr->AsWeak()]() -> ECheckBoxState
+									{
+										bool bEnabled = false;
+
+										if (TSharedPtr<ISequencer> SequencerPtr = WeakSequencer.Pin())
+										{
+											if (UAIESelectionSets* SelectionSets = UAIESelectionSets::GetSelectionSets(SequencerPtr))
+											{
+												bEnabled = SelectionSets->bShowSelectedOnly;
+											}
+										}
+										return bEnabled ? ECheckBoxState::Checked : ECheckBoxState::Unchecked;
+									})
+								.OnCheckStateChanged_Lambda([this, WeakSequencer = SequencerPtr->AsWeak()](ECheckBoxState InCheckBoxState)
+									{
+										const bool bIsChecked = InCheckBoxState == ECheckBoxState::Checked;
+										if (TSharedPtr<ISequencer> SequencerPtr = WeakSequencer.Pin())
+										{
+											if (UAIESelectionSets* SelectionSets = UAIESelectionSets::GetSelectionSets(SequencerPtr))
+											{
+												SelectionSets->bShowSelectedOnly = bIsChecked;
+											}
+										}
+										bPopulateUI = true;
+									})
+								.CheckBoxContentUsesAutoWidth(true)
+								[
+									SNew(SImage)
+										.Image_Lambda([]() -> const FSlateBrush*
+											{
+												return FSlateIcon(FControlRigEditorStyle::Get().GetStyleSetName(), "ControlRig.FilterSelected").GetIcon();
+											})
+										.ColorAndOpacity_Lambda([WeakSequencer = SequencerPtr->AsWeak()]()
+											{
+												FSlateColor SlateColor = FSlateColor::UseForeground();
+
+												if (TSharedPtr<ISequencer> SequencerPtr = WeakSequencer.Pin())
+												{
+													if (UAIESelectionSets* SelectionSets = UAIESelectionSets::GetSelectionSets(SequencerPtr))
+													{
+														bool bValue = SelectionSets->bShowSelectedOnly;
+														if (bValue == true)
+														{
+															SlateColor = FAppStyle::GetSlateColor("SelectionColor");
+
+														}
+													}
+												}
+												return SlateColor;
+											})
+								]
+						]
+
+				];
+	SAssignNew(ActiveActorsSlot, SBox)
+	.Padding(0.0f)
+	.WidthOverride(200.f)
+	.Visibility_Lambda([this]()
+		{
+			return (ActiveActorsWidget.IsValid() && ActiveActorsWidget->GetChildren() && ActiveActorsWidget->GetChildren()->Num() > 0)
+				? EVisibility::Visible : EVisibility::Collapsed;
+		})
+	  [
+		SNew(SBorder)
+			.Padding(2.0)
+			.BorderImage(FAppStyle::GetBrush("Brushes.Panel"))
+			[
+				// Panel background, seen between items
+				SNew(SBorder)
+					.BorderImage(FAppStyle::GetBrush("WhiteBrush"))
+					.BorderBackgroundColor(FStyleColors::Panel)
+					.HAlign(HAlign_Fill)
+					.VAlign(VAlign_Fill)
+					.Padding(0.0f)
+					[
+						SNew(SBox)
+							.MaxDesiredHeight(100)
+							[
+								SNew(SScrollBox)
+									+ SScrollBox::Slot()
+									[
+										SAssignNew(ActiveActorsWidget, SGridPanel)
+									]
+							]
+					]
+			]
+	];
+	SAssignNew(SearchSetsSlot, SBorder)
+		.BorderImage(FAppStyle::Get().GetBrush("Menu.Background"))
+		.Padding(0.0f)
+		[
+			SNew(SHorizontalBox)
+				+ SHorizontalBox::Slot()
+				.Padding(4.0f, 0.f, 0.f, 0.0f)
+				.MaxWidth(20)
+				.FillContentWidth(1.0)
+				.HAlign(HAlign_Left)
+				.VAlign(VAlign_Center)
+				[
+					SNew(SBox)
+						.WidthOverride(MenuIconSize)
+						.HeightOverride(MenuIconSize)
+						[
+							SNew(SImage)
+								.Image_Lambda([]() -> const FSlateBrush*
+									{
+										return FSlateIcon(FControlRigEditorStyle::Get().GetStyleSetName(), "ControlRig.SelectionSet").GetIcon();
+									})
+								.ColorAndOpacity(FSlateColor::UseForeground())
+						]
+				]
+
+			+ SHorizontalBox::Slot()
+				.FillContentWidth(1.0)
+				.MaxWidth(50.)
+				.Padding(4.0f, 0.f, 0.f, 0.0f)
+				.VAlign(VAlign_Center)
+				.HAlign(HAlign_Left)
+				[
+					SNew(STextBlock)
+						.ColorAndOpacity(FSlateColor::UseForeground())
+						.Text_Lambda([InArgs]()-> FText
+							{
+								const FText Set = LOCTEXT("Set", "Set");
+								return Set;
+							})
+				]
+			+ SHorizontalBox::Slot()
+				.AutoWidth()
+				.HAlign(HAlign_Left)
+				.VAlign(VAlign_Center)
+				.Padding(4.0f, 0.f, 0.f, 0.0f)
+				[
+					SNew(SBox)
+						.MinDesiredWidth(150)
+						[
+							SAssignNew(SetSearchBox, SSearchBox)
+								.OnTextChanged_Lambda([this](const FText& InText)
+									{
+										SetTextFilter->SetFilterText(InText);
+										PopulateSelectionSets();
+									})
+						]
+				]
+			+ SHorizontalBox::Slot()
+				.FillWidth(1.f)
+				[
+					SNew(SSpacer)
+				]
+		];
+
+	SAssignNew(ActiveSetsSlot, SBox)
+		.Padding(0.0f)
+		[
+			SNew(SBorder)
+				.Padding(2.0)
+				.BorderImage(FAppStyle::GetBrush("Brushes.Panel"))
+				[
+					// Panel background, seen between items
+					SNew(SBorder)
+						.BorderImage(FAppStyle::GetBrush("WhiteBrush"))
+						.BorderBackgroundColor(FStyleColors::Panel)
+						.HAlign(HAlign_Fill)
+						.VAlign(VAlign_Fill)
+						.Padding(0.0f)
+						[
+
+							SAssignNew(SelectionSetsWidget, SWrapBox)
+								.UseAllottedSize(true) // This is the key property
+								.InnerSlotPadding(FVector2D(10.0f, 10.0f))
+								.Orientation(Orient_Horizontal) // Default, but good to be explicit
+
+						]
+				]
+		];
+
+	bool bIsVertical = true;
+	if (const USelectionSetsSettings* Settings = GetDefault<USelectionSetsSettings>())
+	{
+		bIsVertical = Settings->GetSelectionSetsLayout() == ESelectionSetsLayout::Horizontal ? false : true;
+	}
+	if (bIsVertical)
+	{
+		ChildSlot
+		[
+			SNew(SBox)
+				[
+					SNew(SBorder)
+						.Content()
+						[
+							SNew(SVerticalBox)
+								+ SVerticalBox::Slot()
+								.Padding(FMargin(0.0f, 0.f, 0.f, 0.f))
+								.AutoHeight()
+								[
+									MainAddFilterSetSlot.ToSharedRef()
+								]
+								+ SVerticalBox::Slot()
+								.Padding(FMargin(0.0f, 0.f, 0.f, 0.f))
+								.AutoHeight()
+								[
+									ActiveActorsSlot.ToSharedRef()
+								]
+								+ SVerticalBox::Slot()
+								.Padding(FMargin(0.0f, 0.f, 0.f, 0.f))
+								.AutoHeight()
+								[
+									SearchSetsSlot.ToSharedRef()
+								]
+								+ SVerticalBox::Slot()
+								.Padding(FMargin(0.0f, 0.f, 0.f, 0.f))
+								.AutoHeight()
+								[
+									ActiveSetsSlot.ToSharedRef()
+
+								]
+						]
+				]
+		];
+	}
+	else
+	{
+		ChildSlot
+		[
+			SNew(SBox)
+				[
+					SNew(SBorder)
+						.Content()
+						[
+							SNew(SVerticalBox)
+								+ SVerticalBox::Slot()
+								.Padding(FMargin(0.0f, 0.f, 0.f, 0.f))
+								.AutoHeight()
+								[
+									SNew(SHorizontalBox)
+									+ SHorizontalBox::Slot()
+									.Padding(FMargin(0.0f, 0.f, 0.f, 0.f))
+									.HAlign(EHorizontalAlignment::HAlign_Left)
+									.AutoWidth()
+									[
+										MainAddFilterSetSlot.ToSharedRef()
+									]
+									+ SHorizontalBox::Slot()
+									.FillWidth(1.0)
+									.HAlign(EHorizontalAlignment::HAlign_Center)
+									[
+										SNew(SSpacer)
+									]
+									+ SHorizontalBox::Slot()
+									.Padding(FMargin(0.0f, 0.f, 0.f, 0.f))
+									.HAlign(EHorizontalAlignment::HAlign_Right)
+									.AutoWidth()
+									[
+										SearchSetsSlot.ToSharedRef()
+									]
+								]
+								+ SVerticalBox::Slot()
+								.Padding(FMargin(0.0f, 0.f, 0.f, 0.f))
+								.AutoHeight()
+								[
+									SNew(SHorizontalBox)
+									+ SHorizontalBox::Slot()
+									.Padding(FMargin(0.0f, 0.f, 0.f, 0.f))
+									.AutoWidth()
+									[
+										ActiveActorsSlot.ToSharedRef()
+									]
+									+ SHorizontalBox::Slot()
+									.Padding(FMargin(0.0f, 0.f, 0.f, 0.f))
+									.AutoWidth()
+									[
+										ActiveSetsSlot.ToSharedRef()
+									]
+								]
+						]
+				]
+		];
+	}
+	
+	SetEditMode(InEditMode);
+	RegisterSelectionChanged();
+	SetCanTick(true);
+	bUpdateBindings = true;
+	SetCursor(EMouseCursor::Default);
+
+}
+
+TSharedRef<SWidget> SSelectionSets::MakeImportExportMenu()
+{
+	constexpr bool bCloseMenuAfterSelection = true;
+	FMenuBuilder MenuBuilder(bCloseMenuAfterSelection, nullptr);
+
+
+	MenuBuilder.BeginSection("ImportExport", LOCTEXT("ImportExport", "Import/Export"));
+
+	{
+		FUIAction Action = FUIAction(
+			FExecuteAction::CreateLambda([this]()
+				{
+					ImportLastSelectionSet();
+				}));
+
+		const FText Label = LOCTEXT("FastImport", "Import Last");
+		const FText ToolTipText = LOCTEXT("FastImporttooltip", "Import Last Selection Set that was opened");
+		FSlateIcon Icon(FAppStyle::GetAppStyleSetName(),"Icons.Reimport");
+		MenuBuilder.AddMenuEntry(Label, ToolTipText, Icon, Action);
+	}
+
+	{
+		FUIAction Action = FUIAction(
+			FExecuteAction::CreateLambda([this]()
+				{
+					ImportFromJson();
+				}));
+
+		const FText Label = LOCTEXT("ImportSets", "Import Sets");
+		const FText ToolTipText = LOCTEXT("ImportTooltip", "Import Selection Sets from JSON File");
+		FSlateIcon Icon(FAppStyle::GetAppStyleSetName(), "Icons.Import");
+		MenuBuilder.AddMenuEntry(Label, ToolTipText, Icon, Action);
+	}
+	MenuBuilder.AddSeparator();
+	{
+		FUIAction Action = FUIAction(
+			FExecuteAction::CreateLambda([this]()
+				{
+					ExportToJson();
+				}));
+
+		const FText Label = LOCTEXT("ExportSets", "Export Sets");
+		const FText ToolTipText = LOCTEXT("Exporttooltip", "Export Selection Sets as JSON File ");
+		FSlateIcon Icon(FAppStyle::GetAppStyleSetName(), "Icons.Export");
+		MenuBuilder.AddMenuEntry(Label, ToolTipText, Icon, Action);
+	}
+
+	return MenuBuilder.MakeWidget();
+}
+
+TSharedRef<SWidget> SSelectionSets::CreateActorMenu(FActorWithSelectionSet& ActorWithSelectionSet)
+{
+	constexpr bool bCloseMenuAfterSelection = true;
+	FMenuBuilder MenuBuilder(bCloseMenuAfterSelection, nullptr);
+
+	MenuBuilder.BeginSection("Visibility", LOCTEXT("Visibility", "Visibility"));
+
+	FUIAction Action = FUIAction(
+		FExecuteAction::CreateLambda([&ActorWithSelectionSet]()
+			{
+				ActorWithSelectionSet.ShowAll();
+			}));
+
+	const FText Label = LOCTEXT("ShowAll", "Show All");
+	const FText ToolTipText = LOCTEXT("ShowAlltooltip", "Show all controls or the actor");
+	MenuBuilder.AddMenuEntry(Label, ToolTipText, FSlateIcon(), Action);
+
+	MenuBuilder.EndSection();
+
+	return MenuBuilder.MakeWidget();
+}
+
+void SSelectionSets::PopulateUI()
+{
+	PopulateActiveActors();
+	PopulateSelectionSets();
+}
+
+void SSelectionSets::PopulateActiveActors()
+{
+	if (TSharedPtr<ISequencer> SequencerPtr = UAIESelectionSets::GetSequencerFromAsset())
+	{
+		const float MenuIconSize = FAppStyle::Get().GetFloat("Menu.MenuIconSize");
+
+		auto AddActor = [this, WeakSequencer = SequencerPtr->AsWeak(), MenuIconSize](SGridPanel& Grid, int32 Row, int32 Column, TWeakObjectPtr<AActor>& ActorPtr, FActorWithSelectionSet& ActorWithSeletionSet)
+			{
+				Grid.AddSlot(Column, Row)
+					.Padding(4.0, 4.0, 2.0, 2.0)
+					[
+						SNew(SCheckBox)
+							.OnGetMenuContent_Lambda([WeakSequencer, ActorPtr,this]()
+								{
+									if (TSharedPtr<ISequencer> SequencerPtr = WeakSequencer.Pin())
+									{
+										if (UAIESelectionSets* SelectionSets = UAIESelectionSets::GetSelectionSets(SequencerPtr))
+										{
+											if (FActorWithSelectionSet* ActorWithSeletionSet = SelectionSets->ActorsWithSelectionSets.ActorsWithSet.Find(ActorPtr))
+											{
+												return CreateActorMenu(*ActorWithSeletionSet);
+											}
+										}
+									}
+									return SNullWidget::NullWidget;
+								})
+							.Style(FAppStyle::Get(), "FilterBar.FilterButton")
+							//	.ToolTipText(FilterToolTip)
+							.IsChecked_Lambda( [WeakSequencer, ActorPtr]() -> ECheckBoxState
+								{
+									if (TSharedPtr<ISequencer> SequencerPtr = WeakSequencer.Pin())
+									{
+										if (UAIESelectionSets* SelectionSets = UAIESelectionSets::GetSelectionSets(SequencerPtr))
+										{
+											if (FActorWithSelectionSet* ActorWithSeletionSet = SelectionSets->ActorsWithSelectionSets.ActorsWithSet.Find(ActorPtr))
+											{
+												bool bEnabled = ActorWithSeletionSet->GetSelected();
+												return bEnabled ? ECheckBoxState::Checked : ECheckBoxState::Unchecked;
+											}
+										}
+									}
+									return ECheckBoxState::Unchecked;
+								})
+							.OnCheckStateChanged_Lambda([this, ActorPtr, WeakSequencer](ECheckBoxState InCheckBoxState)
+								{
+									if (TSharedPtr<ISequencer> SequencerPtr = WeakSequencer.Pin())
+									{
+										if (UAIESelectionSets* SelectionSets = UAIESelectionSets::GetSelectionSets(SequencerPtr))
+										{
+											const bool bAdd = FSlateApplication::Get().GetModifierKeys().IsShiftDown();
+											const bool bIsChecked = InCheckBoxState == ECheckBoxState::Checked;
+											if (bAdd == false && bIsChecked)
+											{
+												for (TPair<TWeakObjectPtr<AActor>, FActorWithSelectionSet>& Pair : SelectionSets->ActorsWithSelectionSets.ActorsWithSet)
+												{
+													Pair.Value.SetSelected(false);
+												}
+											}
+											if (FActorWithSelectionSet* ActorWithSeletionSet = SelectionSets->ActorsWithSelectionSets.ActorsWithSet.Find(ActorPtr))
+											{
+												ActorWithSeletionSet->SetSelected(bIsChecked);
+											}
+										}
+									}
+											
+
+									PopulateSelectionSets();
+								})
+							.CheckBoxContentUsesAutoWidth(false)
+							//.OnGetMenuContent(this, &SFilter::GetRightClickMenuContent)
+							[
+								SNew(SBorder)
+									.Padding(1.0f)
+									.BorderImage(FAppStyle::Get().GetBrush("FilterBar.FilterBackground"))
+									[
+										SNew(SHorizontalBox)
+											+ SHorizontalBox::Slot()
+											.VAlign(VAlign_Center)
+											.AutoWidth()
+											[
+												SNew(SImage)
+													.Image(FAppStyle::Get().GetBrush("FilterBar.FilterImage"))
+													.ColorAndOpacity_Lambda([ActorPtr, WeakSequencer]()
+														{
+															if (TSharedPtr<ISequencer> SequencerPtr = WeakSequencer.Pin())
+															{
+																if (UAIESelectionSets* SelectionSets = UAIESelectionSets::GetSelectionSets(SequencerPtr))
+																{
+																	if (FActorWithSelectionSet* ActorWithSeletionSet = SelectionSets->ActorsWithSelectionSets.ActorsWithSet.Find(ActorPtr))
+																	{
+																		return ActorWithSeletionSet->GetSelected() ? FSlateColor(ActorWithSeletionSet->Color) : FAppStyle::Get().GetSlateColor("Colors.Recessed");
+																	}
+																}
+															}
+															return FSlateColor(FLinearColor::White);
+														})
+											]
+										+ SHorizontalBox::Slot()
+											.Padding(4.0f, 0.f, 0.f, 0.0f)
+											.AutoWidth()
+											.HAlign(HAlign_Center)
+											.VAlign(VAlign_Center)
+											[
+												SNew(SBox)
+													.WidthOverride(MenuIconSize)
+													.HeightOverride(MenuIconSize)
+													[
+														SNew(SImage)
+															.Image_Lambda([ActorPtr, WeakSequencer]() -> const FSlateBrush*
+																{
+																	if (TSharedPtr<ISequencer> SequencerPtr = WeakSequencer.Pin())
+																	{
+																		if (UAIESelectionSets* SelectionSets = UAIESelectionSets::GetSelectionSets(SequencerPtr))
+																		{
+																			if (FActorWithSelectionSet* ActorWithSeletionSet = SelectionSets->ActorsWithSelectionSets.ActorsWithSet.Find(ActorPtr))
+																			{
+																				for (const TWeakObjectPtr<UObject>& Object : ActorWithSeletionSet->Actors)
+																				{
+																					if (UControlRig* ControlRig = Cast<UControlRig>(Object.Get()))
+																					{
+																						return FSlateIconFinder::FindIconForClass(USkeleton::StaticClass()).GetIcon();
+																					}
+																				}
+																			}
+																		}
+																	}
+																	return FSlateIconFinder::FindIconForClass(AActor::StaticClass()).GetIcon();
+																})
+															.ColorAndOpacity(FSlateColor::UseForeground())
+													]
+											]
+										+ SHorizontalBox::Slot()
+											//.Padding(TAttribute<FMargin>(this, &SFilter::GetFilterNamePadding))
+											.Padding(FMargin(4, 1, 4, 1))
+											.VAlign(VAlign_Center)
+											[
+												SNew(STextBlock)
+													.Text_Lambda([&ActorWithSeletionSet]()
+														{
+															return ActorWithSeletionSet.Name;
+														})
+													.IsEnabled_Lambda([&ActorWithSeletionSet] {return ActorWithSeletionSet.GetSelected(); })
+
+											]
+									]
+							]
+					];
+			};
+
+		if (ActiveActorsWidget.IsValid())
+		{
+			ActiveActorsWidget->ClearChildren();
+
+			if (UAIESelectionSets* SelectionSets = UAIESelectionSets::GetSelectionSets(SequencerPtr))
+			{
+				int32 Row = 0;
+				int32 Column = 0;
+				constexpr int32 MaxColumn = 1;
+				for (TPair<TWeakObjectPtr<AActor>, FActorWithSelectionSet>& Pair : SelectionSets->ActorsWithSelectionSets.ActorsWithSet)
+				{
+					if (SelectionSets->IsVisible(Pair.Value) == false)
+					{
+						continue;
+					}
+					AddActor(*ActiveActorsWidget, Row, Column, Pair.Key, Pair.Value);
+					if (++Column > MaxColumn)
+					{
+						Column = 0;
+						++Row;
+					}
+				}
+			}
+		}
+	}
+}
+
+void SSelectionSets::PopulateSelectionSets()
+{
+	auto FilterString = [this](const FString& StringToTest) ->bool
+		{
+			return SetTextFilter->TestTextFilter(FBasicStringFilterExpressionContext(StringToTest));
+		};
+	
+	SelectionSetsButton.Reset();
+	auto AddRow = [this](SWrapBox& Grid, TSharedPtr<ISequencer>& SequencerPtr, int32 Row, FAIESelectionSetItem* Item)
+	{
+		if (Item == nullptr)
+		{
+			return;
+		}
+		
+		
+		FGuid ItemGuid = Item->Guid;
+		//Item.ViewData.
+		Grid.AddSlot()
+			.HAlign(HAlign_Center)
+			.VAlign(VAlign_Center)
+			[
+				SNew(SBorder)
+					.BorderImage(FAppStyle::GetBrush("Graph.StateNode.Body"))
+					.BorderBackgroundColor_Lambda([this,ItemGuid, WeakSequencer = SequencerPtr->AsWeak()]()
+						{
+							if (TSharedPtr<ISequencer> SequencerPtr = WeakSequencer.Pin())
+							{
+								if (UAIESelectionSets* SelectionSets = UAIESelectionSets::GetSelectionSets(SequencerPtr))
+								{
+									if (FAIESelectionSetItem* SetItem = SelectionSets->SelectionSets.Find(ItemGuid))
+									{
+										const TSharedPtr<SSelectionSetButton>& Button = SelectionSetsButton.FindRef(ItemGuid);
+										float Alpha = Button.IsValid() && Button->IsHovered() ? 1.0 : 0.5;
+										FLinearColor ItemColor = SetItem->ViewData.Color;
+										return  FSlateColor(FLinearColor(ItemColor.R, ItemColor.G, ItemColor.B, Alpha));
+									}
+								}
+							}
+							return FSlateColor(FLinearColor::White);
+							})
+					.HAlign(HAlign_Fill)
+					.VAlign(VAlign_Center)
+				[
+					SAssignNew(SelectionSetsButton.FindOrAdd(ItemGuid), SSelectionSetButton)
+						.Guid(ItemGuid)
+						.SetsWidget(this)
+						.OnClicked_Lambda([this, WeakSequencer = SequencerPtr->AsWeak(), ItemGuid]()
+							{
+								if (TSharedPtr<ISequencer> SequencerPtr = WeakSequencer.Pin())
+								{
+									if (UAIESelectionSets* SelectionSets = UAIESelectionSets::GetSelectionSets(SequencerPtr))
+									{
+										const bool bToggle = FSlateApplication::Get().GetModifierKeys().IsControlDown();
+										const bool bAdd = FSlateApplication::Get().GetModifierKeys().IsShiftDown();
+
+										SelectionSets->SelectItem(ItemGuid,false,bAdd, bToggle);
+									}
+								}	
+								return FReply::Handled();
+							})
+						[
+							
+							SAssignNew(SelectionSetsTextBlock.FindOrAdd(ItemGuid), SInlineEditableTextBlock)
+								.Justification(ETextJustify::Center)
+								.Text_Lambda([ItemGuid, WeakSequencer = SequencerPtr->AsWeak()]()
+								{
+										if (TSharedPtr<ISequencer> SequencerPtr = WeakSequencer.Pin())
+										{
+											if (UAIESelectionSets* SelectionSets = UAIESelectionSets::GetSelectionSets(SequencerPtr))
+											{
+												if (FAIESelectionSetItem* SetItem = SelectionSets->SelectionSets.Find(ItemGuid))
+												{
+													return SetItem->ItemName;
+												}
+											}
+										}
+										return FText();
+								})
+								.OnTextCommitted_Lambda([WeakSequencer = SequencerPtr->AsWeak(),ItemGuid](const FText & InNewText, ETextCommit::Type InTextCommit)
+								{
+										if (TSharedPtr<ISequencer> SequencerPtr = UAIESelectionSets::GetSequencerFromAsset())
+										{
+											if (UAIESelectionSets* SelectionSets = UAIESelectionSets::GetSelectionSets(SequencerPtr))
+											{
+												SelectionSets->RenameSetItem(ItemGuid, InNewText);
+											}
+										}
+								})
+								
+						]
+				]
+			];
+
+	};
+	if (SelectionSetsWidget.IsValid())
+	{
+		if (SelectionSetsWidget->GetChildren() && SelectionSetsWidget->GetChildren()->Num())
+		{
+			SelectionSetsWidget->ClearChildren();
+		}
+		if (TSharedPtr<ISequencer> SequencerPtr = UAIESelectionSets::GetSequencerFromAsset())
+		{
+			if (UAIESelectionSets* SelectionSets = UAIESelectionSets::GetSelectionSets(SequencerPtr, false))
+			{
+				int32 Row = 0;
+				//int32 Column = 0;
+				//constexpr int32 MaxColumn = 5;
+				TArray<FGuid> Guids = SelectionSets->GetActiveSelectionSets();
+				for (const FGuid& Guid : Guids)
+				{
+					if (FAIESelectionSetItem* SetItem = SelectionSets->SelectionSets.Find(Guid))
+					{
+						if (SetTextFilter->GetFilterText().IsEmpty() || FilterString(SetItem->ItemName.ToString()))
+						{
+							AddRow(*SelectionSetsWidget, SequencerPtr, Row, SetItem);
+
+							/*
+							AddRow(*SelectionSetsWidget, Row, Column, SetItem);
+							if (++Column > MaxColumn)
+							{
+								Column = 0;
+								++Row;
+							}
+							*/
+						}
+					}
+				}
+			}
+		}
+	}
+}
+
+void SSelectionSets::Tick(const FGeometry& AllottedGeometry, const double InCurrentTime, const float InDeltaTime)
+{
+	if (ISequencer* Sequencer = GetSequencer())
+	{
+		FGuid CurrentMovieSceneSig = Sequencer->GetFocusedMovieSceneSequence()->GetMovieScene()->GetSignature();
+		if (LastMovieSceneSig != CurrentMovieSceneSig)
+		{
+			LastMovieSceneSig = CurrentMovieSceneSig;
+		}
+	}
+	if (bUpdateBindings == true)
+	{
+		bUpdateBindings = false;
+		bPopulateUI = false;
+		UpdateBindings();
+	}
+	if (bPopulateUI)
+	{
+		bPopulateUI = false;
+		PopulateUI();
+	}
+	if (FocusOnTextBlock.IsValid())
+	{
+		if (TSharedPtr<SInlineEditableTextBlock>* TextBlock = SelectionSetsTextBlock.Find(FocusOnTextBlock))
+		{
+			if (TextBlock->IsValid())
+			{
+				(*TextBlock)->EnterEditingMode();
+			}
+			FocusOnTextBlock = FGuid();
+		}
+	}
+}
+
+FReply SSelectionSets::OnSelectionFilterClicked()
+{
+	return FReply::Handled();
+}
+
+bool SSelectionSets::IsSelectionFilterActive() const
+{
+	return false;
+}
+
+SSelectionSets::SSelectionSets()
+{
+}
+
+SSelectionSets::~SSelectionSets()
+{
+	FEditorDelegates::PostUndoRedo.Remove(PostUndoRedoHandle);
+	if (TSharedPtr<ISequencer> SequencerPtr = UAIESelectionSets::GetSequencerFromAsset())
+	{
+		SequencerPtr->OnActivateSequence().RemoveAll(this);
+		SequencerPtr->OnMovieSceneDataChanged().RemoveAll(this);
+		SequencerPtr->OnCloseEvent().RemoveAll(this);
+		SequencerPtr->OnMovieSceneBindingsChanged().RemoveAll(this);
+	}
+
+	// unregister previous one
+	if (OnSelectionChangedHandle.IsValid())
+	{
+		FLevelEditorModule& LevelEditor = FModuleManager::GetModuleChecked<FLevelEditorModule>("LevelEditor");
+		FLevelEditorModule::FActorSelectionChangedEvent& ActorSelectionChangedEvent = LevelEditor.OnActorSelectionChanged();
+		ActorSelectionChangedEvent.Remove(OnSelectionChangedHandle);
+		OnSelectionChangedHandle.Reset();
+	}
+
+	for (TWeakObjectPtr<UControlRig>& ControlRig : BoundControlRigs)
+	{
+		if (ControlRig.IsValid())
+		{
+			ControlRig.Get()->ControlRigBound().RemoveAll(this);
+			const TSharedPtr<IControlRigObjectBinding> Binding = ControlRig.Get()->GetObjectBinding();
+			if (Binding)
+			{
+				Binding->OnControlRigBind().RemoveAll(this);
+			}
+		}
+	}
+	BoundControlRigs.SetNum(0);
+}
+
+
+void SSelectionSets::HandleControlSelected(UControlRig* Subject, FRigControlElement* ControlElement, bool bSelected)
+{
+	FControlRigBaseDockableView::HandleControlSelected(Subject, ControlElement, bSelected);
+	bPopulateUI = true;
+}
+
+void SSelectionSets::RegisterSelectionChanged()
+{
+	FLevelEditorModule& LevelEditor = FModuleManager::GetModuleChecked<FLevelEditorModule>("LevelEditor");
+	FLevelEditorModule::FActorSelectionChangedEvent& ActorSelectionChangedEvent = LevelEditor.OnActorSelectionChanged();
+
+	// unregister previous one
+	if (OnSelectionChangedHandle.IsValid())
+	{
+		ActorSelectionChangedEvent.Remove(OnSelectionChangedHandle);
+		OnSelectionChangedHandle.Reset();
+	}
+
+	// register
+	OnSelectionChangedHandle = ActorSelectionChangedEvent.AddRaw(this, &SSelectionSets::OnActorSelectionChanged);
+}
+
+void SSelectionSets::OnActorSelectionChanged(const TArray<UObject*>& NewSelection, bool bForceRefresh)
+{
+	bPopulateUI = true;
+}
+
+FString  SSelectionSets::LastSelectionSetJSON;
+
+void  SSelectionSets::OnCloseEvent(TSharedRef<ISequencer> InSequencer)
+{
+	if (ULevelSequence* LevelSequence = Cast<ULevelSequence>(InSequencer->GetFocusedMovieSceneSequence()))
+	{
+		if (UAIESelectionSets* SelectionSets = UAIESelectionSets::GetSelectionSets(LevelSequence, false /*Add if doesn't exist*/))
+		{
+			if (SelectionSets->ExportAsJsonString(LastSelectionSetJSON) == false)
+			{
+				UE_LOG(LogControlRig, Error, TEXT("Selection Sets: Error fast JSON export"));
+			}
+		}
+	}
+}
+
+void SSelectionSets::UpdateBindings()
+{
+	if (TSharedPtr<ISequencer> SequencerPtr = UAIESelectionSets::GetSequencerFromAsset())
+	{
+		if (UAIESelectionSets* SelectionSets = UAIESelectionSets::GetSelectionSets(SequencerPtr, false /*Add if doesn't exist*/))
+		{
+			SelectionSets->SelectionListChanged().RemoveAll(this);
+			SelectionSets->SelectionListChanged().AddRaw(this, &SSelectionSets::HandleSelectionListChanged);
+
+			SelectionSets->SequencerBindingsAdded(SequencerPtr);
+			PopulateUI();
+		}
+	}
+
+}
+
+void SSelectionSets::OnMovieSceneBindingsChanged()
+{
+	bUpdateBindings = true; //do it next tick
+}
+
+void SSelectionSets::OnActivateSequence(FMovieSceneSequenceIDRef ID)
+{
+	bUpdateBindings = true; //do it next tick;
+}
+
+void SSelectionSets::OnMovieSceneDataChanged(EMovieSceneDataChangeType Type)
+{
+	/** Handle to help updating selection on tick tick to avoid too many flooded selections*/
+	static bool bTimerIsSet = false;
+	if (TSharedPtr<ISequencer> SequencerPtr = UAIESelectionSets::GetSequencerFromAsset())
+	{
+		if (UAIESelectionSets* SelectionSets = UAIESelectionSets::GetSelectionSets(SequencerPtr, false /*Add if doesn't exist*/))
+		{
+			switch (Type)
+			{
+
+				case EMovieSceneDataChangeType::MovieSceneStructureItemAdded:
+				case EMovieSceneDataChangeType::MovieSceneStructureItemRemoved:
+				case EMovieSceneDataChangeType::MovieSceneStructureItemsChanged:
+					if (bSetUpTimerIsSet == false)
+					{
+						bSetUpTimerIsSet = true;
+						GEditor->GetTimerManager()->SetTimerForNextTick([this, WeakSequencer = SequencerPtr->AsWeak(), WeakThis = this->AsWeak()]()
+							{
+								bSetUpTimerIsSet = false;
+								if (!WeakThis.IsValid())
+								{
+									return;
+								}
+								if (TSharedPtr<ISequencer> SequencerPtr = WeakSequencer.Pin())
+								{
+									if (UAIESelectionSets* SelectionSets = UAIESelectionSets::GetSelectionSets(SequencerPtr, false /*Add if doesn't exist*/))
+									{
+										SelectionSets->SequencerBindingsAdded(SequencerPtr);
+										PopulateActiveActors();
+									}
+								}
+							});
+					}
+					break;
+				default:
+					break;
+
+			}
+		}
+	}
+}
+
+void SSelectionSets::SetEditMode(FControlRigEditMode& InEditMode)
+{
+	FControlRigBaseDockableView::SetEditMode(InEditMode);
+	ModeTools = InEditMode.GetModeManager();
+	if (FControlRigEditMode* EditMode = static_cast<FControlRigEditMode*>(ModeTools->GetActiveMode(FControlRigEditMode::ModeName)))
+	{
+		TArrayView<TWeakObjectPtr<UControlRig>> ControlRigs = EditMode->GetControlRigs();
+		for (TWeakObjectPtr<UControlRig>& ControlRig : ControlRigs)
+		{
+			if (ControlRig.IsValid())
+			{
+				if (!ControlRig.Get()->ControlRigBound().IsBoundToObject(this))
+				{
+					ControlRig.Get()->ControlRigBound().AddRaw(this, &SSelectionSets::HandleOnControlRigBound);
+					BoundControlRigs.Add(ControlRig);
+				}
+				const TSharedPtr<IControlRigObjectBinding> Binding = ControlRig.Get()->GetObjectBinding();
+				if (Binding && !Binding->OnControlRigBind().IsBoundToObject(this))
+				{
+					Binding->OnControlRigBind().AddRaw(this, &SSelectionSets::HandleOnObjectBoundToControlRig);
+				}
+			}
+		}
+	}
+}
+
+void SSelectionSets::HandleControlAdded(UControlRig* ControlRig, bool bIsAdded)
+{
+	FControlRigBaseDockableView::HandleControlAdded(ControlRig, bIsAdded);
+	if (ControlRig)
+	{
+		if (bIsAdded == true)
+		{
+			if (!ControlRig->ControlRigBound().IsBoundToObject(this))
+			{
+				ControlRig->ControlRigBound().AddRaw(this, &SSelectionSets::HandleOnControlRigBound);
+				BoundControlRigs.Add(ControlRig);
+			}
+			const TSharedPtr<IControlRigObjectBinding> Binding = ControlRig->GetObjectBinding();
+			if (Binding && !Binding->OnControlRigBind().IsBoundToObject(this))
+			{
+				Binding->OnControlRigBind().AddRaw(this, &SSelectionSets::HandleOnObjectBoundToControlRig);
+			}
+		}
+		else
+		{
+			BoundControlRigs.Remove(ControlRig);
+			ControlRig->ControlRigBound().RemoveAll(this);
+			const TSharedPtr<IControlRigObjectBinding> Binding = ControlRig->GetObjectBinding();
+			if (Binding)
+			{
+				Binding->OnControlRigBind().RemoveAll(this);
+			}
+		}
+	}
+}
+
+void SSelectionSets::HandleOnControlRigBound(UControlRig* InControlRig)
+{
+	if (!InControlRig)
+	{
+		return;
+	}
+
+	const TSharedPtr<IControlRigObjectBinding> Binding = InControlRig->GetObjectBinding();
+
+	if (Binding && !Binding->OnControlRigBind().IsBoundToObject(this))
+	{
+		Binding->OnControlRigBind().AddRaw(this, &SSelectionSets::HandleOnObjectBoundToControlRig);
+	}
+}
+
+//mz todo need to test recompiling
+void SSelectionSets::HandleOnObjectBoundToControlRig(UObject* InObject)
+{
+
+}
+
+void SSelectionSets::NewSelectionSet(UAIESelectionSets* NewSelectionSet)
+{
+	if (NewSelectionSet)
+	{
+		NewSelectionSet->SelectionListChanged().RemoveAll(this);
+		NewSelectionSet->SelectionListChanged().AddRaw(this, &SSelectionSets::HandleSelectionListChanged);
+		if (NewSelectionSet->ExportAsJsonString(LastSelectionSetJSON) == false)
+		{
+			UE_LOG(LogControlRig, Error, TEXT("Selection Sets: Error fast JSON export"));
+		}
+	}
+}
+
+
+FReply SSelectionSets::OnAddClicked()
+{
+	if (TSharedPtr<ISequencer> SequencerPtr = UAIESelectionSets::GetSequencerFromAsset())
+	{
+		if (UAIESelectionSets* SelectionSets = UAIESelectionSets::GetSelectionSets(SequencerPtr, true /*Add if doesn't exist*/))
+		{
+			NewSelectionSet(SelectionSets);
+
+			FocusOnTextBlock = SelectionSets->CreateSetItemFromSelection(SequencerPtr);
+		}
+	}
+	return FReply::Handled();
+}
+
+bool SSelectionSets::ImportOrExportPath(FString& OutPath, bool bImport)
+{
+	TArray<FString> Filenames;
+	IDesktopPlatform* DesktopPlatform = FDesktopPlatformModule::Get();
+	if (DesktopPlatform == nullptr)
+	{
+		return false;
+	}
+	static FString LastFileName;
+
+	bool bOpen;
+
+	if (bImport)
+	{
+		FString LastRigMapperImportDirectory;
+		if (!GConfig->GetString(TEXT("AIESelectionSets"), TEXT("LastImportDirectory"), LastRigMapperImportDirectory, GEditorPerProjectIni))
+		{
+			LastRigMapperImportDirectory = FEditorDirectories::Get().GetLastDirectory(ELastDirectory::GENERIC_IMPORT);
+		}
+
+		bOpen = DesktopPlatform->OpenFileDialog(
+			FSlateApplication::Get().FindBestParentWindowHandleForDialogs(nullptr),
+			NSLOCTEXT("RigMapperDefinitionIO", "ImportFromJson", "Import Selection Set from Json...").ToString(),
+			LastRigMapperImportDirectory,
+			TEXT(""),
+			TEXT("json (*.json)|*.json|"),
+			EFileDialogFlags::None,
+			Filenames
+		);
+	}
+	else
+	{
+		FString LastRigMapperExportDirectory;
+		if (!GConfig->GetString(TEXT("AIESelectionSets"), TEXT("LastExportDirectory"), LastRigMapperExportDirectory, GEditorPerProjectIni))
+		{
+			LastRigMapperExportDirectory = FEditorDirectories::Get().GetLastDirectory(ELastDirectory::GENERIC_EXPORT);
+		}
+
+		bOpen = DesktopPlatform->SaveFileDialog(
+			FSlateApplication::Get().FindBestParentWindowHandleForDialogs(nullptr),
+			LOCTEXT("ExportToJson", "Export Selection Set to Json...").ToString(),
+			LastRigMapperExportDirectory,
+			LastFileName,
+			TEXT("json (*.json)|*.json|"),
+			EFileDialogFlags::None,
+			Filenames
+		);
+	}
+
+	if (bOpen && !Filenames.IsEmpty())
+	{
+		LastFileName = FPaths::GetBaseFilename(Filenames[0]);
+		OutPath = Filenames[0];
+	}
+
+	return bOpen;
+}
+
+
+bool SSelectionSets::ImportLastSelectionSet()
+{
+	if (LastSelectionSetJSON.IsEmpty() == false)
+	{
+		if (TSharedPtr<ISequencer> SequencerPtr = UAIESelectionSets::GetSequencerFromAsset())
+		{
+			if (UAIESelectionSets* SelectionSets = UAIESelectionSets::GetSelectionSets(SequencerPtr, true /*Add if doesn't exist*/))
+			{
+				FScopedTransaction Transaction(LOCTEXT("LoadFromLastJson", "Load Selection Set from last open Selection Set"));
+
+				if (SelectionSets->LoadFromJsonString(SequencerPtr, LastSelectionSetJSON))
+				{
+					if (SelectionSets->ExportAsJsonString(LastSelectionSetJSON) == false) //just resave the string, it may add have more things added to it
+					{
+						UE_LOG(LogControlRig, Error, TEXT("Selection Sets: Error fast JSON export"));
+						return false;
+					}
+				}
+				else
+				{
+					UE_LOG(LogControlRig, Error, TEXT("Selection Sets: Error fast JSON import"));
+					return false;
+				}
+			}
+		}
+	}
+	return true;
+}
+
+bool SSelectionSets::ImportFromJson()
+{
+	FFilePath JsonFile;
+	bool bSucceeded = false;
+	FString LastPath;
+
+	if (ImportOrExportPath(JsonFile.FilePath, true))
+	{
+		if (TSharedPtr<ISequencer> SequencerPtr = UAIESelectionSets::GetSequencerFromAsset())
+		{
+			if (UAIESelectionSets* SelectionSets = UAIESelectionSets::GetSelectionSets(SequencerPtr, true /*Add if doesn't exist*/))
+			{
+				NewSelectionSet(SelectionSets);
+				FScopedTransaction Transaction(LOCTEXT("LoadFromJson", "Load Selection Sets from Json file"));
+
+				if (SelectionSets->LoadFromJsonFile(SequencerPtr, JsonFile))
+				{
+					LastPath = JsonFile.FilePath;
+					// set the new import directory from the last successful save
+					FString NewDirectory = FPaths::GetPath(LastPath);
+					GConfig->SetString(TEXT("AIESelectionSets"), TEXT("LastImportDirectory"), *NewDirectory, GEditorPerProjectIni);
+					GConfig->Flush(false, GEditorPerProjectIni);
+					return true;
+				}
+			}
+		}
+	}
+
+	return false;
+}
+bool SSelectionSets::ExportToJson()
+{
+	if (TSharedPtr<ISequencer> SequencerPtr = UAIESelectionSets::GetSequencerFromAsset())
+	{
+		if (UAIESelectionSets* SelectionSets = UAIESelectionSets::GetSelectionSets(SequencerPtr))
+		{
+			FFilePath Path;
+			bool bSucceeded = false;
+
+			if (ImportOrExportPath(Path.FilePath, false))
+			{
+				if (SelectionSets->ExportAsJsonFile(Path))
+				{
+					// set the new export directory
+					FString NewDirectory = FPaths::GetPath(Path.FilePath);
+					GConfig->SetString(TEXT("AIESelectionSets"), TEXT("LastExportDirectory"), *NewDirectory, GEditorPerProjectIni);
+					GConfig->Flush(false, GEditorPerProjectIni);
+					return true;
+				}
+			}
+		}
+	}
+	return false;
+}
+
+#undef LOCTEXT_NAMESPACE
+
+
+
+
diff --git a/Engine/Plugins/Animation/ControlRig/Source/ControlRigEditor/Private/Sequencer/SelectionSets/SSelectionSets.h b/Engine/Plugins/Animation/ControlRig/Source/ControlRigEditor/Private/Sequencer/SelectionSets/SSelectionSets.h
new file mode 100644
--- /dev/null
+++ b/Engine/Plugins/Animation/ControlRig/Source/ControlRigEditor/Private/Sequencer/SelectionSets/SSelectionSets.h
@@ -0,0 +1,144 @@
+// Copyright Epic Games, Inc. All Rights Reserved.
+#pragma once
+
+#include "CoreMinimal.h"
+#include "EditMode/ControlRigBaseDockableView.h"
+#include "Widgets/SWidget.h"
+#include "Widgets/SCompoundWidget.h"
+#include "ISequencer.h"
+#include "Misc/Guid.h"
+#include "Widgets/DeclarativeSyntaxSupport.h"
+
+class UAIESelectionSets;
+class SSelectionSets;
+class SSearchBox;
+class FTextFilterExpressionEvaluator;
+struct FActorWithSelectionSet;
+class SBox;
+class SBorder;
+class SScrollBox;
+class SWrapBox;
+class SGridPanel;
+class SInlineEditableTextBlock;
+
+class SSelectionSetButton: public SButton
+{
+public:
+	SLATE_BEGIN_ARGS(SSelectionSetButton)
+		{}
+		SLATE_DEFAULT_SLOT(FArguments, Content)
+		SLATE_EVENT(FOnClicked, OnClicked)
+		SLATE_ARGUMENT(FGuid, Guid)
+		SLATE_ARGUMENT(SSelectionSets*, SetsWidget )
+	SLATE_END_ARGS()
+	
+	void Construct(const FArguments& Args);
+	
+protected:
+	virtual FReply OnMouseButtonDown(const FGeometry& MyGeometry, const FPointerEvent& MouseEvent) override;
+	virtual FReply OnMouseButtonDoubleClick(const FGeometry& InMyGeometry, const FPointerEvent& InMouseEvent) override;
+
+private:
+	FGuid SelectionSetGuid;
+	SSelectionSets* SelectionSetsWidget = nullptr;
+	void AddSelectionToSet() const;
+	void RemoveSelectionFromSet() const;
+	void DeleteSet() const;
+	void RenameSet() const;
+	TSharedRef<SWidget> CreateSelectionSetColorWidget(const FGuid& InGuid, TSharedPtr<ISequencer>& SequencerPtr) const;
+	void CreateMirror() const;
+	void SelectMirror() const;
+	void HideSelectionSet() const;
+	void ShowSelectionSet() const;
+	void IsolateSelectionSet() const;
+	void ShowAllSelectionSet() const;
+	void KeyAll() const;
+};
+
+class SSelectionSets : public FControlRigBaseDockableView, public SCompoundWidget
+{
+	SLATE_BEGIN_ARGS(SSelectionSets) {}
+	SLATE_END_ARGS()
+	SSelectionSets();
+	~SSelectionSets();
+
+	void Construct(const FArguments& InArgs, FControlRigEditMode& InEditMode);
+	//SCompuntWidget
+	void Tick(const FGeometry& AllottedGeometry, const double InCurrentTime, const float InDeltaTime) override;
+
+	//FControlRigBaseDockableView overrides
+	virtual void SetEditMode(FControlRigEditMode& InEditMode) override;
+	virtual void HandleControlAdded(UControlRig* ControlRig, bool bIsAdded) override;
+
+	void SetFocusOnTextBlock(const FGuid& InGuid) { FocusOnTextBlock = InGuid; }
+private:
+	void HandleOnControlRigBound(UControlRig* InControlRig);
+	void HandleOnObjectBoundToControlRig(UObject* InObject);
+	virtual void HandleControlSelected(UControlRig* Subject, FRigControlElement* InControl, bool bSelected) override;
+
+	//set of control rigs we are bound too and need to clear delegates from
+	TArray<TWeakObjectPtr<UControlRig>> BoundControlRigs;
+private:
+	//actor selection changing
+	void RegisterSelectionChanged();
+	void OnActorSelectionChanged(const TArray<UObject*>& NewSelection, bool bForceRefresh);
+	FDelegateHandle OnSelectionChangedHandle;
+
+	//sequencer changing so need to refresh list
+	void OnActivateSequence(FMovieSceneSequenceIDRef ID);
+	void OnMovieSceneDataChanged(EMovieSceneDataChangeType);
+	void OnCloseEvent(TSharedRef<ISequencer> InSequencer);
+	void OnMovieSceneBindingsChanged();
+	void UpdateBindings();
+
+	FGuid LastMovieSceneSig = FGuid();
+
+
+	void HandleSelectionListChanged(UAIESelectionSets*) { bPopulateUI = true; }
+private:
+	FReply OnSelectionFilterClicked();
+	bool IsSelectionFilterActive() const;
+	FReply OnAddClicked();
+	bool ImportOrExportPath(FString& OutPath, bool bImport);
+	bool ImportLastSelectionSet();
+	bool ImportFromJson();
+	bool ExportToJson();
+	void NewSelectionSet(UAIESelectionSets* NewSelectionSet);
+
+private:
+	
+	TSharedRef<SWidget> CreateActorMenu(FActorWithSelectionSet& ActorWithSelectionSet);
+	TSharedRef<SWidget> MakeImportExportMenu();
+
+//widgets
+private:
+	//main UI elements
+	TSharedPtr<SBorder> MainAddFilterSetSlot;
+	TSharedPtr<SBox> ActiveActorsSlot;
+	TSharedPtr<SBorder> SearchSetsSlot;
+	TSharedPtr<SBox> ActiveSetsSlot;
+
+	TSharedPtr<SGridPanel> ActiveActorsWidget;
+	TSharedPtr<SWrapBox> SelectionSetsWidget;
+
+	TMap<FGuid, TSharedPtr<SSelectionSetButton>> SelectionSetsButton;
+	TMap<FGuid, TSharedPtr<SInlineEditableTextBlock>> SelectionSetsTextBlock;
+	TSharedPtr<SSearchBox> SetSearchBox;
+	TSharedPtr<FTextFilterExpressionEvaluator>  SetTextFilter;
+
+	FGuid  FocusOnTextBlock;
+
+private:
+	void PopulateUI();
+	void PopulateActiveActors();
+	void PopulateSelectionSets();
+	void PostEditUndo();
+	FDelegateHandle PostUndoRedoHandle;
+	//things we handle on next tick
+	bool bSetUpTimerIsSet = false;
+	bool bUpdateBindings = false;
+	bool bPopulateUI = false;
+private:
+	//cache of last select set closed, for quick loading of sets onto new level sequences
+	static FString  LastSelectionSetJSON;
+};
\ No newline at end of file
diff --git a/Engine/Plugins/Animation/ControlRig/Source/ControlRigEditor/Private/Sequencer/SelectionSets/SelectionSets.cpp b/Engine/Plugins/Animation/ControlRig/Source/ControlRigEditor/Private/Sequencer/SelectionSets/SelectionSets.cpp
new file mode 100644
--- /dev/null
+++ b/Engine/Plugins/Animation/ControlRig/Source/ControlRigEditor/Private/Sequencer/SelectionSets/SelectionSets.cpp
@@ -0,0 +1,1842 @@
+// Copyright Epic Games, Inc. All Rights Reserved.
+#include "SelectionSets.h"
+#include "SelectionSetsSettings.h"
+
+#include "ScopedTransaction.h"
+#include "ControlRig.h"
+#include "IControlRigObjectBinding.h"
+#include "Misc/Optional.h"
+#include "ISequencer.h"
+#include "MVVM/Selection/SequencerSelectionEventSuppressor.h"
+#include "MVVM/Selection/SequencerCoreSelection.h"
+#include "MVVM/Selection/Selection.h"
+#include "MVVM/TrackRowModelStorageExtension.h"
+#include "MVVM/ViewModels/TrackModel.h"
+#include "MVVM/ViewModels/SectionModel.h"
+#include "MVVM/ViewModels/TrackRowModel.h"
+#include "MVVM/ViewModels/SequencerEditorViewModel.h"
+#include "MVVM/ViewModels/SequencerOutlinerViewModel.h"
+#include "MVVM/ViewModels/ObjectBindingModel.h"
+#include "Channels/MovieSceneFloatChannel.h"
+#include "Channels/MovieSceneChannelProxy.h"
+#include "Channels/MovieSceneDoubleChannel.h"
+#include "LevelSequence.h"
+#include "ControlRigSequencerEditorLibrary.h"
+#include "Sequencer/ControlRigParameterTrackEditor.h"
+#include "Sequencer/MovieSceneControlRigParameterTrack.h"
+#include "Sequencer/MovieSceneControlRigParameterSection.h"
+#include "Tracks/MovieScene3DTransformTrack.h"
+#include "ILevelSequenceEditorToolkit.h"
+#include "LevelSequencePlayer.h"
+#include "EditMode/ControlRigEditMode.h"
+#include "LevelSequenceEditorBlueprintLibrary.h"
+#include "MovieScene.h"
+#include "Tools/EvaluateSequencerTools.h"
+#include "Editor.h"
+#include "Engine/Selection.h"
+#include "EditorModeManager.h"
+#include "Tools/ControlRigPose.h"
+#include "Misc/FileHelper.h"
+#include "Serialization/JsonTypes.h"
+#include "Dom/JsonValue.h"
+#include "Dom/JsonObject.h"
+#include "Serialization/JsonReader.h"
+#include "Serialization/JsonWriter.h"
+#include "Serialization/JsonSerializer.h"
+
+#include "Tools/ControlRigPoseMirrorTable.h"
+
+#include "Animation/AnimationSettings.h"
+#include "Animation/MirrorDataTable.h"
+
+#define LOCTEXT_NAMESPACE "SelectionSets"
+
+///////////////////////////////////////////////////////////
+/// FAIESelectionSetItem
+///////////////////////////////////////////////////////////
+
+bool FAIESelectionSetItem::ContainsItem(const FString& InName) const
+{
+	for (const FAIESelectionSetItemName& LocalName : Names)
+	{
+		if (LocalName.Name == InName)
+		{
+			return true;
+		}
+	}
+	return false;
+}
+
+FAIESelectionSetItemName* FAIESelectionSetItem::GetItem(const FString& InName)
+{
+	for (FAIESelectionSetItemName& LocalName : Names)
+	{
+		if (LocalName.Name == InName)
+		{
+			return &LocalName;
+		}
+	}
+	return nullptr;
+}
+
+bool FAIESelectionSetItem::RemoveItem(const FString& InName)
+{
+	int32 RemoveIndex = INDEX_NONE;
+	for (int32 Index = 0; Index < Names.Num(); ++Index)
+	{
+		FAIESelectionSetItemName& LocalName = Names[Index];
+		if (LocalName.Name == InName)
+		{
+			RemoveIndex = Index;
+			break;
+		}
+	}
+	if (RemoveIndex != INDEX_NONE)
+	{
+		Names.RemoveAt(RemoveIndex);
+		return true;
+	}
+	return false;
+}
+
+FString FAIESelectionSetItem::GetMirror(const FString& InString)
+{
+	FString NewString = InString;
+	if (UAnimationSettings* AnimSettings = UAnimationSettings::Get())
+	{
+		FString PossibleNewString = FControlRigPoseMirrorTable::GetMirrorString(InString, AnimSettings->MirrorFindReplaceExpressions);
+		if (PossibleNewString.IsEmpty() == false)
+		{
+			NewString = PossibleNewString;
+		}
+	}
+	return NewString;
+}
+
+AActor* FActorWithSelectionSet::GetControlRigActor(const UControlRig* InControlRig)
+{
+	TSharedPtr<IControlRigObjectBinding> ObjectBinding = InControlRig->GetObjectBinding();
+	if (ObjectBinding.IsValid())
+	{
+		if (USceneComponent* BoundSceneComponent = Cast<USceneComponent>(ObjectBinding->GetBoundObject()))
+		{
+			return BoundSceneComponent->GetOwner();
+		}
+	}
+	return (AActor*)(nullptr);
+}
+
+void FActorWithSelectionSet::SetUpOwnedActors(const TWeakObjectPtr<AActor>& ActorOwner, TSharedPtr<ISequencer>& InSequencer, UAIESelectionSets* SelectionSets)
+{
+	if (AActor* Actor = Cast<AActor>(ActorOwner))
+	{
+		OwnedActors.Reset();
+		for (const FGuid& Guid : SetsThatMatch)
+		{
+			if (const FAIESelectionSetItem* CurrentSetItem = SelectionSets->SelectionSets.Find(Guid))
+			{
+				for (const FAIESelectionSetItemName& SetItem : CurrentSetItem->Names)
+				{
+					if (SetItem.OwnerActorName.Len() > 0 && SetItem.Name == ActorOwner->GetActorLabel(false))
+					{
+						OwnedActors.Add(Guid, SetItem.OwnerActorName);
+					}
+				}
+			}
+		}
+	}
+}
+
+
+void FActorsWithSelectionSets::SetUpFromSequencer(TSharedPtr<ISequencer>& Sequencer, UAIESelectionSets* SelectionSets, const TSet<TWeakObjectPtr<UObject>>& MakeSelectedIfPresent)
+{
+	//if nothing there than automatically selected even if not in MakeSelectedIfPresent
+
+	// const TMap<FGuid, FAIESelectionSetItem>& SelectionSets
+	if (Sequencer.IsValid() == false || SelectionSets == nullptr || SelectionSets->SelectionSets.Num() == 0)
+	{
+		ActorsWithSet.Reset();
+		return;
+	}
+	
+	//used to seee if we have any left overs
+	TArray<TWeakObjectPtr<AActor>> CurrentSet;
+	ActorsWithSet.GenerateKeyArray(CurrentSet);
+
+	TArray<UControlRig*> ControlRigs;
+	FControlRigEditMode* ControlRigEditMode = static_cast<FControlRigEditMode*>(GLevelEditorModeTools().GetActiveMode(FControlRigEditMode::ModeName));
+	if (ControlRigEditMode)
+	{
+		ControlRigs = ControlRigEditMode->GetControlRigsArray(false /*bIsVisible*/);
+	}
+
+	if (ULevelSequence* LevelSequence = Cast<ULevelSequence>(Sequencer->GetFocusedMovieSceneSequence()))
+	{
+		const UMovieScene* MovieScene = LevelSequence->GetMovieScene();
+		if (MovieScene)
+		{
+			const TArray<FMovieSceneBinding>& Bindings = MovieScene->GetBindings();
+			for (const FMovieSceneBinding& Binding : Bindings)
+			{
+				for (TWeakObjectPtr<> BoundObject : Sequencer->FindBoundObjects(Binding.GetObjectGuid(), Sequencer->GetFocusedTemplateID()))
+				{
+					if (AActor* Actor = Cast<AActor>(BoundObject))
+					{
+						FString ActorName = Actor->GetActorLabel(false);
+						for (const TPair<FGuid, FAIESelectionSetItem>& Pair : SelectionSets->SelectionSets)
+						{
+							if (Pair.Value.ContainsItem(ActorName))
+							{
+								const bool bExists = ActorsWithSet.Contains(Actor);
+								FActorWithSelectionSet& Set = ActorsWithSet.FindOrAdd(Actor);
+								Set.SetsThatMatch.Add(Pair.Key);
+								Set.Name = FText::FromString(Actor->GetActorLabel());
+								if (bExists == false)
+								{
+									Set.Actors.Add(Actor);
+									Set.Color = SelectionSets->GetNextActorColor();
+									Set.SetSelected(true);
+								}
+								if (MakeSelectedIfPresent.Contains(Actor))
+								{
+									Set.SetSelected(true);
+								}
+								CurrentSet.Remove(Actor);
+							}
+						}
+					}
+				}
+			}
+		}
+	}
+
+	for (UControlRig* ControlRig : ControlRigs)
+	{
+		if (ControlRig)
+		{
+			const TArray<FRigControlElement*> Controls = ControlRig->AvailableControls();
+			for (const FRigControlElement* Control : Controls)
+			{
+				if (Control)
+				{
+					for (const TPair<FGuid, FAIESelectionSetItem>& Pair : SelectionSets->SelectionSets)
+					{
+						if (Pair.Value.ContainsItem(Control->GetName()))
+						{
+							if (AActor * Actor = FActorWithSelectionSet::GetControlRigActor(ControlRig))
+							{
+								FString ActorName = Actor->GetActorLabel(false);
+
+								const bool bExists = ActorsWithSet.Contains(Actor);
+								FActorWithSelectionSet& Set = ActorsWithSet.FindOrAdd(Actor);
+
+								Set.Name = FText::FromString(ActorName);
+								if (bExists == false)
+								{
+									Set.Color = SelectionSets->GetNextActorColor();
+									Set.Actors.Add(ControlRig);
+									Set.SetSelected(true);
+								}
+								Set.SetsThatMatch.Add(Pair.Key);
+								if (MakeSelectedIfPresent.Contains(ControlRig))
+								{
+									Set.SetSelected(true);
+								}
+								CurrentSet.Remove(Actor);
+							}
+						}
+					}
+				}
+			}
+		}
+	}
+
+	//now remove those that no longer exist with any set's active, BUT if the name matches make them matchstate
+	if (CurrentSet.Num() > 0)
+	{
+		for (TWeakObjectPtr<AActor>& Object : CurrentSet)
+		{
+			if (FActorWithSelectionSet* Set = ActorsWithSet.Find(Object))
+			{
+				for (TPair<TWeakObjectPtr<AActor>, FActorWithSelectionSet>& Pair : ActorsWithSet)
+				{
+					if ((Object != Pair.Key) && (Set->Name.EqualTo(Pair.Value.Name)))
+					{
+						Pair.Value.SetSelected(Set->GetSelected());
+						Pair.Value.Color = Set->Color;
+					}
+				}
+				ActorsWithSet.Remove(Object);
+			}
+		}
+	}
+	for (TPair<TWeakObjectPtr<AActor>, FActorWithSelectionSet>& Pair : ActorsWithSet)
+	{
+		Pair.Value.SetUpOwnedActors(Pair.Key, Sequencer, SelectionSets);
+	}
+}
+
+void FAIESelectionSetItem::GetAllChildren(const UAIESelectionSets* SelectionSets, TArray<FGuid>& AllChildren) const
+{
+	AllChildren.Add(Guid);
+	if (SelectionSets)
+	{
+		for (const FGuid& Child : Children)
+		{
+			if (const FAIESelectionSetItem* Item = SelectionSets->SelectionSets.Find(Child))
+			{
+				Item->GetAllChildren(SelectionSets, AllChildren);
+			}
+		}
+	}
+}
+
+///////////////////////////////////////////////////////////
+/// FActorWithSelectionSet
+///////////////////////////////////////////////////////////
+
+bool FActorWithSelectionSet::IsSelectedInViewport(bool bOnlyInViewport) const
+{
+	USelection* ActorSelection = GEditor->GetSelectedActors();
+	TArray<TWeakObjectPtr<UObject>> SelectedActors;
+	if (ActorSelection)
+	{
+		ActorSelection->GetSelectedObjects(SelectedActors);
+	}
+	for (const TWeakObjectPtr<UObject>& Object : Actors)
+	{
+		if (const AActor* Actor = Cast<const AActor>(Object.Get()))
+		{
+			if (SelectedActors.Contains(Object))
+			{
+				bWasSelectedInViewport = true;
+				return true;
+			}
+		}
+		else if (const UControlRig* ControlRig = Cast<const UControlRig>(Object.Get()))
+		{
+			TArray<FName> SelectedControls = ControlRig->CurrentControlSelection();
+			if (SelectedControls.Num() > 0)
+			{
+				bWasSelectedInViewport = true;
+				return true;
+			}
+			if (AActor* CRActor = GetControlRigActor(ControlRig))
+			{
+				if (SelectedActors.Contains(CRActor))
+				{
+					bWasSelectedInViewport = true;
+					return true;
+				}
+			}
+		}
+	}
+	if (bOnlyInViewport == false && bWasSelectedInViewport == true)
+	{
+		return true;
+	}
+
+	return false;
+
+}
+bool FActorWithSelectionSet::IsControlRig() const
+{
+	for (const TWeakObjectPtr<UObject>& Object : Actors)
+	{
+		if (const UControlRig* ControlRig = Cast<const UControlRig>(Object.Get()))
+		{
+			return true;
+		}
+	}
+	return false;
+}
+
+bool FActorWithSelectionSet::IsActive(bool bShowSelectedOnly) const
+{
+	if (const USelectionSetsSettings* Settings = GetDefault<USelectionSetsSettings>())
+	{
+		if (Settings->GetSelectionSetsActorSetVisibility() == ESelectionSetsActorSetVisibility::MultipleControRigs)
+		{
+			if (IsControlRig() == false)
+			{
+				return true;
+			}
+		}
+	}
+	if (bShowSelectedOnly)
+	{
+		return IsSelectedInViewport();
+	}
+	
+	return bIsSelected;
+}
+
+bool UAIESelectionSets::IsVisible(const FActorWithSelectionSet& ActorWithSelectionSet) const
+{
+	if (const USelectionSetsSettings* Settings = GetDefault<USelectionSetsSettings>())
+	{
+		if (Settings->GetSelectionSetsActorSetVisibility() == ESelectionSetsActorSetVisibility::MultipleControRigs)
+		{
+			if (bShowSelectedOnly)
+			{
+				return false;
+			}
+			if (ActorWithSelectionSet.IsControlRig() == true)
+			{
+				for (const TPair<TWeakObjectPtr<AActor>, FActorWithSelectionSet>& Pair : ActorsWithSelectionSets.ActorsWithSet)
+				{
+					if (Pair.Value.IsControlRig() == true)
+					{
+						if (&ActorWithSelectionSet != &Pair.Value)
+						{
+							return true;
+						}
+					}
+				}
+			}
+			else
+			{
+				return false;
+			}
+		}
+	}
+	if (bShowSelectedOnly)
+	{
+		if (ActorWithSelectionSet.IsSelectedInViewport())
+		{
+			return true;
+		}
+	}
+	return false;
+}
+
+void FActorWithSelectionSet::ShowAll() const
+{
+	FControlRigEditMode* EditMode = static_cast<FControlRigEditMode*>(GLevelEditorModeTools().GetActiveMode(FControlRigEditMode::ModeName));
+	if (EditMode)
+	{
+		for (const TWeakObjectPtr<UObject>& Object : Actors)
+		{
+			if (UControlRig* ControlRig = Cast<UControlRig>(Object.Get()))
+			{
+				TSet<FString> ControlNames;
+				const TArray<FRigControlElement*> AvailbleControls = ControlRig->AvailableControls();
+				ControlNames.Reserve(AvailbleControls.Num());
+				for (const FRigControlElement* Element : AvailbleControls)
+				{
+					ControlNames.Add(Element->GetName());
+				}
+				//not undoable yet
+				EditMode->ShowControlRigControls(ControlRig, ControlNames, true);
+			}
+			else if (AActor* Actor = Cast<AActor>(Object.Get()))
+			{
+				Actor->SetActorHiddenInGame(false);
+				Actor->SetIsTemporarilyHiddenInEditor(false);
+			}
+		}
+	}
+}
+
+///////////////////////////////////////////////////////////
+/// UAIESelectionSets
+///////////////////////////////////////////////////////////
+
+UAIESelectionSets::UAIESelectionSets(class FObjectInitializer const& ObjectInitializer)
+	: Super(ObjectInitializer)
+{
+
+}
+TSharedPtr<ISequencer> UAIESelectionSets::GetSequencerFromAsset()
+{
+	ULevelSequence* LevelSequence = ULevelSequenceEditorBlueprintLibrary::GetCurrentLevelSequence();
+	IAssetEditorInstance* AssetEditor = GEditor->GetEditorSubsystem<UAssetEditorSubsystem>()->FindEditorForAsset(LevelSequence, false);
+	ILevelSequenceEditorToolkit* LevelSequenceEditor = static_cast<ILevelSequenceEditorToolkit*>(AssetEditor);
+	TSharedPtr<ISequencer> SequencerPtr = LevelSequenceEditor ? LevelSequenceEditor->GetSequencer() : nullptr;
+	return SequencerPtr;
+}
+
+UAIESelectionSets* UAIESelectionSets::GetSelectionSets(TSharedPtr<ISequencer>& SequencerPtr, bool bAddIfDoesNotExist)
+{
+	if (SequencerPtr.IsValid() == false)
+	{
+		return nullptr;
+	}
+	ULevelSequence* LevelSequence = Cast<ULevelSequence>(SequencerPtr->GetFocusedMovieSceneSequence());
+	return UAIESelectionSets::GetSelectionSets(LevelSequence, bAddIfDoesNotExist);
+}
+
+UAIESelectionSets* UAIESelectionSets::GetSelectionSets(ULevelSequence* LevelSequence, bool bAddIfDoesNotExist)
+{
+	if (!LevelSequence)
+	{
+		return nullptr;
+	}
+	if (LevelSequence && LevelSequence->GetClass()->ImplementsInterface(UInterface_AssetUserData::StaticClass()))
+	{
+		if (IInterface_AssetUserData* AssetUserDataInterface = Cast< IInterface_AssetUserData >(LevelSequence))
+		{
+			UAIESelectionSets* SelectionSets = AssetUserDataInterface->GetAssetUserData< UAIESelectionSets >();
+			if (!SelectionSets && bAddIfDoesNotExist)
+			{
+				SelectionSets = NewObject<UAIESelectionSets>(LevelSequence, NAME_None, RF_Public | RF_Transactional);
+				AssetUserDataInterface->AddAssetUserData(SelectionSets);
+			}
+			return SelectionSets;
+		}
+	}
+	return nullptr;
+}
+
+TArray<FGuid> UAIESelectionSets::GetActiveSelectionSets() 
+{
+	//when bShowSelected is on we cache the last set and use that when none are selected. Also need to keep track
+	//of when it was previously selected
+	//if not show selected we just use the internal select flag (or it's an actor)
+
+	TSet<FGuid> ActiveSelectionSetsSet;
+	TSet<FGuid> SelectedInViewport;
+	TSet<TWeakObjectPtr<AActor>> SelectedActors;
+	for (const TPair<TWeakObjectPtr<AActor>, FActorWithSelectionSet>& Pair : ActorsWithSelectionSets.ActorsWithSet)
+	{
+		if (bShowSelectedOnly)
+		{
+			constexpr bool bOnlyInViewport = true;
+			if (Pair.Value.IsSelectedInViewport(true))
+			{
+				SelectedInViewport.Append(Pair.Value.SetsThatMatch);
+				ActiveSelectionSetsSet.Append(Pair.Value.SetsThatMatch);
+				SelectedActors.Add(Pair.Key);
+			}
+		}
+		else 
+		{
+			if (Pair.Value.IsActive(bShowSelectedOnly))
+			{
+				ActiveSelectionSetsSet.Append(Pair.Value.SetsThatMatch);
+			}
+		}
+		
+	}
+	
+	if (bShowSelectedOnly)
+	{
+		if (SelectedInViewport.Num() == 0)
+		{
+			for(const FGuid& Guid: LastActiveSelectionSetsSet)
+			{
+				if (SelectionSets.Contains(Guid))
+				{
+					ActiveSelectionSetsSet.Add(Guid);
+				}
+			}
+		}
+		else//something is selected we need to clear the unselected actors previous state
+		{
+			for (const TPair<TWeakObjectPtr<AActor>, FActorWithSelectionSet>& Pair : ActorsWithSelectionSets.ActorsWithSet)
+			{
+				if (SelectedActors.Contains(Pair.Key) == false)
+				{
+					Pair.Value.bWasSelectedInViewport = false;
+				}
+			}
+		}
+		LastActiveSelectionSetsSet = ActiveSelectionSetsSet;
+	}
+	TArray<FGuid> ActiveSelectionSetsArray;
+	if (ActiveSelectionSetsSet.Num() > 0)
+	{
+		ActiveSelectionSetsArray = ActiveSelectionSetsSet.Array();
+		ActiveSelectionSetsArray = SortSelectionSetsByRow(ActiveSelectionSetsArray);
+	}
+	return ActiveSelectionSetsArray;
+}
+
+TArray<FGuid> UAIESelectionSets::SortSelectionSetsByRow(TArray<FGuid>& InSelectionSets) 
+{
+	TSortedMap<int32, FGuid> SortedSelectionSets;
+
+	for (const FGuid& Guid : InSelectionSets)
+	{
+		if (FAIESelectionSetItem* CurrentSetItem = SelectionSets.Find(Guid))
+		{
+			int32 Row = CurrentSetItem->ViewData.Row;
+			while (SortedSelectionSets.Contains(Row))
+			{
+				++Row;
+			}
+			CurrentSetItem->ViewData.Row = Row;
+			SortedSelectionSets.Add(Row, Guid);
+		}
+	}
+	TArray<FGuid> OutSelectionSets;
+	SortedSelectionSets.GenerateValueArray(OutSelectionSets);
+	return OutSelectionSets;
+}
+
+FGuid UAIESelectionSets::CreateSetItemFromSelection(TSharedPtr<ISequencer>& InSequencer)
+{
+	using namespace UE::AIE;
+	TArray<FControlRigAndControlsAndTrack> SelectedCRs;
+	TArray<FObjectAndTrack> SelectedBoundObjects;
+	{
+		const FScopedTransaction Transaction(LOCTEXT("AddSelctionSet_Transcation", "Add Selection Set"), !GIsTransacting);
+		Modify();
+
+		TSet<TWeakObjectPtr<UObject>> ObjectsAdded;
+
+		FAIESelectionSetItem Item;
+		Item.Guid = FGuid::NewGuid();
+		Item.ViewData.Color = GetNextSelectionSetColor();
+		FSequencerSelected::GetSelectedControlRigsAndBoundObjects(InSequencer.Get(), SelectedCRs, SelectedBoundObjects);
+
+	
+		if (SelectedCRs.Num() <= 0 && SelectedBoundObjects.Num() <= 0)
+		{
+			/* no longer do empty sets for now
+			FString Name = FString::Printf(TEXT("Empty Set %d"), SelectionSets.Num());
+			FText Text;
+			Text = Text.FromString(Name);
+			Item.ItemName = Text;
+			SelectionSets.Add(Item.Guid, Item);
+			SelectionSetsChangedBroadcast();
+
+			return Item.Guid;
+			*/
+			return FGuid();
+
+		}
+
+		FString Name = FString::Printf(TEXT("Set %d"), SelectionSets.Num());
+		FText Text;
+		Text = Text.FromString(Name);
+		Item.ItemName = Text;
+
+		TMultiMap<FString, FString> ActiveNames;
+		TArray<FString> Duplicates;
+		FControlRigPoseMirrorTable MirrorTable;
+		FControlRigControlPose TempPose;
+		UControlRig* LastOwnerControlRig = nullptr;
+		for (FControlRigAndControlsAndTrack& CRControls : SelectedCRs)
+		{
+			UControlRig* ControlRig = CRControls.ControlRig;
+			LastOwnerControlRig = ControlRig;
+			TempPose.SavePose(ControlRig, true);
+			for (FName& ControlName : CRControls.Controls)
+			{
+				FAIESelectionSetItemName ItemName;
+				ItemName.Name = ControlName.ToString();
+				if (FRigControlCopy* CopyRigControl = MirrorTable.GetControl(ControlRig, TempPose, ControlName, true))
+				{
+					ItemName.MirrorName = CopyRigControl->Name.ToString();
+				}
+				else
+				{
+					ItemName.MirrorName = ItemName.Name;
+				}
+				
+				ItemName.Type = (uint32)EAIESelectionSetItemType::ControlRig;
+				Duplicates.SetNum(0);
+				ActiveNames.MultiFind(ItemName.Name, Duplicates);
+				ItemName.Duplicates = Duplicates.Num();
+				Item.Names.Add(ItemName);
+				ActiveNames.Add(ItemName.Name, ItemName.Name);
+			}
+			ObjectsAdded.Add(CRControls.ControlRig);
+		}
+		TOptional<FString> OwnerActorName;
+		if (LastOwnerControlRig)
+		{
+			if (AActor* Actor = FActorWithSelectionSet::GetControlRigActor(LastOwnerControlRig))
+			{
+				OwnerActorName = Actor->GetActorLabel(false);
+			}
+		}
+		for (FObjectAndTrack& ObjectAndTrack : SelectedBoundObjects)
+		{
+			if (AActor* Actor = Cast<AActor>(ObjectAndTrack.BoundObject))
+			{
+				FAIESelectionSetItemName ItemName;
+				ItemName.Name = Actor->GetActorLabel(true);
+				if (OwnerActorName.IsSet())
+				{
+					ItemName.OwnerActorName = OwnerActorName.GetValue();
+				}
+				ItemName.MirrorName = ItemName.Name;
+				ItemName.Type = (uint32)EAIESelectionSetItemType::Actor;
+				Duplicates.SetNum(0);
+				ActiveNames.MultiFind(ItemName.Name, Duplicates);
+				ItemName.Duplicates = Duplicates.Num();
+				Item.Names.Add(ItemName);
+				ActiveNames.Add(ItemName.Name, ItemName.Name);
+				ObjectsAdded.Add(Actor);
+			}
+		}
+		Item.ViewData.Row = SelectionSets.Num();
+		SelectionSets.Add(Item.Guid, Item);
+
+		ActorsWithSelectionSets.SetUpFromSequencer(InSequencer, this, ObjectsAdded);
+
+		SelectionSetsChangedBroadcast();
+		return Item.Guid;
+	}
+}
+
+FGuid UAIESelectionSets::CreateMirror(TSharedPtr<ISequencer>& InSequencer, const FGuid& InGuid)
+{
+	using namespace UE::Sequencer;
+
+	TArray<FGuid> AllGuids;
+	FGuid NewGuid;
+	if (const FAIESelectionSetItem* CurrentSetItem = SelectionSets.Find(InGuid))
+	{
+		const FScopedTransaction Transaction(LOCTEXT("CreateMirror", "Create Mirror Selection Set"), !GIsTransacting);
+		FControlRigPoseMirrorTable MirrorTable;
+		bool bNoteSelectionChange = false;
+		FAIESelectionSetItem Item;
+		NewGuid = FGuid::NewGuid();
+		Item.Guid = NewGuid;
+		Item.ViewData.Color = GetNextSelectionSetColor();
+		FString NewName = FAIESelectionSetItem::GetMirror(CurrentSetItem->ItemName.ToString());
+		if (NewName.IsEmpty())
+		{
+			NewName = FString::Printf(TEXT("Set %d"), SelectionSets.Num());
+		}
+		FText Text;
+		Text = Text.FromString(NewName);
+		Item.ItemName = Text;
+		for (const FAIESelectionSetItemName& Name : CurrentSetItem->Names)
+		{
+			FAIESelectionSetItemName NewItem;
+			NewItem.Duplicates = Name.Duplicates;
+			NewItem.Name = Name.MirrorName;
+			NewItem.OwnerActorName = Name.OwnerActorName;
+			NewItem.MirrorName = Name.Name;
+			NewItem.Type = Name.Type;
+			
+			Item.Names.Add(NewItem);
+		}
+		for (const FGuid& ChildGuid : CurrentSetItem->Children)
+		{
+			Item.ViewData.Row = SelectionSets.Num();
+			SelectionSets.Add(Item.Guid, Item);
+			FGuid NewChildGuid = CreateMirror(InSequencer, ChildGuid);
+			Item.Children.Add(NewChildGuid);
+		}
+		if (NewGuid.IsValid())
+		{
+			Item.ViewData.Row = SelectionSets.Num();
+			SelectionSets.Add(Item.Guid, Item);
+			TSet<TWeakObjectPtr<UObject>> NoObjectsAdded;
+			ActorsWithSelectionSets.SetUpFromSequencer(InSequencer, this, NoObjectsAdded);
+
+			SelectionSetsChangedBroadcast();
+		}
+	}
+	return NewGuid;
+}
+
+
+void UAIESelectionSets::SequencerBindingsAdded(TSharedPtr<ISequencer>& InSequencer)
+{
+	TSet<TWeakObjectPtr<UObject>> NoObjectsAdded;
+	ActorsWithSelectionSets.SetUpFromSequencer(InSequencer, this, NoObjectsAdded);
+	SelectionSetsChangedBroadcast();
+
+}
+
+bool UAIESelectionSets::AddSelectionToSetItem(const FGuid& InGuid, TSharedPtr<ISequencer>& InSequencer)
+{
+	using namespace UE::AIE;
+	if (FAIESelectionSetItem* SetItem = SelectionSets.Find(InGuid))
+	{
+		TArray<FControlRigAndControlsAndTrack> SelectedCRs;
+		TArray<FObjectAndTrack> SelectedBoundObjects;
+		
+		FSequencerSelected::GetSelectedControlRigsAndBoundObjects(InSequencer.Get(), SelectedCRs, SelectedBoundObjects);
+
+		if (SelectedCRs.Num() <= 0 && SelectedBoundObjects.Num() <= 0)
+		{
+			return false;
+		}
+		const FScopedTransaction Transaction(LOCTEXT("AddSelectionToSetItem", "Add Selection To Set Item"), !GIsTransacting);
+		Modify();
+		TMultiMap<FString, FString> ActiveNames;
+		TArray<FString> Duplicates;
+		//set up ActiveNames so we can get Duplicate counts
+		for (const FAIESelectionSetItemName& Name : SetItem->Names)
+		{
+			ActiveNames.Add(Name.Name, Name.Name);
+
+		}
+		FControlRigPoseMirrorTable MirrorTable;
+		FControlRigControlPose TempPose;
+		for (FControlRigAndControlsAndTrack& CRControls : SelectedCRs)
+		{
+			UControlRig* ControlRig = CRControls.ControlRig;
+			TempPose.SavePose(ControlRig, true);
+			for (FName& ControlName : CRControls.Controls)
+			{
+				FAIESelectionSetItemName ItemName;
+				ItemName.Name = ControlName.ToString();
+				ItemName.Type = (uint32)EAIESelectionSetItemType::ControlRig;
+				if (FRigControlCopy* CopyRigControl = MirrorTable.GetControl(ControlRig, TempPose, ControlName, true))
+				{
+					ItemName.MirrorName = CopyRigControl->Name.ToString();
+				}
+				Duplicates.SetNum(0);
+				ActiveNames.MultiFind(ItemName.Name, Duplicates);
+				ItemName.Duplicates = Duplicates.Num();
+				SetItem->Names.Add(ItemName);
+			}
+		}
+		for (FObjectAndTrack& ObjectAndTrack : SelectedBoundObjects)
+		{
+			if (AActor* Actor = Cast<AActor>(ObjectAndTrack.BoundObject))
+			{
+				FAIESelectionSetItemName ItemName;
+				ItemName.Name = Actor->GetActorLabel(true);
+				ItemName.Type = (uint32)EAIESelectionSetItemType::Actor;
+				Duplicates.SetNum(0);
+				ActiveNames.MultiFind(ItemName.Name, Duplicates);
+				ItemName.Duplicates = Duplicates.Num();
+				SetItem->Names.Add(ItemName);
+			}
+		}
+		return true;
+		
+	}
+	return false;
+}
+bool UAIESelectionSets::RemoveSelectionFromSetItem(const FGuid& InGuid, TSharedPtr<ISequencer>& InSequencer)
+{
+	using namespace UE::AIE;
+	if (FAIESelectionSetItem* SetItem = SelectionSets.Find(InGuid))
+	{
+		//wrap scoped transaction since it can deselect control rigs
+		TArray<FControlRigAndControlsAndTrack> SelectedCRs;
+		TArray<FObjectAndTrack> SelectedBoundObjects;
+
+		FSequencerSelected::GetSelectedControlRigsAndBoundObjects(InSequencer.Get(), SelectedCRs, SelectedBoundObjects);
+
+		if (SelectedCRs.Num() <= 0 && SelectedBoundObjects.Num() <= 0)
+		{
+			return false;
+		}
+		const FScopedTransaction Transaction(LOCTEXT("RemoveSelectionFromSetItem", "Remove Selection From Set Item"), !GIsTransacting);
+		Modify();
+	
+		for (FControlRigAndControlsAndTrack& CRControls : SelectedCRs)
+		{
+			for (FName& ControlName : CRControls.Controls)
+			{
+				FString Name = ControlName.ToString();
+				SetItem->RemoveItem(Name); //mz todo handle dups
+			}
+		}
+		for (FObjectAndTrack& ObjectAndTrack : SelectedBoundObjects)
+		{
+			if (AActor* Actor = Cast<AActor>(ObjectAndTrack.BoundObject))
+			{
+				FString Name = Actor->GetActorLabel(true);
+				SetItem->RemoveItem(Name);
+			}
+		}
+		return true;
+
+	}
+	return false;
+}
+
+bool UAIESelectionSets::GetItemName(const FGuid& InGuid, FText& OutName) const
+{
+	if (const FAIESelectionSetItem* SetItem = SelectionSets.Find(InGuid))
+	{
+		OutName = SetItem->ItemName;
+		return true;
+	}
+	return false;
+}
+
+bool UAIESelectionSets::RenameSetItem(const FGuid& InGuid, const  FText& NewName)
+{
+	const FScopedTransaction Transaction(LOCTEXT("RenameSetItem_transaction", "Rename Selection Set"), !GIsTransacting);
+	Modify();
+	if (FAIESelectionSetItem* SetItem = SelectionSets.Find(InGuid))
+	{
+		SetItem->ItemName = NewName;
+		return true;
+	}
+	return false;
+}
+bool UAIESelectionSets::DeleteSetItem(const FGuid& InGuid, TSharedPtr<ISequencer>& InSequencer)
+{
+	const FScopedTransaction Transaction(LOCTEXT("DeleteSetItem_transaction", "Delete Selection Set"), !GIsTransacting);
+	Modify();
+	if (const FAIESelectionSetItem* SetItem = SelectionSets.Find(InGuid))
+	{
+		ShowOrHideControls(InGuid, true, false);
+		SelectionSets.Remove(InGuid);
+		TSet<TWeakObjectPtr<UObject>> NoObjectsAdded;
+		ActorsWithSelectionSets.SetUpFromSequencer(InSequencer, this, NoObjectsAdded);
+		SelectionSetsChangedBroadcast();
+
+		return true;
+	}
+	return false;
+}
+
+bool UAIESelectionSets::GetItemColor(const FGuid& InGuid, FLinearColor& OutColor) const
+{
+	if (const FAIESelectionSetItem* SetItem = SelectionSets.Find(InGuid))
+	{
+		OutColor = SetItem->ViewData.Color;
+		return true;
+	}
+	return false;
+}
+
+bool UAIESelectionSets::SetItemColor(const FGuid& InGuid, const FLinearColor& InColor)
+{
+	const FScopedTransaction Transaction(LOCTEXT("ColorSetItem_transaction", "Color Selection Set"), !GIsTransacting);
+	Modify();
+	if (FAIESelectionSetItem* SetItem = SelectionSets.Find(InGuid))
+	{
+		SetItem->ViewData.Color = InColor;
+		return true;
+	}
+	return false;
+}
+
+static void ClearSelection()
+{
+	FControlRigEditMode* EditMode = static_cast<FControlRigEditMode*>(GLevelEditorModeTools().GetActiveMode(FControlRigEditMode::ModeName));
+	if (EditMode)
+	{
+		if ((GEditor->GetSelectedActorCount() || GEditor->GetSelectedComponentCount()))
+		{
+			GEditor->SelectNone(false, true);
+			GEditor->NoteSelectionChange();
+		}
+		TMap<UControlRig*, TArray<FRigElementKey>> SelectedControls;
+		EditMode->GetAllSelectedControls(SelectedControls);
+		for (TPair<UControlRig*, TArray<FRigElementKey>>& CurrentSelection : SelectedControls)
+		{
+			if (CurrentSelection.Key)
+			{
+				CurrentSelection.Key->ClearControlSelection(true);
+			}
+		}
+	}
+}
+
+bool UAIESelectionSets::SelectItem(const FGuid& InGuid, bool bDoMirror, bool bAdd, bool bToggle, bool bSelect) const
+{
+	using namespace UE::Sequencer;
+
+	TSharedPtr<ISequencer> Sequencer = UAIESelectionSets::GetSequencerFromAsset();
+	if (Sequencer.IsValid() == false)
+	{
+		return false;
+	}
+
+	using namespace UE::Sequencer;
+	TUniquePtr<FSelectionEventSuppressor> SequencerSelectionGuard;
+	TSharedPtr<FSequencerEditorViewModel> SequencerViewModel = Sequencer->GetViewModel();
+	if (TSharedPtr<FSequencerSelection> SequencerSelection = SequencerViewModel ? SequencerViewModel->GetSelection() : nullptr)
+	{
+		SequencerSelectionGuard.Reset(new FSelectionEventSuppressor(SequencerSelection.Get()));
+	}
+	
+	TArray<FGuid> AllGuids;
+	if (const FAIESelectionSetItem* SetItem = SelectionSets.Find(InGuid))
+	{
+		//first the do is get cached control rig selections, we do this for two reasons 1) we need to check selection state before clearing selection if bShowSelectedOnly is true
+		//and 2. we need to select the controls AFTER the actors or the actor selection will deselect the controls
+		//so we need do a two pass here
+		const FScopedTransaction Transaction(LOCTEXT("SelectSelectionSet", "Select Selection Set"), !GIsTransacting);
+		SetItem->GetAllChildren(this, AllGuids);
+		bool bNoteSelectionChange = false;
+		//we select actors first then control rigs second
+		TMap<UControlRig*,const TArray<FAIESelectionSetItemName>> CachedControlRigs;
+		TSet<FString> CachedControlRigActorNames;
+		for (const FGuid& Guid : AllGuids)
+		{
+			if (const FAIESelectionSetItem* CurrentSetItem = SelectionSets.Find(Guid))
+			{
+				for (const TPair<TWeakObjectPtr<AActor>, FActorWithSelectionSet>& Pair : ActorsWithSelectionSets.ActorsWithSet)
+				{
+					bool bIsActive = Pair.Value.IsActive(bShowSelectedOnly);
+					if (bIsActive)
+					{
+						if (Pair.Value.SetsThatMatch.Contains(Guid))
+						{
+							if (AActor* Key = Cast<AActor>(Pair.Key.Get()))
+							{
+								for (const TWeakObjectPtr<UObject>& Object : Pair.Value.Actors)
+								{
+									if (UControlRig* ControlRig = Cast<UControlRig>(Object.Get()))
+									{
+										ControlRig->Modify();
+										CachedControlRigs.Add(ControlRig, CurrentSetItem->Names);
+										if (AActor* CRActor = FActorWithSelectionSet::GetControlRigActor(ControlRig))
+										{
+											FString ActorLabel = CRActor->GetActorLabel(false);
+											CachedControlRigActorNames.Add(ActorLabel);
+										}
+									}
+								}
+							}
+						}
+					}
+				}
+			}
+		}
+		//now clear selection
+		if (bSelect == true && (bToggle == false && bAdd == false))
+		{
+			ClearSelection();
+		}
+
+		TArray<TWeakObjectPtr<UObject>> SelectedActors;
+		if (bToggle)
+		{
+			USelection* ActorSelection = GEditor->GetSelectedActors();
+			if (ActorSelection)
+			{
+				ActorSelection->GetSelectedObjects(SelectedActors);
+			}
+		}
+		//now do actors
+		for (const FGuid& Guid : AllGuids)
+		{
+			if (const FAIESelectionSetItem* CurrentSetItem = SelectionSets.Find(Guid))
+			{
+				for (const TPair<TWeakObjectPtr<AActor>, FActorWithSelectionSet>& Pair : ActorsWithSelectionSets.ActorsWithSet)
+				{
+					bool bIsActive = Pair.Value.IsActive(bShowSelectedOnly);
+					if (bIsActive)					
+					{
+						if (Pair.Value.SetsThatMatch.Contains(Guid))
+						{
+							if (AActor* Key = Cast<AActor>(Pair.Key.Get()))
+							{
+								//cache the
+								for (const TWeakObjectPtr<UObject>& Object : Pair.Value.Actors)
+								{
+									if (AActor* Actor = Cast<AActor>(Object.Get()))
+									{
+										if (const FString* OwnedName = Pair.Value.OwnedActors.Find(Guid))
+										{
+											if (CachedControlRigActorNames.Contains(*OwnedName) == false)
+											{
+												continue;
+											}
+										}
+										if (bToggle == false)
+										{
+											GEditor->SelectActor(Actor, bSelect, true);
+										}
+										else
+										{
+											const bool bToggleSelect = !SelectedActors.Contains(Actor);
+											GEditor->SelectActor(Actor, bToggleSelect, true);
+										}
+									}
+								}
+							}
+						}
+					}
+				}
+			}
+		}
+		//do control rigs last
+		TArray<FName> SelectedControls;
+		for (TPair<UControlRig*, const TArray<FAIESelectionSetItemName>>& CachedControlRig : CachedControlRigs)
+		{
+			if (UControlRig* ControlRig = CachedControlRig.Key)
+			{
+				if (bToggle)
+				{
+					SelectedControls = ControlRig->CurrentControlSelection();
+				}
+				for (const FAIESelectionSetItemName& String : CachedControlRig.Value)
+				{
+
+					if (bDoMirror)
+					{
+						const FName Name(*String.MirrorName);
+						bool bControlRigSelect = bSelect;
+						if (bToggle)
+						{
+							bControlRigSelect = !SelectedControls.Contains(Name);
+						}
+						ControlRig->SelectControl(Name, bControlRigSelect,true);
+					}
+					else
+					{
+						const FName Name(*String.Name);
+						bool bControlRigSelect = bSelect;
+						if (bToggle)
+						{
+							bControlRigSelect = !SelectedControls.Contains(Name);
+						}
+						ControlRig->SelectControl(Name, bControlRigSelect,true);
+					}
+				}}
+		}
+	}
+	return true;
+}
+
+bool UAIESelectionSets::ShowOrHideControls(const FGuid& InGuid, bool bShow, bool bDoMirror) const
+{
+	FControlRigEditMode* EditMode = static_cast<FControlRigEditMode*>(GLevelEditorModeTools().GetActiveMode(FControlRigEditMode::ModeName));
+	if (EditMode)
+	{
+		//fscopedtransaction here
+		if (bShow == false)//deselect before hidding
+		{
+			constexpr bool bDoAdd = false;
+			constexpr bool bDoToggle = false;
+			constexpr bool bSelect = false;
+			SelectItem(InGuid, bDoMirror,bDoAdd, bDoToggle, bSelect);
+		}
+		TArray<FGuid> AllGuids;
+		if (const FAIESelectionSetItem* SetItem = SelectionSets.Find(InGuid))
+		{
+			SetItem->GetAllChildren(this, AllGuids);
+			for (const FGuid& Guid : AllGuids)
+			{
+				if (const FAIESelectionSetItem* CurrentSetItem = SelectionSets.Find(Guid))
+				{
+					for (const TPair<TWeakObjectPtr<AActor>, FActorWithSelectionSet>& Pair : ActorsWithSelectionSets.ActorsWithSet)
+					{
+						if (Pair.Value.IsSelectedInViewport())
+						{
+							if (Pair.Value.SetsThatMatch.Contains(Guid))
+							{
+								if (AActor* Key = Cast<AActor>(Pair.Key.Get()))
+								{
+									for (const TWeakObjectPtr<UObject>& Object : Pair.Value.Actors)
+									{
+										if (UControlRig* ControlRig = Cast<UControlRig>(Object.Get()))
+										{
+											TSet<FString> ControlNames;
+											ControlNames.Reserve(CurrentSetItem->Names.Num());
+											for (const FAIESelectionSetItemName& String : CurrentSetItem->Names)
+											{
+												ControlNames.Add(String.Name);
+											}
+											//not undoable yet
+											EditMode->ShowControlRigControls(ControlRig, ControlNames, bShow);
+										}
+										else if (AActor* Actor = Cast<AActor>(Object.Get()))
+										{
+											Actor->SetActorHiddenInGame(!bShow);
+											Actor->SetIsTemporarilyHiddenInEditor(!bShow);
+										}
+									}
+								}
+							}
+						}
+					}
+				}
+			}
+			return true;
+		}
+	}
+	return false;
+}
+
+bool UAIESelectionSets::IsolateControls(const FGuid& InGuid) const
+{
+	FControlRigEditMode* EditMode = static_cast<FControlRigEditMode*>(GLevelEditorModeTools().GetActiveMode(FControlRigEditMode::ModeName));
+	if (EditMode)
+	{
+		TArray<FGuid> AllGuids;
+		if (const FAIESelectionSetItem* SetItem = SelectionSets.Find(InGuid))
+		{
+			SetItem->GetAllChildren(this, AllGuids);
+			for (const FGuid& Guid : AllGuids)
+			{
+				if (const FAIESelectionSetItem* CurrentSetItem = SelectionSets.Find(Guid))
+				{
+					for (const TPair<TWeakObjectPtr<AActor>, FActorWithSelectionSet>& Pair : ActorsWithSelectionSets.ActorsWithSet)
+					{
+						if (Pair.Value.IsSelectedInViewport())
+						{
+							if (Pair.Value.SetsThatMatch.Contains(Guid))
+							{
+								if (AActor* Key = Cast<AActor>(Pair.Key.Get()))
+								{
+									for (const TWeakObjectPtr<UObject>& Object : Pair.Value.Actors)
+									{
+										if (UControlRig* ControlRig = Cast<UControlRig>(Object.Get()))
+										{
+											TSet<FString> ControlNames;
+											ControlNames.Reserve(CurrentSetItem->Names.Num());
+											for (const FAIESelectionSetItemName& String : CurrentSetItem->Names)
+											{
+												ControlNames.Add(String.Name);
+											}
+											TSet<FString> ControlsToHide;
+											TArray<FRigControlElement*> Controls = ControlRig->AvailableControls();
+											for (const FRigControlElement* Control : Controls)
+											{
+												if (ControlNames.Contains(Control->GetName()) == false)
+												{
+													ControlsToHide.Add(Control->GetName());
+												}
+											}
+											//not undoable yet
+											EditMode->ShowControlRigControls(ControlRig, ControlNames, true);
+											EditMode->ShowControlRigControls(ControlRig, ControlsToHide, false);
+
+										}
+										else if (AActor* Actor = Cast<AActor>(Object.Get()))
+										{
+											//to do how to handle actors
+											const bool bShow = true;
+											Actor->SetActorHiddenInGame(!bShow);
+											Actor->SetIsTemporarilyHiddenInEditor(!bShow);
+										}
+									}
+								}
+							}
+						}
+					}
+				}
+			}
+			return true;
+		}
+	}
+	return false;
+}
+
+bool UAIESelectionSets::ShowAllControls(const FGuid& InGuid) const
+{
+	FControlRigEditMode* EditMode = static_cast<FControlRigEditMode*>(GLevelEditorModeTools().GetActiveMode(FControlRigEditMode::ModeName));
+	if (EditMode)
+	{
+		TArray<FGuid> AllGuids;
+		if (const FAIESelectionSetItem* SetItem = SelectionSets.Find(InGuid))
+		{
+			SetItem->GetAllChildren(this, AllGuids);
+			for (const FGuid& Guid : AllGuids)
+			{
+				if (const FAIESelectionSetItem* CurrentSetItem = SelectionSets.Find(Guid))
+				{
+					for (const TPair<TWeakObjectPtr<AActor>, FActorWithSelectionSet>& Pair : ActorsWithSelectionSets.ActorsWithSet)
+					{
+						if (Pair.Value.IsSelectedInViewport())
+						{
+							if (Pair.Value.SetsThatMatch.Contains(Guid))
+							{
+								Pair.Value.ShowAll();
+							}
+						}
+					}
+				}
+			}
+			return true;
+		}
+	}
+	return false;
+}
+
+static UMovieSceneControlRigParameterTrack* GetControlRigTrack(TSharedPtr<ISequencer>& Sequencer, UControlRig* ControlRig)
+{
+	UMovieSceneSequence* OwnerSequence = Sequencer ? Sequencer->GetFocusedMovieSceneSequence() : nullptr;
+	const UMovieScene* MovieScene = OwnerSequence ? OwnerSequence->GetMovieScene() : nullptr;
+	if (!MovieScene)
+	{
+		return  nullptr;
+	}
+
+	const TArray<FMovieSceneBinding>& Bindings = MovieScene->GetBindings();
+	for (const FMovieSceneBinding& Binding : Bindings)
+	{
+		TArray<UMovieSceneTrack*> FoundTracks = MovieScene->FindTracks(UMovieSceneControlRigParameterTrack::StaticClass(), Binding.GetObjectGuid(), NAME_None);
+		for (UMovieSceneTrack* Track : FoundTracks)
+		{
+			if (UMovieSceneControlRigParameterTrack* CRTrack = Cast<UMovieSceneControlRigParameterTrack>(Track))
+			{
+				if (CRTrack->GetControlRig() == ControlRig)
+				{
+					return CRTrack;
+				}
+			}
+		}
+	}
+	return  nullptr;
+}
+
+static UMovieScene3DTransformTrack* Get3DTransformTrack(TSharedPtr<ISequencer>& Sequencer, const AActor* InActor)
+{
+	UMovieSceneSequence* OwnerSequence = Sequencer ? Sequencer->GetFocusedMovieSceneSequence() : nullptr;
+	const UMovieScene* MovieScene = OwnerSequence ? OwnerSequence->GetMovieScene() : nullptr;
+	if (!MovieScene)
+	{
+		return  nullptr;
+	}
+
+	const TArray<FMovieSceneBinding>& Bindings = MovieScene->GetBindings();
+	for (const FMovieSceneBinding& Binding : Bindings)
+	{
+		if (UMovieScene3DTransformTrack* TransformTrack = Sequencer->GetFocusedMovieSceneSequence()->GetMovieScene()->FindTrack<UMovieScene3DTransformTrack>(Binding.GetObjectGuid()))
+		{
+			for (TWeakObjectPtr<> BoundObject : Sequencer->FindBoundObjects(Binding.GetObjectGuid(), Sequencer->GetFocusedTemplateID()))
+			{
+				if (!BoundObject.IsValid())
+				{
+					continue;
+				}
+				if (AActor* BoundActor = Cast<AActor>(BoundObject.Get()))
+				{
+					if (BoundActor == InActor)
+					{
+						return TransformTrack;
+					}
+				}
+				else if (USceneComponent* BoundComponent = Cast<USceneComponent>(BoundObject.Get()))
+				{
+					if (BoundComponent->GetOwner() == InActor)
+					{
+						return TransformTrack;
+					}
+				}
+			}
+		}
+	}
+	return nullptr;
+}
+
+static bool GetStartEndIndicesForControl(UMovieSceneControlRigParameterSection* BaseSection, const FRigControlElement* ControlElement, int& OutStartIndex, int& OutEndIndex)
+{
+	if (FChannelMapInfo* pChannelIndex = BaseSection->ControlChannelMap.Find(ControlElement->GetFName()))
+	{
+		const int32 ChannelIndex = pChannelIndex->ChannelIndex;
+		switch (ControlElement->Settings.ControlType)
+		{
+
+		case ERigControlType::Float:
+		case ERigControlType::ScaleFloat:
+		{
+			OutStartIndex = ChannelIndex;
+			OutEndIndex = ChannelIndex;
+			break;
+		}
+		case ERigControlType::Vector2D:
+		{
+			OutStartIndex = ChannelIndex;
+			OutEndIndex = ChannelIndex + 1;
+			break;
+		}
+		case ERigControlType::Position:
+		case ERigControlType::Scale:
+		case ERigControlType::Rotator:
+		{
+			OutStartIndex = ChannelIndex;
+			OutEndIndex = ChannelIndex + 2;
+			break;
+		}
+
+		case ERigControlType::Transform:
+		case ERigControlType::TransformNoScale:
+		case ERigControlType::EulerTransform:
+		{
+			if (ControlElement->Settings.ControlType == ERigControlType::TransformNoScale)
+			{
+				OutStartIndex = ChannelIndex;
+				OutEndIndex = ChannelIndex + 5;
+
+			}
+			else
+			{
+				OutStartIndex = ChannelIndex;
+				OutEndIndex = ChannelIndex + 8;
+			}
+			break;
+		}
+		default:
+			break;
+		}
+		return true;
+	}
+	return false;
+}
+
+template<typename ChannelType, typename ValueType>
+void AssigneOrSetValue(ChannelType* Channel, ValueType Value, const FFrameNumber& FrameNumber, EMovieSceneKeyInterpolation DefaultInterpolation)
+{
+	TArray<FKeyHandle> KeysAtCurrentTime;
+	Channel->GetKeys(TRange<FFrameNumber>(FrameNumber), nullptr, &KeysAtCurrentTime);
+	if (KeysAtCurrentTime.Num() > 0)
+	{
+		AssignValue(Channel, KeysAtCurrentTime[0], Value);
+	}
+	else
+	{
+		EMovieSceneKeyInterpolation Interpolation = GetInterpolationMode(Channel, FrameNumber, DefaultInterpolation);
+		AddKeyToChannel(Channel, FrameNumber, Value, Interpolation);
+	}
+}
+
+template<typename ChannelType, typename ValueType>
+void SetCurrentKeys(TArrayView<ChannelType*> Channels, int32 StartIndex, int32 EndIndex, EMovieSceneKeyInterpolation DefaultInterpolation, const FFrameNumber& FrameNumber)
+{
+	ValueType Value = 0.0;
+	FFrameTime FrameTime(FrameNumber);
+	for (int32 Index = StartIndex; Index <= EndIndex; ++Index)
+	{
+		ChannelType* Channel = Channels[Index];
+		if (Channel->Evaluate(FrameTime, Value))
+		{
+			AssigneOrSetValue(Channel, Value, FrameNumber, DefaultInterpolation);
+		}
+	}
+}
+
+bool UAIESelectionSets::KeyAll(TSharedPtr<ISequencer>& InSequencer, const FGuid& InGuid) const
+{
+	if (InSequencer.IsValid() == false)
+	{
+		return false;
+	}
+
+	FControlRigEditMode* EditMode = static_cast<FControlRigEditMode*>(GLevelEditorModeTools().GetActiveMode(FControlRigEditMode::ModeName));
+	if (EditMode)
+	{
+		FFrameNumber FrameNumber = InSequencer->GetLocalTime().Time.GetFrame();
+		EMovieSceneKeyInterpolation DefaultInterpolation = InSequencer->GetKeyInterpolation();
+
+		const FScopedTransaction Transaction(LOCTEXT("KeySelectionSet_transaction", "Key Selection Set"), !GIsTransacting);
+
+		TArray<FGuid> AllGuids;
+		if (const FAIESelectionSetItem* SetItem = SelectionSets.Find(InGuid))
+		{
+			SetItem->GetAllChildren(this, AllGuids);
+			for (const FGuid& Guid : AllGuids)
+			{
+				if (const FAIESelectionSetItem* CurrentSetItem = SelectionSets.Find(Guid))
+				{
+					for (const TPair<TWeakObjectPtr<AActor>, FActorWithSelectionSet>& Pair : ActorsWithSelectionSets.ActorsWithSet)
+					{
+						if (Pair.Value.IsActive(bShowSelectedOnly))
+						{
+							if (Pair.Value.SetsThatMatch.Contains(Guid))
+							{
+								if (AActor* Key = Cast<AActor>(Pair.Key.Get()))
+								{
+									for (const TWeakObjectPtr<UObject>& Object : Pair.Value.Actors)
+									{
+										if (UControlRig* ControlRig = Cast<UControlRig>(Object.Get()))
+										{
+											if (UMovieSceneControlRigParameterTrack* CRTrack = GetControlRigTrack(InSequencer, ControlRig))
+											{
+												for (const FAIESelectionSetItemName& Name : CurrentSetItem->Names)
+												{
+													FName ControlName(*Name.Name);
+													if (UMovieSceneControlRigParameterSection* CRSection = Cast<UMovieSceneControlRigParameterSection>(CRTrack->GetSectionToKey(ControlName)))
+													{
+														CRSection->Modify();
+														TArrayView<FMovieSceneFloatChannel*> BaseFloatChannels = CRSection->GetChannelProxy().GetChannels<FMovieSceneFloatChannel>();
+														if (FRigControlElement* Control = ControlRig->FindControl(ControlName))
+														{
+															int32 StartIndex = 0, EndIndex = 0;
+															if (GetStartEndIndicesForControl(CRSection, Control, StartIndex, EndIndex))
+															{
+																SetCurrentKeys<FMovieSceneFloatChannel, float>(BaseFloatChannels, StartIndex, EndIndex, DefaultInterpolation, FrameNumber);
+
+															}
+														}
+													}
+												}
+											}
+										}
+										else if (AActor* Actor = Cast<AActor>(Object.Get()))
+										{
+											if (UMovieScene3DTransformTrack* Track = Get3DTransformTrack(InSequencer, Actor))
+											{
+												if (Track->GetAllSections().Num() > 0)
+												{
+													UMovieSceneSection* Section = Track->GetSectionToKey() ?
+														Track->GetSectionToKey() :
+														Track->GetAllSections()[0];
+													if (Section)
+													{
+														Section->Modify();
+														TArrayView<FMovieSceneDoubleChannel*> BaseDoubleChannels = Section->GetChannelProxy().GetChannels<FMovieSceneDoubleChannel>();
+														const int32 NumChannels = BaseDoubleChannels.Num();
+														const int32 StartIndex = 0;
+														const int32 EndIndex = NumChannels - 1;
+														SetCurrentKeys<FMovieSceneDoubleChannel, double>(BaseDoubleChannels, StartIndex, EndIndex, DefaultInterpolation, FrameNumber);
+													}
+												}
+											}
+										}
+									}
+								}
+							}
+						}
+					}
+				}
+			}
+			InSequencer->NotifyMovieSceneDataChanged(EMovieSceneDataChangeType::TrackValueChanged);
+
+			return true;
+		}
+	}
+	return false;
+}
+
+bool UAIESelectionSets::AddChildToItem(const FGuid& InParentGuid, const FGuid& InChildGuid)
+{
+	return false;
+}
+bool UAIESelectionSets::RemoveChildFromItem(const FGuid& InParentGuid, const FGuid& InChildGuid)
+{
+	return false;
+}
+
+void UAIESelectionSets::SelectionSetsChangedBroadcast() 
+{
+	OnSelectionSetsChanged.Broadcast(this);
+}
+
+namespace UE::AIE
+{
+	FLinearColor GetNextRandomColor(int32& NextIndex, bool bIncrement)
+	{
+		FLinearColor Color = FLinearColor::Gray;
+
+		if (const USelectionSetsSettings* Settings = GetDefault<USelectionSetsSettings>())
+		{
+			const TArray<FLinearColor> CustomColors = Settings->GetCustomColors();
+			if (CustomColors.Num() > 0)
+			{
+				if (NextIndex < 0)
+				{
+					NextIndex = 0;
+				}
+				else if (NextIndex >= CustomColors.Num())
+				{
+					NextIndex = CustomColors.Num() - 1;
+				}
+				Color = CustomColors[NextIndex];
+				if (bIncrement)
+				{
+					++NextIndex;
+					NextIndex = (NextIndex % CustomColors.Num());
+
+				}
+				else
+				{
+					--NextIndex;
+					if (NextIndex < 0)
+					{
+						NextIndex = CustomColors.Num() - 1;
+					}
+				}
+			}
+		}
+		return Color;
+	}
+}
+FLinearColor UAIESelectionSets::GetNextSelectionSetColor()
+{
+	static int32 NextIndex = 0;
+	return UE::AIE::GetNextRandomColor(NextIndex, true);
+}
+
+FLinearColor UAIESelectionSets::GetNextActorColor()
+{
+	static int32 NextIndex = MAX_int32;
+	return UE::AIE::GetNextRandomColor(NextIndex, false);
+}
+
+bool UAIESelectionSets::LoadFromJsonFile(TSharedPtr<ISequencer>& InSequencer, const FFilePath& JsonFilePath)
+{
+	FString JsonAsString;
+
+	if (!FFileHelper::LoadFileToString(JsonAsString, *JsonFilePath.FilePath))
+	{
+		UE_LOG(LogControlRig, Error, TEXT("Selection Sets: Could not load from file"));
+		return false;
+	}
+
+	return LoadFromJsonString(InSequencer, JsonAsString);
+}
+
+bool UAIESelectionSets::ExportAsJsonFile(const FFilePath& JsonFilePath) const
+{
+	FString JsonString;
+
+	if (ExportAsJsonString(JsonString))
+	{
+		return FFileHelper::SaveStringToFile(JsonString, *JsonFilePath.FilePath);
+	}
+	return false;
+}
+
+bool UAIESelectionSets::LoadFromJsonString(TSharedPtr<ISequencer>& InSequencer, const FString& JsonString)
+{
+	TSharedPtr<FJsonObject> JsonObject = MakeShareable(new FJsonObject());
+	const TSharedRef<TJsonReader<>> JsonReader = TJsonReaderFactory<>::Create(JsonString);
+
+	if (!FJsonSerializer::Deserialize(JsonReader, JsonObject) && JsonObject.IsValid() == false)
+	{
+		UE_LOG(LogControlRig, Warning, TEXT("Could not deserialize json data for selection sets"))
+		return false;
+	}
+
+	if (!JsonObject->HasTypedField(TEXT("sets"), EJson::Array))
+	{
+		UE_LOG(LogControlRig, Warning, TEXT("Selection Sets JSON: Missing sets field"));
+		return false;
+	}
+	FString ItemName, Guid, Parent, Name, MirrorName, OwnerActorName;
+	int32 Type, Duplicates, Row, Column;
+	float Red, Green, Blue, Alpha;
+	const TArray<TSharedPtr<FJsonValue>>* NamesArray;
+	//const TArray<TSharedPtr<FJsonValue>>* ChildArray; //not supported yet
+	const FScopedTransaction Transaction(LOCTEXT("ImportSelectionSetJSON_transaction", "Import Selection Set"), !GIsTransacting);
+	Modify();
+	TArray<TSharedPtr<FJsonValue>> SetsArray = JsonObject->GetArrayField(TEXT("sets"));
+	for (TSharedPtr<FJsonValue> SetValue : SetsArray)
+	{
+		TSharedPtr<FJsonObject> SetObject = SetValue->AsObject();
+		if (SetObject.IsValid())
+		{
+			FAIESelectionSetItem SetItem;
+			if (SetObject->TryGetStringField(TEXT("item_name"), ItemName))
+			{
+				SetItem.ItemName = FText::FromString(ItemName);
+			}
+			else
+			{
+				UE_LOG(LogControlRig, Warning, TEXT("Selection Sets JSON: No item_name"));
+				return false;
+			}
+			if (SetObject->TryGetStringField(TEXT("guid"), Guid))
+			{
+				SetItem.Guid = FGuid(Guid);
+			}
+			else
+			{
+				UE_LOG(LogControlRig, Warning, TEXT("Selection Sets JSON: No guid"));
+				return false;
+			}
+			if (SetObject->TryGetStringField(TEXT("parent"), Parent))
+			{
+				SetItem.Parent = FGuid(Parent);
+			}
+			else
+			{
+				UE_LOG(LogControlRig, Warning, TEXT("Selection Sets JSON: No parent"));
+				return false;
+			}
+			if (SetObject->TryGetNumberField(TEXT("row"), Row))
+			{
+				SetItem.ViewData.Row = Row;
+			}
+			else //may not exist on older
+			{
+				SetItem.ViewData.Row = 0;
+			}
+			if (SetObject->TryGetNumberField(TEXT("column"), Column))
+			{
+				SetItem.ViewData.Column = Column;
+			}
+			else //may not exist on older
+			{
+				SetItem.ViewData.Column = 0;
+			}
+			if (SetObject->TryGetNumberField(TEXT("red"), Red))
+			{
+				SetItem.ViewData.Color.R = Red;
+			}
+			else
+			{
+				UE_LOG(LogControlRig, Warning, TEXT("Selection Sets JSON: No red"));
+				return false;
+			}
+			if (SetObject->TryGetNumberField(TEXT("green"), Green))
+			{
+				SetItem.ViewData.Color.G = Green;
+			}
+			else
+			{
+				UE_LOG(LogControlRig, Warning, TEXT("Selection Sets JSON: No green"));
+				return false;
+			}
+			if (SetObject->TryGetNumberField(TEXT("blue"), Blue))
+			{
+				SetItem.ViewData.Color.B = Blue;
+			}
+			else
+			{
+				UE_LOG(LogControlRig, Warning, TEXT("Selection Sets JSON: No blue"));
+				return false;
+			}
+			if (SetObject->TryGetNumberField(TEXT("alpha"), Alpha))
+			{
+				SetItem.ViewData.Color.A = Alpha;
+			}
+			else
+			{
+				UE_LOG(LogControlRig, Warning, TEXT("Selection Sets JSON: No alpha"));
+				return false;
+			}
+			if (SetObject->TryGetArrayField(TEXT("names"), NamesArray))
+			{
+				SetItem.Names.Reserve(NamesArray->Num());
+				for (TSharedPtr<FJsonValue> NamesSource : *NamesArray)
+				{
+					TSharedPtr<FJsonObject> NamesObject = NamesSource->AsObject();
+					if (NamesObject.IsValid())
+					{
+						FAIESelectionSetItemName SetItemName;
+						if (NamesObject->TryGetStringField(TEXT("name"), Name))
+						{
+							SetItemName.Name = Name;
+						}
+						else
+						{
+							UE_LOG(LogControlRig, Warning, TEXT("Selection Sets JSON: No name"));
+							return false;
+						}
+						if (NamesObject->TryGetStringField(TEXT("owner_actor_name"), OwnerActorName))
+						{
+							SetItemName.OwnerActorName = OwnerActorName;
+						}
+						else //may not exist
+						{
+							SetItemName.OwnerActorName = FString("");
+						}
+
+						if (NamesObject->TryGetStringField(TEXT("mirror_name"), MirrorName))
+						{
+							SetItemName.MirrorName = MirrorName;
+						}
+						else
+						{
+							UE_LOG(LogControlRig, Warning, TEXT("Selection Sets JSON: No mirror name"));
+							return false;
+						}
+
+						if (NamesObject->TryGetNumberField(TEXT("type"), Type))
+						{
+							SetItemName.Type = Type;
+						}
+						else
+						{
+							UE_LOG(LogControlRig, Warning, TEXT("Selection Sets JSON: No type"));
+							return false;
+						}
+						if (NamesObject->TryGetNumberField(TEXT("duplicates"), Duplicates))
+						{
+							SetItemName.Duplicates = Duplicates;
+						}
+						else
+						{
+							UE_LOG(LogControlRig, Warning, TEXT("Selection Sets JSON: No type"));
+							return false;
+						}
+
+						SetItem.Names.Add(SetItemName);
+
+					}
+					else
+					{
+						UE_LOG(LogControlRig, Warning, TEXT("Selection Sets JSON: No names Object"));
+						return false;
+					}
+				}
+
+			}
+			else
+			{
+				UE_LOG(LogControlRig, Warning, TEXT("Selection Sets JSON: No names Array"));
+				return false;
+			}
+			SelectionSets.Add(SetItem.Guid, SetItem);
+
+		}
+		else
+		{
+			UE_LOG(LogControlRig, Warning, TEXT("Selection Sets JSON: No set Object"));
+			return false;
+		}
+
+	}
+	
+	TSet<TWeakObjectPtr<UObject>> NoObjectsAdded;
+	ActorsWithSelectionSets.SetUpFromSequencer(InSequencer, this, NoObjectsAdded);
+
+	SelectionSetsChangedBroadcast();	
+	return true;
+}
+
+
+bool UAIESelectionSets::ExportAsJsonString(FString& OutJsonString) const
+{
+	
+	const TSharedRef<TJsonWriter<>> JsonWriter = TJsonWriterFactory<TCHAR, TPrettyJsonPrintPolicy<TCHAR>>::Create(&OutJsonString);
+
+	JsonWriter->WriteObjectStart();
+	JsonWriter->WriteArrayStart(TEXT("sets"));
+	for (const TPair<FGuid, FAIESelectionSetItem>& Pair : SelectionSets)
+	{
+;
+		JsonWriter->WriteObjectStart();
+		JsonWriter->WriteValue(TEXT("item_name"), Pair.Value.ItemName.ToString());
+		JsonWriter->WriteValue(TEXT("guid"), Pair.Value.Guid.ToString());
+		JsonWriter->WriteValue(TEXT("parent"), Pair.Value.Parent.ToString());
+		JsonWriter->WriteValue(TEXT("row"), Pair.Value.ViewData.Row);
+		JsonWriter->WriteValue(TEXT("column"), Pair.Value.ViewData.Column);
+		JsonWriter->WriteValue(TEXT("red"), Pair.Value.ViewData.Color.R);
+		JsonWriter->WriteValue(TEXT("green"), Pair.Value.ViewData.Color.G);
+		JsonWriter->WriteValue(TEXT("blue"), Pair.Value.ViewData.Color.B);
+		JsonWriter->WriteValue(TEXT("alpha"), Pair.Value.ViewData.Color.A);
+
+		JsonWriter->WriteArrayStart(TEXT("children"));
+		for (const FGuid& ChildGuid : Pair.Value.Children)
+		{
+			JsonWriter->WriteObjectStart();
+			JsonWriter->WriteValue(TEXT("child"), ChildGuid.ToString());
+			JsonWriter->WriteObjectEnd();
+		}
+		JsonWriter->WriteArrayEnd();
+		//JsonWriter->WriteValue(TEXT("children"), Pair.Value.Children);
+		JsonWriter->WriteArrayStart(TEXT("names"));
+		for (const FAIESelectionSetItemName& Name : Pair.Value.Names)
+		{
+			JsonWriter->WriteObjectStart();
+			JsonWriter->WriteValue(TEXT("name"), Name.Name);
+			if (Name.OwnerActorName.Len() > 0)
+			{
+				JsonWriter->WriteValue(TEXT("owner_actor_name"), Name.OwnerActorName);
+			}
+			JsonWriter->WriteValue(TEXT("mirror_name"), Name.MirrorName);
+			JsonWriter->WriteValue(TEXT("type"), (int32)Name.Type);
+			JsonWriter->WriteValue(TEXT("duplicates"), Name.Duplicates);
+			JsonWriter->WriteObjectEnd();
+		}
+		JsonWriter->WriteArrayEnd();
+		JsonWriter->WriteObjectEnd();
+
+	}
+	JsonWriter->WriteArrayEnd();
+	JsonWriter->WriteObjectEnd();
+
+	
+	JsonWriter->Close();
+	return true;
+
+}
+
+#undef LOCTEXT_NAMESPACE
+
+
diff --git a/Engine/Plugins/Animation/ControlRig/Source/ControlRigEditor/Private/Sequencer/SelectionSets/SelectionSets.h b/Engine/Plugins/Animation/ControlRig/Source/ControlRigEditor/Private/Sequencer/SelectionSets/SelectionSets.h
new file mode 100644
--- /dev/null
+++ b/Engine/Plugins/Animation/ControlRig/Source/ControlRigEditor/Private/Sequencer/SelectionSets/SelectionSets.h
@@ -0,0 +1,203 @@
+// Copyright Epic Games, Inc. All Rights Reserved.
+#pragma once
+#include "CoreMinimal.h"
+#include "UObject/Object.h"
+#include "Engine/EngineBaseTypes.h"
+#include "Engine/AssetUserData.h"
+
+#include "SelectionSets.generated.h"
+
+
+class ISequencer;
+class ULevelSequence;
+class UControlRig;
+class UAIESelectionSets;
+
+
+//Type the selection set item
+UENUM()
+enum EAIESelectionSetItemType : int32
+{
+	ControlRig = 0x0,
+	Actor = 0x1,
+	Component = 0x2,
+};
+ENUM_CLASS_FLAGS(EAIESelectionSetItemType)
+
+USTRUCT(BlueprintType)
+struct FAIESelectionSetItemName
+{
+	GENERATED_BODY();
+	UPROPERTY(BlueprintReadWrite, EditAnywhere, Category = "Data")
+	FString Name;
+	UPROPERTY(BlueprintReadWrite, EditAnywhere, Category = "Data")
+	FString MirrorName;
+	UPROPERTY(BlueprintReadOnly, EditAnywhere, Category = "Data", meta = (Bitmask, BitmaskEnum = "/Script/ControlRigEditor.EAIESelectionSetItemType"))
+	int32 Type;
+	UPROPERTY(BlueprintReadOnly, EditAnywhere, Category = "Data")
+	int32 Duplicates = 0;
+	UPROPERTY(BlueprintReadWrite, EditAnywhere, Category = "Data")
+	FString OwnerActorName;
+};
+
+USTRUCT(BlueprintType)
+struct FAIESelectionSetItemViewData
+{
+	GENERATED_BODY();
+
+	UPROPERTY(BlueprintReadWrite, EditAnywhere, Category = "Data")
+	FLinearColor Color;
+
+	UPROPERTY()
+	int32 Row;
+	UPROPERTY()
+	int32 Column;
+};
+
+USTRUCT(BlueprintType)
+struct FAIESelectionSetItem
+{
+	GENERATED_BODY();
+
+	//unique ID
+	UPROPERTY()
+	FGuid Guid;
+	
+	UPROPERTY(BlueprintReadWrite, EditAnywhere, Category = "Data")
+	FText ItemName;
+
+	UPROPERTY(BlueprintReadWrite, EditAnywhere, Category = "Data")
+	TArray<FAIESelectionSetItemName> Names;
+	UPROPERTY(BlueprintReadWrite, EditAnywhere, Category = "Data")
+	FAIESelectionSetItemViewData ViewData;
+	UPROPERTY(BlueprintReadWrite, EditAnywhere, Category = "Hierarchy")
+	FGuid Parent;
+	UPROPERTY(BlueprintReadWrite, EditAnywhere, Category = "Hierarchy")
+	TArray<FGuid> Children;
+
+	
+	bool ContainsItem(const FString& Name) const;
+	FAIESelectionSetItemName* GetItem(const FString& Name);
+	bool RemoveItem(const FString& Name);
+
+	void GetAllChildren(const UAIESelectionSets* SelectionSets, TArray<FGuid>& AllChildren) const;
+	static FString GetMirror(const FString& InString);
+};
+
+
+struct FActorWithSelectionSet
+{
+	TSet<TWeakObjectPtr<UObject>> Actors; //actor or crs
+	TSet<FGuid> SetsThatMatch;
+
+	bool IsSelectedInViewport(bool bOnlyInViewport = false) const;
+	bool IsControlRig() const;
+	bool IsActive(bool bShowSelectedOnly) const;
+	void ShowAll() const;
+	//
+	FText Name;
+	FLinearColor Color;
+	void SetSelected(bool bInSelected) { bIsSelected = bInSelected; }
+	bool GetSelected() const {return bIsSelected;}
+
+	//if the set contains owned actors but they came from a different control rig we mark that. This is done in SetUpFromSequencer
+	void SetUpOwnedActors(const TWeakObjectPtr<AActor>& ActorOwner, TSharedPtr<ISequencer>& InSequencer, UAIESelectionSets* SelectionSets);
+	TMap <FGuid,FString> OwnedActors;
+
+	// it was previously selected but no more
+	mutable bool bWasSelectedInViewport = false;
+private:
+	bool bIsSelected = false;
+public:
+	static AActor* GetControlRigActor(const UControlRig* InControlRig);
+
+};
+
+//set of control rigs and actors that may be pasted onto
+struct FActorsWithSelectionSets
+{
+	TMap<TWeakObjectPtr<AActor>, FActorWithSelectionSet> ActorsWithSet;
+	
+	void SetUpFromSequencer(TSharedPtr<ISequencer>& InSequencer,  UAIESelectionSets* SelectionSets, const TSet<TWeakObjectPtr<UObject>>& MakeSelectedIfPresent);
+};
+
+
+/** Link To Set of Anim Sequences that we may be linked to.*/
+
+UCLASS(BlueprintType, MinimalAPI)
+class UAIESelectionSets : public UAssetUserData
+{
+	GENERATED_UCLASS_BODY()
+
+public:
+
+	virtual bool IsEditorOnly() const override { return true; }
+
+	UPROPERTY(BlueprintReadWrite, Category = SelectionSet)
+	TMap<FGuid, FAIESelectionSetItem> SelectionSets;
+	
+	UPROPERTY(BlueprintReadWrite, Category = SelectionSet)
+	bool bShowSelectedOnly = true;
+
+	FActorsWithSelectionSets  ActorsWithSelectionSets;
+
+		
+	//set of selection sets that are active
+	TArray<FGuid> GetActiveSelectionSets();
+
+	FGuid CreateSetItemFromSelection(TSharedPtr<ISequencer>& InSequencer);
+	FGuid CreateMirror(TSharedPtr<ISequencer>& InSequencer, const FGuid& InGuid);
+	bool LoadFromJsonFile(TSharedPtr<ISequencer>& InSequencer, const FFilePath& JsonFilePath);
+	bool ExportAsJsonFile(const FFilePath& JsonFilePath)  const;
+	bool LoadFromJsonString(TSharedPtr<ISequencer>& InSequencer, const FString& JsonString);
+	bool ExportAsJsonString(FString& OutJsonString) const;
+
+	void SequencerBindingsAdded(TSharedPtr<ISequencer>& InSequencer);
+	bool AddSelectionToSetItem(const FGuid& InGuid, TSharedPtr<ISequencer>& InSequencer);
+	bool RemoveSelectionFromSetItem(const FGuid& InGuid, TSharedPtr<ISequencer>& InSequencer);
+	bool RenameSetItem(const FGuid& InGuid,const  FText& NewName);
+	bool GetItemName(const FGuid& InGuid, FText& OutName) const;
+	bool DeleteSetItem(const FGuid& InGuid, TSharedPtr<ISequencer>& InSequencer);
+	bool GetItemColor(const FGuid& InGuid, FLinearColor& OutColor) const;
+	bool SetItemColor(const FGuid& InGuid, const FLinearColor& InColor);
+
+	bool SelectItem(const FGuid& InGuid, bool bDoMirror, bool bAdd, bool bToggle,bool bSelect = true) const;
+
+	bool AddChildToItem(const FGuid& InParentGuid, const FGuid& InChildGuid);
+	bool RemoveChildFromItem(const FGuid& InParentGuid, const FGuid& InChildGuid);
+
+	bool ShowOrHideControls(const FGuid& InGuid, bool bShow, bool bDoMirror) const;
+	bool IsolateControls(const FGuid& InGuid) const;
+	bool ShowAllControls(const FGuid& InGuid) const;
+
+	bool KeyAll(TSharedPtr<ISequencer>& InSequencer, const FGuid& InGuid) const;
+
+	//Get updated visibility for the actor selections, using settings and properties
+	bool IsVisible(const FActorWithSelectionSet& ActorWithSelectionSet) const;
+
+	static UAIESelectionSets* GetSelectionSets(TSharedPtr<ISequencer>& InSequencer, bool bAddIfDoesNotExist = false);
+	static UAIESelectionSets* GetSelectionSets(ULevelSequence* InLevelSequence, bool bAddIfDoesNotExist = false);
+	static 	TSharedPtr<ISequencer> GetSequencerFromAsset();
+
+	
+public:
+	DECLARE_EVENT_OneParam(UAIESelectionSets, FSelectionSetsChanged, UAIESelectionSets*);
+
+	FSelectionSetsChanged& SelectionListChanged() { return OnSelectionSetsChanged; };
+
+
+	void SelectionSetsChangedBroadcast();
+	FSelectionSetsChanged OnSelectionSetsChanged;
+
+	FLinearColor GetNextSelectionSetColor();
+	FLinearColor GetNextActorColor();
+
+
+	friend struct FActorsWithSelectionSets;
+
+private:
+	TArray<FGuid> SortSelectionSetsByRow(TArray<FGuid>& InSelectionSets);
+private:
+	//to handle selecting last selected object when none are selected
+	mutable TSet<FGuid> LastActiveSelectionSetsSet;
+};
diff --git a/Engine/Plugins/Animation/ControlRig/Source/ControlRigEditor/Private/Sequencer/SelectionSets/SelectionSetsOverlayManager.cpp b/Engine/Plugins/Animation/ControlRig/Source/ControlRigEditor/Private/Sequencer/SelectionSets/SelectionSetsOverlayManager.cpp
new file mode 100644
--- /dev/null
+++ b/Engine/Plugins/Animation/ControlRig/Source/ControlRigEditor/Private/Sequencer/SelectionSets/SelectionSetsOverlayManager.cpp
@@ -0,0 +1,44 @@
+﻿// Copyright Epic Games, Inc. All Rights Reserved.
+
+#include "SelectionSetsOverlayManager.h"
+#include "SSelectionSets.h"
+#include "EditMode/ControlRigEditModeCommands.h"
+#include "EditMode/ControlRigEditModeSettings.h"
+#include "Misc/Flyout/Host/ToolkitBasedWidgetHost.h"
+#include "Subsystems/AssetEditorSubsystem.h"
+#include "Toolkits/IToolkitHost.h"
+
+namespace UE::AIE
+{
+namespace Private
+{
+static void SaveWidgetState(FToolWidget_FlyoutSavedState State)
+{
+	UControlRigEditModeSettings* ControlRigEditModeSettings = GetMutableDefault<UControlRigEditModeSettings>();
+	ControlRigEditModeSettings->LastSelectionSetOverlayState = State;
+	ControlRigEditModeSettings->SaveConfig();
+}
+}
+	
+FSelectionSetsOverlayManager::FSelectionSetsOverlayManager(
+	const TSharedRef<IToolkitHost>& InToolkitHost,
+	const TSharedRef<FUICommandList>& InToolkitCommandList,
+	const TSharedRef<FControlRigEditMode>& InOwningEditMode
+	)
+	: SelectionSetWidget(SNew(SSelectionSets, *InOwningEditMode)),
+	  FlyoutWidgetManager(
+		UE::ControlRigEditor::FFlyoutWidgetArgs(
+			SelectionSetWidget.ToSharedRef(),
+			TAttribute<TSharedPtr<FUICommandList>>::CreateLambda([WeakList = InToolkitCommandList.ToWeakPtr()]{ return WeakList.Pin(); }),
+			MakeShared<UE::ControlRigEditor::FToolkitBasedWidgetHost>(InToolkitHost)
+			)
+			.SetStateToRestoreFrom(GetMutableDefault<UControlRigEditModeSettings>()->LastSelectionSetOverlayState)
+			.SetSaveStateDelegate(UE::ControlRigEditor::FSaveFlyoutState::CreateStatic(&UE::AIE::Private::SaveWidgetState))
+			.SetToggleVisibility(FControlRigEditModeCommands::Get().ToggleSelectionSetsWidget)
+			.SetSummonToCursor(FControlRigEditModeCommands::Get().SummonSelectionSetsWidget)
+		)
+{
+}
+
+
+}
diff --git a/Engine/Plugins/Animation/ControlRig/Source/ControlRigEditor/Private/Sequencer/SelectionSets/SelectionSetsOverlayManager.h b/Engine/Plugins/Animation/ControlRig/Source/ControlRigEditor/Private/Sequencer/SelectionSets/SelectionSetsOverlayManager.h
new file mode 100644
--- /dev/null
+++ b/Engine/Plugins/Animation/ControlRig/Source/ControlRigEditor/Private/Sequencer/SelectionSets/SelectionSetsOverlayManager.h
@@ -0,0 +1,38 @@
+﻿// Copyright Epic Games, Inc. All Rights Reserved.
+
+#pragma once
+
+#include "EditMode/ControlRigBaseDockableView.h"
+#include "Misc/Flyout/FlyoutOverlayManager.h"
+
+namespace UE::AIE
+{
+/**
+ * Manages all logic around the selectio set overlay widget.
+ 
+ */
+class FSelectionSetsOverlayManager : public FNoncopyable
+{
+public:
+	
+	explicit FSelectionSetsOverlayManager(
+		const TSharedRef<IToolkitHost>& InToolkitHost,
+		const TSharedRef<FUICommandList>& InToolkitCommandList,
+		const TSharedRef<FControlRigEditMode>& InOwningEditMode
+		);
+
+	void ToggleVisibility() { FlyoutWidgetManager.ToggleVisibility(); }
+	void ShowWidget() { FlyoutWidgetManager.ShowWidget(); }
+	void HideWidget() { FlyoutWidgetManager.HideWidget(); }
+	void DestroyWidget() { FlyoutWidgetManager.DestroyWidget(); }
+	bool IsShowingWidget() const { return FlyoutWidgetManager.IsShowingWidget(); }
+	
+
+private:
+	
+	TSharedPtr<SWidget> SelectionSetWidget;
+
+	UE::ControlRigEditor::FFlyoutOverlayManager FlyoutWidgetManager;
+	
+};
+}
diff --git a/Engine/Plugins/Animation/ControlRig/Source/ControlRigEditor/Private/Sequencer/SelectionSets/SelectionSetsSettings.cpp b/Engine/Plugins/Animation/ControlRig/Source/ControlRigEditor/Private/Sequencer/SelectionSets/SelectionSetsSettings.cpp
new file mode 100644
--- /dev/null
+++ b/Engine/Plugins/Animation/ControlRig/Source/ControlRigEditor/Private/Sequencer/SelectionSets/SelectionSetsSettings.cpp
@@ -0,0 +1,48 @@
+// Copyright Epic Games, Inc. All Rights Reserved.
+
+#include "SelectionSetsSettings.h"
+
+#include UE_INLINE_GENERATED_CPP_BY_NAME(SelectionSetsSettings)
+
+USelectionSetsSettings::USelectionSetsSettings()
+{
+	CategoryName = TEXT("General");
+	SectionName = TEXT("Animation Selection Sets");
+
+	SelectionSetsLayout = ESelectionSetsLayout::Vertical;
+
+	SelectionSetsActorSetVisibility = ESelectionSetsActorSetVisibility::MultipleControRigs;
+
+
+	if (CustomColors.Num() == 0)
+	{
+		CustomColors.Add(FLinearColor(.904, .323, .539)); //pastel red
+		CustomColors.Add(FLinearColor(.552, .737, .328)); //pastel green
+		CustomColors.Add(FLinearColor(.947, .418, .219)); //pastel orange
+		CustomColors.Add(FLinearColor(.156, .624, .921)); //pastel blue
+		CustomColors.Add(FLinearColor(.921, .314, .337)); //pastel red 2
+		CustomColors.Add(FLinearColor(.361, .651, .332)); //pastel green 2
+		CustomColors.Add(FLinearColor(.982, .565, .254)); //pastel orange 2
+		CustomColors.Add(FLinearColor(.246, .223, .514)); //pastel purple
+		CustomColors.Add(FLinearColor(.208, .386, .687)); //pastel blue2
+		CustomColors.Add(FLinearColor(.223, .590, .337)); //pastel green 3
+		CustomColors.Add(FLinearColor(.230, .291, .591)); //pastel blue 3
+	}
+}
+
+void USelectionSetsSettings::PostEditChangeProperty(FPropertyChangedEvent& PropertyChangedEvent)
+{
+	FName PropertyName = (PropertyChangedEvent.Property != NULL) ? PropertyChangedEvent.Property->GetFName() : NAME_None;
+	FName MemberPropertyName = (PropertyChangedEvent.MemberProperty != NULL) ? PropertyChangedEvent.MemberProperty->GetFName() : NAME_None;
+	if (PropertyName == GET_MEMBER_NAME_CHECKED(USelectionSetsSettings, SelectionSetsLayout) ||
+		MemberPropertyName == GET_MEMBER_NAME_CHECKED(USelectionSetsSettings, SelectionSetsLayout))
+	{
+		OnLayoutChanged.Broadcast();
+	}
+	else if (PropertyName == GET_MEMBER_NAME_CHECKED(USelectionSetsSettings, SelectionSetsActorSetVisibility) ||
+		MemberPropertyName == GET_MEMBER_NAME_CHECKED(USelectionSetsSettings, SelectionSetsActorSetVisibility))
+	{
+		OnActorSetVisibilityChanged.Broadcast();
+	}
+	Super::PostEditChangeProperty(PropertyChangedEvent);
+}
\ No newline at end of file
diff --git a/Engine/Plugins/Animation/ControlRig/Source/ControlRigEditor/Private/Sequencer/SelectionSets/SelectionSetsSettings.h b/Engine/Plugins/Animation/ControlRig/Source/ControlRigEditor/Private/Sequencer/SelectionSets/SelectionSetsSettings.h
new file mode 100644
--- /dev/null
+++ b/Engine/Plugins/Animation/ControlRig/Source/ControlRigEditor/Private/Sequencer/SelectionSets/SelectionSetsSettings.h
@@ -0,0 +1,78 @@
+// Copyright Epic Games, Inc. All Rights Reserved.
+
+#pragma once
+
+#include "Containers/UnrealString.h"
+#include "Engine/DeveloperSettings.h"
+#include "Containers/EnumAsByte.h"
+#include "Engine/DeveloperSettings.h"
+#include "UObject/Object.h"
+#include "Math/Color.h"
+#include "Containers/Array.h"
+
+#include "SelectionSetsSettings.generated.h"
+
+
+#define UE_API CONTROLRIGEDITOR_API
+
+UENUM()
+enum class ESelectionSetsLayout : uint8
+{
+	Vertical	  UMETA(DisplayName = "Vertical"),
+	Horizontal	  UMETA(DisplayName = "Horizontal"),
+};
+
+UENUM()
+enum class ESelectionSetsActorSetVisibility : uint8
+{
+	AllVisible			 UMETA(DisplayName = "All Visible"),
+	MultipleControRigs	 UMETA(DisplayName = "Multiple Control Rigs"),
+};
+
+
+/** Serializable options for animation mode selection sets. */
+UCLASS(MinimalAPI, config=EditorPerProjectUserSettings)
+class USelectionSetsSettings : public UDeveloperSettings
+{
+public:
+	GENERATED_BODY()
+
+	DECLARE_MULTICAST_DELEGATE(FOnLayoutChanged);
+	DECLARE_MULTICAST_DELEGATE(FOnActorSetVisibilityChanged);
+
+	UE_API USelectionSetsSettings();
+
+	const TArray<FLinearColor>& GetCustomColors() const { return CustomColors; }
+
+	ESelectionSetsLayout GetSelectionSetsLayout() const { return SelectionSetsLayout; }
+
+	ESelectionSetsActorSetVisibility GetSelectionSetsActorSetVisibility() const { return SelectionSetsActorSetVisibility;  }
+
+
+	FOnLayoutChanged& GetOnLayoutChanged() { return OnLayoutChanged; }
+	FOnActorSetVisibilityChanged& GetOActorSetVisibilityChanged() { return OnActorSetVisibilityChanged; }
+
+protected:
+
+	/* Whether layout is horizontal or vertical. Need to re-enter animation mode for change to take place */
+	UPROPERTY(config, EditAnywhere, Category = "Layout")
+	ESelectionSetsLayout SelectionSetsLayout;
+
+	/* Whether or not we always show valid actors, or only show them when multiple control rigs match sets. Need to re-enter animation mode for change to take place */
+	UPROPERTY(config, EditAnywhere, Category = "Layout")
+	ESelectionSetsActorSetVisibility SelectionSetsActorSetVisibility;
+
+	/** Set of colors to use for selection sets*/
+	UPROPERTY(config, EditAnywhere, Category="Colors")
+	TArray<FLinearColor> CustomColors;
+
+	UE_API virtual void PostEditChangeProperty(FPropertyChangedEvent& PropertyChangedEvent) override;
+
+private:
+
+	FOnLayoutChanged OnLayoutChanged;
+	FOnActorSetVisibilityChanged OnActorSetVisibilityChanged;
+
+};
+
+#undef UE_API
diff --git a/Engine/Plugins/Animation/ControlRig/Source/ControlRigEditor/Private/Tools/EvaluateSequencerTools.cpp b/Engine/Plugins/Animation/ControlRig/Source/ControlRigEditor/Private/Tools/EvaluateSequencerTools.cpp
--- a/Engine/Plugins/Animation/ControlRig/Source/ControlRigEditor/Private/Tools/EvaluateSequencerTools.cpp
+++ b/Engine/Plugins/Animation/ControlRig/Source/ControlRigEditor/Private/Tools/EvaluateSequencerTools.cpp
@@ -31,6 +31,14 @@
 #include "Components/SkeletalMeshComponent.h"
 #include "Transform/TransformConstraintUtil.h"
 #include "EditMode/ControlRigEditMode.h"
+#include "MVVM/Selection/Selection.h"
+#include "MVVM/TrackRowModelStorageExtension.h"
+#include "MVVM/ViewModels/TrackModel.h"
+#include "MVVM/ViewModels/SectionModel.h"
+#include "MVVM/ViewModels/TrackRowModel.h"
+#include "MVVM/ViewModels/SequencerEditorViewModel.h"
+#include "MVVM/ViewModels/SequencerOutlinerViewModel.h"
+#include "MVVM/ViewModels/ObjectBindingModel.h"
 
 class UControlRig;
 
@@ -750,6 +758,99 @@ namespace AIE
 		}
 	}
 
+	TArray<FGuid> FSequencerSelected::GetSelectedOutlinerGuids(ISequencer* SequencerPtr)
+	{
+		using namespace UE::Sequencer;
+		TArray<FGuid> SelectedObjects;// = SequencerPtr->GetViewModel()->GetSelection()->GetBoundObjectsGuids(); //this is not exposed :(
+
+		for (TViewModelPtr<FObjectBindingModel> ObjectBindingNode : SequencerPtr->GetViewModel()->GetSelection()->Outliner.Filter<FObjectBindingModel>())
+		{
+			FGuid Guid = ObjectBindingNode->GetObjectGuid();
+			SelectedObjects.Add(Guid);
+		}
+		return SelectedObjects;
+	}
+
+	void FSequencerSelected::GetSelectedControlRigsAndBoundObjects(ISequencer* SequencerPtr, TArray<FControlRigAndControlsAndTrack>& OutSelectedCRs, TArray<FObjectAndTrack>& OutBoundObjects)
+	{
+		if (SequencerPtr == nullptr || SequencerPtr->GetViewModel() == nullptr)
+		{
+			return;
+		}
+		ULevelSequence* LevelSequence = Cast<ULevelSequence>(SequencerPtr->GetFocusedMovieSceneSequence());
+		if (LevelSequence == nullptr)
+		{
+			return;
+		}
+
+		using namespace UE::Sequencer;
+		TArray<FGuid> SelectedObjects = GetSelectedOutlinerGuids(SequencerPtr);;// = SequencerPtr->GetViewModel()->GetSelection()->GetBoundObjectsGuids();
+
+		const UMovieScene* MovieScene = LevelSequence->GetMovieScene();
+		if (MovieScene)
+		{
+			const TArray<FMovieSceneBinding>& Bindings = MovieScene->GetBindings();
+			for (const FMovieSceneBinding& Binding : Bindings)
+			{
+				bool bHaveControlRig = false;
+				TArray<UMovieSceneTrack*> CRTracks = MovieScene->FindTracks(UMovieSceneControlRigParameterTrack::StaticClass(), Binding.GetObjectGuid(), NAME_None);
+				for (UMovieSceneTrack* AnyOleTrack : CRTracks)
+				{
+					if (UMovieSceneControlRigParameterTrack* Track = Cast<UMovieSceneControlRigParameterTrack>(AnyOleTrack))
+					{
+						if (UControlRig* ControlRig = Track->GetControlRig())
+						{
+							FControlRigAndControlsAndTrack CRControls;
+							CRControls.ControlRig = ControlRig;
+							CRControls.Track = Track;
+							CRControls.Controls = ControlRig->CurrentControlSelection();
+							if (CRControls.Controls.Num() > 0)
+							{
+								bHaveControlRig = true;
+								OutSelectedCRs.Add(CRControls);
+							}
+						}
+					}
+				}
+				//if we have control rig controls don't add the base skel mesh for now
+				if (bHaveControlRig == false && SelectedObjects.Contains(Binding.GetObjectGuid()))
+				{
+					TArray<UMovieSceneTrack*> Tracks = MovieScene->FindTracks(UMovieScenePropertyTrack::StaticClass(), Binding.GetObjectGuid(), NAME_None);
+					if (Tracks.Num() < 1)
+					{
+						continue;
+					}
+					const FMovieSceneSequenceID& SequenceID = SequencerPtr->GetEvaluationState()->FindSequenceId(LevelSequence);
+
+					for (UMovieSceneTrack* AnyOleTrack : Tracks)
+					{
+						if (UMovieScenePropertyTrack* Track = Cast<UMovieScenePropertyTrack>(AnyOleTrack))
+						{
+							const FMovieSceneBlendTypeField SupportedBlendTypes = Track->GetSupportedBlendTypes();
+							if (SupportedBlendTypes.Num() == 0)
+							{
+								continue;
+							}
+							const TArrayView<TWeakObjectPtr<UObject>>& BoundObjects = SequencerPtr->FindBoundObjects(Binding.GetObjectGuid(), SequenceID);
+							for (const TWeakObjectPtr<UObject> CurrentBoundObject : BoundObjects)
+							{
+								if (UObject* BoundObject = CurrentBoundObject.Get())
+								{
+									FObjectAndTrack ObjectAndTrack;
+									ObjectAndTrack.BoundObject = BoundObject;
+									ObjectAndTrack.Track = Track;
+									ObjectAndTrack.SequencerGuid = Binding.GetObjectGuid();
+									OutBoundObjects.Add(ObjectAndTrack);
+								}
+							}
+						}
+					}
+				}
+			}
+		}
+	}
+
+
 } // namespace AIE
 } // namespace UE
 
diff --git a/Engine/Plugins/Animation/ControlRig/Source/ControlRigEditor/Public/EditMode/ControlRigEditMode.h b/Engine/Plugins/Animation/ControlRig/Source/ControlRigEditor/Public/EditMode/ControlRigEditMode.h
--- a/Engine/Plugins/Animation/ControlRig/Source/ControlRigEditor/Public/EditMode/ControlRigEditMode.h
+++ b/Engine/Plugins/Animation/ControlRig/Source/ControlRigEditor/Public/EditMode/ControlRigEditMode.h
@@ -274,7 +274,9 @@ public:
 		return FRigElementTypeHelper::ToMask(ERigElementType::Control);
 	}
 
-
+	//show or hide the specified Controls on this Control Rig;
+	UE_API void ShowControlRigControls(UControlRig* InControlRig,const TSet<FString>& Names, bool bVal);
+	UE_API void ClearOutHidden(UControlRig* InControlRig);
 
 protected:
 
@@ -520,6 +522,7 @@ protected:
 	/** Shape actors */
 	TMap<TWeakObjectPtr<UControlRig>,TArray<TObjectPtr<AControlRigShapeActor>>> ControlRigShapeActors;
 
+
 	/** Manager for anim details proxies */
 	TObjectPtr<UAnimDetailsProxyManager> AnimDetailsProxyManager;
 
diff --git a/Engine/Plugins/Animation/ControlRig/Source/ControlRigEditor/Public/Tools/EvaluateSequencerTools.h b/Engine/Plugins/Animation/ControlRig/Source/ControlRigEditor/Public/Tools/EvaluateSequencerTools.h
--- a/Engine/Plugins/Animation/ControlRig/Source/ControlRigEditor/Public/Tools/EvaluateSequencerTools.h
+++ b/Engine/Plugins/Animation/ControlRig/Source/ControlRigEditor/Public/Tools/EvaluateSequencerTools.h
@@ -4,6 +4,7 @@
 
 #include "ActorForWorldTransforms.h"
 #include "MovieSceneBinding.h"
+#include "Misc/Guid.h"
 
 #define UE_API CONTROLRIGEDITOR_API
 
@@ -18,6 +19,9 @@ class UMovieSceneTrack;
 class UMovieScene;
 struct FGuid;
 class UMovieScene3DTransformSection;
+class UMovieScenePropertyTrack;
+class UMovieSceneControlRigParameterTrack;
+
 namespace UE
 {
 namespace AIE
@@ -257,8 +261,26 @@ struct FEvalHelpers
 	static UE_API bool EvaluateSequencer(UWorld* World, ISequencer* Sequencer, const TSet<const UMovieSceneTrack*>& RelevantTracks);
 };
 
+struct FControlRigAndControlsAndTrack
+{
+	UMovieSceneControlRigParameterTrack* Track;
+	UControlRig* ControlRig;
+	TArray<FName> Controls;
+};
 
+struct FObjectAndTrack
+{
+	UMovieScenePropertyTrack* Track;
+	UObject* BoundObject;
+	FGuid SequencerGuid;
+};
 
+//helper to get selected items in sequencer, control rigs and properties
+struct FSequencerSelected
+{
+	static UE_API TArray<FGuid> GetSelectedOutlinerGuids(ISequencer* SequencerPtr);
+	static UE_API void GetSelectedControlRigsAndBoundObjects(ISequencer* SequencerPtr, TArray<FControlRigAndControlsAndTrack>& OutSelectedCRs, TArray<FObjectAndTrack>& OutBoundObjects);
+};
 
 } // namespace AIE
 } // namespace UE
diff --git a/Engine/Source/Runtime/SlateCore/Private/Styling/StarshipCoreStyle.cpp b/Engine/Source/Runtime/SlateCore/Private/Styling/StarshipCoreStyle.cpp
--- a/Engine/Source/Runtime/SlateCore/Private/Styling/StarshipCoreStyle.cpp
+++ b/Engine/Source/Runtime/SlateCore/Private/Styling/StarshipCoreStyle.cpp
@@ -349,6 +349,7 @@ TSharedRef<ISlateStyle> FStarshipCoreStyle::Create()
 		Style->Set("Icons.SaveModified", new IMAGE_BRUSH_SVG("Starship/Common/save-modified", Icon16x16));
 		Style->Set("Icons.Favorites.Small", new IMAGE_BRUSH_SVG("Starship/Common/Favorite", Icon20x20));
 
+		Style->Set("Icons.Export", new IMAGE_BRUSH_SVG("Starship/Common/export", Icon16x16));
 		Style->Set("Icons.Import", new IMAGE_BRUSH_SVG("Starship/Common/import", Icon16x16));
 		Style->Set("Icons.Reimport", new IMAGE_BRUSH_SVG("Starship/Common/reimport", Icon16x16));
 		Style->Set("Icons.Filter", new IMAGE_BRUSH_SVG("Starship/Common/filter", Icon16x16));
