CIS focus for bogus Mac Component issue.
see https://epicgames.slack.com/archives/C0448GCBM5L/p1757196536711159?thread_ts=1757196536.256029&cid=C0448GCBM5L

#DisabledValidations Tag, Preflight

diff --git a/Engine/Plugins/Animation/ControlRig/Source/ControlRigEditor/Private/Sequencer/SelectionSets/SelectionSets.h b/Engine/Plugins/Animation/ControlRig/Source/ControlRigEditor/Private/Sequencer/SelectionSets/SelectionSets.h
--- a/Engine/Plugins/Animation/ControlRig/Source/ControlRigEditor/Private/Sequencer/SelectionSets/SelectionSets.h
+++ b/Engine/Plugins/Animation/ControlRig/Source/ControlRigEditor/Private/Sequencer/SelectionSets/SelectionSets.h
@@ -20,7 +20,7 @@ enum EAIESelectionSetItemType : int32
 {
 	ControlRig = 0x0,
 	Actor = 0x1,
-	Component = 0x2,
+//caussing mac issues not used yet so removing	Component = 0x2,
 };
 ENUM_CLASS_FLAGS(EAIESelectionSetItemType)
 
