Anim Details: fix for ensures, use proper member name

diff --git a/Engine/Plugins/Animation/ControlRig/Source/ControlRigEditor/Private/AnimDetails/Proxies/AnimDetailsProxyLocation.cpp b/Engine/Plugins/Animation/ControlRig/Source/ControlRigEditor/Private/AnimDetails/Proxies/AnimDetailsProxyLocation.cpp
--- a/Engine/Plugins/Animation/ControlRig/Source/ControlRigEditor/Private/AnimDetails/Proxies/AnimDetailsProxyLocation.cpp
+++ b/Engine/Plugins/Animation/ControlRig/Source/ControlRigEditor/Private/AnimDetails/Proxies/AnimDetailsProxyLocation.cpp
@@ -130,7 +130,8 @@ void UAnimDetailsProxyLocation::AdoptValues(const ERigControlValueType RigContro
 void UAnimDetailsProxyLocation::ResetPropertyToDefault(const FName& PropertyName)
 {
 	using namespace UE::ControlRigEditor;
-	if (PropertyName == FAnimDetailsLocation::StaticStruct()->GetFName())
+	const FName StructName = GET_MEMBER_NAME_CHECKED(UAnimDetailsProxyLocation, Location);
+	if (PropertyName == StructName)
 	{
 		Location = LocationUtils::GetDefaultLocation(*this);
 	}
@@ -159,7 +160,8 @@ bool UAnimDetailsProxyLocation::HasDefaultValue(const FName& PropertyName) const
 
 	const FAnimDetailsLocation DefaultLocation = LocationUtils::GetDefaultLocation(*this);
 
-	if (PropertyName == FAnimDetailsLocation::StaticStruct()->GetFName())
+	const FName StructName = GET_MEMBER_NAME_CHECKED(UAnimDetailsProxyLocation, Location);
+	if (PropertyName == StructName)
 	{
 		return
 			Location.LX == DefaultLocation.LX &&
diff --git a/Engine/Plugins/Animation/ControlRig/Source/ControlRigEditor/Private/AnimDetails/Proxies/AnimDetailsProxyRotation.cpp b/Engine/Plugins/Animation/ControlRig/Source/ControlRigEditor/Private/AnimDetails/Proxies/AnimDetailsProxyRotation.cpp
--- a/Engine/Plugins/Animation/ControlRig/Source/ControlRigEditor/Private/AnimDetails/Proxies/AnimDetailsProxyRotation.cpp
+++ b/Engine/Plugins/Animation/ControlRig/Source/ControlRigEditor/Private/AnimDetails/Proxies/AnimDetailsProxyRotation.cpp
@@ -153,7 +153,8 @@ void UAnimDetailsProxyRotation::ResetPropertyToDefault(const FName& PropertyName
 {
 	using namespace UE::ControlRigEditor;
 
-	if (PropertyName == FAnimDetailsRotation::StaticStruct()->GetFName())
+	const FName StructName = GET_MEMBER_NAME_CHECKED(UAnimDetailsProxyRotation, Rotation);
+	if (PropertyName == StructName)
 	{
 		// Reset the whole rotation struct
 		Rotation = RotationUtils::GetDefaultRotation(*this);
@@ -176,14 +177,13 @@ void UAnimDetailsProxyRotation::ResetPropertyToDefault(const FName& PropertyName
 		Binding.CallFunction<FAnimDetailsRotation>(*this, Rotation);
 	}
 }
-
 bool UAnimDetailsProxyRotation::HasDefaultValue(const FName& PropertyName) const
 {
 	using namespace UE::ControlRigEditor;
 
 	const FAnimDetailsRotation DefaultRotation = RotationUtils::GetDefaultRotation(*this);
-
-	if (PropertyName == FAnimDetailsRotation::StaticStruct()->GetFName())
+	const FName StructName = GET_MEMBER_NAME_CHECKED(UAnimDetailsProxyRotation, Rotation);
+	if (PropertyName == StructName)
 	{
 		return
 			Rotation.RX == DefaultRotation.RX &&
@@ -207,7 +207,6 @@ bool UAnimDetailsProxyRotation::HasDefaultValue(const FName& PropertyName) const
 
 	return true;
 }
-
 EControlRigContextChannelToKey UAnimDetailsProxyRotation::GetChannelToKeyFromPropertyName(const FName& PropertyName) const
 {
 	if (PropertyName == GET_MEMBER_NAME_CHECKED(FAnimDetailsRotation, RX))
diff --git a/Engine/Plugins/Animation/ControlRig/Source/ControlRigEditor/Private/AnimDetails/Proxies/AnimDetailsProxyScale.cpp b/Engine/Plugins/Animation/ControlRig/Source/ControlRigEditor/Private/AnimDetails/Proxies/AnimDetailsProxyScale.cpp
--- a/Engine/Plugins/Animation/ControlRig/Source/ControlRigEditor/Private/AnimDetails/Proxies/AnimDetailsProxyScale.cpp
+++ b/Engine/Plugins/Animation/ControlRig/Source/ControlRigEditor/Private/AnimDetails/Proxies/AnimDetailsProxyScale.cpp
@@ -138,8 +138,8 @@ void UAnimDetailsProxyScale::AdoptValues(const ERigControlValueType RigControlVa
 void UAnimDetailsProxyScale::ResetPropertyToDefault(const FName& PropertyName)
 {
 	using namespace UE::ControlRigEditor;
-
-	if (PropertyName == FAnimDetailsScale::StaticStruct()->GetFName())
+	const FName StructName = GET_MEMBER_NAME_CHECKED(UAnimDetailsProxyScale, Scale);
+	if (PropertyName == StructName)
 	{
 		// Reset the whole scale struct
 		Scale = ScaleUtils::GetDefaultScale(*this);
@@ -169,7 +169,8 @@ bool UAnimDetailsProxyScale::HasDefaultValue(const FName& PropertyName) const
 
 	const FAnimDetailsScale DefaultScale = ScaleUtils::GetDefaultScale(*this);
 
-	if (PropertyName == FAnimDetailsScale::StaticStruct()->GetFName())
+	const FName StructName = GET_MEMBER_NAME_CHECKED(UAnimDetailsProxyScale, Scale);
+	if (PropertyName == StructName)
 	{
 		return
 			Scale.SX == DefaultScale.SX &&
diff --git a/Engine/Plugins/Animation/ControlRig/Source/ControlRigEditor/Private/AnimDetails/Proxies/AnimDetailsProxyVector2D.cpp b/Engine/Plugins/Animation/ControlRig/Source/ControlRigEditor/Private/AnimDetails/Proxies/AnimDetailsProxyVector2D.cpp
--- a/Engine/Plugins/Animation/ControlRig/Source/ControlRigEditor/Private/AnimDetails/Proxies/AnimDetailsProxyVector2D.cpp
+++ b/Engine/Plugins/Animation/ControlRig/Source/ControlRigEditor/Private/AnimDetails/Proxies/AnimDetailsProxyVector2D.cpp
@@ -104,7 +104,8 @@ void UAnimDetailsProxyVector2D::AdoptValues(const ERigControlValueType RigContro
 
 void UAnimDetailsProxyVector2D::ResetPropertyToDefault(const FName& PropertyName)
 {
-	if (PropertyName == FAnimDetailsVector2D::StaticStruct()->GetFName())
+	const FName StructName = GET_MEMBER_NAME_CHECKED(UAnimDetailsProxyVector2D, Vector2D);
+	if (PropertyName == StructName)
 	{
 		// Reset the whole vector 2D struct
 		AdoptValues(ERigControlValueType::Initial);
@@ -154,8 +155,9 @@ bool UAnimDetailsProxyVector2D::HasDefaultValue(const FName& PropertyName) const
 
 			return FAnimDetailsVector2D();
 		}();
+	const FName StructName = GET_MEMBER_NAME_CHECKED(UAnimDetailsProxyVector2D, Vector2D);
 
-	if (PropertyName == FAnimDetailsVector2D::StaticStruct()->GetFName())
+	if (PropertyName == StructName)
 	{
 		return
 			Vector2D.X == DefaultVector2D.X &&
