CIS Fix

diff --git a/Engine/Plugins/Animation/ControlRig/Source/ControlRigEditor/Private/Sequencer/AnimLayers/AnimLayers.cpp b/Engine/Plugins/Animation/ControlRig/Source/ControlRigEditor/Private/Sequencer/AnimLayers/AnimLayers.cpp
--- a/Engine/Plugins/Animation/ControlRig/Source/ControlRigEditor/Private/Sequencer/AnimLayers/AnimLayers.cpp
+++ b/Engine/Plugins/Animation/ControlRig/Source/ControlRigEditor/Private/Sequencer/AnimLayers/AnimLayers.cpp
@@ -62,21 +62,21 @@ static TArray<FGuid> GetSelectedOutlinerGuids(ISequencer* SequencerPtr)
 	return SelectedObjects;
 }
 
-struct FControlRigAndControlsAndTrack
+struct FALControlRigAndControlsAndTrack
 {
 	UMovieSceneControlRigParameterTrack* Track;
 	UControlRig* ControlRig;
 	TArray<FName> Controls;
 };
 
-struct FObjectAndTrack
+struct FALObjectAndTrack
 {
 	UMovieScenePropertyTrack* Track;
 	UObject* BoundObject;
 	FGuid SequencerGuid;
 };
 
-static void GetSelectedControlRigsAndBoundObjects(ISequencer* SequencerPtr, TArray<FControlRigAndControlsAndTrack>& OutSelectedCRs, TArray<FObjectAndTrack>& OutBoundObjects)
+static void GetSelectedControlRigsAndBoundObjects(ISequencer* SequencerPtr, TArray<FALControlRigAndControlsAndTrack>& OutSelectedCRs, TArray<FALObjectAndTrack>& OutBoundObjects)
 {
 	if (SequencerPtr == nullptr || SequencerPtr->GetViewModel() == nullptr)
 	{
@@ -105,7 +105,7 @@ static void GetSelectedControlRigsAndBoundObjects(ISequencer* SequencerPtr, TArr
 				{
 					if (UControlRig* ControlRig = Track->GetControlRig())
 					{
-						FControlRigAndControlsAndTrack CRControls;
+						FALControlRigAndControlsAndTrack CRControls;
 						CRControls.ControlRig = ControlRig;
 						CRControls.Track = Track;
 						CRControls.Controls = ControlRig->CurrentControlSelection();
@@ -141,7 +141,7 @@ static void GetSelectedControlRigsAndBoundObjects(ISequencer* SequencerPtr, TArr
 						{
 							if (UObject* BoundObject = CurrentBoundObject.Get())
 							{
-								FObjectAndTrack ObjectAndTrack;
+								FALObjectAndTrack ObjectAndTrack;
 								ObjectAndTrack.BoundObject = BoundObject;
 								ObjectAndTrack.Track = Track;
 								ObjectAndTrack.SequencerGuid = Binding.GetObjectGuid();
@@ -573,8 +573,8 @@ bool UAnimLayer::AddSelectedInSequencer()
 	{
 		return false;
 	}
-	TArray<FControlRigAndControlsAndTrack> SelectedCRs;
-	TArray<FObjectAndTrack> SelectedBoundObjects;
+	TArray<FALControlRigAndControlsAndTrack> SelectedCRs;
+	TArray<FALObjectAndTrack> SelectedBoundObjects;
 	GetSelectedControlRigsAndBoundObjects(SequencerPtr.Get(), SelectedCRs, SelectedBoundObjects);
 	if (SelectedCRs.Num() <= 0 && SelectedBoundObjects.Num() <= 0)
 	{
@@ -583,7 +583,7 @@ bool UAnimLayer::AddSelectedInSequencer()
 	bool bAddedSomething = false;
 	const FScopedTransaction Transaction(LOCTEXT("AddSelectedAnimLayer_Transaction", "Add Selected"),!GIsTransacting);
 	Modify();
-	for (FControlRigAndControlsAndTrack& CRControls : SelectedCRs)
+	for (FALControlRigAndControlsAndTrack& CRControls : SelectedCRs)
 	{
 		if (FAnimLayerItem* ExistingAnimLayerItem = AnimLayerItems.Find(CRControls.ControlRig))
 		{
@@ -639,7 +639,7 @@ bool UAnimLayer::AddSelectedInSequencer()
 		}
 	}
 
-	for (FObjectAndTrack& ObjectAndTrack : SelectedBoundObjects)
+	for (FALObjectAndTrack& ObjectAndTrack : SelectedBoundObjects)
 	{
 		FAnimLayerItem& AnimLayerItem = AnimLayerItems.FindOrAdd(ObjectAndTrack.BoundObject);
 		AnimLayerItem.SequencerGuid = ObjectAndTrack.SequencerGuid;
@@ -821,8 +821,8 @@ bool UAnimLayer::RemoveSelectedInSequencer()
 		}
 	}
 
-	TArray<FControlRigAndControlsAndTrack> SelectedCRs;
-	TArray<FObjectAndTrack> SelectedBoundObjects;
+	TArray<FALControlRigAndControlsAndTrack> SelectedCRs;
+	TArray<FALObjectAndTrack> SelectedBoundObjects;
 	GetSelectedControlRigsAndBoundObjects(SequencerPtr.Get(), SelectedCRs, SelectedBoundObjects);
 	if (SelectedCRs.Num() <= 0 && SelectedBoundObjects.Num() < 0)
 	{
@@ -832,7 +832,7 @@ bool UAnimLayer::RemoveSelectedInSequencer()
 	UAnimLayer& AnimLayer = *this;
 	const FScopedTransaction Transaction(LOCTEXT("RemoveSelected_Transaction", "Remove Selected"), !GIsTransacting);
 	Modify();
-	for (FControlRigAndControlsAndTrack& CRControls : SelectedCRs)
+	for (FALControlRigAndControlsAndTrack& CRControls : SelectedCRs)
 	{
 		if (FAnimLayerItem* ExistingAnimLayerItem = AnimLayer.AnimLayerItems.Find(CRControls.ControlRig))
 		{
@@ -864,7 +864,7 @@ bool UAnimLayer::RemoveSelectedInSequencer()
 			}
 		}
 	}
-	for (FObjectAndTrack& ObjectAndTrack : SelectedBoundObjects)
+	for (FALObjectAndTrack& ObjectAndTrack : SelectedBoundObjects)
 	{
 		if (AnimLayer.AnimLayerItems.Find(ObjectAndTrack.BoundObject))
 		{
@@ -2913,8 +2913,8 @@ int32 UAnimLayers::AddAnimLayerFromSelection(ISequencer* SequencerPtr)
 {
 	int32 NewIndex = INDEX_NONE;
 	//wrap scoped transaction since it can deselect control rigs
-	TArray<FControlRigAndControlsAndTrack> SelectedCRs;
-	TArray<FObjectAndTrack> SelectedBoundObjects; 
+	TArray<FALControlRigAndControlsAndTrack> SelectedCRs;
+	TArray<FALObjectAndTrack> SelectedBoundObjects; 
 	{
 		const FScopedTransaction Transaction(LOCTEXT("AddAnimLayer_Transaction", "Add Anim Layer"), !GIsTransacting);
 		Modify();
@@ -2942,7 +2942,7 @@ int32 UAnimLayers::AddAnimLayerFromSelection(ISequencer* SequencerPtr)
 
 		bool bItemAdded = false;
 
-		for (FControlRigAndControlsAndTrack& CRControls : SelectedCRs)
+		for (FALControlRigAndControlsAndTrack& CRControls : SelectedCRs)
 		{
 			Modify();
 			FAnimLayerItem AnimLayerItem;
@@ -2980,7 +2980,7 @@ int32 UAnimLayers::AddAnimLayerFromSelection(ISequencer* SequencerPtr)
 
 			}
 		}
-		for (FObjectAndTrack& ObjectAndTrack : SelectedBoundObjects)
+		for (FALObjectAndTrack& ObjectAndTrack : SelectedBoundObjects)
 		{
 			Modify();
 
@@ -3040,7 +3040,7 @@ int32 UAnimLayers::AddAnimLayerFromSelection(ISequencer* SequencerPtr)
 		}
 	}
 	//may need to reselect controls here
-	for (FControlRigAndControlsAndTrack& CRControls : SelectedCRs)
+	for (FALControlRigAndControlsAndTrack& CRControls : SelectedCRs)
 	{
 		for (const FName& Control : CRControls.Controls)
 		{
