CIS Fix: Use updated namepsaced function and remove old one

diff --git a/Engine/Plugins/Animation/ControlRig/Source/ControlRigEditor/Private/Sequencer/AnimLayers/AnimLayers.cpp b/Engine/Plugins/Animation/ControlRig/Source/ControlRigEditor/Private/Sequencer/AnimLayers/AnimLayers.cpp
--- a/Engine/Plugins/Animation/ControlRig/Source/ControlRigEditor/Private/Sequencer/AnimLayers/AnimLayers.cpp
+++ b/Engine/Plugins/Animation/ControlRig/Source/ControlRigEditor/Private/Sequencer/AnimLayers/AnimLayers.cpp
@@ -43,118 +43,13 @@
 #include "MovieSceneToolHelpers.h"
 #include "Tools/EvaluateSequencerTools.h"
 #include "LevelEditorViewport.h"
+
 #include "EulerTransform.h"
 
 #include UE_INLINE_GENERATED_CPP_BY_NAME(AnimLayers)
 #define LOCTEXT_NAMESPACE "AnimLayers"
 
 
-static TArray<FGuid> GetSelectedOutlinerGuids(ISequencer* SequencerPtr)
-{
-	using namespace UE::Sequencer;
-	TArray<FGuid> SelectedObjects;// = SequencerPtr->GetViewModel()->GetSelection()->GetBoundObjectsGuids(); //this is not exposed :(
-
-	for (TViewModelPtr<FObjectBindingModel> ObjectBindingNode : SequencerPtr->GetViewModel()->GetSelection()->Outliner.Filter<FObjectBindingModel>())
-	{
-		FGuid Guid = ObjectBindingNode->GetObjectGuid();
-		SelectedObjects.Add(Guid);
-	}
-	return SelectedObjects;
-}
-
-struct FALControlRigAndControlsAndTrack
-{
-	UMovieSceneControlRigParameterTrack* Track;
-	UControlRig* ControlRig;
-	TArray<FName> Controls;
-};
-
-struct FALObjectAndTrack
-{
-	UMovieScenePropertyTrack* Track;
-	UObject* BoundObject;
-	FGuid SequencerGuid;
-};
-
-static void GetSelectedControlRigsAndBoundObjects(ISequencer* SequencerPtr, TArray<FALControlRigAndControlsAndTrack>& OutSelectedCRs, TArray<FALObjectAndTrack>& OutBoundObjects)
-{
-	if (SequencerPtr == nullptr || SequencerPtr->GetViewModel() == nullptr)
-	{
-		return;
-	}
-	ULevelSequence* LevelSequence = Cast<ULevelSequence>(SequencerPtr->GetFocusedMovieSceneSequence());
-	if (LevelSequence == nullptr)
-	{
-		return;
-	}
-
-	using namespace UE::Sequencer;
-	TArray<FGuid> SelectedObjects = GetSelectedOutlinerGuids(SequencerPtr);;// = SequencerPtr->GetViewModel()->GetSelection()->GetBoundObjectsGuids();
-
-	const UMovieScene* MovieScene = LevelSequence->GetMovieScene();
-	if (MovieScene)
-	{
-		const TArray<FMovieSceneBinding>& Bindings = MovieScene->GetBindings();
-		for (const FMovieSceneBinding& Binding : Bindings)
-		{
-			bool bHaveControlRig = false;
-			TArray<UMovieSceneTrack*> CRTracks = MovieScene->FindTracks(UMovieSceneControlRigParameterTrack::StaticClass(), Binding.GetObjectGuid(), NAME_None);
-			for (UMovieSceneTrack* AnyOleTrack : CRTracks)
-			{
-				if (UMovieSceneControlRigParameterTrack* Track = Cast<UMovieSceneControlRigParameterTrack>(AnyOleTrack))
-				{
-					if (UControlRig* ControlRig = Track->GetControlRig())
-					{
-						FALControlRigAndControlsAndTrack CRControls;
-						CRControls.ControlRig = ControlRig;
-						CRControls.Track = Track;
-						CRControls.Controls = ControlRig->CurrentControlSelection();
-						if (CRControls.Controls.Num() > 0)
-						{
-							bHaveControlRig = true;
-							OutSelectedCRs.Add(CRControls);
-						}
-					}
-				}
-			}
-			//if we have control rig controls don't add the base skel mesh for now
-			if (bHaveControlRig == false && SelectedObjects.Contains(Binding.GetObjectGuid()))
-			{
-				TArray<UMovieSceneTrack*> Tracks = MovieScene->FindTracks(UMovieScenePropertyTrack::StaticClass(), Binding.GetObjectGuid(), NAME_None);
-				if (Tracks.Num() < 1)
-				{
-					continue;
-				}
-				const FMovieSceneSequenceID& SequenceID = SequencerPtr->GetEvaluationState()->FindSequenceId(LevelSequence);
-
-				for (UMovieSceneTrack* AnyOleTrack : Tracks)
-				{
-					if (UMovieScenePropertyTrack* Track = Cast<UMovieScenePropertyTrack>(AnyOleTrack))
-					{
-						const FMovieSceneBlendTypeField SupportedBlendTypes = Track->GetSupportedBlendTypes();
-						if (SupportedBlendTypes.Num() == 0)
-						{
-							continue;
-						}
-						const TArrayView<TWeakObjectPtr<UObject>>& BoundObjects = SequencerPtr->FindBoundObjects(Binding.GetObjectGuid(), SequenceID);
-						for (const TWeakObjectPtr<UObject> CurrentBoundObject : BoundObjects)
-						{
-							if (UObject* BoundObject = CurrentBoundObject.Get())
-							{
-								FALObjectAndTrack ObjectAndTrack;
-								ObjectAndTrack.BoundObject = BoundObject;
-								ObjectAndTrack.Track = Track;
-								ObjectAndTrack.SequencerGuid = Binding.GetObjectGuid();
-								OutBoundObjects.Add(ObjectAndTrack);
-							}
-						}
-					}
-				}
-			}
-		}
-	}
-}
-
 bool FAnimLayerSelectionSet::MergeWithAnotherSelection(const FAnimLayerSelectionSet& Selection)
 {
 	if (BoundObject.IsValid() && (BoundObject.Get() == Selection.BoundObject.Get()))
@@ -568,14 +463,16 @@ void UAnimLayer::SetActive(bool bInActive)
 
 bool UAnimLayer::AddSelectedInSequencer() 
 {
+	using namespace UE::AIE;
+
 	TSharedPtr<ISequencer> SequencerPtr = UAnimLayers::GetSequencerFromAsset();
 	if (SequencerPtr.IsValid() == false)
 	{
 		return false;
 	}
-	TArray<FALControlRigAndControlsAndTrack> SelectedCRs;
-	TArray<FALObjectAndTrack> SelectedBoundObjects;
-	GetSelectedControlRigsAndBoundObjects(SequencerPtr.Get(), SelectedCRs, SelectedBoundObjects);
+	TArray<FControlRigAndControlsAndTrack> SelectedCRs;
+	TArray<FObjectAndTrack> SelectedBoundObjects;
+	FSequencerSelected::GetSelectedControlRigsAndBoundObjects(SequencerPtr.Get(), SelectedCRs, SelectedBoundObjects);
 	if (SelectedCRs.Num() <= 0 && SelectedBoundObjects.Num() <= 0)
 	{
 		return false;
@@ -583,7 +480,7 @@ bool UAnimLayer::AddSelectedInSequencer()
 	bool bAddedSomething = false;
 	const FScopedTransaction Transaction(LOCTEXT("AddSelectedAnimLayer_Transaction", "Add Selected"),!GIsTransacting);
 	Modify();
-	for (FALControlRigAndControlsAndTrack& CRControls : SelectedCRs)
+	for (FControlRigAndControlsAndTrack& CRControls : SelectedCRs)
 	{
 		if (FAnimLayerItem* ExistingAnimLayerItem = AnimLayerItems.Find(CRControls.ControlRig))
 		{
@@ -639,7 +536,7 @@ bool UAnimLayer::AddSelectedInSequencer()
 		}
 	}
 
-	for (FALObjectAndTrack& ObjectAndTrack : SelectedBoundObjects)
+	for (FObjectAndTrack& ObjectAndTrack : SelectedBoundObjects)
 	{
 		FAnimLayerItem& AnimLayerItem = AnimLayerItems.FindOrAdd(ObjectAndTrack.BoundObject);
 		AnimLayerItem.SequencerGuid = ObjectAndTrack.SequencerGuid;
@@ -807,6 +704,8 @@ bool UAnimLayer::RemoveAnimLayerItem(UObject* InObject)
 }
 bool UAnimLayer::RemoveSelectedInSequencer() 
 {
+	using namespace UE::AIE;
+
 	TSharedPtr<ISequencer> SequencerPtr = UAnimLayers::GetSequencerFromAsset();
 	if (SequencerPtr.IsValid() == false)
 	{
@@ -821,9 +720,9 @@ bool UAnimLayer::RemoveSelectedInSequencer()
 		}
 	}
 
-	TArray<FALControlRigAndControlsAndTrack> SelectedCRs;
-	TArray<FALObjectAndTrack> SelectedBoundObjects;
-	GetSelectedControlRigsAndBoundObjects(SequencerPtr.Get(), SelectedCRs, SelectedBoundObjects);
+	TArray<FControlRigAndControlsAndTrack> SelectedCRs;
+	TArray<FObjectAndTrack> SelectedBoundObjects;
+	FSequencerSelected::GetSelectedControlRigsAndBoundObjects(SequencerPtr.Get(), SelectedCRs, SelectedBoundObjects);
 	if (SelectedCRs.Num() <= 0 && SelectedBoundObjects.Num() < 0)
 	{
 		return false;
@@ -832,7 +731,7 @@ bool UAnimLayer::RemoveSelectedInSequencer()
 	UAnimLayer& AnimLayer = *this;
 	const FScopedTransaction Transaction(LOCTEXT("RemoveSelected_Transaction", "Remove Selected"), !GIsTransacting);
 	Modify();
-	for (FALControlRigAndControlsAndTrack& CRControls : SelectedCRs)
+	for (FControlRigAndControlsAndTrack& CRControls : SelectedCRs)
 	{
 		if (FAnimLayerItem* ExistingAnimLayerItem = AnimLayer.AnimLayerItems.Find(CRControls.ControlRig))
 		{
@@ -864,7 +763,7 @@ bool UAnimLayer::RemoveSelectedInSequencer()
 			}
 		}
 	}
-	for (FALObjectAndTrack& ObjectAndTrack : SelectedBoundObjects)
+	for (FObjectAndTrack& ObjectAndTrack : SelectedBoundObjects)
 	{
 		if (AnimLayer.AnimLayerItems.Find(ObjectAndTrack.BoundObject))
 		{
@@ -2911,10 +2810,12 @@ bool UAnimLayers::IsTrackOnSelectedLayer(const UMovieSceneTrack* InTrack)const
 
 int32 UAnimLayers::AddAnimLayerFromSelection(ISequencer* SequencerPtr)
 {
+	using namespace UE::AIE;
+
 	int32 NewIndex = INDEX_NONE;
 	//wrap scoped transaction since it can deselect control rigs
-	TArray<FALControlRigAndControlsAndTrack> SelectedCRs;
-	TArray<FALObjectAndTrack> SelectedBoundObjects; 
+	TArray<FControlRigAndControlsAndTrack> SelectedCRs;
+	TArray<FObjectAndTrack> SelectedBoundObjects; 
 	{
 		const FScopedTransaction Transaction(LOCTEXT("AddAnimLayer_Transaction", "Add Anim Layer"), !GIsTransacting);
 		Modify();
@@ -2924,7 +2825,7 @@ int32 UAnimLayers::AddAnimLayerFromSelection(ISequencer* SequencerPtr)
 		}
 		UAnimLayer* AnimLayer = NewObject<UAnimLayer>(this, NAME_None, RF_Transactional);
 
-		GetSelectedControlRigsAndBoundObjects(SequencerPtr, SelectedCRs, SelectedBoundObjects);
+		FSequencerSelected::GetSelectedControlRigsAndBoundObjects(SequencerPtr, SelectedCRs, SelectedBoundObjects);
 
 		if (SelectedCRs.Num() <= 0 && SelectedBoundObjects.Num() <= 0)
 		{
@@ -2942,7 +2843,7 @@ int32 UAnimLayers::AddAnimLayerFromSelection(ISequencer* SequencerPtr)
 
 		bool bItemAdded = false;
 
-		for (FALControlRigAndControlsAndTrack& CRControls : SelectedCRs)
+		for (FControlRigAndControlsAndTrack& CRControls : SelectedCRs)
 		{
 			Modify();
 			FAnimLayerItem AnimLayerItem;
@@ -2980,7 +2881,7 @@ int32 UAnimLayers::AddAnimLayerFromSelection(ISequencer* SequencerPtr)
 
 			}
 		}
-		for (FALObjectAndTrack& ObjectAndTrack : SelectedBoundObjects)
+		for (FObjectAndTrack& ObjectAndTrack : SelectedBoundObjects)
 		{
 			Modify();
 
@@ -3040,7 +2941,7 @@ int32 UAnimLayers::AddAnimLayerFromSelection(ISequencer* SequencerPtr)
 		}
 	}
 	//may need to reselect controls here
-	for (FALControlRigAndControlsAndTrack& CRControls : SelectedCRs)
+	for (FControlRigAndControlsAndTrack& CRControls : SelectedCRs)
 	{
 		for (const FName& Control : CRControls.Controls)
 		{
diff --git a/Engine/Plugins/Animation/ControlRig/Source/ControlRigEditor/Private/Sequencer/AnimLayers/SAnimLayers.cpp b/Engine/Plugins/Animation/ControlRig/Source/ControlRigEditor/Private/Sequencer/AnimLayers/SAnimLayers.cpp
--- a/Engine/Plugins/Animation/ControlRig/Source/ControlRigEditor/Private/Sequencer/AnimLayers/SAnimLayers.cpp
+++ b/Engine/Plugins/Animation/ControlRig/Source/ControlRigEditor/Private/Sequencer/AnimLayers/SAnimLayers.cpp
@@ -1798,6 +1798,7 @@ void SAnimLayers::Tick(const FGeometry& AllottedGeometry, const double InCurrent
 				AnimLayersPtr->AnimLayerListChanged().AddRaw(AnimLayerController.Get(), &FAnimLayerController::HandleOnAnimLayerListChanged);
 				AnimLayers = AnimLayersPtr;
 			}
+			AnimLayerController->RefreshSourceData(true);
 			AnimLayerController->RefreshSelectionData();
 		}
 	}
