CIS Fix: Use updated namepsaced function and remove old one part 2

diff --git a/Engine/Plugins/Animation/ControlRig/Source/ControlRigEditor/Private/Sequencer/AnimLayers/AnimLayers.cpp b/Engine/Plugins/Animation/ControlRig/Source/ControlRigEditor/Private/Sequencer/AnimLayers/AnimLayers.cpp
--- a/Engine/Plugins/Animation/ControlRig/Source/ControlRigEditor/Private/Sequencer/AnimLayers/AnimLayers.cpp
+++ b/Engine/Plugins/Animation/ControlRig/Source/ControlRigEditor/Private/Sequencer/AnimLayers/AnimLayers.cpp
@@ -1804,7 +1804,7 @@ void AddKeyToChannel(ChannelType* Channel, EMovieSceneKeyInterpolation DefaultIn
 		break;
 	}
 }
-static bool GetStartEndIndicesForControl(UMovieSceneControlRigParameterSection* BaseSection, const FRigControlElement* ControlElement, int& OutStartIndex, int& OutEndIndex)
+static bool ALGetStartEndIndicesForControl(UMovieSceneControlRigParameterSection* BaseSection, const FRigControlElement* ControlElement, int& OutStartIndex, int& OutEndIndex)
 {
 	if (FChannelMapInfo* pChannelIndex = BaseSection->ControlChannelMap.Find(ControlElement->GetFName()))
 	{
@@ -2021,7 +2021,7 @@ bool UAnimLayers::SetKeyValueOrPassthrough(ISequencer* InSequencer, int32 InInde
 													if (FRigControlElement* Control = CRSection->GetControlRig()->FindControl(ControlName))
 													{
 														int32 StartIndex = 0, EndIndex = 0;
-														if (GetStartEndIndicesForControl(CRSection, Control, StartIndex, EndIndex))
+														if (ALGetStartEndIndicesForControl(CRSection, Control, StartIndex, EndIndex))
 														{
 															if (bJustValue)
 															{
@@ -2087,7 +2087,7 @@ bool UAnimLayers::SetKeyValueOrPassthrough(ISequencer* InSequencer, int32 InInde
 															if (FRigControlElement* Control = ControlPair.Key->FindControl(ControlName))
 															{
 																int32 StartIndex = 0, EndIndex = 0;
-																if (GetStartEndIndicesForControl(CRSection, Control, StartIndex, EndIndex))
+																if (ALGetStartEndIndicesForControl(CRSection, Control, StartIndex, EndIndex))
 																{
 																	if (bJustValue)
 																	{
@@ -3774,7 +3774,6 @@ bool UAnimLayers::AdjustmentBlendLayers(ISequencer* InSequencer, int32 LayerInde
 				if (FAnimLayerItem* BaseLayerItem = BaseLayer->AnimLayerItems.Find(Pair.Key))
 				{
 					AdjustmentBlendWorldSpace(InSequencer, *BaseLayerItem, Pair.Value);
-
 				}
 			}
 			else
