CIS Fix: Use updated namepsaced function until refactor tomorrow

diff --git a/Engine/Plugins/Animation/ControlRig/Source/ControlRigEditor/Private/Sequencer/AnimLayers/AnimLayers.cpp b/Engine/Plugins/Animation/ControlRig/Source/ControlRigEditor/Private/Sequencer/AnimLayers/AnimLayers.cpp
--- a/Engine/Plugins/Animation/ControlRig/Source/ControlRigEditor/Private/Sequencer/AnimLayers/AnimLayers.cpp
+++ b/Engine/Plugins/Animation/ControlRig/Source/ControlRigEditor/Private/Sequencer/AnimLayers/AnimLayers.cpp
@@ -1783,7 +1783,8 @@ FAnimLayerSectionItem* FAnimLayerItem::FindMatchingSectionItem(UMovieSceneSectio
 	return nullptr;
 }
 
-
+namespace UE::AnimLayers 
+{
 template<typename ChannelType, typename ValueType>
 void AddKeyToChannel(ChannelType* Channel, EMovieSceneKeyInterpolation DefaultInterpolation, const FFrameNumber& FrameNumber, ValueType Value)
 {
@@ -1804,7 +1805,7 @@ void AddKeyToChannel(ChannelType* Channel, EMovieSceneKeyInterpolation DefaultIn
 		break;
 	}
 }
-static bool ALGetStartEndIndicesForControl(UMovieSceneControlRigParameterSection* BaseSection, const FRigControlElement* ControlElement, int& OutStartIndex, int& OutEndIndex)
+static bool GetStartEndIndicesForControl(UMovieSceneControlRigParameterSection* BaseSection, const FRigControlElement* ControlElement, int& OutStartIndex, int& OutEndIndex)
 {
 	if (FChannelMapInfo* pChannelIndex = BaseSection->ControlChannelMap.Find(ControlElement->GetFName()))
 	{
@@ -1921,7 +1922,7 @@ void SetCurrentKeys(TArrayView<ChannelType*> Channels, int32 StartIndex, int32 E
 		}
 	}
 }
-
+}
 bool UAnimLayers::SetPassthroughKey(ISequencer* InSequencer, int32 InIndex)
 {
 	return SetKeyValueOrPassthrough(InSequencer, InIndex, false /*IsValue*/);
@@ -1936,6 +1937,7 @@ bool UAnimLayers::SetKey(ISequencer* InSequencer, int32 InIndex)
 
 bool UAnimLayers::SetKeyValueOrPassthrough(ISequencer* InSequencer, int32 InIndex, bool bJustValue)
 {
+	using namespace UE::AnimLayers;
 	auto GetSelectedRigElements = [](UControlRig* ControlRig, TArray<FRigElementKey>& OutSelectedKeys) 
 	{
 		if (ControlRig)
@@ -2021,7 +2023,7 @@ bool UAnimLayers::SetKeyValueOrPassthrough(ISequencer* InSequencer, int32 InInde
 													if (FRigControlElement* Control = CRSection->GetControlRig()->FindControl(ControlName))
 													{
 														int32 StartIndex = 0, EndIndex = 0;
-														if (ALGetStartEndIndicesForControl(CRSection, Control, StartIndex, EndIndex))
+														if (GetStartEndIndicesForControl(CRSection, Control, StartIndex, EndIndex))
 														{
 															if (bJustValue)
 															{
@@ -2087,7 +2089,7 @@ bool UAnimLayers::SetKeyValueOrPassthrough(ISequencer* InSequencer, int32 InInde
 															if (FRigControlElement* Control = ControlPair.Key->FindControl(ControlName))
 															{
 																int32 StartIndex = 0, EndIndex = 0;
-																if (ALGetStartEndIndicesForControl(CRSection, Control, StartIndex, EndIndex))
+																if (GetStartEndIndicesForControl(CRSection, Control, StartIndex, EndIndex))
 																{
 																	if (bJustValue)
 																	{
@@ -3236,6 +3238,7 @@ template<typename ChannelType>
 void AdjustmentBlend(UMovieSceneSection* Section, TArrayView<ChannelType*>& BaseChannels, 
 	TArrayView<ChannelType*>& LayerChannels,ISequencer* InSequencer)
 {
+	using namespace UE::AnimLayers;
 	if (BaseChannels.Num() != LayerChannels.Num())
 	{
 		return;
