Control Rig: Finishing UX polish on dependency viewer

#rnx

diff --git a/Engine/Plugins/Animation/ControlRig/Source/ControlRig/Private/Settings/ControlRigSettings.cpp b/Engine/Plugins/Animation/ControlRig/Source/ControlRig/Private/Settings/ControlRigSettings.cpp
--- a/Engine/Plugins/Animation/ControlRig/Source/ControlRig/Private/Settings/ControlRigSettings.cpp
+++ b/Engine/Plugins/Animation/ControlRig/Source/ControlRig/Private/Settings/ControlRigSettings.cpp
@@ -25,6 +25,7 @@ UControlRigEditorSettings::UControlRigEditorSettings(const FObjectInitializer& O
 	bResetControlsOnPinValueInteraction = false;
 	bResetPoseWhenTogglingEventQueue = true;
 	bEnableUndoForPoseInteraction = true;
+	bEnableFlashlightInDependencyViewer = true;
 
 	ConstructionEventBorderColor = FLinearColor::Red;
 	BackwardsSolveBorderColor = FLinearColor::Yellow;
diff --git a/Engine/Plugins/Animation/ControlRig/Source/ControlRig/Public/Settings/ControlRigSettings.h b/Engine/Plugins/Animation/ControlRig/Source/ControlRig/Public/Settings/ControlRigSettings.h
--- a/Engine/Plugins/Animation/ControlRig/Source/ControlRig/Public/Settings/ControlRigSettings.h
+++ b/Engine/Plugins/Animation/ControlRig/Source/ControlRig/Public/Settings/ControlRigSettings.h
@@ -102,7 +102,14 @@ public:
 	 */
 	UPROPERTY(EditAnywhere, config, Category = NodeGraph)
 	TMap<FString, FControlRigSettingsPerPinBool> RigUnitPinExpansion;
-	
+
+	/**
+	 * Whether or not to use a flash light effect around the mouse pointer
+	 * in the dependency viewer to brighten up faded out nodes
+	 */
+	UPROPERTY(EditAnywhere, config, Category = NodeGraph)
+	bool bEnableFlashlightInDependencyViewer;
+
 	/**
 	 * The border color of the viewport when entering "Construction Event" mode
 	 */
diff --git a/Engine/Plugins/Animation/ControlRig/Source/ControlRigEditor/Private/ControlRigEditorStyle.h b/Engine/Plugins/Animation/ControlRig/Source/ControlRigEditor/Private/ControlRigEditorStyle.h
--- a/Engine/Plugins/Animation/ControlRig/Source/ControlRigEditor/Private/ControlRigEditorStyle.h
+++ b/Engine/Plugins/Animation/ControlRig/Source/ControlRigEditor/Private/ControlRigEditorStyle.h
@@ -49,6 +49,7 @@ public:
 		const FVector2D Icon40x40(40.0f, 40.0f);
 		const FVector2D Icon128x128(128.0f, 128.0f);
 		const FString ControlRigPluginContentDir = FPaths::EnginePluginsDir() / TEXT("Animation/ControlRig/Content");
+		const FString EngineSlateDir = FPaths::EngineContentDir() / TEXT("Slate");
 		const FString EngineEditorSlateDir = FPaths::EngineContentDir() / TEXT("Editor/Slate");
 		SetContentRoot(ControlRigPluginContentDir);
 		SetCoreContentRoot(EngineEditorSlateDir);
@@ -191,7 +192,6 @@ public:
 				.SetPressedForeground(FLinearColor::White)
 				.SetPressedPadding(FMargin(0.0, 1.0, 0.0, 0.0));
 
-			const FString EngineSlateDir = FPaths::EngineContentDir() / TEXT("Slate");
 			FContentRootBracket Bracket(this, EngineSlateDir);		
 
 			FComboButtonStyle ConstraintComboButton = FComboButtonStyle()
@@ -209,6 +209,11 @@ public:
 				.SetColorAndOpacity(FLinearColor(0.0f, 0.0f, 0.0f, 1.0f))
 				.SetFont(DEFAULT_FONT("Regular", 8)));
 
+			{
+				FContentRootBracket Bracket(this, EngineSlateDir);		
+				Set("ControlRig.DependencyGraph.Menu", new IMAGE_BRUSH_SVG("Starship/Common/menu", Icon20x20));
+			}
+
 			Set( "ControlRig.DependencyGraph.NodeBody", new BOX_BRUSH( "Slate/DependencyGraphNode_Body", FMargin(16.f/64.f, 25.f/64.f, 16.f/64.f, 16.f/64.f) ) );
 
 			Set("ControlRig.DependencyGraph.Colors.Instruction", FLinearColor(FColor::FromHex(TEXT("#00AAAA"))));
@@ -221,13 +226,12 @@ public:
 			Set("ControlRig.DependencyGraph.Colors.Metadata", FLinearColor(FColor::FromHex(TEXT("#BB6BF0"))));
 			Set("ControlRig.DependencyGraph.Colors.Variable", FLinearColor(FColor::FromHex(TEXT("#CCCC00"))));
 
-			Set("ControlRig.DependencyGraph.SyncContent", new CORE_IMAGE_BRUSH_SVG("Starship/Common/Undo", Icon16x16));
-			{
-				FContentRootBracket Bracket(this, EngineEditorSlateDir);
- 				Set( "ControlRig.DependencyGraph.HideUnrelated", new IMAGE_BRUSH_SVG( "Starship/GraphEditors/HideUnrelated", Icon16x16 ) );
- 				Set( "ControlRig.DependencyGraph.IsolateSelection", new IMAGE_BRUSH_SVG( "Sequencer/TrackIsolate", Icon16x16 ) );
-			}
-
+			Set( "ControlRig.DependencyGraph.HideUnrelated", new IMAGE_BRUSH_SVG( "Slate/GraphHideUnrelated_20", Icon20x20 ) );
+			Set( "ControlRig.DependencyGraph.IsolateSelection", new IMAGE_BRUSH_SVG( "Slate/GraphIsolate_20", Icon20x20 ) );
+			Set( "ControlRig.DependencyGraph.ParentChild", new IMAGE_BRUSH_SVG( "Slate/GraphParentChild_20", Icon20x20 ) );
+			Set( "ControlRig.DependencyGraph.VMRelationShips", new IMAGE_BRUSH_SVG( "Slate/GraphVMRelationships_20", Icon20x20 ) );
+			Set( "ControlRig.DependencyGraph.Cleanup", new IMAGE_BRUSH_SVG( "Slate/GraphCleanup_20", Icon20x20 ) );
+			Set( "ControlRig.DependencyGraph.Flashlight", new IMAGE_BRUSH_SVG( "Slate/Flashlight_20", Icon20x20 ) );
 		}
 
 		// Tabs
diff --git a/Engine/Plugins/Animation/ControlRig/Source/ControlRigEditor/Private/RigDependencyGraph/RigDependencyGraph.cpp b/Engine/Plugins/Animation/ControlRig/Source/ControlRigEditor/Private/RigDependencyGraph/RigDependencyGraph.cpp
--- a/Engine/Plugins/Animation/ControlRig/Source/ControlRigEditor/Private/RigDependencyGraph/RigDependencyGraph.cpp
+++ b/Engine/Plugins/Animation/ControlRig/Source/ControlRigEditor/Private/RigDependencyGraph/RigDependencyGraph.cpp
@@ -982,11 +982,23 @@ void URigDependencyGraph::RemoveUnrelatedNodes()
 		return;
 	}
 
+	TArray<FGuid> Guids;
+	for (const FNodeId& NodeId : SelectedNodes)
+	{
+		if (const URigDependencyGraphNode* Node = FindNode(NodeId))
+		{
+			if (Node->IslandGuid.IsSet())
+			{
+				Guids.AddUnique(Node->IslandGuid.GetValue());
+			}
+		}
+	}
+
 	bool bRemovedAnyNodes = false;
 	TArray<TObjectPtr<URigDependencyGraphNode>> PreviousNodes = DependencyGraphNodes;
 	for (int32 NodeIndex = 0; NodeIndex < PreviousNodes.Num(); ++NodeIndex)
 	{
-		if (PreviousNodes[NodeIndex]->IsUnrelated())
+		if (!Guids.Contains(PreviousNodes[NodeIndex]->IslandGuid.Get(FGuid())))
 		{
 			DependencyGraphNodes.Remove(PreviousNodes[NodeIndex]);
 			SelectedNodes.Remove(PreviousNodes[NodeIndex]->NodeId);
diff --git a/Engine/Plugins/Animation/ControlRig/Source/ControlRigEditor/Private/RigDependencyGraph/RigDependencyGraphCommands.cpp b/Engine/Plugins/Animation/ControlRig/Source/ControlRigEditor/Private/RigDependencyGraph/RigDependencyGraphCommands.cpp
--- a/Engine/Plugins/Animation/ControlRig/Source/ControlRigEditor/Private/RigDependencyGraph/RigDependencyGraphCommands.cpp
+++ b/Engine/Plugins/Animation/ControlRig/Source/ControlRigEditor/Private/RigDependencyGraph/RigDependencyGraphCommands.cpp
@@ -7,18 +7,19 @@
 void FRigDependencyGraphCommands::RegisterCommands()
 {
 	UI_COMMAND(FrameSelection, "Frame Selection", "Expands and frames the selection in the tree", EUserInterfaceActionType::Button, FInputChord(EKeys::F));
+	UI_COMMAND(JumpToSource, "Jump to Source", "Jumps to the instruction node backing up the dependency node", EUserInterfaceActionType::Button, FInputChord(EKeys::F12));
 	UI_COMMAND(ShowParentChildRelationships, "Show Parent/Child", "Whether to show or hide parent and child relationships", EUserInterfaceActionType::ToggleButton, FInputChord(EKeys::P));
 	UI_COMMAND(ShowInstructionRelationships, "Show VM Nodes", "Whether to show or hide VM based relationships", EUserInterfaceActionType::ToggleButton, FInputChord(EKeys::O));
-	UI_COMMAND(SyncContent, "Sync Content", "Whether or not to sync the content with selection in other views", EUserInterfaceActionType::ToggleButton, FInputChord(EKeys::L));
+	UI_COMMAND(LockContent, "Lock Content", "Whether or not to sync the content with selection in other views", EUserInterfaceActionType::ToggleButton, FInputChord(EKeys::L));
+	UI_COMMAND(EnableFlashlight, "Enable Flashlight", "Enable a flashlight around the mouse cursor to brighten up faded out nodes", EUserInterfaceActionType::ToggleButton, FInputChord());
 	UI_COMMAND(SelectAllNodes, "Select All", "Selects all nodes in the graph", EUserInterfaceActionType::Button, FInputChord(EModifierKey::Control, EKeys::A));
 	UI_COMMAND(SelectInputNodes, "Select Inputs", "Selects the linked inputs", EUserInterfaceActionType::Button, FInputChord(EKeys::MiddleMouseButton));
 	UI_COMMAND(SelectOutputNodes, "Select Outputs", "Selects the linked outputs", EUserInterfaceActionType::Button, FInputChord(EModifierKey::Control, EKeys::MiddleMouseButton));
 	UI_COMMAND(SelectNodeIsland, "Select Island", "Selects the linked island", EUserInterfaceActionType::Button, FInputChord(EModifierKey::Alt, EKeys::MiddleMouseButton));
 	UI_COMMAND(RemoveAllNodes, "Clear View", "Hides all nodes", EUserInterfaceActionType::Button, FInputChord(EKeys::X));
 	UI_COMMAND(RemoveSelected, "Hide Selection", "Hides the selected nodes", EUserInterfaceActionType::Button, FInputChord(EKeys::Delete));
-	UI_COMMAND(RemoveUnrelated, "Hide Unrelated", "Hides nodes without any dependencies", EUserInterfaceActionType::Button, FInputChord());
+	UI_COMMAND(RemoveUnrelated, "Hide Unrelated", "Hides nodes that not related to the selection", EUserInterfaceActionType::Button, FInputChord());
 	UI_COMMAND(IsolateSelected, "Isolate Selection", "Hides all non-selected nodes", EUserInterfaceActionType::Button, FInputChord(EKeys::I));
-	UI_COMMAND(SyncWithHierarchy, "Sync Hierarchy", "Syncs the content of the graph with the hierarchy selection", EUserInterfaceActionType::Button, FInputChord(EKeys::F5));
 	UI_COMMAND(RunLayoutSimulation, "Auto Layout", "Starts a layout simulation to relax the nodes", EUserInterfaceActionType::Button, FInputChord(EKeys::R));
 	UI_COMMAND(CancelLayoutSimulation, "Cancel Layout", "Cancels the currently running layout simulation", EUserInterfaceActionType::Button, FInputChord(EKeys::Escape));
 	UI_COMMAND(FindNodesByName, "Find Nodes By Name", "Finds nodes by name", EUserInterfaceActionType::Button, FInputChord(EModifierKey::Control, EKeys::F));
diff --git a/Engine/Plugins/Animation/ControlRig/Source/ControlRigEditor/Private/RigDependencyGraph/RigDependencyGraphCommands.h b/Engine/Plugins/Animation/ControlRig/Source/ControlRigEditor/Private/RigDependencyGraph/RigDependencyGraphCommands.h
--- a/Engine/Plugins/Animation/ControlRig/Source/ControlRigEditor/Private/RigDependencyGraph/RigDependencyGraphCommands.h
+++ b/Engine/Plugins/Animation/ControlRig/Source/ControlRigEditor/Private/RigDependencyGraph/RigDependencyGraphCommands.h
@@ -21,14 +21,20 @@ public:
 	/* frames the selection in the graph */
 	TSharedPtr<FUICommandInfo> FrameSelection;
 
+	/* Jumps to the instruction node backing up the dependency node */
+	TSharedPtr<FUICommandInfo> JumpToSource;
+
 	/* Show parent-child relationships */
 	TSharedPtr< FUICommandInfo > ShowParentChildRelationships;
 
 	/* Show instruction based relationships */
 	TSharedPtr< FUICommandInfo > ShowInstructionRelationships;
 
-	/* Syncs the content with other views */
-	TSharedPtr< FUICommandInfo > SyncContent;
+	/* Locks the content of the view and stops syncing based on selection changes */
+	TSharedPtr< FUICommandInfo > LockContent;
+
+	/* Enable a flashlight around the mouse cursor to brighten up faded out nodes */
+	TSharedPtr< FUICommandInfo > EnableFlashlight;
 
 	/* select all nodes */
 	TSharedPtr<FUICommandInfo> SelectAllNodes;
@@ -54,9 +60,6 @@ public:
 	/* remove unrealted nodes */
 	TSharedPtr<FUICommandInfo> RemoveUnrelated;
 
-	/* sync the view with the hierarchy */
-	TSharedPtr<FUICommandInfo> SyncWithHierarchy;
-
 	/* starts a layout simulation to relax the nodes */
 	TSharedPtr<FUICommandInfo> RunLayoutSimulation;
 
diff --git a/Engine/Plugins/Animation/ControlRig/Source/ControlRigEditor/Private/RigDependencyGraph/RigDependencyGraphNode.cpp b/Engine/Plugins/Animation/ControlRig/Source/ControlRigEditor/Private/RigDependencyGraph/RigDependencyGraphNode.cpp
--- a/Engine/Plugins/Animation/ControlRig/Source/ControlRigEditor/Private/RigDependencyGraph/RigDependencyGraphNode.cpp
+++ b/Engine/Plugins/Animation/ControlRig/Source/ControlRigEditor/Private/RigDependencyGraph/RigDependencyGraphNode.cpp
@@ -374,11 +374,6 @@ const FRigVMExternalVariableDef* URigDependencyGraphNode::GetExternalVariable()
 	return &CachedRigVMExternalVariableDef;
 }
 
-bool URigDependencyGraphNode::IsUnrelated() const
-{
-	return InputPin->bHidden && OutputPin->bHidden;
-}
-
 bool URigDependencyGraphNode::IsFadedOut() const
 {
 	return GetFadedOutState() < 1.f - SMALL_NUMBER;
diff --git a/Engine/Plugins/Animation/ControlRig/Source/ControlRigEditor/Private/RigDependencyGraph/RigDependencyGraphNode.h b/Engine/Plugins/Animation/ControlRig/Source/ControlRigEditor/Private/RigDependencyGraph/RigDependencyGraphNode.h
--- a/Engine/Plugins/Animation/ControlRig/Source/ControlRigEditor/Private/RigDependencyGraph/RigDependencyGraphNode.h
+++ b/Engine/Plugins/Animation/ControlRig/Source/ControlRigEditor/Private/RigDependencyGraph/RigDependencyGraphNode.h
@@ -64,7 +64,6 @@ class URigDependencyGraphNode : public UEdGraphNode
 	const URigVMNode* GetRigVMNodeForInstruction() const;
 	const FRigVMExternalVariableDef* GetExternalVariable() const;
 
-	bool IsUnrelated() const;
 	bool IsFadedOut() const;
 	float GetFadedOutState() const;
 	void OverrideFadeOutState(TOptional<float> InFadedOutState);
diff --git a/Engine/Plugins/Animation/ControlRig/Source/ControlRigEditor/Private/RigDependencyGraph/SRigDependencyGraph.cpp b/Engine/Plugins/Animation/ControlRig/Source/ControlRigEditor/Private/RigDependencyGraph/SRigDependencyGraph.cpp
--- a/Engine/Plugins/Animation/ControlRig/Source/ControlRigEditor/Private/RigDependencyGraph/SRigDependencyGraph.cpp
+++ b/Engine/Plugins/Animation/ControlRig/Source/ControlRigEditor/Private/RigDependencyGraph/SRigDependencyGraph.cpp
@@ -11,8 +11,8 @@
 #include "RigDependencyGraph/RigDependencyGraphCommands.h"
 #include "Widgets/Input/SSearchBox.h"
 #include "Widgets/Input/SButton.h"
-#include "Widgets/Input/SCheckBox.h"
 #include "Widgets/Layout/SSpacer.h"
+#include "Widgets/Layout/SSeparator.h"
 #include "Widgets/Notifications/SProgressBar.h"
 #include "Algo/Transform.h"
 #include "Editor/ControlRigEditor.h"
@@ -21,6 +21,7 @@
 #include "Framework/MultiBox/MultiBoxBuilder.h"
 #include "Rigs/RigHierarchyController.h"
 #include "Editor/RigVMEditorStyle.h"
+#include "Settings/ControlRigSettings.h"
 
 #define LOCTEXT_NAMESPACE "RigDependencyGraph"
 
@@ -86,149 +87,230 @@ void SRigDependencyGraph::Construct(const FArguments& InArgs, const TSharedRef<I
 				.Visibility(EVisibility::Visible)
 				.ComboButtonStyle(&FAppStyle::Get().GetWidgetStyle<FComboButtonStyle>("SimpleComboButtonWithIcon"))
 				.ForegroundColor(FSlateColor::UseStyle())
-				.ContentPadding(0.0f)
-				.OnGetMenuContent(this, &SRigDependencyGraph::CreateFilterMenu)
+				.OnGetMenuContent(this, &SRigDependencyGraph::CreateTopLeftMenu)
+				.ContentPadding(2)
 				.ButtonContent()
 				[
 					SNew(SImage)
-					.Image(FAppStyle::Get().GetBrush("Icons.Filter"))
+					.Image(FAppStyle::Get().GetBrush("Icons.Menu"))
+					.DesiredSizeOverride(FVector2D(20, 20))
 					.ColorAndOpacity(FSlateColor::UseForeground())
+					.Image(FControlRigEditorStyle::Get().GetBrush("ControlRig.DependencyGraph.Menu"))
 				 ]
 			]
 
+			+SHorizontalBox::Slot()
+			.MinWidth(100)
+			.MaxWidth(300)
+			.VAlign(VAlign_Center)
+			.Padding(0.0f, 0.0f, 2.0f, 0.0f)
+			[
+				SAssignNew(FilterBox, SSearchBox)
+				.OnTextChanged(this, &SRigDependencyGraph::OnSearchBarTextChanged)
+			]
+
 			+SHorizontalBox::Slot()
 			.AutoWidth()
-			.HAlign(HAlign_Left)
 			.VAlign(VAlign_Center)
 			.Padding(0.0f, 0.0f, 2.0f, 0.0f)
 			[
-				SNew(SCheckBox)
-				.Padding(FMargin(8.0f, 2.0f, 8.0f, 2.0f))
-				.Style(FAppStyle::Get(), "ToggleButtonCheckbox")
-				.ToolTipText(LOCTEXT("SyncContentTooltip", "Sync the view with the selection in the hierarchy and / or rig graphs"))
-				.IsChecked_Lambda( [this] () -> ECheckBoxState
+				SNew(SButton)
+				.ButtonStyle(FAppStyle::Get(), "SimpleButton")
+				.OnClicked_Lambda([this]() -> FReply
+				{
+					OnSearchMatchPicked(MatchIndex.Get(0) - 1);
+					return FReply::Handled();
+				})
+				.Visibility_Lambda([this]() -> EVisibility
 				{
-					return GraphObj->bLockContent ? ECheckBoxState::Unchecked : ECheckBoxState::Checked;
+					return NodesMatchingSearch.IsEmpty() ? EVisibility::Collapsed : EVisibility::Visible;
 				})
-				.OnCheckStateChanged_Lambda([this](ECheckBoxState InCheckBoxState)
+				.IsEnabled_Lambda([this]() -> bool
 				{
-					GraphObj->bLockContent = InCheckBoxState == ECheckBoxState::Unchecked;
+					return MatchIndex.Get(0) > 0;
 				})
+				.ToolTipText(LOCTEXT("SelectPreviousMatch", "Select the previous node matching the search"))
+				.ContentPadding(2)
 				[
 					SNew(SImage)
+					.DesiredSizeOverride(FVector2D(16, 16))
 					.ColorAndOpacity(FSlateColor::UseForeground())
-					.DesiredSizeOverride(FVector2D(20, 20))
-					.Image_Lambda([this]() -> const FSlateBrush*
-					{
-						static const FSlateBrush* Unlocked = FAppStyle::GetBrush("Icons.Unlock");
-						static const FSlateBrush* Locked = FAppStyle::GetBrush("Icons.Lock");
-						return GraphObj->bLockContent ? Locked : Unlocked;
-					})
+					.Image(FAppStyle::Get().GetBrush("Animation.Backward_Step"))
 				]
 			]
 
 			+SHorizontalBox::Slot()
 			.AutoWidth()
-			.HAlign(HAlign_Left)
 			.VAlign(VAlign_Center)
 			.Padding(0.0f, 0.0f, 2.0f, 0.0f)
 			[
-				SNew(SCheckBox)
-				.Padding(FMargin(8.0f, 2.0f, 8.0f, 2.0f))
-				.Style(FAppStyle::Get(), "ToggleButtonCheckbox")
-				.ToolTipText(LOCTEXT("ShowParentChildRelationShips", "Show the parent-child relationships"))
-				.IsChecked_Lambda( [this] () -> ECheckBoxState
+				SNew(SButton)
+				.ButtonStyle(FAppStyle::Get(), "SimpleButton")
+				.OnClicked_Lambda([this]() -> FReply
+				{
+					OnSearchMatchPicked(MatchIndex.Get(0) + 1);
+					return FReply::Handled();
+				})
+				.Visibility_Lambda([this]() -> EVisibility
 				{
-					return GraphObj->bFollowParentRelationShips ? ECheckBoxState::Checked : ECheckBoxState::Unchecked;
+					return NodesMatchingSearch.IsEmpty() ? EVisibility::Collapsed : EVisibility::Visible;
 				})
-				.OnCheckStateChanged_Lambda([this](ECheckBoxState InCheckBoxState)
+				.IsEnabled_Lambda([this]() -> bool
 				{
-					GraphObj->SetFollowParentRelationShips(InCheckBoxState != ECheckBoxState::Unchecked);
+					return MatchIndex.Get(0) < NodesMatchingSearch.Num() - 1;
 				})
+				.ToolTipText(LOCTEXT("SelectNextMatch", "Select the next node matching the search"))
+				.ContentPadding(2)
 				[
 					SNew(SImage)
+					.DesiredSizeOverride(FVector2D(16, 16))
 					.ColorAndOpacity(FSlateColor::UseForeground())
-					.DesiredSizeOverride(FVector2D(20,20))
-					.Image(FAppStyle::Get().GetBrush("LevelEditor.CreateBlankBlueprintClass"))
+					.Image(FAppStyle::Get().GetBrush("Animation.Forward_Step"))
 				]
 			]
-			
+
+			+ SHorizontalBox::Slot()
+			.AutoWidth()
+			.Padding(8.0f, 0.0f, 10.0f, 0.0f)
+			[
+				SNew(SSeparator)
+				.Orientation(Orient_Vertical)
+			]
+
 			+SHorizontalBox::Slot()
 			.AutoWidth()
 			.HAlign(HAlign_Left)
 			.VAlign(VAlign_Center)
 			.Padding(0.0f, 0.0f, 2.0f, 0.0f)
 			[
-				SNew(SCheckBox)
-				.Padding(FMargin(8.0f, 2.0f, 8.0f, 2.0f))
-				.Style(FAppStyle::Get(), "ToggleButtonCheckbox")
-				.ToolTipText(LOCTEXT("ShowVMRelationShips", "Show the relationships introduced by the rig nodes"))
-				.IsChecked_Lambda( [this] () -> ECheckBoxState
-				{
-					return GraphObj->bFollowVMRelationShips ? ECheckBoxState::Checked : ECheckBoxState::Unchecked;
-				})
-				.OnCheckStateChanged_Lambda([this](ECheckBoxState InCheckBoxState)
+				SNew(SButton)
+				.ButtonStyle( FAppStyle::Get(), "SimpleButton" )
+				.ToolTipText(LOCTEXT("LockContentTooltip", "Lock the content of this view / stop syncing content based on selection"))
+				.OnClicked_Lambda([this]() -> FReply
 				{
-					GraphObj->SetFollowVMRelationShips(InCheckBoxState != ECheckBoxState::Unchecked);
+					GraphObj->bLockContent = !GraphObj->bLockContent;
+					return FReply::Handled();
 				})
+				.ContentPadding(2)
 				[
 					SNew(SImage)
-					.ColorAndOpacity(FSlateColor::UseForeground())
+					.ColorAndOpacity_Lambda([this]()
+					{
+						if (GraphObj->bLockContent)
+						{
+							return FStyleColors::AccentBlue;
+						}
+						return FSlateColor::UseForeground();
+					})
 					.DesiredSizeOverride(FVector2D(20, 20))
-					.Image(FRigVMEditorStyle::Get().GetBrush("RigVM.ReleaseMode"))
+					.Image_Lambda([this]() -> const FSlateBrush*
+					{
+						static const FSlateBrush* Unlocked = FAppStyle::GetBrush("Icons.Unlock");
+						static const FSlateBrush* Locked = FAppStyle::GetBrush("Icons.Lock");
+						return GraphObj->bLockContent ? Locked : Unlocked;
+					})
 				]
 			]
 
 			+SHorizontalBox::Slot()
 			.AutoWidth()
+			.HAlign(HAlign_Left)
 			.VAlign(VAlign_Center)
 			.Padding(0.0f, 0.0f, 2.0f, 0.0f)
 			[
 				SNew(SButton)
-				.ButtonStyle(FAppStyle::Get(), "SimpleButton")
+				.ButtonStyle( FAppStyle::Get(), "SimpleButton" )
+				.ToolTipText(LOCTEXT("EnableFlashLight", "Enable a flashlight around the mouse cursor to brighten up faded out nodes"))
 				.OnClicked_Lambda([this]() -> FReply
 				{
-					GraphObj->RemoveUnrelatedNodes();
+					UControlRigEditorSettings::Get()->bEnableFlashlightInDependencyViewer = !UControlRigEditorSettings::Get()->bEnableFlashlightInDependencyViewer;
+					UControlRigEditorSettings::Get()->SaveConfig();
 					return FReply::Handled();
 				})
-				.ToolTipText(LOCTEXT("HideUnrelatedTooltip", "Removes unrelated nodes from the view"))
+				.ContentPadding(2)
 				[
 					SNew(SImage)
-					.DesiredSizeOverride(FVector2D(20, 20))
-					.Image(FControlRigEditorStyle::Get().GetBrush("ControlRig.DependencyGraph.HideUnrelated"))
+					.ColorAndOpacity_Lambda([this]()
+					{
+						if (UControlRigEditorSettings::Get()->bEnableFlashlightInDependencyViewer)
+						{
+							return FStyleColors::AccentBlue;
+						}
+						return FSlateColor::UseForeground();
+					})
+					.DesiredSizeOverride(FVector2D(20,20))
+					.Image(FControlRigEditorStyle::Get().GetBrush("ControlRig.DependencyGraph.Flashlight"))
 				]
 			]
-
+			
 			+SHorizontalBox::Slot()
 			.AutoWidth()
+			.HAlign(HAlign_Left)
 			.VAlign(VAlign_Center)
 			.Padding(0.0f, 0.0f, 2.0f, 0.0f)
 			[
 				SNew(SButton)
-				.ButtonStyle(FAppStyle::Get(), "SimpleButton")
+				.ButtonStyle( FAppStyle::Get(), "SimpleButton" )
+				.ToolTipText(LOCTEXT("ShowParentChildRelationShips", "Show the parent-child relationships"))
 				.OnClicked_Lambda([this]() -> FReply
 				{
-					GraphObj->RemoveUnselectedNodes();
+					GraphObj->SetFollowParentRelationShips(!GraphObj->bFollowParentRelationShips);
 					return FReply::Handled();
 				})
-				.ToolTipText(LOCTEXT("IsolateSelectionTooltip", "Isolates the selection by removing unselected nodes"))
+				.ContentPadding(2)
 				[
 					SNew(SImage)
-					.DesiredSizeOverride(FVector2D(20, 20))
-					.Image(FControlRigEditorStyle::Get().GetBrush("ControlRig.DependencyGraph.IsolateSelection"))
+					.ColorAndOpacity_Lambda([this]()
+					{
+						if (GraphObj->bFollowParentRelationShips)
+						{
+							return FStyleColors::AccentBlue;
+						}
+						return FSlateColor::UseForeground();
+					})
+					.DesiredSizeOverride(FVector2D(20,20))
+					.Image(FControlRigEditorStyle::Get().GetBrush("ControlRig.DependencyGraph.ParentChild"))
 				]
 			]
-
+			
 			+SHorizontalBox::Slot()
-			.MinWidth(100)
-			.MaxWidth(300)
+			.AutoWidth()
+			.HAlign(HAlign_Left)
 			.VAlign(VAlign_Center)
 			.Padding(0.0f, 0.0f, 2.0f, 0.0f)
 			[
-				SAssignNew(FilterBox, SSearchBox)
-				.OnTextChanged(this, &SRigDependencyGraph::OnSearchBarTextChanged)
+				SNew(SButton)
+				.ButtonStyle( FAppStyle::Get(), "SimpleButton" )
+				.ToolTipText(LOCTEXT("ShowVMRelationShips", "Show the relationships introduced by the rig nodes"))
+				.OnClicked_Lambda([this]() -> FReply
+				{
+					GraphObj->SetFollowVMRelationShips(!GraphObj->bFollowVMRelationShips);
+					return FReply::Handled();
+				})
+				.ContentPadding(2)
+				[
+					SNew(SImage)
+					.ColorAndOpacity_Lambda([this]()
+					{
+						if (GraphObj->bFollowVMRelationShips)
+						{
+							return FStyleColors::AccentBlue;
+						}
+						return FSlateColor::UseForeground();
+					})
+					.DesiredSizeOverride(FVector2D(20,20))
+					.Image(FControlRigEditorStyle::Get().GetBrush("ControlRig.DependencyGraph.VMRelationShips"))
+				]
 			]
 
+			+ SHorizontalBox::Slot()
+			.AutoWidth()
+			.Padding(8.0f, 0.0f, 10.0f, 0.0f)
+			[
+				SNew(SSeparator)
+				.Orientation(Orient_Vertical)
+			]
+		
 			+SHorizontalBox::Slot()
 			.AutoWidth()
 			.VAlign(VAlign_Center)
@@ -238,22 +320,15 @@ void SRigDependencyGraph::Construct(const FArguments& InArgs, const TSharedRef<I
 				.ButtonStyle(FAppStyle::Get(), "SimpleButton")
 				.OnClicked_Lambda([this]() -> FReply
 				{
-					OnSearchMatchPicked(MatchIndex.Get(0) - 1);
+					GraphObj->RemoveUnrelatedNodes();
 					return FReply::Handled();
 				})
-				.Visibility_Lambda([this]() -> EVisibility
-				{
-					return NodesMatchingSearch.IsEmpty() ? EVisibility::Collapsed : EVisibility::Visible;
-				})
-				.IsEnabled_Lambda([this]() -> bool
-				{
-					return MatchIndex.Get(0) > 0;
-				})
-				.ToolTipText(LOCTEXT("SelectPreviousMatch", "Select the previous node matching the search"))
+				.ToolTipText(LOCTEXT("HideUnrelatedTooltip", "Hides unrelated nodes from the view"))
 				[
 					SNew(SImage)
-					.DesiredSizeOverride(FVector2D(16, 16))
-					.Image(FAppStyle::Get().GetBrush("Animation.Backward_Step"))
+					.DesiredSizeOverride(FVector2D(20, 20))
+					.ColorAndOpacity(FSlateColor::UseForeground())
+					.Image(FControlRigEditorStyle::Get().GetBrush("ControlRig.DependencyGraph.HideUnrelated"))
 				]
 			]
 
@@ -266,34 +341,68 @@ void SRigDependencyGraph::Construct(const FArguments& InArgs, const TSharedRef<I
 				.ButtonStyle(FAppStyle::Get(), "SimpleButton")
 				.OnClicked_Lambda([this]() -> FReply
 				{
-					OnSearchMatchPicked(MatchIndex.Get(0) + 1);
+					GraphObj->RemoveUnselectedNodes();
 					return FReply::Handled();
 				})
+				.ToolTipText(LOCTEXT("IsolateSelectionTooltip", "Isolates the selection by removing unselected nodes"))
+				[
+					SNew(SImage)
+					.DesiredSizeOverride(FVector2D(20, 20))
+					.ColorAndOpacity(FSlateColor::UseForeground())
+					.Image(FControlRigEditorStyle::Get().GetBrush("ControlRig.DependencyGraph.IsolateSelection"))
+				]
+			]
+
+			+ SHorizontalBox::Slot()
+			.AutoWidth()
+			.Padding(8.0f, 0.0f, 10.0f, 0.0f)
+			[
+				SNew(SSeparator)
+				.Orientation(Orient_Vertical)
 				.Visibility_Lambda([this]() -> EVisibility
 				{
-					return NodesMatchingSearch.IsEmpty() ? EVisibility::Collapsed : EVisibility::Visible;
+					return GraphObj->LayoutIterationsLeft > 0 ? EVisibility::Visible : EVisibility::Collapsed;
 				})
-				.IsEnabled_Lambda([this]() -> bool
+			]
+
+			+ SHorizontalBox::Slot()
+			.AutoWidth()
+			.Padding(8.0f, 0.0f, 10.0f, 0.0f)
+			[
+				SNew(SSeparator)
+				.Orientation(Orient_Vertical)
+			]
+
+			+SHorizontalBox::Slot()
+			.AutoWidth()
+			.VAlign(VAlign_Center)
+			.Padding(0.0f, 0.0f, 2.0f, 0.0f)
+			[
+				SNew(SButton)
+				.ButtonStyle(FAppStyle::Get(), "SimpleButton")
+				.OnClicked_Lambda([this]() -> FReply
 				{
-					return MatchIndex.Get(0) < NodesMatchingSearch.Num() - 1;
+					GraphObj->RemoveAllNodes();
+					return FReply::Handled();
 				})
-				.ToolTipText(LOCTEXT("SelectNextMatch", "Select the next node matching the search"))
+				.ToolTipText(LOCTEXT("RemoveAllNodesTooltip", "Removes all of the nodes from the view"))
 				[
 					SNew(SImage)
-					.DesiredSizeOverride(FVector2D(16, 16))
-					.Image(FAppStyle::Get().GetBrush("Animation.Forward_Step"))
+					.DesiredSizeOverride(FVector2D(20, 20))
+					.ColorAndOpacity(FSlateColor::UseForeground())
+					.Image(FControlRigEditorStyle::Get().GetBrush("ControlRig.DependencyGraph.Cleanup"))
 				]
 			]
 
 			+SHorizontalBox::Slot()
 			.VAlign(VAlign_Center)
 			.HAlign(HAlign_Fill)
-			.Padding(4.f, 0.f, 4.f, 0.f)
+			.Padding(2.f, 0.f, 4.f, 0.f)
 			[
 				SAssignNew(LayoutProgressBar, SProgressBar)
 				.Visibility_Lambda([this]() -> EVisibility
 				{
-					return GraphObj->LayoutIterationsLeft > 0 ? EVisibility::Visible : EVisibility::Hidden;
+					return GraphObj->LayoutIterationsLeft > 0 ? EVisibility::Visible : EVisibility::Collapsed;
 				})
 				.Percent_Lambda([this]() -> TOptional<float>
 				{
@@ -313,13 +422,13 @@ void SRigDependencyGraph::Construct(const FArguments& InArgs, const TSharedRef<I
 			+SHorizontalBox::Slot()
 			.AutoWidth()
 			.VAlign(VAlign_Center)
-			.Padding(4)
+			.Padding(0.0f, 0.0f, 2.0f, 0.0f)
 			[
 				SNew(SButton)
 				.ButtonStyle(FAppStyle::Get(), "SimpleButton")
 				.Visibility_Lambda([this]() -> EVisibility
 				{
-					return GraphObj->LayoutIterationsLeft > 0 ? EVisibility::Visible : EVisibility::Hidden;
+					return GraphObj->LayoutIterationsLeft > 0 ? EVisibility::Visible : EVisibility::Collapsed;
 				})
 				.OnClicked_Lambda([this]() -> FReply
 				{
@@ -329,7 +438,8 @@ void SRigDependencyGraph::Construct(const FArguments& InArgs, const TSharedRef<I
 				.ToolTipText(LOCTEXT("CancelLayoutTooltip", "Cancels the layout simulation"))
 				[
 					SNew(SImage)
-					.DesiredSizeOverride(FVector2D(16, 16))
+					.DesiredSizeOverride(FVector2D(20, 20))
+					.ColorAndOpacity(FSlateColor::UseForeground())
 					.Image(FAppStyle::Get().GetBrush("Icons.X"))
 				]
 			]
@@ -417,29 +527,6 @@ void SRigDependencyGraph::HandleNodeDoubleClicked(UEdGraphNode* InNode)
 		GraphObj->ConstructNodes(AllNodeIds, false);
 		return;
 	}
-
-	// if this is an instruction - also navigate to the instruction
-	const bool bIsShiftPressed = FSlateApplication::Get().GetModifierKeys().IsShiftDown();
-	if (RigDependencyGraphNode->GetNodeId().IsInstruction() && bIsShiftPressed)
-	{
-		TSharedPtr<IControlRigBaseEditor> ControlRigEditor = WeakControlRigEditor.Pin();
-		if (ControlRigEditor.IsValid())
-		{
-			if (FControlRigAssetInterfacePtr ControlRigBlueprint = ControlRigEditor->GetControlRigAssetInterface())
-			{
-				if (FRigVMClient* Client = ControlRigBlueprint->GetRigVMClient())
-				{
-					if (URigVMController* Controller = Client->GetOrCreateController(Client->GetDefaultModel()))
-					{
-						if (const URigVMNode* Node = RigDependencyGraphNode->GetRigVMNodeForInstruction())
-						{
-							Controller->RequestJumpToHyperLink(Node);
-						}
-					}
-				}
-			}
-		}
-	}
 }
 
 void SRigDependencyGraph::OnClearNodeSelection()
@@ -450,7 +537,7 @@ void SRigDependencyGraph::OnClearNodeSelection()
 	}
 }
 
-TSharedRef<SWidget> SRigDependencyGraph::CreateFilterMenu()
+TSharedRef<SWidget> SRigDependencyGraph::CreateTopLeftMenu()
 {
 	const FRigDependencyGraphCommands& Actions = FRigDependencyGraphCommands::Get();
 
@@ -459,6 +546,8 @@ TSharedRef<SWidget> SRigDependencyGraph::CreateFilterMenu()
 
 	MenuBuilder.BeginSection("FilterOptions", LOCTEXT("OptionsMenuHeading", "Options"));
 	{
+		MenuBuilder.AddMenuEntry(Actions.LockContent);
+		MenuBuilder.AddMenuEntry(Actions.EnableFlashlight);
 		MenuBuilder.AddMenuEntry(Actions.ShowParentChildRelationships);
 		MenuBuilder.AddMenuEntry(Actions.ShowInstructionRelationships);
 	}
@@ -467,12 +556,12 @@ TSharedRef<SWidget> SRigDependencyGraph::CreateFilterMenu()
 	MenuBuilder.BeginSection("Navigation", LOCTEXT("NavigationMenuHeading", "Navigation"));
 	{
 		MenuBuilder.AddMenuEntry(Actions.FrameSelection);
+		MenuBuilder.AddMenuEntry(Actions.JumpToSource);
 	}
 	MenuBuilder.EndSection();
 
 	MenuBuilder.BeginSection("Graph", LOCTEXT("GraphMenuHeading", "Graph"));
 	{
-		MenuBuilder.AddMenuEntry(Actions.SyncContent);
 		MenuBuilder.AddMenuEntry(Actions.SelectAllNodes);
 		MenuBuilder.AddMenuEntry(Actions.SelectInputNodes);
 		MenuBuilder.AddMenuEntry(Actions.SelectOutputNodes);
@@ -481,7 +570,6 @@ TSharedRef<SWidget> SRigDependencyGraph::CreateFilterMenu()
 		MenuBuilder.AddMenuEntry(Actions.RemoveSelected);
 		MenuBuilder.AddMenuEntry(Actions.IsolateSelected);
 		MenuBuilder.AddMenuEntry(Actions.RemoveUnrelated);
-		MenuBuilder.AddMenuEntry(Actions.SyncWithHierarchy);
 		MenuBuilder.AddMenuEntry(Actions.RunLayoutSimulation);
 		MenuBuilder.AddMenuEntry(Actions.CancelLayoutSimulation);
 	}
@@ -497,13 +585,64 @@ void SRigDependencyGraph::BindCommands()
 
 	CommandList = MakeShared<FUICommandList>();
 	
-	CommandList->MapAction(Commands.FrameSelection, FExecuteAction::CreateLambda([this] ()
-	{
-		if (GraphEditor)
+	CommandList->MapAction(Commands.FrameSelection,
+		FExecuteAction::CreateLambda([this] ()
 		{
-			GraphEditor->ZoomToFit(true);
-		}
-	}));
+			if (GraphEditor)
+			{
+				GraphEditor->ZoomToFit(true);
+			}
+		}),
+		FCanExecuteAction::CreateLambda([this] () -> bool
+		{
+			return GraphObj->GetSelectedNodeIds().Num() > 0;
+		})		
+	);
+
+	CommandList->MapAction(Commands.JumpToSource,
+		FExecuteAction::CreateLambda([this] ()
+		{
+			for (const FNodeId& SelectedNodeId : GraphObj->GetSelectedNodeIds())
+			{
+				if (SelectedNodeId.IsInstruction())
+				{
+					TSharedPtr<IControlRigBaseEditor> ControlRigEditor = WeakControlRigEditor.Pin();
+					if (ControlRigEditor.IsValid())
+					{
+						if (FControlRigAssetInterfacePtr ControlRigBlueprint = ControlRigEditor->GetControlRigAssetInterface())
+						{
+							if (FRigVMClient* Client = ControlRigBlueprint->GetRigVMClient())
+							{
+								if (URigVMController* Controller = Client->GetOrCreateController(Client->GetDefaultModel()))
+								{
+									if (const URigDependencyGraphNode* RigDependencyGraphNode = GraphObj->FindNode(SelectedNodeId))
+									{
+										if (const URigVMNode* Node = RigDependencyGraphNode->GetRigVMNodeForInstruction())
+										{
+											Controller->RequestJumpToHyperLink(Node);
+											break;
+										}
+									}
+								}
+							}
+						}
+					}
+				}
+			}
+		}),
+		FCanExecuteAction::CreateLambda([this] () -> bool
+		{
+			for (const FNodeId& SelectedNodeId : GraphObj->GetSelectedNodeIds())
+			{
+				if (SelectedNodeId.IsInstruction())
+				{
+					return true;
+				}
+				break;
+			}
+			return false;
+		})
+	);
 
 	CommandList->MapAction(Commands.SelectAllNodes, FExecuteAction::CreateLambda([this] ()
 	{
@@ -545,27 +684,6 @@ void SRigDependencyGraph::BindCommands()
 		GraphObj->RemoveUnrelatedNodes();
 	}));
 
-	CommandList->MapAction(Commands.SyncWithHierarchy, FExecuteAction::CreateLambda([this] ()
-	{
-		URigHierarchy* Hierarchy = GraphObj->GetRigHierarchy();
-		if (Hierarchy == nullptr)
-		{
-			return;
-		}
-
-		GraphObj->RemoveAllNodes();
-			
-		const TArray<FRigElementKey> SelectedElements = Hierarchy->GetSelectedKeys();
-
-		TArray<FNodeId> NodesToSelect;
-		for (const FRigElementKey& SelectedKey : SelectedElements)
-		{
-			NodesToSelect.Emplace(FNodeId::EType_RigElement, Hierarchy->GetIndex(SelectedKey), NAME_None);
-		}
-
-		GraphObj->ConstructNodes(NodesToSelect, false);
-	}));
-
 	CommandList->MapAction(Commands.RunLayoutSimulation,
 		FExecuteAction::CreateLambda([this] ()
 		{
@@ -625,18 +743,32 @@ void SRigDependencyGraph::BindCommands()
 	);
 
 	CommandList->MapAction(
-		Commands.SyncContent,
+		Commands.EnableFlashlight,
 		FExecuteAction::CreateLambda([this]()
 		{
-			GraphObj->bLockContent = !GraphObj->bLockContent;
+			UControlRigEditorSettings::Get()->bEnableFlashlightInDependencyViewer = !UControlRigEditorSettings::Get()->bEnableFlashlightInDependencyViewer;
+			UControlRigEditorSettings::Get()->SaveConfig();
 		}),
 		FCanExecuteAction(),
 		FIsActionChecked::CreateLambda([this]()
 		{
-			return !GraphObj->bLockContent;
+			return UControlRigEditorSettings::Get()->bEnableFlashlightInDependencyViewer;
 		})
 	);
 
+	CommandList->MapAction(
+		Commands.LockContent,
+		FExecuteAction::CreateLambda([this]()
+		{
+			GraphObj->bLockContent = !GraphObj->bLockContent;
+		}),
+		FCanExecuteAction(),
+		FIsActionChecked::CreateLambda([this]()
+		{
+			return GraphObj->bLockContent;
+		})
+	);
+	
 	CommandList->MapAction(Commands.FindNodesByName,
 	FExecuteAction::CreateLambda([this] ()
 	{
@@ -805,7 +937,7 @@ void SRigDependencyGraph::OnGraphEditorTick(const FGeometry& AllottedGeometry, c
 
 	const float FlashLightRadius = FlashLightRadiusAttribute.Get();
 
-	if(PaintGeometry.IsUnderLocation(MousePosition) && FlashLightRadius > SMALL_NUMBER)
+	if(UControlRigEditorSettings::Get()->bEnableFlashlightInDependencyViewer && PaintGeometry.IsUnderLocation(MousePosition) && FlashLightRadius > SMALL_NUMBER)
 	{
 		FVector2f GraphEditorLocation = FVector2f::ZeroVector;
 		float GraphEditorZoomAmount = 0;
diff --git a/Engine/Plugins/Animation/ControlRig/Source/ControlRigEditor/Private/RigDependencyGraph/SRigDependencyGraph.h b/Engine/Plugins/Animation/ControlRig/Source/ControlRigEditor/Private/RigDependencyGraph/SRigDependencyGraph.h
--- a/Engine/Plugins/Animation/ControlRig/Source/ControlRigEditor/Private/RigDependencyGraph/SRigDependencyGraph.h
+++ b/Engine/Plugins/Animation/ControlRig/Source/ControlRigEditor/Private/RigDependencyGraph/SRigDependencyGraph.h
@@ -67,7 +67,7 @@ private:
 
 	void OnClearNodeSelection();
 
-	TSharedRef< SWidget > CreateFilterMenu();
+	TSharedRef< SWidget > CreateTopLeftMenu();
 
 	void BindCommands();
 
