Anim in Engine: Add default Hotkeys for Control Rig Editor

* Enable profiling- P
* Show Controls as Overlay - U
* Display Nulls - G
* Display Socket - H
* Display Axes for selection - J

diff --git a/Engine/Plugins/Animation/ControlRig/Source/ControlRigEditor/Private/Editor/ControlRigEditorCommands.cpp b/Engine/Plugins/Animation/ControlRig/Source/ControlRigEditor/Private/Editor/ControlRigEditorCommands.cpp
--- a/Engine/Plugins/Animation/ControlRig/Source/ControlRigEditor/Private/Editor/ControlRigEditorCommands.cpp
+++ b/Engine/Plugins/Animation/ControlRig/Source/ControlRigEditor/Private/Editor/ControlRigEditorCommands.cpp
@@ -14,10 +14,10 @@ void FControlRigEditorCommands::RegisterCommands()
 	UI_COMMAND(RequestDirectManipulationRotation, "Request Direct Manipulation for Rotation", "Request per node direct manipulation on a rotation", EUserInterfaceActionType::Button, FInputChord(EKeys::E));
 	UI_COMMAND(RequestDirectManipulationScale, "Request Direct Manipulation for Scale", "Request per node direct manipulation on a scale", EUserInterfaceActionType::Button, FInputChord(EKeys::R));
 	UI_COMMAND(ToggleControlVisibility, "Show Controls", "Toggles the visibility of the controls.", EUserInterfaceActionType::ToggleButton, FInputChord(EKeys::T));
-	UI_COMMAND(ToggleControlsAsOverlay, "Show Controls as Overlay", "If checked controls will be rendered on top of other controls.", EUserInterfaceActionType::ToggleButton, FInputChord());
-	UI_COMMAND(ToggleDrawNulls, "Display Nulls", "If checked all nulls are drawn as axes.", EUserInterfaceActionType::ToggleButton, FInputChord());
-	UI_COMMAND(ToggleDrawSockets, "Display Sockets", "If checked all sockets are drawn.", EUserInterfaceActionType::ToggleButton, FInputChord());
-	UI_COMMAND(ToggleDrawAxesOnSelection, "Display Axes On Selection", "If checked axes will be drawn for all selected rig elements.", EUserInterfaceActionType::ToggleButton, FInputChord());
+	UI_COMMAND(ToggleControlsAsOverlay, "Show Controls as Overlay", "If checked controls will be rendered on top of other controls.", EUserInterfaceActionType::ToggleButton, FInputChord(EKeys::U));
+	UI_COMMAND(ToggleDrawNulls, "Display Nulls", "If checked all nulls are drawn as axes.", EUserInterfaceActionType::ToggleButton, FInputChord(EKeys::G));
+	UI_COMMAND(ToggleDrawSockets, "Display Sockets", "If checked all sockets are drawn.", EUserInterfaceActionType::ToggleButton, FInputChord(EKeys::H));
+	UI_COMMAND(ToggleDrawAxesOnSelection, "Display Axes On Selection", "If checked axes will be drawn for all selected rig elements.", EUserInterfaceActionType::ToggleButton, FInputChord(EKeys::J));
 	UI_COMMAND(ToggleSchematicViewportVisibility, "Show Schematic Viewport", "Toggles the visibility of the viewport schematic", EUserInterfaceActionType::ToggleButton, FInputChord(EKeys::Y));
 	UI_COMMAND(SwapModuleWithinAsset, "Swap Module (Asset)", "Swaps a module for all occurrences within this asset.", EUserInterfaceActionType::Button, FInputChord());
 	UI_COMMAND(SwapModuleAcrossProject, "Swap Module (Project)", "Swaps a module for all occurrences in the project.", EUserInterfaceActionType::Button, FInputChord());
diff --git a/Engine/Plugins/Runtime/RigVM/Source/RigVMDeveloper/Private/RigVMAsset.cpp b/Engine/Plugins/Runtime/RigVM/Source/RigVMDeveloper/Private/RigVMAsset.cpp
--- a/Engine/Plugins/Runtime/RigVM/Source/RigVMDeveloper/Private/RigVMAsset.cpp
+++ b/Engine/Plugins/Runtime/RigVM/Source/RigVMDeveloper/Private/RigVMAsset.cpp
@@ -3687,6 +3687,13 @@ void IRigVMAssetInterface::BroadcastPostEditChangeChainProperty(FPropertyChanged
 	PostEditChangeChainPropertyEvent.Broadcast(PropertyChangedChainEvent);
 }
 
+void IRigVMAssetInterface::SetProfilingEnabled(const bool bEnabled)
+{
+	GetVMRuntimeSettings().bEnableProfiling = bEnabled;
+	PropagateRuntimeSettingsFromBPToInstances();
+	RequestAutoVMRecompilation();
+}
+
 void IRigVMAssetInterface::RefreshAllNodes()
 {
 	DECLARE_SCOPE_HIERARCHICAL_COUNTER_FUNC()
diff --git a/Engine/Plugins/Runtime/RigVM/Source/RigVMDeveloper/Public/RigVMAsset.h b/Engine/Plugins/Runtime/RigVM/Source/RigVMDeveloper/Public/RigVMAsset.h
--- a/Engine/Plugins/Runtime/RigVM/Source/RigVMDeveloper/Public/RigVMAsset.h
+++ b/Engine/Plugins/Runtime/RigVM/Source/RigVMDeveloper/Public/RigVMAsset.h
@@ -350,6 +350,10 @@ public:
 	virtual URigVMGraph* GetDefaultModel() const;
 	virtual TArray<TObjectPtr<UEdGraph>>& GetUberGraphs() = 0;
 	virtual const TArray<TObjectPtr<UEdGraph>>& GetUberGraphs() const { return const_cast<IRigVMAssetInterface*>(this)->GetUberGraphs(); };
+
+	/** Enables or disables profiling for this asset */
+	void SetProfilingEnabled(const bool bEnabled);
+
 protected:
 	IRigVMAssetInterface() {}
 	IRigVMAssetInterface(const FObjectInitializer& ObjectInitializer);
diff --git a/Engine/Plugins/Runtime/RigVM/Source/RigVMEditor/Private/Editor/RigVMEditor.cpp b/Engine/Plugins/Runtime/RigVM/Source/RigVMEditor/Private/Editor/RigVMEditor.cpp
--- a/Engine/Plugins/Runtime/RigVM/Source/RigVMEditor/Private/Editor/RigVMEditor.cpp
+++ b/Engine/Plugins/Runtime/RigVM/Source/RigVMEditor/Private/Editor/RigVMEditor.cpp
@@ -1435,6 +1435,11 @@ void FRigVMEditorBase::BindCommands()
 		EditorCommands.SwapAssetReferences,
 		FExecuteAction::CreateSP(ThisRef, &FRigVMEditorBase::SwapAssetReferences),
 		FCanExecuteAction());
+
+	ToolkitCommands->MapAction(
+		EditorCommands.ToggleProfiling,
+		FExecuteAction::CreateSP(ThisRef, &FRigVMEditorBase::ToggleProfiling),
+		FCanExecuteAction());
 }
 
 void FRigVMEditorBase::ToggleAutoCompileGraph()
@@ -3786,6 +3791,14 @@ void FRigVMEditorBase::SwapAssetReferences()
 	SwapFunctionsDialog->ShowNormal();
 }
 
+void FRigVMEditorBase::ToggleProfiling()
+{
+	if (TScriptInterface<IRigVMAssetInterface> Asset = GetRigVMAssetInterface())
+	{
+		Asset->SetProfilingEnabled(!Asset->GetVMRuntimeSettings().bEnableProfiling);
+	}
+}
+
 void FRigVMEditorBase::OnOpenSelectedNodesInNewTab()
 {
 	const FRigVMAssetInterfacePtr RigVMAsset = GetRigVMAssetInterface();
diff --git a/Engine/Plugins/Runtime/RigVM/Source/RigVMEditor/Private/Editor/RigVMEditorCommands.cpp b/Engine/Plugins/Runtime/RigVM/Source/RigVMEditor/Private/Editor/RigVMEditorCommands.cpp
--- a/Engine/Plugins/Runtime/RigVM/Source/RigVMEditor/Private/Editor/RigVMEditorCommands.cpp
+++ b/Engine/Plugins/Runtime/RigVM/Source/RigVMEditor/Private/Editor/RigVMEditorCommands.cpp
@@ -24,6 +24,8 @@ void FRigVMEditorCommands::RegisterCommands()
 	UI_COMMAND(SwapFunctionWithinAsset, "Swap Function (Asset)", "Swaps a function for all occurrences within this asset.", EUserInterfaceActionType::Button, FInputChord());
 	UI_COMMAND(SwapFunctionAcrossProject, "Swap Function (Project)", "Swaps a function for all occurrences in the project.", EUserInterfaceActionType::Button, FInputChord());
 	UI_COMMAND(SwapAssetReferences, "Swap Asset References", "Swaps an asset reference for another asset.", EUserInterfaceActionType::Button, FInputChord());
+
+	UI_COMMAND(ToggleProfiling, "Toggle Profiling", "Enables or disables heat map profiling for the graph.", EUserInterfaceActionType::ToggleButton, FInputChord(EKeys::P));
 }
 
 #undef LOCTEXT_NAMESPACE
diff --git a/Engine/Plugins/Runtime/RigVM/Source/RigVMEditor/Private/Editor/RigVMEditorCommands.h b/Engine/Plugins/Runtime/RigVM/Source/RigVMEditor/Private/Editor/RigVMEditorCommands.h
--- a/Engine/Plugins/Runtime/RigVM/Source/RigVMEditor/Private/Editor/RigVMEditorCommands.h
+++ b/Engine/Plugins/Runtime/RigVM/Source/RigVMEditor/Private/Editor/RigVMEditorCommands.h
@@ -54,6 +54,9 @@ public:
 	/** Swap Asset References */
 	TSharedPtr< FUICommandInfo > SwapAssetReferences;
 
+	/** Toggle heatmap profiling for the graph */
+	TSharedPtr<FUICommandInfo> ToggleProfiling;
+
 	/**
 	 * Initialize commands
 	 */
diff --git a/Engine/Plugins/Runtime/RigVM/Source/RigVMEditor/Private/RigVMEditorModule.cpp b/Engine/Plugins/Runtime/RigVM/Source/RigVMEditor/Private/RigVMEditorModule.cpp
--- a/Engine/Plugins/Runtime/RigVM/Source/RigVMEditor/Private/RigVMEditorModule.cpp
+++ b/Engine/Plugins/Runtime/RigVM/Source/RigVMEditor/Private/RigVMEditorModule.cpp
@@ -1480,9 +1480,7 @@ void FRigVMEditorModule::GetNodeDisplayContextMenuActions(TScriptInterface<IRigV
 		{
 			FScopedTransaction Transaction(LOCTEXT("ToggleProfiler", "Toggle Profiler"));
 			Blueprint.GetObject()->Modify();
-			Blueprint->GetVMRuntimeSettings().bEnableProfiling = !Blueprint->GetVMRuntimeSettings().bEnableProfiling;
-			Blueprint->PropagateRuntimeSettingsFromBPToInstances();
-			Blueprint->RequestAutoVMRecompilation();
+			Blueprint->SetProfilingEnabled(!Blueprint->GetVMRuntimeSettings().bEnableProfiling);
 		}),
 		FCanExecuteAction(),
 		FGetActionCheckState::CreateLambda([Blueprint]() -> ECheckBoxState
diff --git a/Engine/Plugins/Runtime/RigVM/Source/RigVMEditor/Public/Editor/RigVMEditor.h b/Engine/Plugins/Runtime/RigVM/Source/RigVMEditor/Public/Editor/RigVMEditor.h
--- a/Engine/Plugins/Runtime/RigVM/Source/RigVMEditor/Public/Editor/RigVMEditor.h
+++ b/Engine/Plugins/Runtime/RigVM/Source/RigVMEditor/Public/Editor/RigVMEditor.h
@@ -540,6 +540,9 @@ protected:
 	UE_API void SwapFunctionForAssets(const TArray<FAssetData>& InAssets, bool bSetupUndo);
 	UE_API void SwapAssetReferences();
 
+	/** Enables or disables heat map profiling for the graph */
+	UE_API void ToggleProfiling();
+
 	UE_API void OnOpenSelectedNodesInNewTab();
 	UE_API bool CanOpenSelectedNodesInNewTab() const;
 
