Control Rig: Follow up to 45575153, remove unrelated to be a no-op when nothing is selected

#rnx

diff --git a/Engine/Plugins/Animation/ControlRig/Source/ControlRigEditor/Private/RigDependencyGraph/RigDependencyGraph.cpp b/Engine/Plugins/Animation/ControlRig/Source/ControlRigEditor/Private/RigDependencyGraph/RigDependencyGraph.cpp
--- a/Engine/Plugins/Animation/ControlRig/Source/ControlRigEditor/Private/RigDependencyGraph/RigDependencyGraph.cpp
+++ b/Engine/Plugins/Animation/ControlRig/Source/ControlRigEditor/Private/RigDependencyGraph/RigDependencyGraph.cpp
@@ -977,7 +977,7 @@ void URigDependencyGraph::RemoveUnselectedNodes()
 
 void URigDependencyGraph::RemoveUnrelatedNodes()
 {
-	if (DependencyGraphNodes.IsEmpty())
+	if (DependencyGraphNodes.IsEmpty() || SelectedNodes.IsEmpty())
 	{
 		return;
 	}
