Control Rig: Follow to dependency viewer submit, fix hotkey

#preflighit skip
#rnx

diff --git a/Engine/Plugins/Animation/ControlRig/Source/ControlRigEditor/Private/RigDependencyGraph/RigDependencyGraphCommands.cpp b/Engine/Plugins/Animation/ControlRig/Source/ControlRigEditor/Private/RigDependencyGraph/RigDependencyGraphCommands.cpp
--- a/Engine/Plugins/Animation/ControlRig/Source/ControlRigEditor/Private/RigDependencyGraph/RigDependencyGraphCommands.cpp
+++ b/Engine/Plugins/Animation/ControlRig/Source/ControlRigEditor/Private/RigDependencyGraph/RigDependencyGraphCommands.cpp
@@ -7,7 +7,7 @@
 void FRigDependencyGraphCommands::RegisterCommands()
 {
 	UI_COMMAND(FrameSelection, "Frame Selection", "Expands and frames the selection in the tree", EUserInterfaceActionType::Button, FInputChord(EKeys::F));
-	UI_COMMAND(JumpToSource, "Jump to Source", "Jumps to the instruction node backing up the dependency node", EUserInterfaceActionType::Button, FInputChord(EKeys::F12));
+	UI_COMMAND(JumpToSource, "Jump to Source", "Jumps to the instruction node backing up the dependency node", EUserInterfaceActionType::Button, FInputChord(EModifierKey::Control, EKeys::J));
 	UI_COMMAND(ShowParentChildRelationships, "Show Parent/Child", "Whether to show or hide parent and child relationships", EUserInterfaceActionType::ToggleButton, FInputChord(EKeys::P));
 	UI_COMMAND(ShowInstructionRelationships, "Show VM Nodes", "Whether to show or hide VM based relationships", EUserInterfaceActionType::ToggleButton, FInputChord(EKeys::O));
 	UI_COMMAND(LockContent, "Lock Content", "Whether or not to sync the content with selection in other views", EUserInterfaceActionType::ToggleButton, FInputChord(EKeys::L));
