Anim Layers: Fix regression so they don't constantly refresh

diff --git a/Engine/Plugins/Animation/ControlRig/Source/ControlRigEditor/Private/Sequencer/AnimLayers/SAnimLayers.cpp b/Engine/Plugins/Animation/ControlRig/Source/ControlRigEditor/Private/Sequencer/AnimLayers/SAnimLayers.cpp
--- a/Engine/Plugins/Animation/ControlRig/Source/ControlRigEditor/Private/Sequencer/AnimLayers/SAnimLayers.cpp
+++ b/Engine/Plugins/Animation/ControlRig/Source/ControlRigEditor/Private/Sequencer/AnimLayers/SAnimLayers.cpp
@@ -1691,7 +1691,7 @@ void SAnimLayers::Construct(const FArguments& InArgs, FControlRigEditMode& InEdi
 
 		if (UAnimLayers* AnimLayersPtr = UAnimLayers::GetAnimLayers(SequencerPtr.Get()))
 		{
-			AnimLayersPtr->AnimLayerListChanged().AddRaw(AnimLayerController.Get(), &FAnimLayerController::HandleOnAnimLayerListChanged);
+			OnAnimLayerListChangedHandle = AnimLayersPtr->AnimLayerListChanged().AddRaw(AnimLayerController.Get(), &FAnimLayerController::HandleOnAnimLayerListChanged);
 			AnimLayers = AnimLayersPtr;
 		}
 	}
@@ -1793,12 +1793,16 @@ void SAnimLayers::Tick(const FGeometry& AllottedGeometry, const double InCurrent
 		if (LastMovieSceneSig != CurrentMovieSceneSig)
 		{
 			LastMovieSceneSig = CurrentMovieSceneSig;
-			if (UAnimLayers* AnimLayersPtr = UAnimLayers::GetAnimLayers(Sequencer))
+
+			if (AnimLayers.IsValid() == false)
 			{
-				AnimLayersPtr->AnimLayerListChanged().AddRaw(AnimLayerController.Get(), &FAnimLayerController::HandleOnAnimLayerListChanged);
-				AnimLayers = AnimLayersPtr;
+				if (UAnimLayers* AnimLayersPtr = UAnimLayers::GetAnimLayers(Sequencer))
+				{
+					OnAnimLayerListChangedHandle = AnimLayersPtr->AnimLayerListChanged().AddRaw(AnimLayerController.Get(), &FAnimLayerController::HandleOnAnimLayerListChanged);
+					AnimLayers = AnimLayersPtr;
+					AnimLayerController->RefreshSourceData(true);
+				}
 			}
-			AnimLayerController->RefreshSourceData(true);
 			AnimLayerController->RefreshSelectionData();
 		}
 	}
@@ -1843,7 +1847,8 @@ SAnimLayers::~SAnimLayers()
 {
 	if (AnimLayers.IsValid())
 	{
-		AnimLayers.Get()->AnimLayerListChanged().RemoveAll(AnimLayerController.Get());
+		AnimLayers.Get()->AnimLayerListChanged().Remove(OnAnimLayerListChangedHandle);
+		OnAnimLayerListChangedHandle.Reset();
 	}
 
 	if (TSharedPtr<ISequencer> SequencerPtr = UAnimLayers::GetSequencerFromAsset())
diff --git a/Engine/Plugins/Animation/ControlRig/Source/ControlRigEditor/Private/Sequencer/AnimLayers/SAnimLayers.h b/Engine/Plugins/Animation/ControlRig/Source/ControlRigEditor/Private/Sequencer/AnimLayers/SAnimLayers.h
--- a/Engine/Plugins/Animation/ControlRig/Source/ControlRigEditor/Private/Sequencer/AnimLayers/SAnimLayers.h
+++ b/Engine/Plugins/Animation/ControlRig/Source/ControlRigEditor/Private/Sequencer/AnimLayers/SAnimLayers.h
@@ -46,6 +46,7 @@ private:
 	void RegisterSelectionChanged();
 	void OnActorSelectionChanged(const TArray<UObject*>& NewSelection, bool bForceRefresh);
 	FDelegateHandle OnSelectionChangedHandle;
+	FDelegateHandle OnAnimLayerListChangedHandle;
 
 	//sequencer changing so need to refresh list or update weight
 	void OnActivateSequence(FMovieSceneSequenceIDRef ID);
