Anim In Engine: Moved shared code to EvaluateSequencerTools and fixed localization issues

diff --git a/Engine/Plugins/Animation/ControlRig/Source/ControlRigEditor/Private/Sequencer/AnimLayers/AnimLayers.cpp b/Engine/Plugins/Animation/ControlRig/Source/ControlRigEditor/Private/Sequencer/AnimLayers/AnimLayers.cpp
--- a/Engine/Plugins/Animation/ControlRig/Source/ControlRigEditor/Private/Sequencer/AnimLayers/AnimLayers.cpp
+++ b/Engine/Plugins/Animation/ControlRig/Source/ControlRigEditor/Private/Sequencer/AnimLayers/AnimLayers.cpp
@@ -1783,146 +1783,41 @@ FAnimLayerSectionItem* FAnimLayerItem::FindMatchingSectionItem(UMovieSceneSectio
 	return nullptr;
 }
 
-namespace UE::AnimLayers 
+namespace UE::AIE 
 {
-template<typename ChannelType, typename ValueType>
-void AddKeyToChannel(ChannelType* Channel, EMovieSceneKeyInterpolation DefaultInterpolation, const FFrameNumber& FrameNumber, ValueType Value)
-{
-	switch (DefaultInterpolation)
-	{
-	case EMovieSceneKeyInterpolation::Linear:
-		Channel->AddLinearKey(FrameNumber, Value);
-		break;
-	case EMovieSceneKeyInterpolation::Constant:
-		Channel->AddConstantKey(FrameNumber, Value);
-		break;
-	case  EMovieSceneKeyInterpolation::Auto:
-		Channel->AddCubicKey(FrameNumber, Value, ERichCurveTangentMode::RCTM_Auto);
-		break;
-	case  EMovieSceneKeyInterpolation::SmartAuto:
-	default:
-		Channel->AddCubicKey(FrameNumber, Value, ERichCurveTangentMode::RCTM_SmartAuto);
-		break;
-	}
-}
-static bool GetStartEndIndicesForControl(UMovieSceneControlRigParameterSection* BaseSection, const FRigControlElement* ControlElement, int& OutStartIndex, int& OutEndIndex)
-{
-	if (FChannelMapInfo* pChannelIndex = BaseSection->ControlChannelMap.Find(ControlElement->GetFName()))
+	template<typename ChannelType, typename ValueType>
+	void SetPassthroughKeys(TArrayView<ChannelType*> Channels, int32 StartIndex, int32 EndIndex, EAnimLayerType BlendType, const TArray<UMovieSceneSection*>& Sections, const TArray<UMovieSceneSection*>& AbsoluteSections,
+		EMovieSceneKeyInterpolation DefaultInterpolation, const FFrameNumber& FrameNumber)
 	{
-		const int32 ChannelIndex = pChannelIndex->ChannelIndex;
-		switch (ControlElement->Settings.ControlType)
+		switch (BlendType)
 		{
-
-		case ERigControlType::Float:
-		case ERigControlType::ScaleFloat:
-		{
-			OutStartIndex = ChannelIndex;
-			OutEndIndex = ChannelIndex;
-			break;
-		}
-		case ERigControlType::Vector2D:
-		{
-			OutStartIndex = ChannelIndex;
-			OutEndIndex = ChannelIndex + 1;
-			break;
-		}
-		case ERigControlType::Position:
-		case ERigControlType::Scale:
-		case ERigControlType::Rotator:
+		case EAnimLayerType::Override:
 		{
-			OutStartIndex = ChannelIndex;
-			OutEndIndex = ChannelIndex + 2;
-			break;
-		}
-
-		case ERigControlType::Transform:
-		case ERigControlType::TransformNoScale:
-		case ERigControlType::EulerTransform:
-		{
-			if (ControlElement->Settings.ControlType == ERigControlType::TransformNoScale)
+			TArray<ValueType> ChannelValues = MovieSceneToolHelpers::GetChannelValues<ChannelType,
+				ValueType>(StartIndex, EndIndex, Sections, AbsoluteSections, FrameNumber);
+			int32  ChannelValueIndex = 0;
+			for (int32 Index = StartIndex; Index <= EndIndex; ++Index)
 			{
-				OutStartIndex = ChannelIndex;
-				OutEndIndex = ChannelIndex + 5;
-
+				ChannelType* Channel = Channels[Index];
+				ValueType Value = ChannelValues[ChannelValueIndex];
+				++ChannelValueIndex;
+				AssigneOrSetValue(Channel, Value, FrameNumber, DefaultInterpolation);
 			}
-			else
-			{
-				OutStartIndex = ChannelIndex;
-				OutEndIndex = ChannelIndex + 8;
-			}
-			break;
-		}
-		default:
-			break;
 		}
-		return true;
-	}
-	return false;
-}
-template<typename ChannelType, typename ValueType>
-void AssigneOrSetValue(ChannelType* Channel, ValueType Value, const FFrameNumber& FrameNumber, EMovieSceneKeyInterpolation DefaultInterpolation)
-{
-	TArray<FKeyHandle> KeysAtCurrentTime;
-	Channel->GetKeys(TRange<FFrameNumber>(FrameNumber), nullptr, &KeysAtCurrentTime);
-	if (KeysAtCurrentTime.Num() > 0)
-	{
-		AssignValue(Channel, KeysAtCurrentTime[0], Value);
-	}
-	else
-	{
-		EMovieSceneKeyInterpolation Interpolation = GetInterpolationMode(Channel, FrameNumber, DefaultInterpolation);
-		AddKeyToChannel(Channel, FrameNumber, Value, Interpolation);
-	}
-}
-
-template<typename ChannelType,typename ValueType>
-void SetPassthroughKeys(TArrayView<ChannelType*> Channels,int32 StartIndex, int32 EndIndex, EAnimLayerType BlendType, const TArray<UMovieSceneSection*>& Sections, const TArray<UMovieSceneSection*>& AbsoluteSections,
-	EMovieSceneKeyInterpolation DefaultInterpolation, const FFrameNumber& FrameNumber)
-{
-	switch (BlendType)
-	{
-	case EAnimLayerType::Override:
-	{
-		TArray<ValueType> ChannelValues = MovieSceneToolHelpers::GetChannelValues<ChannelType,
-			ValueType>(StartIndex, EndIndex, Sections,AbsoluteSections,FrameNumber);
-		int32  ChannelValueIndex = 0;
-		for (int32 Index = StartIndex; Index <= EndIndex; ++Index)
-		{
-			ChannelType* Channel = Channels[Index];
-			ValueType Value = ChannelValues[ChannelValueIndex];
-			++ChannelValueIndex;
-			AssigneOrSetValue(Channel, Value, FrameNumber, DefaultInterpolation);
-		}
-	}
-	break;
-	case EAnimLayerType::Additive:
-	{
-		const ValueType Value = 0.0;
-		for (int32 Index = StartIndex; Index <= EndIndex; ++Index)
+		break;
+		case EAnimLayerType::Additive:
 		{
-			ChannelType* Channel = Channels[Index];
-			AssigneOrSetValue(Channel, Value, FrameNumber, DefaultInterpolation);
+			const ValueType Value = 0.0;
+			for (int32 Index = StartIndex; Index <= EndIndex; ++Index)
+			{
+				ChannelType* Channel = Channels[Index];
+				AssigneOrSetValue(Channel, Value, FrameNumber, DefaultInterpolation);
+			}
 		}
-	}
-	break;
-	}
-}
-
-template<typename ChannelType, typename ValueType>
-void SetCurrentKeys(TArrayView<ChannelType*> Channels, int32 StartIndex, int32 EndIndex ,EMovieSceneKeyInterpolation DefaultInterpolation, const FFrameNumber& FrameNumber)
-{
-	ValueType Value = 0.0;
-	FFrameTime FrameTime(FrameNumber);
-	for (int32 Index = StartIndex; Index <= EndIndex; ++Index)
-	{
-		ChannelType* Channel = Channels[Index];
-		if (Channel->Evaluate(FrameTime, Value))
-		{
-			AssigneOrSetValue(Channel, Value, FrameNumber, DefaultInterpolation);
+		break;
 		}
 	}
 }
-}
 bool UAnimLayers::SetPassthroughKey(ISequencer* InSequencer, int32 InIndex)
 {
 	return SetKeyValueOrPassthrough(InSequencer, InIndex, false /*IsValue*/);
@@ -1937,7 +1832,7 @@ bool UAnimLayers::SetKey(ISequencer* InSequencer, int32 InIndex)
 
 bool UAnimLayers::SetKeyValueOrPassthrough(ISequencer* InSequencer, int32 InIndex, bool bJustValue)
 {
-	using namespace UE::AnimLayers;
+	using namespace UE::AIE;
 	auto GetSelectedRigElements = [](UControlRig* ControlRig, TArray<FRigElementKey>& OutSelectedKeys) 
 	{
 		if (ControlRig)
@@ -2023,7 +1918,7 @@ bool UAnimLayers::SetKeyValueOrPassthrough(ISequencer* InSequencer, int32 InInde
 													if (FRigControlElement* Control = CRSection->GetControlRig()->FindControl(ControlName))
 													{
 														int32 StartIndex = 0, EndIndex = 0;
-														if (GetStartEndIndicesForControl(CRSection, Control, StartIndex, EndIndex))
+														if (FControlRigKeys::GetStartEndIndicesForControl(CRSection, Control, StartIndex, EndIndex))
 														{
 															if (bJustValue)
 															{
@@ -2089,7 +1984,7 @@ bool UAnimLayers::SetKeyValueOrPassthrough(ISequencer* InSequencer, int32 InInde
 															if (FRigControlElement* Control = ControlPair.Key->FindControl(ControlName))
 															{
 																int32 StartIndex = 0, EndIndex = 0;
-																if (GetStartEndIndicesForControl(CRSection, Control, StartIndex, EndIndex))
+																if (FControlRigKeys::GetStartEndIndicesForControl(CRSection, Control, StartIndex, EndIndex))
 																{
 																	if (bJustValue)
 																	{
@@ -3238,7 +3133,7 @@ template<typename ChannelType>
 void AdjustmentBlend(UMovieSceneSection* Section, TArrayView<ChannelType*>& BaseChannels, 
 	TArrayView<ChannelType*>& LayerChannels,ISequencer* InSequencer)
 {
-	using namespace UE::AnimLayers;
+	using namespace UE::AIE;
 	if (BaseChannels.Num() != LayerChannels.Num())
 	{
 		return;
diff --git a/Engine/Plugins/Animation/ControlRig/Source/ControlRigEditor/Private/Sequencer/SelectionSets/SSelectionSets.cpp b/Engine/Plugins/Animation/ControlRig/Source/ControlRigEditor/Private/Sequencer/SelectionSets/SSelectionSets.cpp
--- a/Engine/Plugins/Animation/ControlRig/Source/ControlRigEditor/Private/Sequencer/SelectionSets/SSelectionSets.cpp
+++ b/Engine/Plugins/Animation/ControlRig/Source/ControlRigEditor/Private/Sequencer/SelectionSets/SSelectionSets.cpp
@@ -271,7 +271,7 @@ FReply SSelectionSetButton::OnMouseButtonDown(const FGeometry& InMyGeometry, con
 			const bool bShouldCloseWindowAfterMenuSelection = true;
 			FMenuBuilder MenuBuilder(bShouldCloseWindowAfterMenuSelection, nullptr);
 
-			MenuBuilder.BeginSection("SSelectionSetButton", LOCTEXT("SelectionSet", "Selection Sets"));
+			MenuBuilder.BeginSection("SSelectionSetButton", LOCTEXT("SelectionSets", "Selection Sets"));
 			{
 				FUIAction Action = FUIAction(FExecuteAction::CreateRaw((this), &SSelectionSetButton::AddSelectionToSet));
 				const FText Label = LOCTEXT("AddSelectionToSet", "Add Selection");
@@ -420,7 +420,7 @@ void SSelectionSets::Construct(const FArguments& InArgs, FControlRigEditMode& In
 							SNew(SActionButton)
 								.OnClicked(this, &SSelectionSets::OnAddClicked)
 								.Icon(FAppStyle::Get().GetBrush("Icons.Plus"))
-								.Text(LOCTEXT("SelectionSet", "Add Set"))
+								.Text(LOCTEXT("AddSet", "Add Set"))
 								.ToolTipText(LOCTEXT("AddSelectionSetTooltip", "Add a new Selection Set"))
 						]
 					+ SHorizontalBox::Slot()
diff --git a/Engine/Plugins/Animation/ControlRig/Source/ControlRigEditor/Private/Sequencer/SelectionSets/SelectionSets.cpp b/Engine/Plugins/Animation/ControlRig/Source/ControlRigEditor/Private/Sequencer/SelectionSets/SelectionSets.cpp
--- a/Engine/Plugins/Animation/ControlRig/Source/ControlRigEditor/Private/Sequencer/SelectionSets/SelectionSets.cpp
+++ b/Engine/Plugins/Animation/ControlRig/Source/ControlRigEditor/Private/Sequencer/SelectionSets/SelectionSets.cpp
@@ -690,7 +690,7 @@ FGuid UAIESelectionSets::CreateMirror(TSharedPtr<ISequencer>& InSequencer, const
 	FGuid NewGuid;
 	if (const FAIESelectionSetItem* CurrentSetItem = SelectionSets.Find(InGuid))
 	{
-		const FScopedTransaction Transaction(LOCTEXT("CreateMirror", "Create Mirror Selection Set"), !GIsTransacting);
+		const FScopedTransaction Transaction(LOCTEXT("CreateMirrorSelectionSet", "Create Mirror Selection Set"), !GIsTransacting);
 		FControlRigPoseMirrorTable MirrorTable;
 		bool bNoteSelectionChange = false;
 		FAIESelectionSetItem Item;
@@ -1312,91 +1312,6 @@ static UMovieScene3DTransformTrack* Get3DTransformTrack(TSharedPtr<ISequencer>&
 	return nullptr;
 }
 
-static bool GetStartEndIndicesForControl(UMovieSceneControlRigParameterSection* BaseSection, const FRigControlElement* ControlElement, int& OutStartIndex, int& OutEndIndex)
-{
-	if (FChannelMapInfo* pChannelIndex = BaseSection->ControlChannelMap.Find(ControlElement->GetFName()))
-	{
-		const int32 ChannelIndex = pChannelIndex->ChannelIndex;
-		switch (ControlElement->Settings.ControlType)
-		{
-
-		case ERigControlType::Float:
-		case ERigControlType::ScaleFloat:
-		{
-			OutStartIndex = ChannelIndex;
-			OutEndIndex = ChannelIndex;
-			break;
-		}
-		case ERigControlType::Vector2D:
-		{
-			OutStartIndex = ChannelIndex;
-			OutEndIndex = ChannelIndex + 1;
-			break;
-		}
-		case ERigControlType::Position:
-		case ERigControlType::Scale:
-		case ERigControlType::Rotator:
-		{
-			OutStartIndex = ChannelIndex;
-			OutEndIndex = ChannelIndex + 2;
-			break;
-		}
-
-		case ERigControlType::Transform:
-		case ERigControlType::TransformNoScale:
-		case ERigControlType::EulerTransform:
-		{
-			if (ControlElement->Settings.ControlType == ERigControlType::TransformNoScale)
-			{
-				OutStartIndex = ChannelIndex;
-				OutEndIndex = ChannelIndex + 5;
-
-			}
-			else
-			{
-				OutStartIndex = ChannelIndex;
-				OutEndIndex = ChannelIndex + 8;
-			}
-			break;
-		}
-		default:
-			break;
-		}
-		return true;
-	}
-	return false;
-}
-
-template<typename ChannelType, typename ValueType>
-void AssigneOrSetValue(ChannelType* Channel, ValueType Value, const FFrameNumber& FrameNumber, EMovieSceneKeyInterpolation DefaultInterpolation)
-{
-	TArray<FKeyHandle> KeysAtCurrentTime;
-	Channel->GetKeys(TRange<FFrameNumber>(FrameNumber), nullptr, &KeysAtCurrentTime);
-	if (KeysAtCurrentTime.Num() > 0)
-	{
-		AssignValue(Channel, KeysAtCurrentTime[0], Value);
-	}
-	else
-	{
-		EMovieSceneKeyInterpolation Interpolation = GetInterpolationMode(Channel, FrameNumber, DefaultInterpolation);
-		AddKeyToChannel(Channel, FrameNumber, Value, Interpolation);
-	}
-}
-
-template<typename ChannelType, typename ValueType>
-void SetCurrentKeys(TArrayView<ChannelType*> Channels, int32 StartIndex, int32 EndIndex, EMovieSceneKeyInterpolation DefaultInterpolation, const FFrameNumber& FrameNumber)
-{
-	ValueType Value = 0.0;
-	FFrameTime FrameTime(FrameNumber);
-	for (int32 Index = StartIndex; Index <= EndIndex; ++Index)
-	{
-		ChannelType* Channel = Channels[Index];
-		if (Channel->Evaluate(FrameTime, Value))
-		{
-			AssigneOrSetValue(Channel, Value, FrameNumber, DefaultInterpolation);
-		}
-	}
-}
 
 bool UAIESelectionSets::KeyAll(TSharedPtr<ISequencer>& InSequencer, const FGuid& InGuid) const
 {
@@ -1404,7 +1319,7 @@ bool UAIESelectionSets::KeyAll(TSharedPtr<ISequencer>& InSequencer, const FGuid&
 	{
 		return false;
 	}
-
+	using namespace UE::AIE;
 	FControlRigEditMode* EditMode = static_cast<FControlRigEditMode*>(GLevelEditorModeTools().GetActiveMode(FControlRigEditMode::ModeName));
 	if (EditMode)
 	{
@@ -1445,7 +1360,7 @@ bool UAIESelectionSets::KeyAll(TSharedPtr<ISequencer>& InSequencer, const FGuid&
 														if (FRigControlElement* Control = ControlRig->FindControl(ControlName))
 														{
 															int32 StartIndex = 0, EndIndex = 0;
-															if (GetStartEndIndicesForControl(CRSection, Control, StartIndex, EndIndex))
+															if (FControlRigKeys::GetStartEndIndicesForControl(CRSection, Control, StartIndex, EndIndex))
 															{
 																SetCurrentKeys<FMovieSceneFloatChannel, float>(BaseFloatChannels, StartIndex, EndIndex, DefaultInterpolation, FrameNumber);
 
diff --git a/Engine/Plugins/Animation/ControlRig/Source/ControlRigEditor/Private/Tools/EvaluateSequencerTools.cpp b/Engine/Plugins/Animation/ControlRig/Source/ControlRigEditor/Private/Tools/EvaluateSequencerTools.cpp
--- a/Engine/Plugins/Animation/ControlRig/Source/ControlRigEditor/Private/Tools/EvaluateSequencerTools.cpp
+++ b/Engine/Plugins/Animation/ControlRig/Source/ControlRigEditor/Private/Tools/EvaluateSequencerTools.cpp
@@ -851,6 +851,61 @@ namespace AIE
 	}
 
 
+	bool FControlRigKeys::GetStartEndIndicesForControl(UMovieSceneControlRigParameterSection* BaseSection, const FRigControlElement* ControlElement, int& OutStartIndex, int& OutEndIndex)
+	{
+		if (FChannelMapInfo* pChannelIndex = BaseSection->ControlChannelMap.Find(ControlElement->GetFName()))
+		{
+			const int32 ChannelIndex = pChannelIndex->ChannelIndex;
+			switch (ControlElement->Settings.ControlType)
+			{
+
+			case ERigControlType::Float:
+			case ERigControlType::ScaleFloat:
+			{
+				OutStartIndex = ChannelIndex;
+				OutEndIndex = ChannelIndex;
+				break;
+			}
+			case ERigControlType::Vector2D:
+			{
+				OutStartIndex = ChannelIndex;
+				OutEndIndex = ChannelIndex + 1;
+				break;
+			}
+			case ERigControlType::Position:
+			case ERigControlType::Scale:
+			case ERigControlType::Rotator:
+			{
+				OutStartIndex = ChannelIndex;
+				OutEndIndex = ChannelIndex + 2;
+				break;
+			}
+
+			case ERigControlType::Transform:
+			case ERigControlType::TransformNoScale:
+			case ERigControlType::EulerTransform:
+			{
+				if (ControlElement->Settings.ControlType == ERigControlType::TransformNoScale)
+				{
+					OutStartIndex = ChannelIndex;
+					OutEndIndex = ChannelIndex + 5;
+
+				}
+				else
+				{
+					OutStartIndex = ChannelIndex;
+					OutEndIndex = ChannelIndex + 8;
+				}
+				break;
+			}
+			default:
+				break;
+			}
+			return true;
+		}
+		return false;
+	}
+
 } // namespace AIE
 } // namespace UE
 
diff --git a/Engine/Plugins/Animation/ControlRig/Source/ControlRigEditor/Public/Tools/EvaluateSequencerTools.h b/Engine/Plugins/Animation/ControlRig/Source/ControlRigEditor/Public/Tools/EvaluateSequencerTools.h
--- a/Engine/Plugins/Animation/ControlRig/Source/ControlRigEditor/Public/Tools/EvaluateSequencerTools.h
+++ b/Engine/Plugins/Animation/ControlRig/Source/ControlRigEditor/Public/Tools/EvaluateSequencerTools.h
@@ -6,13 +6,19 @@
 #include "MovieSceneBinding.h"
 #include "Misc/Guid.h"
 
+#include "Channels/MovieSceneFloatChannel.h"
+#include "Channels/MovieSceneBoolChannel.h"
+#include "Channels/MovieSceneChannelProxy.h"
+#include "Channels/MovieSceneDoubleChannel.h"
+#include "Channels/MovieSceneIntegerChannel.h"
+#include "Channels/MovieSceneByteChannel.h"
+
 #define UE_API CONTROLRIGEDITOR_API
 
 class UControlRig;
 class UWorld;
 class ISequencer;
 struct FRigControlElement;
-struct FRigControlElement;
 struct FRigControlModifiedContext;
 struct FEulerTransform;
 class UMovieSceneTrack;
@@ -21,6 +27,7 @@ struct FGuid;
 class UMovieScene3DTransformSection;
 class UMovieScenePropertyTrack;
 class UMovieSceneControlRigParameterTrack;
+class UMovieSceneControlRigParameterSection;
 
 namespace UE
 {
@@ -282,6 +289,65 @@ struct FSequencerSelected
 	static UE_API void GetSelectedControlRigsAndBoundObjects(ISequencer* SequencerPtr, TArray<FControlRigAndControlsAndTrack>& OutSelectedCRs, TArray<FObjectAndTrack>& OutBoundObjects);
 };
 
+//get start end times for control sections keys
+struct FControlRigKeys
+{
+	static UE_API bool GetStartEndIndicesForControl(UMovieSceneControlRigParameterSection* BaseSection, const FRigControlElement* ControlElement, int& OutStartIndex, int& OutEndIndex);
+};
+
+template<typename ChannelType, typename ValueType>
+void AddKeyToChannel(ChannelType* Channel, EMovieSceneKeyInterpolation DefaultInterpolation, const FFrameNumber& FrameNumber, ValueType Value)
+{
+	switch (DefaultInterpolation)
+	{
+	case EMovieSceneKeyInterpolation::Linear:
+		Channel->AddLinearKey(FrameNumber, Value);
+		break;
+	case EMovieSceneKeyInterpolation::Constant:
+		Channel->AddConstantKey(FrameNumber, Value);
+		break;
+	case  EMovieSceneKeyInterpolation::Auto:
+		Channel->AddCubicKey(FrameNumber, Value, ERichCurveTangentMode::RCTM_Auto);
+		break;
+	case  EMovieSceneKeyInterpolation::SmartAuto:
+	default:
+		Channel->AddCubicKey(FrameNumber, Value, ERichCurveTangentMode::RCTM_SmartAuto);
+		break;
+	}
+}
+
+template<typename ChannelType, typename ValueType>
+void AssigneOrSetValue(ChannelType* Channel, ValueType Value, const FFrameNumber& FrameNumber, EMovieSceneKeyInterpolation DefaultInterpolation)
+{
+	TArray<FKeyHandle> KeysAtCurrentTime;
+	Channel->GetKeys(TRange<FFrameNumber>(FrameNumber), nullptr, &KeysAtCurrentTime);
+	if (KeysAtCurrentTime.Num() > 0)
+	{
+		AssignValue(Channel, KeysAtCurrentTime[0], Value);
+	}
+	else
+	{
+		EMovieSceneKeyInterpolation Interpolation = GetInterpolationMode(Channel, FrameNumber, DefaultInterpolation);
+		AddKeyToChannel(Channel, FrameNumber, Value, Interpolation);
+	}
+}
+
+
+
+template<typename ChannelType, typename ValueType>
+void SetCurrentKeys(TArrayView<ChannelType*> Channels, int32 StartIndex, int32 EndIndex, EMovieSceneKeyInterpolation DefaultInterpolation, const FFrameNumber& FrameNumber)
+{
+	ValueType Value = 0.0;
+	FFrameTime FrameTime(FrameNumber);
+	for (int32 Index = StartIndex; Index <= EndIndex; ++Index)
+	{
+		ChannelType* Channel = Channels[Index];
+		if (Channel->Evaluate(FrameTime, Value))
+		{
+			AssigneOrSetValue(Channel, Value, FrameNumber, DefaultInterpolation);
+		}
+	}
+}
 } // namespace AIE
 } // namespace UE
 
