Anim In Engine: Deleting Bool channel may set key at another time if AutoKey and KeyAll is on.
#trivial

diff --git a/Engine/Plugins/Animation/ControlRig/Source/ControlRigEditor/Private/Sequencer/ControlRigParameterTrackEditor.cpp b/Engine/Plugins/Animation/ControlRig/Source/ControlRigEditor/Private/Sequencer/ControlRigParameterTrackEditor.cpp
--- a/Engine/Plugins/Animation/ControlRig/Source/ControlRigEditor/Private/Sequencer/ControlRigParameterTrackEditor.cpp
+++ b/Engine/Plugins/Animation/ControlRig/Source/ControlRigEditor/Private/Sequencer/ControlRigParameterTrackEditor.cpp
@@ -2529,7 +2529,7 @@ static void EvaluateThisControl(UMovieSceneControlRigParameterSection* Section,
 					TOptional <bool> Value = Section->EvaluateBoolParameter(FrameTime, ControlName);
 					if (Value.IsSet())
 					{
-						ControlRig->SetControlValue<bool>(ControlName, Value.GetValue(), true, EControlRigSetKey::DoNotCare, bSetupUndo);
+						ControlRig->SetControlValue<bool>(ControlName, Value.GetValue(), true, EControlRigSetKey::Never, bSetupUndo);
 					}
 				}
 				break;
@@ -2544,7 +2544,7 @@ static void EvaluateThisControl(UMovieSceneControlRigParameterSection* Section,
 						if (Value.IsSet())
 						{
 							int32 IVal = (int32)Value.GetValue();
-							ControlRig->SetControlValue<int32>(ControlName, IVal, true, EControlRigSetKey::DoNotCare, bSetupUndo);
+							ControlRig->SetControlValue<int32>(ControlName, IVal, true, EControlRigSetKey::Never, bSetupUndo);
 						}
 					}
 					else
@@ -2552,7 +2552,7 @@ static void EvaluateThisControl(UMovieSceneControlRigParameterSection* Section,
 						TOptional <int32> Value = Section->EvaluateIntegerParameter(FrameTime, ControlName);
 						if (Value.IsSet())
 						{
-							ControlRig->SetControlValue<int32>(ControlName, Value.GetValue(), true, EControlRigSetKey::DoNotCare, bSetupUndo);
+							ControlRig->SetControlValue<int32>(ControlName, Value.GetValue(), true, EControlRigSetKey::Never, bSetupUndo);
 						}
 					}
 				}
@@ -2565,7 +2565,7 @@ static void EvaluateThisControl(UMovieSceneControlRigParameterSection* Section,
 				TOptional <float> Value = Section->EvaluateScalarParameter(FrameTime, ControlName);
 				if (Value.IsSet())
 				{
-					ControlRig->SetControlValue<float>(ControlName, Value.GetValue(), true, EControlRigSetKey::DoNotCare, bSetupUndo);
+					ControlRig->SetControlValue<float>(ControlName, Value.GetValue(), true, EControlRigSetKey::Never, bSetupUndo);
 				}
 				break;
 			}
@@ -2574,7 +2574,7 @@ static void EvaluateThisControl(UMovieSceneControlRigParameterSection* Section,
 				TOptional <FVector2D> Value = Section->EvaluateVector2DParameter(FrameTime, ControlName);
 				if (Value.IsSet())
 				{
-					ControlRig->SetControlValue<FVector2D>(ControlName, Value.GetValue(), true, EControlRigSetKey::DoNotCare, bSetupUndo);
+					ControlRig->SetControlValue<FVector2D>(ControlName, Value.GetValue(), true, EControlRigSetKey::Never, bSetupUndo);
 				}
 				break;
 			}
@@ -2586,7 +2586,7 @@ static void EvaluateThisControl(UMovieSceneControlRigParameterSection* Section,
 				if (Value.IsSet())
 				{
 					FVector3f FloatVal = (FVector3f)Value.GetValue();
-					ControlRig->SetControlValue<FVector3f>(ControlName, FloatVal, true, EControlRigSetKey::DoNotCare, bSetupUndo);
+					ControlRig->SetControlValue<FVector3f>(ControlName, FloatVal, true, EControlRigSetKey::Never, bSetupUndo);
 				}
 				break;
 			}
@@ -2606,14 +2606,14 @@ static void EvaluateThisControl(UMovieSceneControlRigParameterSection* Section,
 					{
 						FVector EulerAngle(Value.GetValue().Rotation.Roll, Value.GetValue().Rotation.Pitch, Value.GetValue().Rotation.Yaw);
 						RigHierarchy->SetControlSpecifiedEulerAngle(ControlElement, EulerAngle);
-						ControlRig->SetControlValue<FRigControlValue::FTransform_Float>(ControlName, Value.GetValue().ToFTransform(), true, EControlRigSetKey::DoNotCare, bSetupUndo);
+						ControlRig->SetControlValue<FRigControlValue::FTransform_Float>(ControlName, Value.GetValue().ToFTransform(), true, EControlRigSetKey::Never, bSetupUndo);
 					}
 					else if (ControlElement->Settings.ControlType == ERigControlType::TransformNoScale)
 					{
 						FTransformNoScale NoScale = Value.GetValue().ToFTransform();
 						FVector EulerAngle(Value.GetValue().Rotation.Roll, Value.GetValue().Rotation.Pitch, Value.GetValue().Rotation.Yaw);
 						RigHierarchy->SetControlSpecifiedEulerAngle(ControlElement, EulerAngle);
-						ControlRig->SetControlValue<FRigControlValue::FTransformNoScale_Float>(ControlName, NoScale, true, EControlRigSetKey::DoNotCare, bSetupUndo);
+						ControlRig->SetControlValue<FRigControlValue::FTransformNoScale_Float>(ControlName, NoScale, true, EControlRigSetKey::Never, bSetupUndo);
 					}
 					else if (ControlElement->Settings.ControlType == ERigControlType::EulerTransform)
 					{
@@ -2624,7 +2624,7 @@ static void EvaluateThisControl(UMovieSceneControlRigParameterSection* Section,
 						FRotator UERotator(Quat);
 						FEulerTransform Transform = Euler;
 						Transform.Rotation = UERotator;
-						ControlRig->SetControlValue<FRigControlValue::FEulerTransform_Float>(ControlName, Transform, true, EControlRigSetKey::DoNotCare, bSetupUndo);
+						ControlRig->SetControlValue<FRigControlValue::FEulerTransform_Float>(ControlName, Transform, true, EControlRigSetKey::Never, bSetupUndo);
 
 					}
 				}
