Control Rig: Open UI state is now saved in config so it is restored when the user restarts the editor. The UI now consists of toggle buttons, showing with a blue background when the tab is open and an empty background when the tab is not open.

diff --git a/Engine/Plugins/Animation/ControlRig/Source/ControlRigEditor/Private/EditMode/ControlRigBaseDockableView.cpp b/Engine/Plugins/Animation/ControlRig/Source/ControlRigEditor/Private/EditMode/ControlRigBaseDockableView.cpp
--- a/Engine/Plugins/Animation/ControlRig/Source/ControlRigEditor/Private/EditMode/ControlRigBaseDockableView.cpp
+++ b/Engine/Plugins/Animation/ControlRig/Source/ControlRigEditor/Private/EditMode/ControlRigBaseDockableView.cpp
@@ -18,7 +18,7 @@ FControlRigBaseDockableView::FControlRigBaseDockableView()
 
 FControlRigBaseDockableView::~FControlRigBaseDockableView()
 {
-	if (FControlRigEditMode* EditMode = static_cast<FControlRigEditMode*>(ModeTools->GetActiveMode(FControlRigEditMode::ModeName)))
+	if (FControlRigEditMode* EditMode = GetEditMode())
 	{
 		EditMode->OnControlRigAddedOrRemoved().RemoveAll(this);
 		EditMode->OnControlRigSelected().RemoveAll(this);
@@ -29,7 +29,7 @@ FControlRigBaseDockableView::~FControlRigBaseDockableView()
 TArray<UControlRig*> FControlRigBaseDockableView::GetControlRigs()  const
 {
 	TArray<UControlRig*> ControlRigs;
-	if (FControlRigEditMode* EditMode = static_cast<FControlRigEditMode*>(ModeTools->GetActiveMode(FControlRigEditMode::ModeName)))
+	if (FControlRigEditMode* EditMode = GetEditMode())
 	{
 		ControlRigs = EditMode->GetControlRigsArray(false);
 	}
@@ -38,9 +38,8 @@ TArray<UControlRig*> FControlRigBaseDockableView::GetControlRigs()  const
 
 void FControlRigBaseDockableView::SetEditMode(FControlRigEditMode& InEditMode)
 {
-
 	ModeTools = InEditMode.GetModeManager();
-	if (FControlRigEditMode* EditMode = static_cast<FControlRigEditMode*>(ModeTools->GetActiveMode(FControlRigEditMode::ModeName)))
+	if (FControlRigEditMode* EditMode = GetEditMode())
 	{
 		EditMode->OnControlRigAddedOrRemoved().RemoveAll(this);
 		EditMode->OnControlRigAddedOrRemoved().AddRaw(this, &FControlRigBaseDockableView::HandleControlAdded);
@@ -82,12 +81,17 @@ void FControlRigBaseDockableView::HandlElementSelected(UControlRig* ControlRig,
 
 FControlRigEditMode* FControlRigBaseDockableView::GetEditMode() const
 {
-	FControlRigEditMode* EditMode = static_cast<FControlRigEditMode*>(ModeTools->GetActiveMode(FControlRigEditMode::ModeName));
+	FControlRigEditMode* EditMode = ModeTools ? static_cast<FControlRigEditMode*>(ModeTools->GetActiveMode(FControlRigEditMode::ModeName)) : nullptr;
 	return EditMode;
 }
 
 ISequencer* FControlRigBaseDockableView::GetSequencer() const
 {
+	if (!ModeTools)
+	{
+		return nullptr;
+	}
+	
 	if (FControlRigEditMode* EditMode = static_cast<FControlRigEditMode*>(ModeTools->GetActiveMode(FControlRigEditMode::ModeName)))
 	{
 		TWeakPtr<ISequencer> Sequencer = EditMode->GetWeakSequencer();
diff --git a/Engine/Plugins/Animation/ControlRig/Source/ControlRigEditor/Private/EditMode/ControlRigEditModeSettings.h b/Engine/Plugins/Animation/ControlRig/Source/ControlRigEditor/Private/EditMode/ControlRigEditModeSettings.h
--- a/Engine/Plugins/Animation/ControlRig/Source/ControlRigEditor/Private/EditMode/ControlRigEditModeSettings.h
+++ b/Engine/Plugins/Animation/ControlRig/Source/ControlRigEditor/Private/EditMode/ControlRigEditModeSettings.h
@@ -4,8 +4,7 @@
 
 #include "CoreMinimal.h"
 #include "Engine/DeveloperSettings.h"
-#include "Misc/Flyout/FlyoutSavedState.h"
-#include "Overlay/DragBoxPosition.h"
+#include "Settings/ControlRigUIRestoreStates.h"
 #include "UObject/Object.h"
 #include "ControlRigEditModeSettings.generated.h"
 
@@ -32,17 +31,6 @@ class UControlRigEditModeSettings : public UDeveloperSettings
 		, GizmoScale(1.0f)
 	{
 		DrivenControlColor = FLinearColor::White * FLinearColor(FVector::OneVector * 0.8f);
-		LastTweenOverlayState = FToolWidget_FlyoutSavedState
-		{
-			FToolWidget_DragBoxPosition(FVector2f(20.f, 30.f), HAlign_Right, VAlign_Bottom),
-			false
-		};
-
-		LastSelectionSetOverlayState = FToolWidget_FlyoutSavedState
-		{
-			FToolWidget_DragBoxPosition(FVector2f(20.f, 30.f), HAlign_Right, VAlign_Top),
-			false
-		};
 	}
 
 	// UObject interface
@@ -53,6 +41,8 @@ class UControlRigEditModeSettings : public UDeveloperSettings
 #endif
 public:
 
+	static UControlRigEditModeSettings* Get() { return GetMutableDefault<UControlRigEditModeSettings>(); }
+
 	/** Whether to show all bones in the hierarchy */
 	UPROPERTY(config, EditAnywhere, Category = "Animation Settings")
 	bool bDisplayHierarchy;
@@ -105,11 +95,9 @@ public:
 	UPROPERTY(config, EditAnywhere, Category = "Animation Settings")
 	float GizmoScale;
 
+	/** Data to restore the UI state control rig had last time it was open. */
 	UPROPERTY(config)
-	FToolWidget_FlyoutSavedState LastTweenOverlayState;
-
-	UPROPERTY(config)
-	FToolWidget_FlyoutSavedState LastSelectionSetOverlayState;
+	FControlRigUIRestoreStates LastUIStates;
 
 	/** The tint to apply when the tween overlay is out of focus. */
 	UPROPERTY(config, EditAnywhere, Category = "Animation Settings|Tween")
diff --git a/Engine/Plugins/Animation/ControlRig/Source/ControlRigEditor/Private/EditMode/ControlRigEditModeToolkit.cpp b/Engine/Plugins/Animation/ControlRig/Source/ControlRigEditor/Private/EditMode/ControlRigEditModeToolkit.cpp
--- a/Engine/Plugins/Animation/ControlRig/Source/ControlRigEditor/Private/EditMode/ControlRigEditModeToolkit.cpp
+++ b/Engine/Plugins/Animation/ControlRig/Source/ControlRigEditor/Private/EditMode/ControlRigEditModeToolkit.cpp
@@ -49,15 +49,6 @@ namespace
 	const TArray<FName> AnimationPaletteNames = { AnimationName };
 }
 
-bool FControlRigEditModeToolkit::bMotionTrailsOn = false;
-bool FControlRigEditModeToolkit::bAnimLayerTabOpen = false;
-bool FControlRigEditModeToolkit::bPoseTabOpen = false;
-bool FControlRigEditModeToolkit::bSnapperTabOpen = false;
-bool FControlRigEditModeToolkit::bTweenOpen = false;
-bool FControlRigEditModeToolkit::bConstraintsOpen = false;
-bool FControlRigEditModeToolkit::bSelectionSetsOpen = false;
-
-
 const FName FControlRigEditModeToolkit::PoseTabName = FName(TEXT("PoseTab"));
 const FName FControlRigEditModeToolkit::MotionTrailTabName = FName(TEXT("MotionTrailTab"));
 const FName FControlRigEditModeToolkit::SnapperTabName = FName(TEXT("SnapperTab"));
@@ -76,7 +67,6 @@ const FName FControlRigEditModeToolkit::SelectionSetsOverlayName = FName(TEXT("S
 #endif
 const FName FControlRigEditModeToolkit::OutlinerTabName = FName(TEXT("ControlRigOutlinerTab"));
 const FName FControlRigEditModeToolkit::DetailsTabName = FName(TEXT("ControlRigDetailsTab"));
-const FName FControlRigEditModeToolkit::SpacePickerTabName = FName(TEXT("ControlRigSpacePicker"));
 
 TSharedPtr<UE::ControlRigEditor::SAnimDetailsView> FControlRigEditModeToolkit::Details = nullptr;
 TSharedPtr<SControlRigOutliner> FControlRigEditModeToolkit::Outliner = nullptr;
@@ -138,64 +128,121 @@ void FControlRigEditModeToolkit::OnToolPaletteChanged(FName PaletteName)
 
 }
 
-void FControlRigEditModeToolkit::TryInvokeToolkitUI(const FName InName)
+void FControlRigEditModeToolkit::TryToggleToolkitUI(const FName InName)
+{
+	if (IsToolkitUIActive(InName))
+	{
+		TryCloseToolkitUI(InName);
+	}
+	else
+	{
+		TryInvokeToolkitUI(InName);
+	}
+}
+
+namespace UE::ControlRigEditor::ToolkitDetail
 {
-	TSharedPtr<FAssetEditorModeUILayer> ModeUILayerPtr = ModeUILayer.Pin();
+/**
+ * @return The tab ID for one of the logical FControlRigEditModeToolkit::xxxTabName.
+ * They mostly coincide but are different for those that re-use the UAssetEditorUISubsystem tab names.
+ */
+static TOptional<FName> GetTabIdForUI(const FName InName)
+{
+	if (InName == FControlRigEditModeToolkit::MotionTrailTabName)
+	{
+		return FControlRigEditModeToolkit::MotionTrailTabName;
+	}
 
-	if (InName == MotionTrailTabName)
+	if (InName == FControlRigEditModeToolkit::AnimLayerTabName)
 	{
-		FTabId TabID(MotionTrailTabName);
-		ModeUILayerPtr->GetTabManager()->TryInvokeTab(TabID, false /*bIsActive*/);
+		return FControlRigEditModeToolkit::AnimLayerTabName;
 	}
-	else if (InName == AnimLayerTabName)
+
+	if (InName == FControlRigEditModeToolkit::PoseTabName)
 	{
-		FTabId TabID(AnimLayerTabName);
-		ModeUILayerPtr->GetTabManager()->TryInvokeTab(TabID, false /*bIsActive*/);
+		return FControlRigEditModeToolkit::PoseTabName;
 	}
-	else if (InName == PoseTabName)
+
+	if (InName == FControlRigEditModeToolkit::SnapperTabName)
 	{
-		FTabId TabID(PoseTabName);
-		ModeUILayerPtr->GetTabManager()->TryInvokeTab(TabID, false /*bIsActive*/);
+		return FControlRigEditModeToolkit::SnapperTabName;
 	}
-	else if (InName == SnapperTabName)
+
+	if (InName == FControlRigEditModeToolkit::OutlinerTabName)
 	{
-		FTabId TabID(SnapperTabName);
-		ModeUILayerPtr->GetTabManager()->TryInvokeTab(TabID, false /*bIsActive*/);
+		return UAssetEditorUISubsystem::TopRightTabID;
 	}
-	else if (InName == OutlinerTabName)
+
+	if (InName == FControlRigEditModeToolkit::DetailsTabName)
 	{
-		ModeUILayerPtr->GetTabManager()->TryInvokeTab(UAssetEditorUISubsystem::TopRightTabID);
+		return UAssetEditorUISubsystem::BottomRightTabID;
 	}
-	else if (InName == SpacePickerTabName)
+
+#if SELECTION_SETS_AS_TAB
+	if (InName == FControlRigEditModeToolkit::SelectionSetsTabName)
 	{
-		ModeUILayerPtr->GetTabManager()->TryInvokeTab(UAssetEditorUISubsystem::BottomLeftTabID);
+		return FControlRigEditModeToolkit::SelectionSetsTabName;
 	}
-	else if (InName == DetailsTabName)
+#endif
+
+	return {};
+}
+}
+
+void FControlRigEditModeToolkit::TryInvokeToolkitUI(const FName InName)
+{
+	const TSharedPtr<FAssetEditorModeUILayer> ModeUILayerPtr = ModeUILayer.Pin();
+	const TOptional<FName> TabName = UE::ControlRigEditor::ToolkitDetail::GetTabIdForUI(InName);
+	
+	if (ModeUILayerPtr && TabName)
 	{
-		ModeUILayerPtr->GetTabManager()->TryInvokeTab(UAssetEditorUISubsystem::BottomRightTabID);
+		const FTabId TabID(*TabName);
+		ModeUILayerPtr->GetTabManager()->TryInvokeTab(TabID, false /*bIsActive*/);
 	}
 	else if (InName == TweenOverlayName)
 	{
-		TweenOverlayManager->ToggleVisibility();
+		TweenOverlayManager->ShowWidget();
 	}
 	else if (InName == ConstrainingTabName)
 	{
-		ConstrainToolsManager->ToggleVisibility();
+		ConstrainToolsManager->ShowWidget();
+	}
+#if SELECTION_SETS_AS_TAB == 0
+	else if (InName == SelectionSetsOverlayName)
+	{
+		SelectionSetsOverlayManager->ShowWidget();
 	}
+#endif
+}
 
-#if SELECTION_SETS_AS_TAB
-	else if (InName == SelectionSetsTabName)
+void FControlRigEditModeToolkit::TryCloseToolkitUI(const FName InName) const
+{
+	const TSharedPtr<FAssetEditorModeUILayer> ModeUILayerPtr = ModeUILayer.Pin();
+	const TOptional<FName> TabName = UE::ControlRigEditor::ToolkitDetail::GetTabIdForUI(InName);
+	const TSharedPtr<SDockTab> LiveTab = ModeUILayerPtr && TabName ? ModeUILayerPtr->GetTabManager()->FindExistingLiveTab(FTabId(*TabName)) : nullptr;
+	if (TabName && LiveTab)
 	{
-		FTabId TabID(SelectionSetsTabName);
-		ModeUILayerPtr->GetTabManager()->TryInvokeTab(TabID, false /*bIsActive*/);
+		LiveTab->RequestCloseTab();
 	}
-#else
+	else if (TabName)
+	{
+		return;
+	}
+	
+	if (InName == TweenOverlayName)
+	{
+		TweenOverlayManager->HideWidget();
+	}
+	else if (InName == ConstrainingTabName)
+	{
+		ConstrainToolsManager->HideWidget();
+	}
+#if SELECTION_SETS_AS_TAB == 0
 	else if (InName == SelectionSetsOverlayName)
 	{
-		SelectionSetsOverlayManager->ToggleVisibility();
+		SelectionSetsOverlayManager->HideWidget();
 	}
 #endif
-	
 }
 
 bool FControlRigEditModeToolkit::IsToolkitUIActive(const FName InName) const
@@ -215,9 +262,11 @@ bool FControlRigEditModeToolkit::IsToolkitUIActive(const FName InName) const
 	{
 		return ConstrainToolsManager->IsShowingWidget();
 	}
-
-	TSharedPtr<FAssetEditorModeUILayer> ModeUILayerPtr = ModeUILayer.Pin();
-	return ModeUILayerPtr->GetTabManager()->FindExistingLiveTab(FTabId(InName)).IsValid();
+	
+	const TSharedPtr<FAssetEditorModeUILayer> ModeUILayerPtr = ModeUILayer.Pin();
+	const TOptional<FName> TabName = UE::ControlRigEditor::ToolkitDetail::GetTabIdForUI(InName);
+	return ModeUILayerPtr && ensureMsgf(TabName, TEXT("You passed in an unknown UI name!"))
+		&& ModeUILayerPtr->GetTabManager()->FindExistingLiveTab(FTabId(*TabName)).IsValid();
 }
 
 void FControlRigEditModeToolkit::OnControlsChanged(TConstArrayView<TWeakObjectPtr<UControlRig>> InControlRigs) const
@@ -403,88 +452,92 @@ void FControlRigEditModeToolkit::InvokeUI()
 {
 	FModeToolkit::InvokeUI();
 
-	if (ModeUILayer.IsValid())
+	const TSharedPtr<FAssetEditorModeUILayer> ModeUILayerPtr = ModeUILayer.Pin();
+	if (!ModeUILayerPtr)
 	{
-		TSharedPtr<FAssetEditorModeUILayer> ModeUILayerPtr = ModeUILayer.Pin();
-		ModeUILayerPtr->GetTabManager()->TryInvokeTab(UAssetEditorUISubsystem::TopRightTabID);
-		// doesn't work as expected todo ModeUILayerPtr->GetTabManager()->TryInvokeTab(UAssetEditorUISubsystem::TopLeftTabID);
-		ModeUILayerPtr->GetTabManager()->TryInvokeTab(UAssetEditorUISubsystem::BottomRightTabID);
-
-		SelectionViewModel = MakeShared<UE::ControlRigEditor::FRigSelectionViewModel>();
-		
-		TweenOverlayManager = MakeUnique<UE::ControlRigEditor::FTweenOverlayManager>(GetToolkitHost(), GetToolkitCommands(), SharedThis(&EditMode));
-		if (bTweenOpen)
-		{
-			TweenOverlayManager->ShowWidget();
-		}
+		return;
+	}
 
-		ConstrainToolsManager = MakeUnique<UE::ControlRigEditor::FConstrainToolsManager>(
-			ModeUILayerPtr->GetTabManager().ToSharedRef(), EditMode, SelectionViewModel.ToSharedRef()
-			);
-	if (bConstraintsOpen)
-		{
-			ConstrainToolsManager->ShowWidget();
-		}
+	ModeUILayerPtr->GetTabManager()->TryInvokeTab(UAssetEditorUISubsystem::TopRightTabID);
+	// doesn't work as expected todo ModeUILayerPtr->GetTabManager()->TryInvokeTab(UAssetEditorUISubsystem::TopLeftTabID);
+	ModeUILayerPtr->GetTabManager()->TryInvokeTab(UAssetEditorUISubsystem::BottomRightTabID);
+
+	// TweenOverlayManager will automatically restore its last visibility state upon construction.
+	TweenOverlayManager = MakeUnique<UE::ControlRigEditor::FTweenOverlayManager>(GetToolkitHost(), GetToolkitCommands(), SharedThis(&EditMode));
+
+	SelectionViewModel = MakeShared<UE::ControlRigEditor::FRigSelectionViewModel>();
+	ConstrainToolsManager = MakeUnique<UE::ControlRigEditor::FConstrainToolsManager>(
+		ModeUILayerPtr->GetTabManager().ToSharedRef(), EditMode, SelectionViewModel.ToSharedRef()
+		);
+	
+	const FControlRigUIRestoreStates& RestoreStates = UControlRigEditModeSettings::Get()->LastUIStates;
+	
 #if SELECTION_SETS_AS_TAB
-		if (bSelectionSetsOpen)
-		{
-			TryInvokeToolkitUI(SelectionSetsTabName);
-		}
-		/*
-		else if (TSharedPtr<ISequencer> SequencerPtr = UAIESelectionSets::GetSequencerFromAsset())
+	if (RestoreStates.bSelectionSetsOpen)
+	{
+		TryInvokeToolkitUI(SelectionSetsTabName);
+	}
+	/*
+	else if (TSharedPtr<ISequencer> SequencerPtr = UAIESelectionSets::GetSequencerFromAsset())
+	{
+		//wasn't open but there are layers in open level sequence we try to open it
+		if (UAIESelectionSets* SelectionSets = UAIESelectionSets::GetSelectionSets(SequencerPtr))
 		{
-			//wasn't open but there are layers in open level sequence we try to open it
-			if (UAIESelectionSets* SelectionSets = UAIESelectionSets::GetSelectionSets(SequencerPtr))
+			if (SelectionSets && SelectionSets->bOpenUIOnOpen)
 			{
-				if (SelectionSets && SelectionSets->bOpenUIOnOpen)
-				{
-					TryInvokeToolkitUI(SelectionSetsTabName);
-				}
+				TryInvokeToolkitUI(SelectionSetsTabName);
 			}
 		}
-		*/
-
+	}
+	*/
 #else
-		SelectionSetsOverlayManager = MakeUnique<UE::AIE::FSelectionSetsOverlayManager>(GetToolkitHost(), GetToolkitCommands(), SharedThis(&EditMode));
-		if (bSelectionSetsOpen)
-		{
-			SelectionSetsOverlayManager->ShowWidget();
-		}
-#endif		
-		if (bAnimLayerTabOpen)
-		{
-			TryInvokeToolkitUI(AnimLayerTabName);
-		}
-		else if (TSharedPtr<ISequencer> SequencerPtr = UAnimLayers::GetSequencerFromAsset())
-		{
-			//wasn't open but there are layers in open level sequence we try to open it
-			if (UAnimLayers* AnimLayers = UAnimLayers::GetAnimLayers(SequencerPtr.Get()))
-			{
-				if (AnimLayers && AnimLayers->bOpenUIOnOpen)
-				{
-					TryInvokeToolkitUI(AnimLayerTabName);
-				}
-			}
-		}
-		if (bSnapperTabOpen)
-		{
-			TryInvokeToolkitUI(SnapperTabName);
-		}
-		if (bPoseTabOpen)
-		{
-			TryInvokeToolkitUI(PoseTabName);
-		}
+	SelectionSetsOverlayManager = MakeUnique<UE::AIE::FSelectionSetsOverlayManager>(GetToolkitHost(), GetToolkitCommands(), SharedThis(&EditMode));
+	if (RestoreStates.bSelectionSetsOpen)
+	{
+		SelectionSetsOverlayManager->ShowWidget();
+	}
+#endif	
 
-		if (bMotionTrailsOn)
+	if (RestoreStates.bConstraintsOpen)
+	{
+		ConstrainToolsManager->ShowWidget();
+	}
+	
+	if (RestoreStates.bAnimLayerTabOpen)
+	{
+		TryInvokeToolkitUI(AnimLayerTabName);
+	}
+	else if (TSharedPtr<ISequencer> SequencerPtr = UAnimLayers::GetSequencerFromAsset())
+	{
+		// Wasn't open but if there are layers in open level sequence we try to open it
+		if (UAnimLayers* AnimLayers = UAnimLayers::GetAnimLayers(SequencerPtr.Get()))
 		{
-			if (UMotionTrailToolOptions* Settings = GetMutableDefault<UMotionTrailToolOptions>())
+			if (AnimLayers && AnimLayers->bOpenUIOnOpen)
 			{
-				Settings->bShowTrails = true;
-				FPropertyChangedEvent ShowTrailEvent(UMotionTrailToolOptions::StaticClass()->FindPropertyByName(GET_MEMBER_NAME_CHECKED(UMotionTrailToolOptions, bShowTrails)));
-				Settings->PostEditChangeProperty(ShowTrailEvent);
+				TryInvokeToolkitUI(AnimLayerTabName);
 			}
 		}
-	}	
+	}
+	
+	if (RestoreStates.bSnapperTabOpen)
+	{
+		TryInvokeToolkitUI(SnapperTabName);
+	}
+	
+	if (RestoreStates.bPoseTabOpen)
+	{
+		TryInvokeToolkitUI(PoseTabName);
+	}
+
+	UMotionTrailToolOptions* Settings = GetMutableDefault<UMotionTrailToolOptions>();
+	if (RestoreStates.bMotionTrailsOn && Settings)
+	{
+		Settings->bShowTrails = true;
+		FPropertyChangedEvent ShowTrailEvent(
+			UMotionTrailToolOptions::StaticClass()->FindPropertyByName(GET_MEMBER_NAME_CHECKED(UMotionTrailToolOptions, bShowTrails))
+			);
+		Settings->PostEditChangeProperty(ShowTrailEvent);
+	}
 }
 
 void FControlRigEditModeToolkit::ShutdownUI()
@@ -495,121 +548,93 @@ void FControlRigEditModeToolkit::ShutdownUI()
 
 void FControlRigEditModeToolkit::UnregisterAndRemoveFloatingTabs()
 {
-	if (FSlateApplication::IsInitialized())
+	if (!FSlateApplication::IsInitialized())
 	{
-		if (TweenOverlayManager)
-		{
-			bTweenOpen = TweenOverlayManager->IsShowingWidget();
-			TweenOverlayManager.Reset();
-		}
+		return;
+	}
+
+	SaveLayoutState();
+	FControlRigUIRestoreStates& RestoreStates = UControlRigEditModeSettings::Get()->LastUIStates;
+
+	if (TweenOverlayManager)
+	{
+		TweenOverlayManager.Reset();
+	}
+	
+	if (ConstrainToolsManager)
+	{
+		ConstrainToolsManager.Reset();
+	}
+	
+	SelectionViewModel.Reset();
 
-		if (ConstrainToolsManager)
-		{
-			bConstraintsOpen = ConstrainToolsManager->IsShowingWidget();
-			ConstrainToolsManager.Reset();
-		}
-		
-		SelectionViewModel.Reset();
 #if SELECTION_SETS_AS_TAB == 0
-		if (SelectionSetsOverlayManager)
+	if (SelectionSetsOverlayManager)
+	{
+		if (RestoreStates.bSelectionSetsOpen)
 		{
-			bSelectionSetsOpen = SelectionSetsOverlayManager->IsShowingWidget();
-			if (bSelectionSetsOpen)
-			{
-				SelectionSetsOverlayManager->DestroyWidget();
-			}
-			SelectionSetsOverlayManager.Reset();
+			SelectionSetsOverlayManager->DestroyWidget();
 		}
+		SelectionSetsOverlayManager.Reset();
+	}
 #endif
-		
-		if (ModeUILayer.IsValid())
-		{
-			TSharedPtr<FAssetEditorModeUILayer> ModeUILayerPtr = ModeUILayer.Pin();
-
-			bMotionTrailsOn = false;
-			if (UMotionTrailToolOptions* Settings = GetMutableDefault<UMotionTrailToolOptions>())
-			{
-				if (Settings->bShowTrails)
-				{
-					bMotionTrailsOn = true;
-					Settings->bShowTrails = false;
-					FPropertyChangedEvent ShowTrailEvent(UMotionTrailToolOptions::StaticClass()->FindPropertyByName(GET_MEMBER_NAME_CHECKED(UMotionTrailToolOptions, bShowTrails)));
-					Settings->PostEditChangeProperty(ShowTrailEvent);
-				}
-			}
 
-			ModeUILayerPtr->GetTabManager()->UnregisterTabSpawner(MotionTrailTabName);
+	if (const TSharedPtr<FAssetEditorModeUILayer> ModeUILayerPtr = ModeUILayer.Pin())
+	{
+		TryCloseToolkitUI(MotionTrailTabName);
+		TryCloseToolkitUI(AnimLayerTabName);
+		TryCloseToolkitUI(SnapperTabName);
+		TryCloseToolkitUI(PoseTabName);
 #if SELECTION_SETS_AS_TAB
-			TSharedPtr<SDockTab> SelectionSetTab = ModeUILayerPtr->GetTabManager()->FindExistingLiveTab(FTabId(SelectionSetsTabName));
-			if (SelectionSetTab)
-			{
-				bSelectionSetsOpen = true;
-				SelectionSetTab->RequestCloseTab();
-			}
-			else
-			{
-				bSelectionSetsOpen = false;
-			}
-			if (TSharedPtr<ISequencer> SequencerPtr = UAIESelectionSets::GetSequencerFromAsset())
-			{
-				//wasn't open but there are layers in open level sequence we try to open it
-				if (UAIESelectionSets* SelectionSets = UAIESelectionSets::GetSelectionSets(SequencerPtr))
-				{
-					/*
-					if (SelectionSets && SelectionSets->bOpenUIOnOpen)
-					{
-						AnimLayers->bOpenUIOnOpen = bAnimLayerTabOpen;
-					}
-					*/
-				}
-			}
-			ModeUILayerPtr->GetTabManager()->UnregisterTabSpawner(SelectionSetsTabName);
+		TryCloseToolkitUI(SelectionSetsTabName);
+#endif
 
+		ModeUILayerPtr->GetTabManager()->UnregisterTabSpawner(MotionTrailTabName);
+		ModeUILayerPtr->GetTabManager()->UnregisterTabSpawner(AnimLayerTabName);
+		ModeUILayerPtr->GetTabManager()->UnregisterTabSpawner(SnapperTabName);
+		ModeUILayerPtr->GetTabManager()->UnregisterTabSpawner(PoseTabName);
+#if SELECTION_SETS_AS_TAB
+		ModeUILayerPtr->GetTabManager()->UnregisterTabSpawner(SelectionSetsTabName);
 #endif
-			TSharedPtr<SDockTab> AnimLayerTab = ModeUILayerPtr->GetTabManager()->FindExistingLiveTab(FTabId(AnimLayerTabName));
-			if (AnimLayerTab)
-			{
-				bAnimLayerTabOpen = true;
-				AnimLayerTab->RequestCloseTab();
-			}
-			else
-			{
-				bAnimLayerTabOpen = false;
-			}
-			if (TSharedPtr<ISequencer> SequencerPtr = UAnimLayers::GetSequencerFromAsset())
-			{
-				if (UAnimLayers* AnimLayers = UAnimLayers::GetAnimLayers(SequencerPtr.Get()))
-				{
-					AnimLayers->bOpenUIOnOpen = bAnimLayerTabOpen;
-				}
-			}
-			ModeUILayerPtr->GetTabManager()->UnregisterTabSpawner(AnimLayerTabName);
 
-			TSharedPtr<SDockTab> SnapperTab = ModeUILayerPtr->GetTabManager()->FindExistingLiveTab(FTabId(SnapperTabName));
-			if (SnapperTab)
-			{
-				bSnapperTabOpen = true;
-				SnapperTab->RequestCloseTab();
-			}
-			else
-			{
-				bSnapperTabOpen = false;
-			}
-			ModeUILayerPtr->GetTabManager()->UnregisterTabSpawner(SnapperTabName);
-		
-			TSharedPtr<SDockTab> PoseTab = ModeUILayerPtr->GetTabManager()->FindExistingLiveTab(FTabId(PoseTabName));
-			if (PoseTab)
-			{
-				bPoseTabOpen = true;
-				PoseTab->RequestCloseTab();
-			}
-			else
+		if (UMotionTrailToolOptions* Settings = GetMutableDefault<UMotionTrailToolOptions>())
+		{
+			if (Settings->bShowTrails)
 			{
-				bPoseTabOpen = false;
+				Settings->bShowTrails = false;
+				FPropertyChangedEvent ShowTrailEvent(
+					UMotionTrailToolOptions::StaticClass()->FindPropertyByName(GET_MEMBER_NAME_CHECKED(UMotionTrailToolOptions, bShowTrails))
+					);
+				Settings->PostEditChangeProperty(ShowTrailEvent);
 			}
-			ModeUILayerPtr->GetTabManager()->UnregisterTabSpawner(PoseTabName);
+		}
+
+		const TSharedPtr<ISequencer> SequencerPtr = UAnimLayers::GetSequencerFromAsset();
+		UAnimLayers* AnimLayers = SequencerPtr ? UAnimLayers::GetAnimLayers(SequencerPtr.Get()) : nullptr;
+		if (AnimLayers)
+		{
+			AnimLayers->bOpenUIOnOpen = RestoreStates.bAnimLayerTabOpen;
 		}
 	}
 }
 
+void FControlRigEditModeToolkit::SaveLayoutState() const
+{
+	UControlRigEditModeSettings* Settings = UControlRigEditModeSettings::Get();
+	FControlRigUIRestoreStates& RestoreStates = Settings->LastUIStates;
+	
+	RestoreStates.bMotionTrailsOn = IsToolkitUIActive(MotionTrailTabName);
+	RestoreStates.bAnimLayerTabOpen = IsToolkitUIActive(AnimLayerTabName);
+	RestoreStates.bPoseTabOpen = IsToolkitUIActive(PoseTabName);
+	RestoreStates.bSnapperTabOpen = IsToolkitUIActive(SnapperTabName);
+	RestoreStates.bConstraintsOpen = IsToolkitUIActive(ConstrainingTabName);
+#if SELECTION_SETS_AS_TAB
+	RestoreStates.bSelectionSetsOpen = IsToolkitUIActive(SelectionSetsTabName);
+#else
+	RestoreStates.bSelectionSetsOpen = IsToolkitUIActive(SelectionSetsOverlayName);
+#endif
+	
+	Settings->SaveConfig();
+}
+
 #undef LOCTEXT_NAMESPACE
diff --git a/Engine/Plugins/Animation/ControlRig/Source/ControlRigEditor/Private/EditMode/ControlRigEditModeToolkit.h b/Engine/Plugins/Animation/ControlRig/Source/ControlRigEditor/Private/EditMode/ControlRigEditModeToolkit.h
--- a/Engine/Plugins/Animation/ControlRig/Source/ControlRigEditor/Private/EditMode/ControlRigEditModeToolkit.h
+++ b/Engine/Plugins/Animation/ControlRig/Source/ControlRigEditor/Private/EditMode/ControlRigEditModeToolkit.h
@@ -65,7 +65,13 @@ public:
 	virtual FText GetActiveToolMessage() const override;
 	virtual void OnToolPaletteChanged(FName PaletteName) override;
 
+	/** If the named UI is shown, hide it. If it's hidden, show it. */
+	void TryToggleToolkitUI(const FName InName);
+	/** Tries to spawn the named UI. If already spawned, and it's a tab, highlight it (little flash). */
 	void TryInvokeToolkitUI(const FName InName);
+	/** Tries to hide the named UI. */
+	void TryCloseToolkitUI(const FName InName) const;
+	/** @return Whether the named UI is open */
 	bool IsToolkitUIActive(const FName InName) const;
 
 	/** Called when the global array of control rigs changed the user is working on has changed. */
@@ -89,7 +95,6 @@ public:
 	static const FName SnapperTabName;
 	static const FName DetailsTabName;
 	static const FName OutlinerTabName;
-	static const FName SpacePickerTabName;
 
 	static TSharedPtr<UE::ControlRigEditor::SAnimDetailsView> Details;
 	static TSharedPtr<SControlRigOutliner> Outliner;
@@ -101,17 +106,10 @@ protected:
 	virtual void InvokeUI() override;
 	virtual void ShutdownUI() override;
 
-	//this also saves the layout
+	/** Removes all UI, unregisters tab spawners, and saves the state so it can be restored next time. */
 	void UnregisterAndRemoveFloatingTabs();
-
-	static bool bMotionTrailsOn;
-	static bool bAnimLayerTabOpen;
-	static bool bPoseTabOpen;
-	static bool bSnapperTabOpen;
-	static bool bTweenOpen;
-	
-	static bool bConstraintsOpen;
-	static bool bSelectionSetsOpen;
+	/** Saves the layout state into config so it can be restored next time the user restarts the editor. */
+	void SaveLayoutState() const;
 
 private:
 	
diff --git a/Engine/Plugins/Animation/ControlRig/Source/ControlRigEditor/Private/EditMode/SControlRigEditModeTools.cpp b/Engine/Plugins/Animation/ControlRig/Source/ControlRigEditor/Private/EditMode/SControlRigEditModeTools.cpp
--- a/Engine/Plugins/Animation/ControlRig/Source/ControlRigEditor/Private/EditMode/SControlRigEditModeTools.cpp
+++ b/Engine/Plugins/Animation/ControlRig/Source/ControlRigEditor/Private/EditMode/SControlRigEditModeTools.cpp
@@ -335,24 +335,26 @@ bool SControlRigEditModeTools::IsReadOnlyPropertyOnDetailCustomization(const FPr
 
 void SControlRigEditModeTools::CustomizeToolBarPalette(FToolBarBuilder& ToolBarBuilder)
 {
+	const TSharedPtr<FControlRigEditModeToolkit> OwningToolkitPin = OwningToolkit.Pin();
+	if (!ensure(OwningToolkitPin))
+	{
+		return;
+	}
+	
 	//TOGGLE SELECTED RIG CONTROLS
 	ToolBarBuilder.AddToolBarButton(
 		FUIAction(
 			FExecuteAction::CreateLambda([this] {
-			FControlRigEditMode* ControlRigEditMode = static_cast<FControlRigEditMode*>(ModeTools->GetActiveMode(FControlRigEditMode::ModeName));
-			if (ControlRigEditMode)
-			{
-				ControlRigEditMode->SetOnlySelectRigControls(!ControlRigEditMode->GetOnlySelectRigControls());
-			}
+				FControlRigEditMode* ControlRigEditMode = static_cast<FControlRigEditMode*>(ModeTools->GetActiveMode(FControlRigEditMode::ModeName));
+				if (ControlRigEditMode)
+				{
+					ControlRigEditMode->SetOnlySelectRigControls(!ControlRigEditMode->GetOnlySelectRigControls());
+				}
 			}),
 			FCanExecuteAction(),
 			FIsActionChecked::CreateLambda([this] {
-			FControlRigEditMode* ControlRigEditMode = static_cast<FControlRigEditMode*>(ModeTools->GetActiveMode(FControlRigEditMode::ModeName));
-			if (ControlRigEditMode)
-			{
-				return ControlRigEditMode->GetOnlySelectRigControls();
-			}
-			return false;
+				FControlRigEditMode* ControlRigEditMode = static_cast<FControlRigEditMode*>(ModeTools->GetActiveMode(FControlRigEditMode::ModeName));
+				return ControlRigEditMode && ControlRigEditMode->GetOnlySelectRigControls();
 			})
 
 		),
@@ -367,21 +369,25 @@ void SControlRigEditModeTools::CustomizeToolBarPalette(FToolBarBuilder& ToolBarB
 
 	//POSES
 	ToolBarBuilder.AddToolBarButton(
-		FExecuteAction::CreateRaw(OwningToolkit.Pin().Get(), &FControlRigEditModeToolkit::TryInvokeToolkitUI, FControlRigEditModeToolkit::PoseTabName),
+		FUIAction(
+			FExecuteAction::CreateSP(OwningToolkitPin.Get(), &FControlRigEditModeToolkit::TryToggleToolkitUI, FControlRigEditModeToolkit::PoseTabName),
+			FCanExecuteAction(),
+			FIsActionChecked::CreateSP(OwningToolkitPin.Get(), &FControlRigEditModeToolkit::IsToolkitUIActive, FControlRigEditModeToolkit::PoseTabName)
+			),
 		NAME_None,
 		LOCTEXT("Poses", "Poses"),
 		LOCTEXT("PosesTooltip", "Show Poses"),
 		FSlateIcon(TEXT("ControlRigEditorStyle"), TEXT("ControlRig.PoseTool")),
-		EUserInterfaceActionType::Button
+		EUserInterfaceActionType::ToggleButton
 	);
 	ToolBarBuilder.AddSeparator();
 
 	// Tweens
 	ToolBarBuilder.AddToolBarButton(
 		FUIAction(
-		FExecuteAction::CreateRaw(OwningToolkit.Pin().Get(), &FControlRigEditModeToolkit::TryInvokeToolkitUI, FControlRigEditModeToolkit::TweenOverlayName),
-		FCanExecuteAction(),
-		FIsActionChecked::CreateRaw(OwningToolkit.Pin().Get(), &FControlRigEditModeToolkit::IsToolkitUIActive, FControlRigEditModeToolkit::TweenOverlayName)
+			FExecuteAction::CreateSP(OwningToolkitPin.Get(), &FControlRigEditModeToolkit::TryToggleToolkitUI, FControlRigEditModeToolkit::TweenOverlayName),
+			FCanExecuteAction(),
+			FIsActionChecked::CreateSP(OwningToolkitPin.Get(), &FControlRigEditModeToolkit::IsToolkitUIActive, FControlRigEditModeToolkit::TweenOverlayName)
 		),		
 		NAME_None,
 		LOCTEXT("Tweens", "Tweens"),
@@ -392,9 +398,9 @@ void SControlRigEditModeTools::CustomizeToolBarPalette(FToolBarBuilder& ToolBarB
 	
 	ToolBarBuilder.AddToolBarButton(
 		FUIAction(
-			FExecuteAction::CreateRaw(OwningToolkit.Pin().Get(), &FControlRigEditModeToolkit::TryInvokeToolkitUI, FControlRigEditModeToolkit::ConstrainingTabName),
+			FExecuteAction::CreateSP(OwningToolkitPin.Get(), &FControlRigEditModeToolkit::TryToggleToolkitUI, FControlRigEditModeToolkit::ConstrainingTabName),
 			FCanExecuteAction(),
-			FIsActionChecked::CreateRaw(OwningToolkit.Pin().Get(), &FControlRigEditModeToolkit::IsToolkitUIActive, FControlRigEditModeToolkit::ConstrainingTabName)
+			FIsActionChecked::CreateSP(OwningToolkitPin.Get(), &FControlRigEditModeToolkit::IsToolkitUIActive, FControlRigEditModeToolkit::ConstrainingTabName)
 		),
 		NAME_None,
 		LOCTEXT("Constrain", "Constrain"),
@@ -407,7 +413,11 @@ void SControlRigEditModeTools::CustomizeToolBarPalette(FToolBarBuilder& ToolBarB
 #if SELECTION_SETS_AS_TAB
 
 	ToolBarBuilder.AddToolBarButton(
-		FExecuteAction::CreateRaw(OwningToolkit.Pin().Get(), &FControlRigEditModeToolkit::TryInvokeToolkitUI, FControlRigEditModeToolkit::SelectionSetsTabName),
+		FUIAction(
+			FExecuteAction::CreateRaw(OwningToolkit.Pin().Get(), &FControlRigEditModeToolkit::TryToggleToolkitUI, FControlRigEditModeToolkit::SelectionSetsTabName),
+			FCanExecuteAction(),
+			FIsActionChecked::CreateRaw(OwningToolkit.Pin().Get(), &FControlRigEditModeToolkit::IsToolkitUIActive, FControlRigEditModeToolkit::SelectionSetsTabName)
+		),
 		NAME_None,
 		LOCTEXT("SelectionSets", "Selection Sets"),
 		LOCTEXT("SelectionSetsTooltip", "Open Selection Sets"),
@@ -418,7 +428,7 @@ void SControlRigEditModeTools::CustomizeToolBarPalette(FToolBarBuilder& ToolBarB
 
 	ToolBarBuilder.AddToolBarButton(
 		FUIAction(
-			FExecuteAction::CreateRaw(OwningToolkit.Pin().Get(), &FControlRigEditModeToolkit::TryInvokeToolkitUI, FControlRigEditModeToolkit::SelectionSetsOverlayName),
+			FExecuteAction::CreateRaw(OwningToolkit.Pin().Get(), &FControlRigEditModeToolkit::TryToggleToolkitUI, FControlRigEditModeToolkit::SelectionSetsOverlayName),
 			FCanExecuteAction(),
 			FIsActionChecked::CreateRaw(OwningToolkit.Pin().Get(), &FControlRigEditModeToolkit::IsToolkitUIActive, FControlRigEditModeToolkit::SelectionSetsOverlayName)
 		),
@@ -445,12 +455,16 @@ void SControlRigEditModeTools::CustomizeToolBarPalette(FToolBarBuilder& ToolBarB
 
 	// Anim Layer
 	ToolBarBuilder.AddToolBarButton(
-		FExecuteAction::CreateRaw(OwningToolkit.Pin().Get(), &FControlRigEditModeToolkit::TryInvokeToolkitUI, FControlRigEditModeToolkit::AnimLayerTabName),
+		FUIAction(
+			FExecuteAction::CreateSP(OwningToolkitPin.Get(), &FControlRigEditModeToolkit::TryToggleToolkitUI, FControlRigEditModeToolkit::AnimLayerTabName),
+			FCanExecuteAction(),
+			FIsActionChecked::CreateSP(OwningToolkitPin.Get(), &FControlRigEditModeToolkit::IsToolkitUIActive, FControlRigEditModeToolkit::AnimLayerTabName)
+		),
 		NAME_None,
 		LOCTEXT("Layers", "Layers"),
 		LOCTEXT("AnimLayersTooltip", "Display animation layers"),
 		FSlateIcon(TEXT("ControlRigEditorStyle"), TEXT("ControlRig.AnimLayers")),
-		EUserInterfaceActionType::Button
+		EUserInterfaceActionType::ToggleButton
 	);
 
 	//Pivot
diff --git a/Engine/Plugins/Animation/ControlRig/Source/ControlRigEditor/Private/EditMode/Settings/ControlRigUIRestoreStates.h b/Engine/Plugins/Animation/ControlRig/Source/ControlRigEditor/Private/EditMode/Settings/ControlRigUIRestoreStates.h
new file mode 100644
--- /dev/null
+++ b/Engine/Plugins/Animation/ControlRig/Source/ControlRigEditor/Private/EditMode/Settings/ControlRigUIRestoreStates.h
@@ -0,0 +1,45 @@
+﻿// Copyright Epic Games, Inc. All Rights Reserved.
+
+#pragma once
+
+#include "Misc/Flyout/FlyoutSavedState.h"
+#include "Overlay/DragBoxPosition.h"
+#include "ControlRigUIRestoreStates.generated.h"
+
+/** Data to restore the UI state control rig had last time it was open. */
+USTRUCT()
+struct FControlRigUIRestoreStates
+{
+	GENERATED_BODY()
+
+	UPROPERTY()
+	bool bMotionTrailsOn = false;
+	
+	UPROPERTY()
+	bool bAnimLayerTabOpen = true;
+	
+	UPROPERTY()
+	bool bPoseTabOpen = false;
+	
+	UPROPERTY()
+	bool bSnapperTabOpen = false;
+	
+	UPROPERTY()
+	bool bConstraintsOpen = false;
+
+	/** Whether the selection set was open when SELECTION_SETS_AS_TAB == 1 */
+	UPROPERTY()
+	bool bSelectionSetsOpen = false;
+
+	/** Selection set view state if SELECTION_SETS_AS_TAB == 0 */
+	UPROPERTY()
+	FToolWidget_FlyoutSavedState SelectionSetOverlayState;
+
+	UPROPERTY()
+	FToolWidget_FlyoutSavedState TweenOverlayState;
+
+	explicit FControlRigUIRestoreStates()
+		: SelectionSetOverlayState(FToolWidget_DragBoxPosition(FVector2f(20.f, 30.f), HAlign_Right, VAlign_Top), false)
+		, TweenOverlayState(FToolWidget_DragBoxPosition(FVector2f(20.f, 30.f), HAlign_Right, VAlign_Bottom), false)
+	{}
+};
\ No newline at end of file
diff --git a/Engine/Plugins/Animation/ControlRig/Source/ControlRigEditor/Private/EditMode/Tween/TweenOverlayManager.cpp b/Engine/Plugins/Animation/ControlRig/Source/ControlRigEditor/Private/EditMode/Tween/TweenOverlayManager.cpp
--- a/Engine/Plugins/Animation/ControlRig/Source/ControlRigEditor/Private/EditMode/Tween/TweenOverlayManager.cpp
+++ b/Engine/Plugins/Animation/ControlRig/Source/ControlRigEditor/Private/EditMode/Tween/TweenOverlayManager.cpp
@@ -23,7 +23,7 @@ namespace Private
 static void SaveWidgetState(FToolWidget_FlyoutSavedState State)
 {
 	UControlRigEditModeSettings* ControlRigEditModeSettings = GetMutableDefault<UControlRigEditModeSettings>();
-	ControlRigEditModeSettings->LastTweenOverlayState = State;
+	ControlRigEditModeSettings->LastUIStates.TweenOverlayState = State;
 	ControlRigEditModeSettings->SaveConfig();
 }
 }
@@ -50,7 +50,7 @@ FTweenOverlayManager::FTweenOverlayManager(
 			TAttribute<TSharedPtr<FUICommandList>>::CreateLambda([WeakList = InToolkitCommandList.ToWeakPtr()]{ return WeakList.Pin(); }),
 			MakeShared<FToolkitBasedWidgetHost>(InToolkitHost)
 			)
-			.SetStateToRestoreFrom(GetMutableDefault<UControlRigEditModeSettings>()->LastTweenOverlayState)
+			.SetStateToRestoreFrom(GetMutableDefault<UControlRigEditModeSettings>()->LastUIStates.TweenOverlayState)
 			.SetSaveStateDelegate(FSaveFlyoutState::CreateStatic(&Private::SaveWidgetState))
 			.SetCanSubdue(TAttribute<bool>::CreateLambda([this]{ return !TweenControllers.GetControllers().MouseSlidingController.IsSliding(); }))
 			.SetSubduedOpacity(TAttribute<FLinearColor>::CreateLambda([]{ return GetMutableDefault<UControlRigEditModeSettings>()->TweenOutOfFocusTint; }))
diff --git a/Engine/Plugins/Animation/ControlRig/Source/ControlRigEditor/Private/Sequencer/SelectionSets/SelectionSetsOverlayManager.cpp b/Engine/Plugins/Animation/ControlRig/Source/ControlRigEditor/Private/Sequencer/SelectionSets/SelectionSetsOverlayManager.cpp
--- a/Engine/Plugins/Animation/ControlRig/Source/ControlRigEditor/Private/Sequencer/SelectionSets/SelectionSetsOverlayManager.cpp
+++ b/Engine/Plugins/Animation/ControlRig/Source/ControlRigEditor/Private/Sequencer/SelectionSets/SelectionSetsOverlayManager.cpp
@@ -15,7 +15,7 @@ namespace Private
 static void SaveWidgetState(FToolWidget_FlyoutSavedState State)
 {
 	UControlRigEditModeSettings* ControlRigEditModeSettings = GetMutableDefault<UControlRigEditModeSettings>();
-	ControlRigEditModeSettings->LastSelectionSetOverlayState = State;
+	ControlRigEditModeSettings->LastUIStates.SelectionSetOverlayState = State;
 	ControlRigEditModeSettings->SaveConfig();
 }
 }
@@ -32,7 +32,7 @@ FSelectionSetsOverlayManager::FSelectionSetsOverlayManager(
 			TAttribute<TSharedPtr<FUICommandList>>::CreateLambda([WeakList = InToolkitCommandList.ToWeakPtr()]{ return WeakList.Pin(); }),
 			MakeShared<UE::ControlRigEditor::FToolkitBasedWidgetHost>(InToolkitHost)
 			)
-			.SetStateToRestoreFrom(GetMutableDefault<UControlRigEditModeSettings>()->LastSelectionSetOverlayState)
+			.SetStateToRestoreFrom(GetMutableDefault<UControlRigEditModeSettings>()->LastUIStates.SelectionSetOverlayState)
 			.SetSaveStateDelegate(UE::ControlRigEditor::FSaveFlyoutState::CreateStatic(&UE::AIE::Private::SaveWidgetState))
 			.SetToggleVisibility(FControlRigEditModeCommands::Get().ToggleSelectionSetsWidget)
 			.SetSummonToCursor(FControlRigEditModeCommands::Get().SummonSelectionSetsWidget)
