Control Rig: Avoid warning when opening modular rig editor
#rnx

diff --git a/Engine/Plugins/Animation/ControlRig/Source/ControlRigEditor/Private/Editor/ControlRigEditorMode.cpp b/Engine/Plugins/Animation/ControlRig/Source/ControlRigEditor/Private/Editor/ControlRigEditorMode.cpp
--- a/Engine/Plugins/Animation/ControlRig/Source/ControlRigEditor/Private/Editor/ControlRigEditorMode.cpp
+++ b/Engine/Plugins/Animation/ControlRig/Source/ControlRigEditor/Private/Editor/ControlRigEditorMode.cpp
@@ -103,7 +103,7 @@ FControlRigLegacyEditorMode::FControlRigLegacyEditorMode(const TSharedRef<FContr
 							// Middle top - document edit area
 							FTabManager::NewStack()
 							->SetSizeCoefficient(0.8f)
-							->AddTab("Document", ETabState::ClosedTab)
+							->AddTab("GraphEditor", ETabState::ClosedTab)
 						)
 						->Split
 						(
@@ -165,7 +165,7 @@ FModularRigLegacyEditorMode::FModularRigLegacyEditorMode(const TSharedRef<FContr
 	TabFactories.RegisterFactory(MakeShared<FRigModuleAssetBrowserTabSummoner>(InControlRigEditor));
 	TabFactories.RegisterFactory(MakeShared<FModularRigEventQueueTabSummoner>(InControlRigEditor));
 
-	TabLayout = FTabManager::NewLayout("Standalone_ModularRigEditMode_Layout_v1.3")
+	TabLayout = FTabManager::NewLayout("Standalone_ModularRigEditMode_Layout_v1.4")
 		->AddArea
 		(
 			// Main application area
@@ -218,6 +218,7 @@ FModularRigLegacyEditorMode::FModularRigLegacyEditorMode(const TSharedRef<FContr
 					->AddTab(FBlueprintEditorTabs::DetailsID, ETabState::OpenedTab)
 					->AddTab(FPersonaTabs::AdvancedPreviewSceneSettingsID, ETabState::OpenedTab)
 					->AddTab(FRigHierarchyTabSummoner::TabID, ETabState::OpenedTab)
+					->AddTab("GraphEditor", ETabState::ClosedTab) // Adding the graph editor to avoid a warning about blueprint assumption there will be a graph editor
 				)
 			)
 		);
@@ -335,7 +336,7 @@ FControlRigEditorMode::FControlRigEditorMode(const TSharedRef<FControlRigEditor>
 							// Middle top - document edit area
 							FTabManager::NewStack()
 							->SetSizeCoefficient(0.8f)
-							->AddTab("Document", ETabState::ClosedTab)
+							->AddTab("GraphEditor", ETabState::ClosedTab)
 						)
 						->Split
 						(
diff --git a/Engine/Plugins/Runtime/RigVM/Source/RigVMEditor/Private/Editor/RigVMLegacyEditor.cpp b/Engine/Plugins/Runtime/RigVM/Source/RigVMEditor/Private/Editor/RigVMLegacyEditor.cpp
--- a/Engine/Plugins/Runtime/RigVM/Source/RigVMEditor/Private/Editor/RigVMLegacyEditor.cpp
+++ b/Engine/Plugins/Runtime/RigVM/Source/RigVMEditor/Private/Editor/RigVMLegacyEditor.cpp
@@ -12,7 +12,7 @@ struct FRigVMEditorZoomLevelsContainer;
 
 FRigVMLegacyEditor::FRigVMLegacyEditor()
 {
-	
+	DocumentManager = MakeShareable(new FDocumentTracker(NAME_None));
 }
 
 void FRigVMLegacyEditor::OnClose()
diff --git a/Engine/Plugins/Runtime/RigVM/Source/RigVMEditor/Private/Editor/RigVMLegacyEditorMode.cpp b/Engine/Plugins/Runtime/RigVM/Source/RigVMEditor/Private/Editor/RigVMLegacyEditorMode.cpp
--- a/Engine/Plugins/Runtime/RigVM/Source/RigVMEditor/Private/Editor/RigVMLegacyEditorMode.cpp
+++ b/Engine/Plugins/Runtime/RigVM/Source/RigVMEditor/Private/Editor/RigVMLegacyEditorMode.cpp
@@ -48,7 +48,7 @@ FRigVMLegacyEditorMode::FRigVMLegacyEditorMode(const TSharedRef<FRigVMLegacyEdit
 						// Middle top - document edit area
 						FTabManager::NewStack()
 						->SetSizeCoefficient(0.8f)
-						->AddTab("Document", ETabState::ClosedTab)
+						->AddTab("GraphEditor", ETabState::ClosedTab)
 					)
 					->Split
 					(
diff --git a/Engine/Plugins/Runtime/RigVM/Source/RigVMEditor/Private/Editor/RigVMNewEditor.cpp b/Engine/Plugins/Runtime/RigVM/Source/RigVMEditor/Private/Editor/RigVMNewEditor.cpp
--- a/Engine/Plugins/Runtime/RigVM/Source/RigVMEditor/Private/Editor/RigVMNewEditor.cpp
+++ b/Engine/Plugins/Runtime/RigVM/Source/RigVMEditor/Private/Editor/RigVMNewEditor.cpp
@@ -187,7 +187,7 @@ FRigVMNewEditor::FRigVMNewEditor()
 	: FRigVMEditorBase()
 	, bIsActionMenuContextSensitive(true)
 {
-	DocumentManager = MakeShareable(new FDocumentTracker);
+	DocumentManager = MakeShareable(new FDocumentTracker(NAME_None));
 }
 
 void FRigVMNewEditor::OnClose()
diff --git a/Engine/Plugins/Runtime/RigVM/Source/RigVMEditor/Private/Editor/RigVMNewEditorMode.cpp b/Engine/Plugins/Runtime/RigVM/Source/RigVMEditor/Private/Editor/RigVMNewEditorMode.cpp
--- a/Engine/Plugins/Runtime/RigVM/Source/RigVMEditor/Private/Editor/RigVMNewEditorMode.cpp
+++ b/Engine/Plugins/Runtime/RigVM/Source/RigVMEditor/Private/Editor/RigVMNewEditorMode.cpp
@@ -52,7 +52,7 @@ FRigVMNewEditorMode::FRigVMNewEditorMode(const TSharedRef<FRigVMNewEditor>& InRi
 						// Middle top
 						FTabManager::NewStack()
 						->SetSizeCoefficient(0.8f)
-						->AddTab("Document", ETabState::ClosedTab)
+						->AddTab("GraphEditor", ETabState::ClosedTab)
 					)
 					->Split
 						(
