Control Rig: Attempt to fix crash in UControlRigComponent when accessing deleted USceneComponent in a mapped element
#rnx

diff --git a/Engine/Plugins/Animation/ControlRig/Source/ControlRig/Private/ControlRigComponent.cpp b/Engine/Plugins/Animation/ControlRig/Source/ControlRig/Private/ControlRigComponent.cpp
--- a/Engine/Plugins/Animation/ControlRig/Source/ControlRig/Private/ControlRigComponent.cpp
+++ b/Engine/Plugins/Animation/ControlRig/Source/ControlRig/Private/ControlRigComponent.cpp
@@ -1501,6 +1501,11 @@ void UControlRigComponent::TransferInputs()
 				continue;
 			}
 
+			if (!IsValid(MappedElement.SceneComponent))
+			{
+				continue;
+			}
+
 			FTransform Transform = FTransform::Identity;
 			if (MappedElement.SubIndex >= 0)
 			{
@@ -1674,6 +1679,11 @@ void UControlRigComponent::TransferOutputs()
 					continue;
 				}
 
+				if (!IsValid(MappedElement.SceneComponent))
+				{
+					continue;
+				}
+
 				if (MappedElement.ElementType == ERigElementType::Bone ||
 					MappedElement.ElementType == ERigElementType::Control ||
 					MappedElement.ElementType == ERigElementType::Null)
