Control Rig: Show space relationships in dependency viewer

#rnx

diff --git a/Engine/Plugins/Animation/ControlRig/Source/ControlRigEditor/Private/RigDependencyGraph/RigDependencyConnectionDrawingPolicy.cpp b/Engine/Plugins/Animation/ControlRig/Source/ControlRigEditor/Private/RigDependencyGraph/RigDependencyConnectionDrawingPolicy.cpp
--- a/Engine/Plugins/Animation/ControlRig/Source/ControlRigEditor/Private/RigDependencyGraph/RigDependencyConnectionDrawingPolicy.cpp
+++ b/Engine/Plugins/Animation/ControlRig/Source/ControlRigEditor/Private/RigDependencyGraph/RigDependencyConnectionDrawingPolicy.cpp
@@ -1,6 +1,8 @@
 // Copyright Epic Games, Inc. All Rights Reserved.
 
 #include "RigDependencyGraph/RigDependencyConnectionDrawingPolicy.h"
+
+#include "RigDependencyGraph.h"
 #include "RigDependencyGraph/RigDependencyGraphNode.h"
 
 void FRigDependencyConnectionDrawingPolicy::DetermineWiringStyle(UEdGraphPin* OutputPin, UEdGraphPin* InputPin, /*inout*/ FConnectionParams& Params)
@@ -13,11 +15,24 @@ void FRigDependencyConnectionDrawingPolicy::DetermineWiringStyle(UEdGraphPin* Ou
 		if (OutputNode->GetNodeId().IsElement() && InputNode->GetNodeId().IsElement())
 		{
 			static const FLinearColor ParentChildWireColor = FLinearColor(FColor::FromHex(TEXT("#AAAAAA")));
+			static const FLinearColor ControlSpaceWireColor = FLinearColor(FColor::FromHex(TEXT("#22AA22")));
 			Params.WireColor = ParentChildWireColor;
+
+			if (const URigDependencyGraph* DependencyGraph = Cast<URigDependencyGraph>(OutputNode->GetGraph()))
+			{
+				const uint32 SpaceToControlhash = HashCombine(GetTypeHash(OutputNode->GetNodeId().Index), GetTypeHash(InputNode->GetNodeId().Index));
+				if (DependencyGraph->HierarchySpaceToControlHashes.Contains(SpaceToControlhash))
+				{
+					Params.WireColor = ControlSpaceWireColor;
+				}
+			}
+
 			Params.WireColor.A = 0.75f;
 		}
 
 		const float FadedState = FMath::Min(OutputNode->GetFadedOutState(), InputNode->GetFadedOutState());
 		Params.WireColor.A *= FadedState;
+
+		Params.bUserFlag1 = OutputNode->GetNodeId().IsInstruction() || InputNode->GetNodeId().IsInstruction();
 	}
 }
\ No newline at end of file
diff --git a/Engine/Plugins/Animation/ControlRig/Source/ControlRigEditor/Private/RigDependencyGraph/RigDependencyGraph.cpp b/Engine/Plugins/Animation/ControlRig/Source/ControlRigEditor/Private/RigDependencyGraph/RigDependencyGraph.cpp
--- a/Engine/Plugins/Animation/ControlRig/Source/ControlRigEditor/Private/RigDependencyGraph/RigDependencyGraph.cpp
+++ b/Engine/Plugins/Animation/ControlRig/Source/ControlRigEditor/Private/RigDependencyGraph/RigDependencyGraph.cpp
@@ -18,6 +18,7 @@ URigDependencyGraph::URigDependencyGraph(const FObjectInitializer& ObjectInitial
 	, LayoutIterationsTotal(0)
 	, LayoutIterationsMax(DefaultLayoutIterationsPerNode * 1000)
 	, bFollowParentRelationShips(true)
+	, bFollowControlSpaceRelationships(true)
 	, bFollowVMRelationShips(false)
 	, bLockContent(false)
 	, bIsPerformingGridLayout(true)
@@ -295,17 +296,26 @@ void URigDependencyGraph::RebuildLinks()
 		{
 			InSourceNode->OutputPin->MakeLinkTo(InTargetNode->InputPin);
 		}
-			if (InSourceNode->GetNodeId().IsInstruction() &&
-				InTargetNode->OutputPin->LinkedTo.Contains(InSourceNode->InputPin))
-			{
-				InTargetNode->OutputPin->BreakLinkTo(InSourceNode->InputPin);
-			}
+		if (InSourceNode->GetNodeId().IsInstruction() &&
+			InTargetNode->OutputPin->LinkedTo.Contains(InSourceNode->InputPin))
+		{
+			InTargetNode->OutputPin->BreakLinkTo(InSourceNode->InputPin);
+		}
 	};
 
+	HierarchySpaceToControlHashes.Reset();
+	HierarchyControlToSpace.Reset();
+	HierarchySpaceToControl.Reset();
+
 	if (bFollowParentRelationShips)
 	{
 		if (const URigHierarchy* Hierarchy = GetRigHierarchy())
 		{
+			if (bFollowControlSpaceRelationships)
+			{
+				ComputeControlToSpaceMap(Hierarchy, &HierarchySpaceToControlHashes, &HierarchyControlToSpace, &HierarchySpaceToControl);
+			}
+			
 			for (TObjectPtr<URigDependencyGraphNode>& ParentElementNode : DependencyGraphNodes)
 			{
 				if (!ParentElementNode->NodeId.IsElement())
@@ -314,7 +324,15 @@ void URigDependencyGraph::RebuildLinks()
 				}
 				if (Hierarchy->IsValidIndex(ParentElementNode->NodeId.Index))
 				{
-					const TArray<int32> Children = Hierarchy->GetChildren(ParentElementNode->NodeId.Index);
+					TArray<int32> Children = Hierarchy->GetChildren(ParentElementNode->NodeId.Index);
+					if (bFollowControlSpaceRelationships)
+					{
+						if (const TArray<int32>* ControlsWithThisSpace = HierarchySpaceToControl.Find(ParentElementNode->NodeId.Index))
+						{
+							Children.Append(*ControlsWithThisSpace);
+						}
+					}
+					
 					for (const int32& ChildIndex : Children)
 					{
 						const FNodeId ChildId(FNodeId::EType_RigElement, ChildIndex, NAME_None);
@@ -334,7 +352,15 @@ void URigDependencyGraph::RebuildLinks()
 				}
 				if (Hierarchy->IsValidIndex(ChildElementNode->NodeId.Index))
 				{
-					const TArray<int32> Parents = Hierarchy->GetParents(ChildElementNode->NodeId.Index);
+					TArray<int32> Parents = Hierarchy->GetParents(ChildElementNode->NodeId.Index);
+					if (bFollowControlSpaceRelationships)
+					{
+						if (const TArray<int32>* Spaces = HierarchyControlToSpace.Find(ChildElementNode->NodeId.Index))
+						{
+							Parents.Append(*Spaces);
+						}
+					}
+
 					for (const int32& ParentIndex : Parents)
 					{
 						const FNodeId ParentId(FNodeId::EType_RigElement, ParentIndex, NAME_None);
@@ -402,7 +428,7 @@ void URigDependencyGraph::RebuildIslands()
 void URigDependencyGraph::UpdatePinVisibility()
 {
 	const URigHierarchy* Hierarchy = GetRigHierarchy();
-	
+
 	for (const TObjectPtr<URigDependencyGraphNode>& Node : DependencyGraphNodes)
 	{
 		const FNodeId& NodeId = Node->GetNodeId();
@@ -420,6 +446,13 @@ void URigDependencyGraph::UpdatePinVisibility()
 				{
 					Node->InputPin->bHidden = false;
 				}
+				if (bFollowControlSpaceRelationships)
+				{
+					if (HierarchyControlToSpace.Contains(NodeId.Index))
+					{
+						Node->InputPin->bHidden = false;
+					}
+				}
 			}
 
 			if (Node->InputPin->bHidden && RigDependencyProvider.IsValid() && bFollowVMRelationShips)
@@ -443,6 +476,13 @@ void URigDependencyGraph::UpdatePinVisibility()
 					{
 						Node->OutputPin->bHidden = false;
 					}
+					if (bFollowControlSpaceRelationships)
+					{
+						if (HierarchySpaceToControl.Contains(NodeId.Index))
+						{
+							Node->OutputPin->bHidden = false;
+						}
+					}
 				}
 			}
 
@@ -529,6 +569,12 @@ void URigDependencyGraph::SetFollowParentRelationShips(bool InFollowParentRelati
 	RebuildGraph(false);
 }
 
+void URigDependencyGraph::SetFollowControlSpaceRelationships(bool InFollowControlSpaceRelationships)
+{
+	bFollowControlSpaceRelationships = InFollowControlSpaceRelationships;
+	RebuildGraph(false);
+}
+
 void URigDependencyGraph::SetFollowVMRelationShips(bool InFollowVMRelationShips)
 {
 	bFollowVMRelationShips = InFollowVMRelationShips;
@@ -591,6 +637,12 @@ TArray<URigDependencyGraph::FNodeId> URigDependencyGraph::GetNodesToExpand(const
 	{
 		if (const URigHierarchy* Hierarchy = GetRigHierarchy())
 		{
+			TMap<int32, TArray<int32>> ControlToSpace, SpaceToControl;
+			if (bFollowControlSpaceRelationships)
+			{
+				ComputeControlToSpaceMap(Hierarchy, nullptr, &ControlToSpace, &SpaceToControl);
+			}
+
 			// also expand the nodes to construct parent and child nodes
 			for (int32 NodeIndex = 0; NodeIndex < InNodeIds.Num(); ++NodeIndex)
 			{
@@ -605,10 +657,24 @@ TArray<URigDependencyGraph::FNodeId> URigDependencyGraph::GetNodesToExpand(const
 					if (bExpandInputs)
 					{
 						ParentsAndChildren.Append(Hierarchy->GetParents(InNodeIds[NodeIndex].Index));
+						if (bFollowControlSpaceRelationships)
+						{
+							if (const TArray<int32>* Spaces = ControlToSpace.Find(InNodeIds[NodeIndex].Index))
+							{
+								ParentsAndChildren.Append(*Spaces);
+							}
+						}
 					}
 					if (bExpandOutputs)
 					{
 						ParentsAndChildren.Append(Hierarchy->GetChildren(InNodeIds[NodeIndex].Index));
+						if (bFollowControlSpaceRelationships)
+						{
+							if (const TArray<int32>* ControlsUsingThisSpace = SpaceToControl.Find(InNodeIds[NodeIndex].Index))
+							{
+								ParentsAndChildren.Append(*ControlsUsingThisSpace);
+							}
+						}
 					}
 
 					for (const int32& ParentOrChildIndex : ParentsAndChildren)
@@ -1289,6 +1355,51 @@ void URigDependencyGraph::OnRigVMGraphModified(ERigVMGraphNotifType InNotificati
 	}
 }
 
+void URigDependencyGraph::ComputeControlToSpaceMap(const URigHierarchy* InHierarchy, TSet<uint32>* OutHashes, TMap<int32, TArray<int32>>* OutControlToSpace, TMap<int32, TArray<int32>>* OutSpaceToControl)
+{
+	if (OutHashes)
+	{
+		OutHashes->Reset();
+	}
+	if (OutControlToSpace)
+	{
+		OutControlToSpace->Reset();
+	}
+	if (OutSpaceToControl)
+	{
+		OutSpaceToControl->Reset();
+	}
+	if (!InHierarchy || (!OutHashes && !OutControlToSpace && !OutSpaceToControl))
+	{
+		return;
+	}
+
+	TArray<FRigControlElement*> Controls = InHierarchy->GetControls();
+	for (const FRigControlElement* Control : Controls)
+	{
+		for (const FRigElementKeyWithLabel& AvailableSpace : Control->Settings.Customization.AvailableSpaces)
+		{
+			const int32 SpaceIndex = InHierarchy->GetIndex(AvailableSpace.Key);
+			if (SpaceIndex == INDEX_NONE)
+			{
+				continue;
+			}
+			if (OutHashes)
+			{
+				OutHashes->Add(HashCombine(GetTypeHash(SpaceIndex), GetTypeHash(Control->GetIndex())));
+			}
+			if (OutControlToSpace)
+			{
+				OutControlToSpace->FindOrAdd(Control->GetIndex()).Add(SpaceIndex);
+			}
+			if (OutSpaceToControl)
+			{
+				OutSpaceToControl->FindOrAdd(SpaceIndex).Add(Control->GetIndex());
+			}
+		}
+	}
+}
+
 void URigDependencyGraph::SelectNodes(const TArray<FNodeId>& InNodeIds)
 {
 	BlockSelectionCounter = 2;
diff --git a/Engine/Plugins/Animation/ControlRig/Source/ControlRigEditor/Private/RigDependencyGraph/RigDependencyGraph.h b/Engine/Plugins/Animation/ControlRig/Source/ControlRigEditor/Private/RigDependencyGraph/RigDependencyGraph.h
--- a/Engine/Plugins/Animation/ControlRig/Source/ControlRigEditor/Private/RigDependencyGraph/RigDependencyGraph.h
+++ b/Engine/Plugins/Animation/ControlRig/Source/ControlRigEditor/Private/RigDependencyGraph/RigDependencyGraph.h
@@ -85,6 +85,7 @@ public:
 	bool IsContentLocked() const { return bLockContent; } 
 
 	void SetFollowParentRelationShips(bool InFollowParentRelationShips);
+	void SetFollowControlSpaceRelationships(bool InFollowControlSpaceRelationships);
 	void SetFollowVMRelationShips(bool InFollowVMRelationShips);
 	void SetContentLocked(bool InLockContent);
 
@@ -148,6 +149,8 @@ private:
 	FDelegateHandle OnRigVMGraphModifiedHandle;
 	
 private:
+
+	static void ComputeControlToSpaceMap(const URigHierarchy* InHierarchy, TSet<uint32>* OutHashes, TMap<int32, TArray<int32>>* OutControlToSpace, TMap<int32, TArray<int32>>* OutSpaceToControl);
 	
 	TArray<FNodeId> SelectedNodes;
 
@@ -171,10 +174,15 @@ private:
 	int32 LayoutIterationsTotal;
 	int32 LayoutIterationsMax;
 	bool bFollowParentRelationShips;
+	bool bFollowControlSpaceRelationships;
 	bool bFollowVMRelationShips;
 	bool bLockContent;
 	bool bIsPerformingGridLayout;
 
+	TSet<uint32> HierarchySpaceToControlHashes;
+	TMap<int32, TArray<int32>> HierarchyControlToSpace;
+	TMap<int32, TArray<int32>> HierarchySpaceToControl;
+
 	/** As selection broadcasting is deferred in the graph, we need to block re-broadcasting a for a few frames on refresh */
 	int32 BlockSelectionCounter;
 
@@ -260,5 +268,6 @@ private:
 
 	friend class SRigDependencyGraph;
 	friend class URigDependencyGraphSchema;
+	friend class FRigDependencyConnectionDrawingPolicy;
 };
 
diff --git a/Engine/Plugins/Animation/ControlRig/Source/ControlRigEditor/Private/RigDependencyGraph/RigDependencyGraphCommands.cpp b/Engine/Plugins/Animation/ControlRig/Source/ControlRigEditor/Private/RigDependencyGraph/RigDependencyGraphCommands.cpp
--- a/Engine/Plugins/Animation/ControlRig/Source/ControlRigEditor/Private/RigDependencyGraph/RigDependencyGraphCommands.cpp
+++ b/Engine/Plugins/Animation/ControlRig/Source/ControlRigEditor/Private/RigDependencyGraph/RigDependencyGraphCommands.cpp
@@ -9,6 +9,7 @@ void FRigDependencyGraphCommands::RegisterCommands()
 	UI_COMMAND(FrameSelection, "Frame Selection", "Expands and frames the selection in the tree", EUserInterfaceActionType::Button, FInputChord(EKeys::F));
 	UI_COMMAND(JumpToSource, "Jump to Source", "Jumps to the instruction node backing up the dependency node", EUserInterfaceActionType::Button, FInputChord(EModifierKey::Control, EKeys::J));
 	UI_COMMAND(ShowParentChildRelationships, "Show Parent/Child", "Whether to show or hide parent and child relationships", EUserInterfaceActionType::ToggleButton, FInputChord(EKeys::P));
+	UI_COMMAND(ShowControlSpaceRelationships, "Show Control Spaces", "Whether to show or hide space relationships for controls", EUserInterfaceActionType::ToggleButton, FInputChord(EModifierKey::Shift, EKeys::P));
 	UI_COMMAND(ShowInstructionRelationships, "Show VM Nodes", "Whether to show or hide VM based relationships", EUserInterfaceActionType::ToggleButton, FInputChord(EKeys::O));
 	UI_COMMAND(LockContent, "Lock Content", "Whether or not to sync the content with selection in other views", EUserInterfaceActionType::ToggleButton, FInputChord(EKeys::L));
 	UI_COMMAND(EnableFlashlight, "Enable Flashlight", "Enable a flashlight around the mouse cursor to brighten up faded out nodes", EUserInterfaceActionType::ToggleButton, FInputChord());
diff --git a/Engine/Plugins/Animation/ControlRig/Source/ControlRigEditor/Private/RigDependencyGraph/RigDependencyGraphCommands.h b/Engine/Plugins/Animation/ControlRig/Source/ControlRigEditor/Private/RigDependencyGraph/RigDependencyGraphCommands.h
--- a/Engine/Plugins/Animation/ControlRig/Source/ControlRigEditor/Private/RigDependencyGraph/RigDependencyGraphCommands.h
+++ b/Engine/Plugins/Animation/ControlRig/Source/ControlRigEditor/Private/RigDependencyGraph/RigDependencyGraphCommands.h
@@ -27,6 +27,9 @@ public:
 	/* Show parent-child relationships */
 	TSharedPtr< FUICommandInfo > ShowParentChildRelationships;
 
+	/* Show control-space relationships */
+	TSharedPtr< FUICommandInfo > ShowControlSpaceRelationships;
+
 	/* Show instruction based relationships */
 	TSharedPtr< FUICommandInfo > ShowInstructionRelationships;
 
diff --git a/Engine/Plugins/Animation/ControlRig/Source/ControlRigEditor/Private/RigDependencyGraph/SRigDependencyGraph.cpp b/Engine/Plugins/Animation/ControlRig/Source/ControlRigEditor/Private/RigDependencyGraph/SRigDependencyGraph.cpp
--- a/Engine/Plugins/Animation/ControlRig/Source/ControlRigEditor/Private/RigDependencyGraph/SRigDependencyGraph.cpp
+++ b/Engine/Plugins/Animation/ControlRig/Source/ControlRigEditor/Private/RigDependencyGraph/SRigDependencyGraph.cpp
@@ -549,6 +549,7 @@ TSharedRef<SWidget> SRigDependencyGraph::CreateTopLeftMenu()
 		MenuBuilder.AddMenuEntry(Actions.LockContent);
 		MenuBuilder.AddMenuEntry(Actions.EnableFlashlight);
 		MenuBuilder.AddMenuEntry(Actions.ShowParentChildRelationships);
+		MenuBuilder.AddMenuEntry(Actions.ShowControlSpaceRelationships);
 		MenuBuilder.AddMenuEntry(Actions.ShowInstructionRelationships);
 	}
 	MenuBuilder.EndSection();
@@ -728,6 +729,20 @@ void SRigDependencyGraph::BindCommands()
 		})
 	);
 
+	CommandList->MapAction(
+		Commands.ShowControlSpaceRelationships,
+		FExecuteAction::CreateLambda([this]()
+		{
+			GraphObj->bFollowControlSpaceRelationships = !GraphObj->bFollowControlSpaceRelationships;
+			GraphObj->RebuildLinks();
+		}),
+		FCanExecuteAction(),
+		FIsActionChecked::CreateLambda([this]()
+		{
+			return GraphObj->bFollowControlSpaceRelationships;
+		})
+	);
+
 	CommandList->MapAction(
 		Commands.ShowInstructionRelationships,
 		FExecuteAction::CreateLambda([this]()
