Control Rig: Fix not being able to reopen the control rig editor of an asset after closing it
#rnx

diff --git a/Engine/Plugins/Animation/ControlRig/Source/ControlRigEditor/Private/Editor/ControlRigEditor.cpp b/Engine/Plugins/Animation/ControlRig/Source/ControlRigEditor/Private/Editor/ControlRigEditor.cpp
--- a/Engine/Plugins/Animation/ControlRig/Source/ControlRigEditor/Private/Editor/ControlRigEditor.cpp
+++ b/Engine/Plugins/Animation/ControlRig/Source/ControlRigEditor/Private/Editor/ControlRigEditor.cpp
@@ -4469,6 +4469,23 @@ void FControlRigBaseEditor::BindCommandsImpl()
 		FCanExecuteAction::CreateSP(SharedEditor, &FControlRigBaseEditor::IsRigModule));
 }
 
+void FControlRigBaseEditor::UnbindCommandsImpl()
+{
+	UnbindCommandsSuper();
+	GetToolkitCommands()->UnmapAction(FControlRigEditorCommands::Get().ConstructionEvent);
+	GetToolkitCommands()->UnmapAction(FControlRigEditorCommands::Get().ForwardsSolveEvent);
+	GetToolkitCommands()->UnmapAction(FControlRigEditorCommands::Get().BackwardsSolveEvent);
+	GetToolkitCommands()->UnmapAction(FControlRigEditorCommands::Get().BackwardsAndForwardsSolveEvent);
+	GetToolkitCommands()->UnmapAction(FControlRigEditorCommands::Get().ToggleControlVisibility);
+	GetToolkitCommands()->UnmapAction(FControlRigEditorCommands::Get().ToggleControlsAsOverlay);
+	GetToolkitCommands()->UnmapAction(FControlRigEditorCommands::Get().ToggleDrawNulls);
+	GetToolkitCommands()->UnmapAction(FControlRigEditorCommands::Get().ToggleDrawSockets);
+	GetToolkitCommands()->UnmapAction(FControlRigEditorCommands::Get().ToggleDrawAxesOnSelection);
+	GetToolkitCommands()->UnmapAction(FControlRigEditorCommands::Get().ToggleSchematicViewportVisibility);
+	GetToolkitCommands()->UnmapAction(FControlRigEditorCommands::Get().SwapModuleWithinAsset);
+	GetToolkitCommands()->UnmapAction(FControlRigEditorCommands::Get().SwapModuleAcrossProject);
+}
+
 FMenuBuilder FControlRigBaseEditor::GenerateBulkEditMenuImpl()
 {
 	FMenuBuilder MenuBuilder = GenerateBulkEditMenuSuper();
diff --git a/Engine/Plugins/Animation/ControlRig/Source/ControlRigEditor/Private/Editor/ControlRigEditor.h b/Engine/Plugins/Animation/ControlRig/Source/ControlRigEditor/Private/Editor/ControlRigEditor.h
--- a/Engine/Plugins/Animation/ControlRig/Source/ControlRigEditor/Private/Editor/ControlRigEditor.h
+++ b/Engine/Plugins/Animation/ControlRig/Source/ControlRigEditor/Private/Editor/ControlRigEditor.h
@@ -172,6 +172,7 @@ protected:
 	virtual void OnWrappedPropertyChangedChainEventSuper(URigVMDetailsViewWrapperObject* InWrapperObject, const FString& InPropertyPath, FPropertyChangedChainEvent& InPropertyChangedChainEvent) = 0;
 	virtual const TArray<TStrongObjectPtr<URigVMDetailsViewWrapperObject>>& GetWrapperObjects() const = 0;
 	virtual void BindCommandsSuper() = 0;
+	virtual void UnbindCommandsSuper() = 0;
 	virtual FMenuBuilder GenerateBulkEditMenuSuper() = 0;
 	virtual TWeakPtr<class SGraphEditor> GetFocusedGraphEd() const = 0;
 	virtual void OnGraphNodeDropToPerformSuper(TSharedPtr<FDragDropOperation> InDragDropOp, UEdGraph* InGraph, const FVector2f& InNodePosition, const FVector2f& InScreenPosition) = 0;
@@ -335,6 +336,7 @@ protected:
 
 
 	virtual void BindCommandsImpl();
+	virtual void UnbindCommandsImpl();
 	virtual FMenuBuilder GenerateBulkEditMenuImpl();
 
 	virtual void OnHierarchyChanged() override;
diff --git a/Engine/Plugins/Animation/ControlRig/Source/ControlRigEditor/Private/Editor/ControlRigLegacyEditor.h b/Engine/Plugins/Animation/ControlRig/Source/ControlRigEditor/Private/Editor/ControlRigLegacyEditor.h
--- a/Engine/Plugins/Animation/ControlRig/Source/ControlRigEditor/Private/Editor/ControlRigLegacyEditor.h
+++ b/Engine/Plugins/Animation/ControlRig/Source/ControlRigEditor/Private/Editor/ControlRigLegacyEditor.h
@@ -118,6 +118,8 @@ protected:
 
 	virtual void BindCommands() override { return FControlRigBaseEditor::BindCommandsImpl(); }
 	virtual void BindCommandsSuper() override { return IControlRigLegacyEditor::BindCommands(); }
+	virtual void UnbindCommands() override { return FControlRigBaseEditor::UnbindCommandsImpl(); }
+	virtual void UnbindCommandsSuper() override { return IControlRigLegacyEditor::UnbindCommands(); }
 	virtual FMenuBuilder GenerateBulkEditMenu() override { return FControlRigBaseEditor::GenerateBulkEditMenuImpl(); }
 	virtual FMenuBuilder GenerateBulkEditMenuSuper() override { return IControlRigLegacyEditor::GenerateBulkEditMenu(); }
 
diff --git a/Engine/Plugins/Animation/ControlRig/Source/ControlRigEditor/Private/Editor/ControlRigNewEditor.h b/Engine/Plugins/Animation/ControlRig/Source/ControlRigEditor/Private/Editor/ControlRigNewEditor.h
--- a/Engine/Plugins/Animation/ControlRig/Source/ControlRigEditor/Private/Editor/ControlRigNewEditor.h
+++ b/Engine/Plugins/Animation/ControlRig/Source/ControlRigEditor/Private/Editor/ControlRigNewEditor.h
@@ -116,6 +116,8 @@ protected:
 
 	virtual void BindCommands() override { return FControlRigBaseEditor::BindCommandsImpl(); }
 	virtual void BindCommandsSuper() override { return IControlRigNewEditor::BindCommands(); }
+	virtual void UnbindCommands() override { return FControlRigBaseEditor::UnbindCommandsImpl(); }
+	virtual void UnbindCommandsSuper() override { return IControlRigNewEditor::UnbindCommands(); }
 	virtual FMenuBuilder GenerateBulkEditMenu() override { return FControlRigBaseEditor::GenerateBulkEditMenuImpl(); }
 	virtual FMenuBuilder GenerateBulkEditMenuSuper() override { return IControlRigNewEditor::GenerateBulkEditMenu(); }
 
diff --git a/Engine/Plugins/Runtime/RigVM/Source/RigVMEditor/Private/Editor/RigVMEditor.cpp b/Engine/Plugins/Runtime/RigVM/Source/RigVMEditor/Private/Editor/RigVMEditor.cpp
--- a/Engine/Plugins/Runtime/RigVM/Source/RigVMEditor/Private/Editor/RigVMEditor.cpp
+++ b/Engine/Plugins/Runtime/RigVM/Source/RigVMEditor/Private/Editor/RigVMEditor.cpp
@@ -86,6 +86,8 @@ void FRigVMEditorBase::UnbindEditor()
 	GEditor->GetEditorSubsystem<UAssetEditorSubsystem>()->OnAssetEditorRequestedOpen().RemoveAll(this);
 	GEditor->GetEditorSubsystem<UAssetEditorSubsystem>()->OnAssetEditorRequestClose().RemoveAll(this);
 
+	UnbindCommands();
+
 	if (RigVMBlueprint)
 	{
 		// clear editor related data from the debugged control rig instance 
@@ -1477,6 +1479,25 @@ void FRigVMEditorBase::BindCommands()
 		}));
 }
 
+void FRigVMEditorBase::UnbindCommands()
+{
+	DECLARE_SCOPE_HIERARCHICAL_COUNTER_FUNC()
+
+	TSharedRef<FRigVMEditorBase> ThisRef = StaticCastSharedRef<FRigVMEditorBase>(SharedRef());
+	TSharedRef<FUICommandList> ToolkitCommands = GetToolkitCommands();
+	FRigVMEditorCommands& EditorCommands = FRigVMEditorCommands::Get();
+
+	ToolkitCommands->UnmapAction(EditorCommands.AutoCompileGraph);
+	ToolkitCommands->UnmapAction(EditorCommands.ToggleEventQueue);
+	ToolkitCommands->UnmapAction(EditorCommands.FrameSelection);
+	ToolkitCommands->UnmapAction(EditorCommands.SwapFunctionWithinAsset);
+	ToolkitCommands->UnmapAction(EditorCommands.SwapFunctionAcrossProject);
+	ToolkitCommands->UnmapAction(EditorCommands.SwapAssetReferences);
+	ToolkitCommands->UnmapAction(EditorCommands.ToggleProfiling);
+	ToolkitCommands->UnmapAction(EditorCommands.TogglePreviewHere);
+	ToolkitCommands->UnmapAction(EditorCommands.PreviewHereStepForward);
+}
+
 void FRigVMEditorBase::ToggleAutoCompileGraph()
 {
 	if (FRigVMAssetInterfacePtr RigVMBlueprint = GetRigVMAssetInterface())
diff --git a/Engine/Plugins/Runtime/RigVM/Source/RigVMEditor/Public/Editor/RigVMEditor.h b/Engine/Plugins/Runtime/RigVM/Source/RigVMEditor/Public/Editor/RigVMEditor.h
--- a/Engine/Plugins/Runtime/RigVM/Source/RigVMEditor/Public/Editor/RigVMEditor.h
+++ b/Engine/Plugins/Runtime/RigVM/Source/RigVMEditor/Public/Editor/RigVMEditor.h
@@ -456,6 +456,7 @@ protected:
 	UE_API void AddReferencedObjects( FReferenceCollector& Collector );
 
 	UE_API virtual void BindCommands();
+	UE_API virtual void UnbindCommands();
 
 	UE_API void ToggleAutoCompileGraph();
 	UE_API bool IsAutoCompileGraphOn() const;
