Control Rig: Fix crash when closing CR editor and sequencer
#rnx

diff --git a/Engine/Plugins/Animation/ControlRig/Source/ControlRig/Private/ControlRig.cpp b/Engine/Plugins/Animation/ControlRig/Source/ControlRig/Private/ControlRig.cpp
--- a/Engine/Plugins/Animation/ControlRig/Source/ControlRig/Private/ControlRig.cpp
+++ b/Engine/Plugins/Animation/ControlRig/Source/ControlRig/Private/ControlRig.cpp
@@ -858,6 +858,11 @@ bool UControlRig::AllConnectorsAreResolved(FString* OutFailureReason, FRigElemen
 	return true;
 }
 
+bool UControlRig::CanExecute() const
+{
+	return Super::CanExecute() && GetHierarchy() != nullptr;
+}
+
 bool UControlRig::Execute(const FName& InEventName)
 {
 	if(!CanExecute())
diff --git a/Engine/Plugins/Animation/ControlRig/Source/ControlRig/Public/ControlRig.h b/Engine/Plugins/Animation/ControlRig/Source/ControlRig/Public/ControlRig.h
--- a/Engine/Plugins/Animation/ControlRig/Source/ControlRig/Public/ControlRig.h
+++ b/Engine/Plugins/Animation/ControlRig/Source/ControlRig/Public/ControlRig.h
@@ -271,6 +271,7 @@ public:
 	UPROPERTY()
 	FRigHierarchySettings HierarchySettings;
 
+	CONTROLRIG_API virtual bool CanExecute() const override;
 	CONTROLRIG_API virtual bool Execute(const FName& InEventName) override;
 	CONTROLRIG_API virtual bool Execute_Internal(const FName& InEventName) override;
 	CONTROLRIG_API virtual void RequestInit() override;
diff --git a/Engine/Plugins/Animation/ControlRig/Source/ControlRigEditor/Private/EditMode/ControlRigEditMode.cpp b/Engine/Plugins/Animation/ControlRig/Source/ControlRigEditor/Private/EditMode/ControlRigEditMode.cpp
--- a/Engine/Plugins/Animation/ControlRig/Source/ControlRigEditor/Private/EditMode/ControlRigEditMode.cpp
+++ b/Engine/Plugins/Animation/ControlRig/Source/ControlRigEditor/Private/EditMode/ControlRigEditMode.cpp
@@ -5901,7 +5901,10 @@ void FControlRigEditMode::RemoveControlRig(UControlRig* InControlRig)
 	if (!URigVMHost::IsGarbageOrDestroyed(InControlRig))
 	{
 		InControlRig->ControlModified().RemoveAll(this);
-		InControlRig->GetHierarchy()->OnModified().RemoveAll(this);
+		if (URigHierarchy* Hierarchy = InControlRig->GetHierarchy())
+		{
+			Hierarchy->OnModified().RemoveAll(this);
+		}
 		InControlRig->OnPreConstructionForUI_AnyThread().RemoveAll(this);
 		InControlRig->OnPostConstruction_AnyThread().RemoveAll(this);
 	}
