Control Rig: Hierarchy to treat zero weight parent relationships as no-op

Also stop using the vm based dependency provider for space switching

#rnx

diff --git a/Engine/Plugins/Animation/ControlRig/Source/ControlRig/Private/Constraints/ControlRigTransformableHandle.cpp b/Engine/Plugins/Animation/ControlRig/Source/ControlRig/Private/Constraints/ControlRigTransformableHandle.cpp
--- a/Engine/Plugins/Animation/ControlRig/Source/ControlRig/Private/Constraints/ControlRigTransformableHandle.cpp
+++ b/Engine/Plugins/Animation/ControlRig/Source/ControlRig/Private/Constraints/ControlRigTransformableHandle.cpp
@@ -411,12 +411,7 @@ bool UTransformableControlHandle::HasDirectDependencyWith(const UTransformableHa
 		{
 			if (InControlRig && InRigHierarchy)
 			{
-#if WITH_EDITOR
-				const FRigDependenciesProviderForControlRig SameModuleDependencyProvider(InControlRig);
-				return InRigHierarchy->IsParentedTo(ControlElement, PossibleParent, SameModuleDependencyProvider);
-#else
 				return InRigHierarchy->IsParentedTo(ControlElement, PossibleParent);
-#endif			
 			}
 			return false;
 		};
diff --git a/Engine/Plugins/Animation/ControlRig/Source/ControlRig/Private/ControlRig.cpp b/Engine/Plugins/Animation/ControlRig/Source/ControlRig/Private/ControlRig.cpp
--- a/Engine/Plugins/Animation/ControlRig/Source/ControlRig/Private/ControlRig.cpp
+++ b/Engine/Plugins/Animation/ControlRig/Source/ControlRig/Private/ControlRig.cpp
@@ -573,12 +573,7 @@ void UControlRig::Evaluate_AnyThread()
 			{
 				if (FRigBaseElement* NewParent = Hierarchy->Find(Pair.Value.NewParent))
 				{
-#if WITH_EDITOR
-					FRigDependenciesProviderForControlRig DependencyProvider(this);
-					Hierarchy->SwitchToParent(Element, NewParent, Pair.Value.bInitial, Pair.Value.bAffectChildren, DependencyProvider, nullptr);
-#else
 					Hierarchy->SwitchToParent(Element, NewParent, Pair.Value.bInitial, Pair.Value.bAffectChildren);
-#endif
 				}
 			}
 		}
@@ -2362,14 +2357,10 @@ void UControlRig::SwitchToParent(const FRigElementKey& InElementKey, const FRigE
 	}
 	else
 	{
-#if WITH_EDITOR
 		// when simply executing space switching we need to make sure it runs through,
 		// and during that we'll only respect parent child relationships (while ignoring VM dependencies).
 		static const FEmptyRigDependenciesProvider EmptyRigDependenciesProvider;
 		DynamicHierarchy->SwitchToParent(InElementKey, InNewParentKey, bInitial, bAffectChildren, EmptyRigDependenciesProvider, nullptr);
-#else
-		DynamicHierarchy->SwitchToParent(InElementKey, InNewParentKey, bInitial, bAffectChildren);
-#endif
 	}
 }
 
diff --git a/Engine/Plugins/Animation/ControlRig/Source/ControlRig/Private/Rigs/RigHierarchy.cpp b/Engine/Plugins/Animation/ControlRig/Source/ControlRig/Private/Rigs/RigHierarchy.cpp
--- a/Engine/Plugins/Animation/ControlRig/Source/ControlRig/Private/Rigs/RigHierarchy.cpp
+++ b/Engine/Plugins/Animation/ControlRig/Source/ControlRig/Private/Rigs/RigHierarchy.cpp
@@ -110,6 +110,7 @@ const FLazyName URigHierarchy::WorldSpaceKeyLabel = TEXT("World");
 
 URigHierarchy::URigHierarchy()
 : TopologyVersion(0)
+, ParentWeightVersion(0)
 , MetadataVersion(0)
 , MetadataTagVersion(0)
 , bEnableDirtyPropagation(true)
@@ -797,6 +798,7 @@ void URigHierarchy::Reset()
 void URigHierarchy::Reset_Impl(bool bResetElements)
 {
 	TopologyVersion = 0;
+	ParentWeightVersion = 0;
 	MetadataVersion = 0;
 	bEnableDirtyPropagation = true;
 	NonUniqueShortNamesCache.Reset();
@@ -1191,6 +1193,8 @@ void URigHierarchy::CopyHierarchy(URigHierarchy* InHierarchy)
 		// Increment the topology version to invalidate our cached children.
 		IncrementTopologyVersion();
 
+		ParentWeightVersion += InHierarchy->GetParentWeightVersion();
+
 		// Keep incrementing the metadata version so that the UI can refresh.
 		MetadataVersion += InHierarchy->GetMetadataVersion();
 		MetadataTagVersion += InHierarchy->GetMetadataTagVersion();
@@ -3641,6 +3645,18 @@ bool URigHierarchy::SetParentWeight(FRigBaseElement* InChild, int32 InParentInde
 			{
 				return false;
 			}
+
+			// if we are setting a weight affecting the hierarchy
+			// we need to guard against cycles
+			if (!bInitial && !InWeight.IsAlmostZero())
+			{
+				if (!CanSwitchToParent(
+					InChild->GetKey(),
+					MultiParentElement->ParentConstraints[InParentIndex].ParentElement->GetKey()))
+				{
+					return false;
+				}
+			}
 			
 			const ERigTransformType::Type LocalType = bInitial ? InitialLocal : CurrentLocal;
 			const ERigTransformType::Type GlobalType = SwapLocalAndGlobal(LocalType);
@@ -3675,6 +3691,7 @@ bool URigHierarchy::SetParentWeight(FRigBaseElement* InChild, int32 InParentInde
 
 			PropagateDirtyFlags(MultiParentElement, ERigTransformType::IsInitial(LocalType), bAffectChildren);
 			EnsureCacheValidity();
+			IncrementParentWeightVersion();
 			
 #if WITH_EDITOR
 			if (!bPropagatingChange)
@@ -3788,8 +3805,43 @@ bool URigHierarchy::SetParentWeightArray(FRigBaseElement* InChild,  const TArray
 				MultiParentElement->GetDirtyState().MarkDirty(LocalType);
 			}
 
+			// first set all the zero weights
 			for(int32 WeightIndex=0; WeightIndex < InWeights.Num(); WeightIndex++)
 			{
+				if (!InWeights[WeightIndex].IsAlmostZero())
+				{
+					continue;
+				}
+				if(bInitial)
+				{
+					MultiParentElement->ParentConstraints[WeightIndex].InitialWeight = InputWeights[WeightIndex];
+				}
+				else
+				{
+					MultiParentElement->ParentConstraints[WeightIndex].Weight = InputWeights[WeightIndex];
+				}
+			}
+
+			IncrementParentWeightVersion();
+
+			// now set all the non-zero weights
+			for(int32 WeightIndex=0; WeightIndex < InWeights.Num(); WeightIndex++)
+			{
+				if (InWeights[WeightIndex].IsAlmostZero())
+				{
+					continue;
+				}
+
+				if (!bInitial)
+				{
+					if (!CanSwitchToParent(
+						InChild->GetKey(),
+						MultiParentElement->ParentConstraints[WeightIndex].ParentElement->GetKey()))
+					{
+						return false;
+					}
+				}
+
 				if(bInitial)
 				{
 					MultiParentElement->ParentConstraints[WeightIndex].InitialWeight = InputWeights[WeightIndex];
@@ -3808,6 +3860,7 @@ bool URigHierarchy::SetParentWeightArray(FRigBaseElement* InChild,  const TArray
 
 			PropagateDirtyFlags(MultiParentElement, ERigTransformType::IsInitial(LocalType), bAffectChildren);
 			EnsureCacheValidity();
+			IncrementParentWeightVersion();
 			
 #if WITH_EDITOR
 			if (!bPropagatingChange)
@@ -3895,7 +3948,7 @@ bool URigHierarchy::CanSwitchToParent(FRigElementKey InChild, FRigElementKey InP
 		}
 	}
 
-	TArray<FRigHierarchyRecord> DependencyChain;
+	TArray<TPair<FRigHierarchyRecord,FRigHierarchyRecord>> DependencyChain;
 	if(IsParentedTo(Parent, Child, InDependencyProvider, &DependencyChain))
 	{
 #if WITH_EDITOR
@@ -3908,110 +3961,118 @@ bool URigHierarchy::CanSwitchToParent(FRigElementKey InChild, FRigElementKey InP
 			}
 			else
 			{
-				TArray<FString> Labels;
-				Labels.Reserve(DependencyChain.Num() + 1);
-				for (int32 DependencyIndex = DependencyChain.Num() - 1; DependencyIndex >= 0; --DependencyIndex)
+				TArray<FString> Messages;
+
+				for (const TPair<FRigHierarchyRecord,FRigHierarchyRecord>& DependencyPair : DependencyChain)
 				{
-					const FRigHierarchyRecord& Dependency = DependencyChain[DependencyIndex];
-					switch (Dependency.Type)
+					auto GetLabelForRecord = [this](const FRigHierarchyRecord& InRecord) -> FString
 					{
-						case FRigHierarchyRecord::EType_RigElement:
-						case FRigHierarchyRecord::EType_Metadata:
+						FString Label;
+							
+						switch (InRecord.Type)
 						{
-							Labels.Add(GetKey(Dependency.Index).ToString());
-							if (Dependency.IsMetadata())
+							case FRigHierarchyRecord::EType_RigElement:
+							case FRigHierarchyRecord::EType_Metadata:
 							{
-								Labels.Last() += FString::Printf(TEXT(" Metadata(%s)"), *Dependency.Name.ToString());
-							}
-							break;
-						}
-						case FRigHierarchyRecord::EType_Variable:
-						case FRigHierarchyRecord::EType_Instruction:
-						{
-							if (Dependency.IsVariable())
-							{
-								Labels.Add(TEXT("Variable"));
-							}
-							else
-							{
-								Labels.Add(TEXT("Node"));
-							}
-
-							UControlRig* ControlRig = Cast<UControlRig>(GetOuter());
-							if (!ControlRig)
-							{
-								break;
-							}
-
-							UModularRig* ModularRig = Cast<UModularRig>(ControlRig);
-
-							const URigVM* VM = nullptr;
-							if (Dependency.Name.IsNone() || !ModularRig)
-							{
-								VM = ControlRig->GetVM();
+								Label = GetKey(InRecord.Index).ToString();
+								if (InRecord.IsMetadata())
+								{
+									Label += FString::Printf(TEXT(" Metadata(%s)"), *InRecord.Name.ToString());
+								}
+								return Label;
 							}
-							else
+							case FRigHierarchyRecord::EType_Variable:
+							case FRigHierarchyRecord::EType_Instruction:
 							{
-								const FRigModuleInstance* RigModule = ModularRig->FindModule(Dependency.Name);
-								if (!RigModule)
+								if (InRecord.IsVariable())
 								{
-									break;
+									Label = TEXT("Variable");
 								}
-								UControlRig* ModuleRig = RigModule->GetRig();
-								if (!ModuleRig)
+								else
+								{
+									Label = TEXT("Node");
+								}
+
+								UControlRig* ControlRig = Cast<UControlRig>(GetOuter());
+								if (!ControlRig)
 								{
 									break;
 								}
-								VM = ModuleRig->GetVM();
-							}
 
-							check(VM);
+								UModularRig* ModularRig = Cast<UModularRig>(ControlRig);
 
-							if (Dependency.IsVariable())
-							{
-								const TArray<FRigVMExternalVariableDef>& ExternalVariableDefs = VM->GetExternalVariableDefs();
-								if (ExternalVariableDefs.IsValidIndex(Dependency.Index))
+								const URigVM* VM = nullptr;
+								if (InRecord.Name.IsNone() || !ModularRig)
 								{
-									Labels.Last() += FString::Printf(TEXT("(%s)"), *ExternalVariableDefs[Dependency.Index].Name.ToString());
+									VM = ControlRig->GetVM();
 								}
-							}
-							else
-							{
-								if (!VM->GetByteCode().GetInstructions().IsValidIndex(Dependency.Index))
+								else
 								{
-									break;
+									const FRigModuleInstance* RigModule = ModularRig->FindModule(InRecord.Name);
+									if (!RigModule)
+									{
+										break;
+									}
+									UControlRig* ModuleRig = RigModule->GetRig();
+									if (!ModuleRig)
+									{
+										break;
+									}
+									VM = ModuleRig->GetVM();
 								}
-	
-								const FRigVMInstruction Instruction = VM->GetByteCode().GetInstructions()[Dependency.Index];
-								if (Instruction.OpCode != ERigVMOpCode::Execute)
+
+								check(VM);
+
+								if (InRecord.IsVariable())
 								{
-									break;
+									const TArray<FRigVMExternalVariableDef>& ExternalVariableDefs = VM->GetExternalVariableDefs();
+									if (ExternalVariableDefs.IsValidIndex(InRecord.Index))
+									{
+										Label += FString::Printf(TEXT("(%s)"), *ExternalVariableDefs[InRecord.Index].Name.ToString());
+									}
 								}
-
-								URigVMNode* Node = Cast<URigVMNode>(VM->GetByteCode().GetSubjectForInstruction(Dependency.Index));
-								if (Node == nullptr)
+								else
 								{
-									break;
+									if (!VM->GetByteCode().GetInstructions().IsValidIndex(InRecord.Index))
+									{
+										break;
+									}
+		
+									const FRigVMInstruction Instruction = VM->GetByteCode().GetInstructions()[InRecord.Index];
+									if (Instruction.OpCode != ERigVMOpCode::Execute)
+									{
+										break;
+									}
+
+									URigVMNode* Node = Cast<URigVMNode>(VM->GetByteCode().GetSubjectForInstruction(InRecord.Index));
+									if (Node == nullptr)
+									{
+										break;
+									}
+
+									Label += FString::Printf(TEXT("(%s)"), *Node->GetPathName());
 								}
 
-								Labels.Last() += FString::Printf(TEXT("(%s)"), *Node->GetNodeTitle());
+								break;
+							}
+							default:
+							{
+								Label = TEXT("Unknown");
+								break;
 							}
-
-							break;
-						}
-						default:
-						{
-							Labels.Add(TEXT("Unknown"));
-							break;
 						}
-					}
-				}
 
-				const FString LastItem = Labels.Last(); 
-				Labels.Insert(LastItem, 0);
+						return Label;
+					};
+
+					const FString TargetLabel = GetLabelForRecord(DependencyPair.Value);
+					const FString SourceLabel = GetLabelForRecord(DependencyPair.Key);
+					Messages.Add(FString::Printf(TEXT("'%s' depends on '%s'."), *TargetLabel, *SourceLabel));
+				}
 
-				const FString JoinedLabels = FString::Join(Labels, TEXT(" -> ")); 
-				OutFailureReason->Appendf(TEXT("Cannot switch '%s' to '%s' - would cause a cycle.\n\n%s"), *InChild.ToString(), *InParent.ToString(), *JoinedLabels);
+				const FString JoinedMessages = FString::Join(Messages, TEXT("\n"));
+				
+				OutFailureReason->Appendf(TEXT("Cannot switch '%s' to '%s' - would cause a cycle.\n\n%s"), *InChild.ToString(), *InParent.ToString(), *JoinedMessages);
 			}
 		}
 #endif
@@ -4587,6 +4648,11 @@ void URigHierarchy::IncrementTopologyVersion()
 	KeyCollectionCache.Reset();
 }
 
+void URigHierarchy::IncrementParentWeightVersion()
+{
+	ParentWeightVersion++;
+}
+
 FRigPose URigHierarchy::GetPose(
 	bool bInitial,
 	ERigElementType InElementType,
@@ -5351,7 +5417,7 @@ void URigHierarchy::SetTransform(FRigTransformElement* InTransformElement, const
 			{			
 				if(FRigTransformElement* ListeningElement = Cast<FRigTransformElement>(ListeningHierarchy->Find(InTransformElement->GetKey())))
 				{
-					// bSetupUndo = false such that all listening hierarchies performs undo at the same time the root hierachy undos
+					// bSetupUndo = false such that all listening hierarchies performs undo at the same time the root hierarchy undos
 					ListeningHierarchy->SetTransform(ListeningElement, InTransform, InTransformType, bAffectChildren, false, bForce);
 				}
 			}
@@ -5584,7 +5650,7 @@ void URigHierarchy::SetControlOffsetTransform(FRigControlElement* InControlEleme
 			{	
 				if(FRigControlElement* ListeningElement = Cast<FRigControlElement>(ListeningHierarchy->Find(InControlElement->GetKey())))
 				{
-					// bSetupUndo = false such that all listening hierarchies performs undo at the same time the root hierachy undos
+					// bSetupUndo = false such that all listening hierarchies performs undo at the same time the root hierarchy undos
 					ListeningHierarchy->SetControlOffsetTransform(ListeningElement, InTransform, InTransformType, bAffectChildren, false, bForce);
 				}
 			}
@@ -5718,7 +5784,7 @@ void URigHierarchy::SetControlShapeTransform(FRigControlElement* InControlElemen
 			{	
 				if(FRigControlElement* ListeningElement = Cast<FRigControlElement>(ListeningHierarchy->Find(InControlElement->GetKey())))
 				{
-					// bSetupUndo = false such that all listening hierarchies performs undo at the same time the root hierachy undos
+					// bSetupUndo = false such that all listening hierarchies performs undo at the same time the root hierarchy undos
 					ListeningHierarchy->SetControlShapeTransform(ListeningElement, InTransform, InTransformType, false, bForce);
 				}
 			}
@@ -5785,7 +5851,7 @@ void URigHierarchy::SetControlSettings(FRigControlElement* InControlElement, FRi
 			{	
 				if(FRigControlElement* ListeningElement = Cast<FRigControlElement>(ListeningHierarchy->Find(InControlElement->GetKey())))
 				{
-					// bSetupUndo = false such that all listening hierarchies performs undo at the same time the root hierachy undos
+					// bSetupUndo = false such that all listening hierarchies performs undo at the same time the root hierarchy undos
 					ListeningHierarchy->SetControlSettings(ListeningElement, InSettings, false, bForce);
 				}
 			}
@@ -6210,7 +6276,7 @@ void URigHierarchy::SetConnectorSettings(FRigConnectorElement* InConnectorElemen
 			{	
 				if(FRigConnectorElement* ListeningElement = Cast<FRigConnectorElement>(ListeningHierarchy->Find(InConnectorElement->GetKey())))
 				{
-					// bSetupUndo = false such that all listening hierarchies performs undo at the same time the root hierachy undos
+					// bSetupUndo = false such that all listening hierarchies performs undo at the same time the root hierarchy undos
 					ListeningHierarchy->SetConnectorSettings(ListeningElement, InSettings, false, bForce);
 				}
 			}
@@ -6402,14 +6468,16 @@ FRigHierarchyKey URigHierarchy::GetPreviousHierarchyParent(const FRigHierarchyKe
 	return FRigHierarchyKey();
 }
 
-bool URigHierarchy::IsParentedTo(FRigBaseElement* InChild, FRigBaseElement* InParent, const IRigDependenciesProvider& InDependencyProvider, TArray<FRigHierarchyRecord>* OutDependencyChain) const
+bool URigHierarchy::IsParentedTo(FRigBaseElement* InChild, FRigBaseElement* InParent, const IRigDependenciesProvider& InDependencyProvider, TArray<TPair<FRigHierarchyRecord,FRigHierarchyRecord>>* OutDependencyChain) const
 {
-	TArray<FRigHierarchyRecord> DependencyChain;
+	TArray<TPair<FRigHierarchyRecord,FRigHierarchyRecord>> DependencyChain;
 
+	bool bActualTopology = false;
 	if (OutDependencyChain == nullptr)
 	{
 		DependencyChain.Reserve(64);
 		OutDependencyChain = &DependencyChain;
+		bActualTopology = true;
 	}
 
 	check(OutDependencyChain != nullptr);
@@ -6417,7 +6485,7 @@ bool URigHierarchy::IsParentedTo(FRigBaseElement* InChild, FRigBaseElement* InPa
 
 	{
 		UE::TScopeLock Lock(IsDependentOnLock);
-		if (!IsDependentOn(InChild, InParent, InDependencyProvider, true, *OutDependencyChain))
+		if (!IsDependentOn(InChild, InParent, InDependencyProvider, bActualTopology, *OutDependencyChain))
 		{
 			return false;
 		}
@@ -6426,34 +6494,38 @@ bool URigHierarchy::IsParentedTo(FRigBaseElement* InChild, FRigBaseElement* InPa
 	return true;
 }
 
-bool URigHierarchy::IsDependentOn(FRigBaseElement* InDependent, FRigBaseElement* InDependency, const IRigDependenciesProvider& InDependencyProvider, bool bIsOnActualTopology, TArray<FRigHierarchyRecord>& OutDependencyChain) const
+bool URigHierarchy::IsDependentOn(FRigBaseElement* InDependent, FRigBaseElement* InDependency, const IRigDependenciesProvider& InDependencyProvider, bool bIsOnActualTopology, TArray<TPair<FRigHierarchyRecord,FRigHierarchyRecord>>& OutDependencyChain) const
 {
 	if((InDependent == nullptr) || (InDependency == nullptr))
 	{
 		return false;
 	}
 
+	if(InDependent == InDependency)
+	{
+		return true;
+	}
+
 	const int32 DependentElementIndex = InDependent->GetIndex();
 	const int32 DependencyElementIndex = InDependency->GetIndex();
 	const TTuple<int32,int32> CacheKey(DependentElementIndex, DependencyElementIndex);
 
 	const FRigHierarchyRecord DependentRecord(FRigHierarchyRecord::EType_RigElement, DependentElementIndex, NAME_None);
-	if (OutDependencyChain.Contains(DependentRecord))
-	{
-		return false;
-	}
-	OutDependencyChain.Push(DependentRecord);
 
-	if(InDependent == InDependency)
+	if (OutDependencyChain.ContainsByPredicate([DependentRecord](const TPair<FRigHierarchyRecord, FRigHierarchyRecord>& InPair) -> bool 
 	{
-		return true;
+		return InPair.Value == DependentRecord;
+	}))
+	{
+		return false;
 	}
-
+	
 	if (bIsOnActualTopology)
 	{
-		if(!ElementDependencyCache.IsValid(GetTopologyVersion()))
+		const uint32 ElementDependencyVersion = HashCombine(GetTopologyVersion(), GetParentWeightVersion()); 
+		if(!ElementDependencyCache.IsValid(ElementDependencyVersion))
 		{
-			ElementDependencyCache.Set(TMap<TTuple<int32, int32>, bool>(), GetTopologyVersion());
+			ElementDependencyCache.Set(TMap<TTuple<int32, int32>, bool>(), ElementDependencyVersion);
 		}
 
 		// we'll only update the caches if we are following edges on the actual topology
@@ -6477,21 +6549,37 @@ bool URigHierarchy::IsDependentOn(FRigBaseElement* InDependent, FRigBaseElement*
 	// collect all possible parents of the dependent
 	if(const FRigSingleParentElement* SingleParentElement = Cast<FRigSingleParentElement>(InDependent))
 	{
-		if(IsDependentOn(SingleParentElement->ParentElement, InDependency, InDependencyProvider, true, OutDependencyChain))
+		if (SingleParentElement->ParentElement)
 		{
-			// we'll only update the caches if we are following edges on the actual topology
-			if(bIsOnActualTopology)
+			const FRigHierarchyRecord DependencyRecord(FRigHierarchyRecord::EType_RigElement, SingleParentElement->ParentElement->GetIndex(), NAME_None);
+			OutDependencyChain.Emplace(DependencyRecord, DependentRecord);
+		
+			if(IsDependentOn(SingleParentElement->ParentElement, InDependency, InDependencyProvider, bIsOnActualTopology, OutDependencyChain))
 			{
-				ElementDependencyCache.Get().FindOrAdd(CacheKey, true);
+				// we'll only update the caches if we are following edges on the actual topology
+				if(bIsOnActualTopology)
+				{
+					ElementDependencyCache.Get().FindOrAdd(CacheKey, true);
+				}
+				return true;
 			}
-			return true;
+
+			OutDependencyChain.Pop();
 		}
 	}
 	else if(const FRigMultiParentElement* MultiParentElement = Cast<FRigMultiParentElement>(InDependent))
 	{
 		for(const FRigElementParentConstraint& ParentConstraint : MultiParentElement->ParentConstraints)
 		{
-			if(IsDependentOn(ParentConstraint.ParentElement, InDependency, InDependencyProvider, true, OutDependencyChain))
+			if (!ParentConstraint.ParentElement || ParentConstraint.Weight.IsAlmostZero())
+			{
+				continue;
+			}
+
+			const FRigHierarchyRecord DependencyRecord(FRigHierarchyRecord::EType_RigElement, ParentConstraint.ParentElement->GetIndex(), NAME_None);
+			OutDependencyChain.Emplace(DependencyRecord, DependentRecord);
+
+			if(IsDependentOn(ParentConstraint.ParentElement, InDependency, InDependencyProvider, bIsOnActualTopology, OutDependencyChain))
 			{
 				// we'll only update the caches if we are following edges on the actual topology
 				if(bIsOnActualTopology)
@@ -6500,6 +6588,8 @@ bool URigHierarchy::IsDependentOn(FRigBaseElement* InDependent, FRigBaseElement*
 				}
 				return true;
 			}
+
+			OutDependencyChain.Pop();
 		}
 	}
 
@@ -6512,19 +6602,22 @@ bool URigHierarchy::IsDependentOn(FRigBaseElement* InDependent, FRigBaseElement*
 			bool bIsOnActualTopology,
 			const TTuple<int32,int32>& CacheKey,
 			const FRigHierarchyRecord& InDependentRecord,
-			TArray<FRigHierarchyRecord>& OutDependencyChain)
+			TArray<TPair<FRigHierarchyRecord,FRigHierarchyRecord>>& OutDependencyChain)
 		{
 			const TRigHierarchyDependencyMap& DependencyMap = InDependencyProvider.GetRigHierarchyDependencies();
 			
-			if(const TSet<FRigHierarchyRecord>* DependentRecordsPtr = DependencyMap.Find(InDependentRecord))
+			if(const TSet<FRigHierarchyRecord>* DependencyRecordsPtr = DependencyMap.Find(InDependentRecord))
 			{
-				const TSet<FRigHierarchyRecord>& DependentRecords = *DependentRecordsPtr;
-				for(const FRigHierarchyRecord& DependentRecord : DependentRecords)
+				const TSet<FRigHierarchyRecord>& DependencyRecords = *DependencyRecordsPtr;
+				for(const FRigHierarchyRecord& DependencyRecord : DependencyRecords)
 				{
-					if (DependentRecord.IsElement())
+					if (DependencyRecord.IsElement())
 					{
-						ensure(InHierarchy->Elements.IsValidIndex(DependentRecord.Index));
-						if(InHierarchy->IsDependentOn(InHierarchy->Elements[DependentRecord.Index], InDependency, InDependencyProvider, false, OutDependencyChain))
+						ensure(InHierarchy->Elements.IsValidIndex(DependencyRecord.Index));
+
+						OutDependencyChain.Emplace(DependencyRecord, InDependentRecord);
+
+						if(InHierarchy->IsDependentOn(InHierarchy->Elements[DependencyRecord.Index], InDependency, InDependencyProvider, bIsOnActualTopology, OutDependencyChain))
 						{
 							// we'll only update the caches if we are following edges on the actual topology
 							if(bIsOnActualTopology)
@@ -6533,16 +6626,22 @@ bool URigHierarchy::IsDependentOn(FRigBaseElement* InDependent, FRigBaseElement*
 							}
 							return true;
 						}
+
+						OutDependencyChain.Pop();
 					}
-					else if (DependentRecord.IsInstruction() || DependentRecord.IsVariable())
+					else if (DependencyRecord.IsInstruction() || DependencyRecord.IsVariable())
 					{
-						if (OutDependencyChain.Contains(DependentRecord))
+						if (OutDependencyChain.ContainsByPredicate([DependencyRecord](const TPair<FRigHierarchyRecord,FRigHierarchyRecord>& InPair) -> bool
+						{
+							return InPair.Key == DependencyRecord;
+						}))
 						{
 							return false;
 						}
-						OutDependencyChain.Push(DependentRecord);
-						
-						if (FoundDependencyInProvider(InHierarchy, InDependency, InDependencyProvider, bIsOnActualTopology, CacheKey, DependentRecord, OutDependencyChain))
+
+						OutDependencyChain.Emplace(DependencyRecord, InDependentRecord);
+
+						if (FoundDependencyInProvider(InHierarchy, InDependency, InDependencyProvider, bIsOnActualTopology, CacheKey, DependencyRecord, OutDependencyChain))
 						{
 							return true;
 						}
@@ -6572,8 +6671,6 @@ bool URigHierarchy::IsDependentOn(FRigBaseElement* InDependent, FRigBaseElement*
 		ElementDependencyCache.Get().FindOrAdd(CacheKey, false);
 	}
 
-	OutDependencyChain.Pop();
-
 	return false;
 }
 
diff --git a/Engine/Plugins/Animation/ControlRig/Source/ControlRig/Public/Rigs/RigHierarchy.h b/Engine/Plugins/Animation/ControlRig/Source/ControlRig/Public/Rigs/RigHierarchy.h
--- a/Engine/Plugins/Animation/ControlRig/Source/ControlRig/Public/Rigs/RigHierarchy.h
+++ b/Engine/Plugins/Animation/ControlRig/Source/ControlRig/Public/Rigs/RigHierarchy.h
@@ -3717,7 +3717,7 @@ public:
 	 * @param OutDependencyChain An optional output error message indicating the parent cycle
 	 * @return True if the given parent and given child have a parent-child relationship
 	 */
-    bool IsParentedTo(int32 InChildIndex, int32 InParentIndex, const IRigDependenciesProvider& InDependencyProvider, TArray<FRigHierarchyRecord>* OutDependencyChain = nullptr) const
+    bool IsParentedTo(int32 InChildIndex, int32 InParentIndex, const IRigDependenciesProvider& InDependencyProvider, TArray<TPair<FRigHierarchyRecord,FRigHierarchyRecord>>* OutDependencyChain = nullptr) const
 	{
 		if(Elements.IsValidIndex(InChildIndex) && Elements.IsValidIndex(InParentIndex))
 		{
@@ -3915,6 +3915,16 @@ public:
 	 */
 	CONTROLRIG_API void IncrementTopologyVersion();
 
+	/**
+	 * Returns the parent weight version of this hierarchy
+	 */
+	uint32 GetParentWeightVersion() const { return ParentWeightVersion; }
+
+	/**
+	 * Increments the parent weight version of this hierarchy
+	 */
+	CONTROLRIG_API void IncrementParentWeightVersion();
+
 	/**
 	 * Returns the metadata version of this hierarchy
 	 */
@@ -4451,7 +4461,7 @@ public:
 	CONTROLRIG_API bool IsParentedTo(FRigBaseElement* InChild,
 		FRigBaseElement* InParent,
 		const IRigDependenciesProvider& InDependencyProvider = FEmptyRigDependenciesProvider(),
-		TArray<FRigHierarchyRecord>* OutDependencyChain = nullptr) const;
+		TArray<TPair<FRigHierarchyRecord,FRigHierarchyRecord>>* OutDependencyChain = nullptr) const;
 
 private:
 	/**
@@ -4467,7 +4477,7 @@ private:
 		FRigBaseElement* InDependency,
 		const IRigDependenciesProvider& InDependencyProvider,
 		bool bIsOnActualTopology,
-		TArray<FRigHierarchyRecord>& OutDependencyChain) const;
+		TArray<TPair<FRigHierarchyRecord,FRigHierarchyRecord>>& OutDependencyChain) const;
 
 public:
 	/**
@@ -4663,6 +4673,13 @@ private:
 	UPROPERTY(transient)
 	uint32 TopologyVersion;
 
+	/**
+	 * The parent weight version of this hierarchy changes whenever parents are
+	 * added / removed or when the weights of a parent relationship changes.
+	 */
+	UPROPERTY(transient)
+	uint32 ParentWeightVersion;
+	
 	/**
 	 * The metadata version of the hierarchy changes when metadata is being
 	 * created or removed (not when the metadata values changes)
diff --git a/Engine/Plugins/Animation/ControlRig/Source/ControlRigEditor/Private/EditMode/Picker/SFloatingSpacePicker_AssetEditor.cpp b/Engine/Plugins/Animation/ControlRig/Source/ControlRigEditor/Private/EditMode/Picker/SFloatingSpacePicker_AssetEditor.cpp
--- a/Engine/Plugins/Animation/ControlRig/Source/ControlRigEditor/Private/EditMode/Picker/SFloatingSpacePicker_AssetEditor.cpp
+++ b/Engine/Plugins/Animation/ControlRig/Source/ControlRigEditor/Private/EditMode/Picker/SFloatingSpacePicker_AssetEditor.cpp
@@ -60,8 +60,8 @@ void SFloatingSpacePicker_AssetEditor::OnActiveSpaceChanged(
 	{
 		constexpr bool bInitial = false, bAffectChildren = true;
 		const FTransform Transform = InHierarchy->GetGlobalTransform(InControlKey);
-		FRigDependenciesProviderForControlRig DependenciesProvider(RuntimeRig);
 		FString OutFailureReason;
+		const FEmptyRigDependenciesProvider DependenciesProvider;
 		if (InHierarchy->SwitchToParent(InControlKey, InSpaceKey, bInitial, bAffectChildren, DependenciesProvider, &OutFailureReason))
 		{
 			InHierarchy->SetGlobalTransform(InControlKey, Transform);
diff --git a/Engine/Plugins/Animation/ControlRig/Source/ControlRigEditor/Private/EditMode/SControlRigSpacePicker.cpp b/Engine/Plugins/Animation/ControlRig/Source/ControlRigEditor/Private/EditMode/SControlRigSpacePicker.cpp
--- a/Engine/Plugins/Animation/ControlRig/Source/ControlRigEditor/Private/EditMode/SControlRigSpacePicker.cpp
+++ b/Engine/Plugins/Animation/ControlRig/Source/ControlRigEditor/Private/EditMode/SControlRigSpacePicker.cpp
@@ -111,7 +111,7 @@ void SControlRigSpacePicker::HandleActiveSpaceChanged(URigHierarchy* InHierarchy
 		if (UControlRig* ControlRig = GetControlRig())
 		{
 			FString FailureReason;
-			FRigDependenciesProviderForControlRig DependencyProvider(ControlRig);
+			const FEmptyRigDependenciesProvider DependencyProvider;
 			if (!InHierarchy->CanSwitchToParent(InControlKey, InSpaceKey, DependencyProvider, &FailureReason))
 			{
 				// notification
diff --git a/Engine/Plugins/Animation/ControlRig/Source/ControlRigEditor/Private/Editor/SRigSpacePickerWidget.cpp b/Engine/Plugins/Animation/ControlRig/Source/ControlRigEditor/Private/Editor/SRigSpacePickerWidget.cpp
--- a/Engine/Plugins/Animation/ControlRig/Source/ControlRigEditor/Private/Editor/SRigSpacePickerWidget.cpp
+++ b/Engine/Plugins/Animation/ControlRig/Source/ControlRigEditor/Private/Editor/SRigSpacePickerWidget.cpp
@@ -785,21 +785,8 @@ void SRigSpacePickerWidget::HandleClickTreeItemInAddMenu(TSharedPtr<FRigTreeElem
 				URigHierarchy* StrongHierarchy = Hierarchy.Get();
 				for(const FRigElementKey& ControlKey : ControlKeys)
 				{
-					FRigDependenciesProviderForControlRig DependencyProvider;
 					FString FailureReason;
-
-					if(UControlRig* ControlRig = StrongHierarchy->GetTypedOuter<UControlRig>())
-					{
-						DependencyProvider = FRigDependenciesProviderForControlRig(ControlRig); 
-					}
-					else if(UControlRigBlueprint* RigBlueprint = StrongHierarchy->GetTypedOuter<UControlRigBlueprint>())
-					{
-						if(UControlRig* CDO = Cast<UControlRig>(RigBlueprint->GetRigVMBlueprintGeneratedClass()->GetDefaultObject()))
-						{
-							DependencyProvider = FRigDependenciesProviderForControlRig(CDO); 
-						}
-					}
-					
+					const FEmptyRigDependenciesProvider DependencyProvider;
 					if (!StrongHierarchy->CanSwitchToParent(ControlKey, Key, DependencyProvider, &FailureReason))
 					{
 						// notification
