Control Rig: Update default hotkeys and related functionality

* Set Solve Mode: Construction Event  - 1
* Set Solve Mode: Forwards Solve - 2
* Set Solve Mode: Backwards Solve - 3
* Set Solve Mode: Backwards and Forwards - 4
* Set Next Solve Mode - M
* Using hotkey for Profiling (P by default) now creates an undo step

diff --git a/Engine/Plugins/Animation/ControlRig/Source/ControlRigEditor/Private/Editor/ControlRigEditor.cpp b/Engine/Plugins/Animation/ControlRig/Source/ControlRigEditor/Private/Editor/ControlRigEditor.cpp
--- a/Engine/Plugins/Animation/ControlRig/Source/ControlRigEditor/Private/Editor/ControlRigEditor.cpp
+++ b/Engine/Plugins/Animation/ControlRig/Source/ControlRigEditor/Private/Editor/ControlRigEditor.cpp
@@ -3054,6 +3054,29 @@ void FControlRigBaseEditor::HandleViewportCreated(const TSharedRef<class IPerson
 	}
 }
 
+void FControlRigBaseEditor::SetNextSolveMode()
+{
+	const TArray<FName> EventQueue = GetEventQueue();
+
+	constexpr bool bCompile = false;
+	if (EventQueue == ConstructionEventQueue)
+	{
+		SetEventQueueSuper(ForwardsSolveEventQueue, bCompile);
+	}
+	if (EventQueue == ForwardsSolveEventQueue)
+	{
+		SetEventQueueSuper(BackwardsSolveEventQueue, bCompile);
+	}
+	if (EventQueue == BackwardsSolveEventQueue)
+	{
+		SetEventQueueSuper(BackwardsAndForwardsSolveEventQueue, bCompile);
+	}
+	if (EventQueue == BackwardsAndForwardsSolveEventQueue)
+	{
+		SetEventQueueSuper(ConstructionEventQueue, bCompile);
+	}
+}
+
 void FControlRigBaseEditor::HandleToggleControlVisibility()
 {
 	if (FControlRigEditMode* EditMode = GetEditMode())
@@ -4422,6 +4445,11 @@ void FControlRigBaseEditor::BindCommandsImpl()
 		FExecuteAction::CreateSP(SharedEditor, &FControlRigBaseEditor::SetEventQueueSuper, TArray<FName>(BackwardsAndForwardsSolveEventQueue)),
 		FCanExecuteAction());
 
+	GetToolkitCommands()->MapAction(
+		FControlRigEditorCommands::Get().SetNextSolveMode,
+		FExecuteAction::CreateSP(SharedEditor, &FControlRigBaseEditor::SetNextSolveMode),
+		FCanExecuteAction());
+
 	GetToolkitCommands()->MapAction(
 		FControlRigEditorCommands::Get().ToggleControlVisibility,
 		FExecuteAction::CreateSP(SharedEditor, &FControlRigBaseEditor::HandleToggleControlVisibility),
@@ -4476,6 +4504,7 @@ void FControlRigBaseEditor::UnbindCommandsImpl()
 	GetToolkitCommands()->UnmapAction(FControlRigEditorCommands::Get().ForwardsSolveEvent);
 	GetToolkitCommands()->UnmapAction(FControlRigEditorCommands::Get().BackwardsSolveEvent);
 	GetToolkitCommands()->UnmapAction(FControlRigEditorCommands::Get().BackwardsAndForwardsSolveEvent);
+	GetToolkitCommands()->UnmapAction(FControlRigEditorCommands::Get().SetNextSolveMode);
 	GetToolkitCommands()->UnmapAction(FControlRigEditorCommands::Get().ToggleControlVisibility);
 	GetToolkitCommands()->UnmapAction(FControlRigEditorCommands::Get().ToggleControlsAsOverlay);
 	GetToolkitCommands()->UnmapAction(FControlRigEditorCommands::Get().ToggleDrawNulls);
diff --git a/Engine/Plugins/Animation/ControlRig/Source/ControlRigEditor/Private/Editor/ControlRigEditor.h b/Engine/Plugins/Animation/ControlRig/Source/ControlRigEditor/Private/Editor/ControlRigEditor.h
--- a/Engine/Plugins/Animation/ControlRig/Source/ControlRigEditor/Private/Editor/ControlRigEditor.h
+++ b/Engine/Plugins/Animation/ControlRig/Source/ControlRigEditor/Private/Editor/ControlRigEditor.h
@@ -376,6 +376,9 @@ protected:
 public:
 	void HandleViewportCreated(const TSharedRef<class IPersonaViewport>& InViewport);
 protected:
+	/** Sets the next solve mode, for example Backwards Solve if the last entry in the Queue is ForwardsSolve */
+	void SetNextSolveMode();
+
 	void HandleToggleControlVisibility();
 	bool AreControlsVisible() const;
 	void HandleToggleControlsAsOverlay();
diff --git a/Engine/Plugins/Animation/ControlRig/Source/ControlRigEditor/Private/Editor/ControlRigEditorCommands.cpp b/Engine/Plugins/Animation/ControlRig/Source/ControlRigEditor/Private/Editor/ControlRigEditorCommands.cpp
--- a/Engine/Plugins/Animation/ControlRig/Source/ControlRigEditor/Private/Editor/ControlRigEditorCommands.cpp
+++ b/Engine/Plugins/Animation/ControlRig/Source/ControlRigEditor/Private/Editor/ControlRigEditorCommands.cpp
@@ -6,10 +6,11 @@
 
 void FControlRigEditorCommands::RegisterCommands()
 {
-	UI_COMMAND(ConstructionEvent, "Construction Event", "Enable the construction mode for the rig", EUserInterfaceActionType::Button, FInputChord());
-	UI_COMMAND(ForwardsSolveEvent, "Forwards Solve", "Run the forwards solve graph", EUserInterfaceActionType::Button, FInputChord());
-	UI_COMMAND(BackwardsSolveEvent, "Backwards Solve", "Run the backwards solve graph", EUserInterfaceActionType::Button, FInputChord());
-	UI_COMMAND(BackwardsAndForwardsSolveEvent, "Backwards and Forwards", "Run backwards solve followed by forwards solve", EUserInterfaceActionType::Button, FInputChord());
+	UI_COMMAND(ConstructionEvent, "Construction Event", "Enable the construction mode for the rig", EUserInterfaceActionType::Button, FInputChord(EKeys::One));
+	UI_COMMAND(ForwardsSolveEvent, "Forwards Solve", "Run the forwards solve graph", EUserInterfaceActionType::Button, FInputChord(EKeys::Two));
+	UI_COMMAND(BackwardsSolveEvent, "Backwards Solve", "Run the backwards solve graph", EUserInterfaceActionType::Button, FInputChord(EKeys::Three));
+	UI_COMMAND(BackwardsAndForwardsSolveEvent, "Backwards and Forwards", "Run backwards solve followed by forwards solve", EUserInterfaceActionType::Button, FInputChord(EKeys::Four));
+	UI_COMMAND(SetNextSolveMode, "Next Solve Mode", "Sets the next solve mode, for example backwards solve when the current mode is forwards solve", EUserInterfaceActionType::Button, FInputChord(EKeys::M));
 	UI_COMMAND(RequestDirectManipulationPosition, "Request Direct Manipulation for Position", "Request per node direct manipulation on a position", EUserInterfaceActionType::Button, FInputChord(EKeys::W));
 	UI_COMMAND(RequestDirectManipulationRotation, "Request Direct Manipulation for Rotation", "Request per node direct manipulation on a rotation", EUserInterfaceActionType::Button, FInputChord(EKeys::E));
 	UI_COMMAND(RequestDirectManipulationScale, "Request Direct Manipulation for Scale", "Request per node direct manipulation on a scale", EUserInterfaceActionType::Button, FInputChord(EKeys::R));
diff --git a/Engine/Plugins/Animation/ControlRig/Source/ControlRigEditor/Private/Editor/ControlRigEditorCommands.h b/Engine/Plugins/Animation/ControlRig/Source/ControlRigEditor/Private/Editor/ControlRigEditorCommands.h
--- a/Engine/Plugins/Animation/ControlRig/Source/ControlRigEditor/Private/Editor/ControlRigEditorCommands.h
+++ b/Engine/Plugins/Animation/ControlRig/Source/ControlRigEditor/Private/Editor/ControlRigEditorCommands.h
@@ -30,6 +30,9 @@ public:
 	/** Run the backwards solve graph followed by the forwards solve graph */
 	TSharedPtr< FUICommandInfo > BackwardsAndForwardsSolveEvent;
 
+	/** Sets the next solve mode, for example Backwards Solve if the last entry in the Queue is ForwardsSolve */
+	TSharedPtr< FUICommandInfo > SetNextSolveMode;
+
 	/** Request per node direct manipulation on a position */
 	TSharedPtr< FUICommandInfo > RequestDirectManipulationPosition;
 
diff --git a/Engine/Plugins/Runtime/RigVM/Source/RigVMEditor/Private/Editor/RigVMEditor.cpp b/Engine/Plugins/Runtime/RigVM/Source/RigVMEditor/Private/Editor/RigVMEditor.cpp
--- a/Engine/Plugins/Runtime/RigVM/Source/RigVMEditor/Private/Editor/RigVMEditor.cpp
+++ b/Engine/Plugins/Runtime/RigVM/Source/RigVMEditor/Private/Editor/RigVMEditor.cpp
@@ -3849,8 +3849,11 @@ void FRigVMEditorBase::SwapAssetReferences()
 
 void FRigVMEditorBase::ToggleProfiling()
 {
-	if (TScriptInterface<IRigVMAssetInterface> Asset = GetRigVMAssetInterface())
+	TScriptInterface<IRigVMAssetInterface> Asset = GetRigVMAssetInterface();
+	if (Asset && Asset->GetObject())
 	{
+		const FScopedTransaction Transaction(LOCTEXT("ToggleProfiler", "Toggle Profiler"));
+		Asset.GetObject()->Modify();
 		Asset->SetProfilingEnabled(!Asset->GetVMRuntimeSettings().bEnableProfiling);
 	}
 }
