Anim In Engine: Initialize Sets and fix localization clashes.

diff --git a/Engine/Plugins/Animation/ControlRig/Source/ControlRigEditor/Private/Sequencer/SelectionSets/SSelectionSets.cpp b/Engine/Plugins/Animation/ControlRig/Source/ControlRigEditor/Private/Sequencer/SelectionSets/SSelectionSets.cpp
--- a/Engine/Plugins/Animation/ControlRig/Source/ControlRigEditor/Private/Sequencer/SelectionSets/SSelectionSets.cpp
+++ b/Engine/Plugins/Animation/ControlRig/Source/ControlRigEditor/Private/Sequencer/SelectionSets/SSelectionSets.cpp
@@ -1449,7 +1449,7 @@ bool SSelectionSets::ImportOrExportPath(FString& OutPath, bool bImport)
 
 		bOpen = DesktopPlatform->OpenFileDialog(
 			FSlateApplication::Get().FindBestParentWindowHandleForDialogs(nullptr),
-			NSLOCTEXT("RigMapperDefinitionIO", "ImportFromJson", "Import Selection Set from Json...").ToString(),
+			LOCTEXT("ImportSelectionSetFromJson", "Import Selection Set from Json...").ToString(),
 			LastRigMapperImportDirectory,
 			TEXT(""),
 			TEXT("json (*.json)|*.json|"),
diff --git a/Engine/Plugins/Animation/ControlRig/Source/ControlRigEditor/Private/Sequencer/SelectionSets/SelectionSets.cpp b/Engine/Plugins/Animation/ControlRig/Source/ControlRigEditor/Private/Sequencer/SelectionSets/SelectionSets.cpp
--- a/Engine/Plugins/Animation/ControlRig/Source/ControlRigEditor/Private/Sequencer/SelectionSets/SelectionSets.cpp
+++ b/Engine/Plugins/Animation/ControlRig/Source/ControlRigEditor/Private/Sequencer/SelectionSets/SelectionSets.cpp
@@ -51,6 +51,7 @@
 
 #define LOCTEXT_NAMESPACE "SelectionSets"
 
+
 ///////////////////////////////////////////////////////////
 /// FAIESelectionSetItem
 ///////////////////////////////////////////////////////////
@@ -371,45 +372,6 @@ bool FActorWithSelectionSet::IsActive(bool bShowSelectedOnly) const
 	return bIsSelected;
 }
 
-bool UAIESelectionSets::IsVisible(const FActorWithSelectionSet& ActorWithSelectionSet) const
-{
-	if (const USelectionSetsSettings* Settings = GetDefault<USelectionSetsSettings>())
-	{
-		if (Settings->GetSelectionSetsActorSetVisibility() == ESelectionSetsActorSetVisibility::MultipleControRigs)
-		{
-			if (bShowSelectedOnly)
-			{
-				return false;
-			}
-			if (ActorWithSelectionSet.IsControlRig() == true)
-			{
-				for (const TPair<TWeakObjectPtr<AActor>, FActorWithSelectionSet>& Pair : ActorsWithSelectionSets.ActorsWithSet)
-				{
-					if (Pair.Value.IsControlRig() == true)
-					{
-						if (&ActorWithSelectionSet != &Pair.Value)
-						{
-							return true;
-						}
-					}
-				}
-			}
-			else
-			{
-				return false;
-			}
-		}
-	}
-	if (bShowSelectedOnly)
-	{
-		if (ActorWithSelectionSet.IsSelectedInViewport())
-		{
-			return true;
-		}
-	}
-	return false;
-}
-
 void FActorWithSelectionSet::ShowAll() const
 {
 	FControlRigEditMode* EditMode = static_cast<FControlRigEditMode*>(GLevelEditorModeTools().GetActiveMode(FControlRigEditMode::ModeName));
@@ -488,6 +450,45 @@ UAIESelectionSets* UAIESelectionSets::GetSelectionSets(ULevelSequence* LevelSequ
 	return nullptr;
 }
 
+bool UAIESelectionSets::IsVisible(const FActorWithSelectionSet& ActorWithSelectionSet) const
+{
+	if (const USelectionSetsSettings* Settings = GetDefault<USelectionSetsSettings>())
+	{
+		if (Settings->GetSelectionSetsActorSetVisibility() == ESelectionSetsActorSetVisibility::MultipleControRigs)
+		{
+			if (bShowSelectedOnly)
+			{
+				return false;
+			}
+			if (ActorWithSelectionSet.IsControlRig() == true)
+			{
+				for (const TPair<TWeakObjectPtr<AActor>, FActorWithSelectionSet>& Pair : ActorsWithSelectionSets.ActorsWithSet)
+				{
+					if (Pair.Value.IsControlRig() == true)
+					{
+						if (&ActorWithSelectionSet != &Pair.Value)
+						{
+							return true;
+						}
+					}
+				}
+			}
+			else
+			{
+				return false;
+			}
+		}
+	}
+	if (bShowSelectedOnly)
+	{
+		if (ActorWithSelectionSet.IsSelectedInViewport())
+		{
+			return true;
+		}
+	}
+	return false;
+}
+
 TArray<FGuid> UAIESelectionSets::GetActiveSelectionSets() 
 {
 	//when bShowSelected is on we cache the last set and use that when none are selected. Also need to keep track
@@ -635,7 +636,7 @@ FGuid UAIESelectionSets::CreateSetItemFromSelection(TSharedPtr<ISequencer>& InSe
 					ItemName.MirrorName = ItemName.Name;
 				}
 				
-				ItemName.Type = (uint32)EAIESelectionSetItemType::ControlRig;
+				ItemName.Type = (int32)EAIESelectionSetItemType::ControlRig;
 				Duplicates.SetNum(0);
 				ActiveNames.MultiFind(ItemName.Name, Duplicates);
 				ItemName.Duplicates = Duplicates.Num();
@@ -663,7 +664,7 @@ FGuid UAIESelectionSets::CreateSetItemFromSelection(TSharedPtr<ISequencer>& InSe
 					ItemName.OwnerActorName = OwnerActorName.GetValue();
 				}
 				ItemName.MirrorName = ItemName.Name;
-				ItemName.Type = (uint32)EAIESelectionSetItemType::Actor;
+				ItemName.Type = (int32)EAIESelectionSetItemType::Actor;
 				Duplicates.SetNum(0);
 				ActiveNames.MultiFind(ItemName.Name, Duplicates);
 				ItemName.Duplicates = Duplicates.Num();
@@ -779,7 +780,7 @@ bool UAIESelectionSets::AddSelectionToSetItem(const FGuid& InGuid, TSharedPtr<IS
 			{
 				FAIESelectionSetItemName ItemName;
 				ItemName.Name = ControlName.ToString();
-				ItemName.Type = (uint32)EAIESelectionSetItemType::ControlRig;
+				ItemName.Type = (int32)EAIESelectionSetItemType::ControlRig;
 				if (FRigControlCopy* CopyRigControl = MirrorTable.GetControl(ControlRig, TempPose, ControlName, true))
 				{
 					ItemName.MirrorName = CopyRigControl->Name.ToString();
@@ -796,7 +797,7 @@ bool UAIESelectionSets::AddSelectionToSetItem(const FGuid& InGuid, TSharedPtr<IS
 			{
 				FAIESelectionSetItemName ItemName;
 				ItemName.Name = Actor->GetActorLabel(true);
-				ItemName.Type = (uint32)EAIESelectionSetItemType::Actor;
+				ItemName.Type = (int32)EAIESelectionSetItemType::Actor;
 				Duplicates.SetNum(0);
 				ActiveNames.MultiFind(ItemName.Name, Duplicates);
 				ItemName.Duplicates = Duplicates.Num();
diff --git a/Engine/Plugins/Animation/ControlRig/Source/ControlRigEditor/Private/Sequencer/SelectionSets/SelectionSets.h b/Engine/Plugins/Animation/ControlRig/Source/ControlRigEditor/Private/Sequencer/SelectionSets/SelectionSets.h
--- a/Engine/Plugins/Animation/ControlRig/Source/ControlRigEditor/Private/Sequencer/SelectionSets/SelectionSets.h
+++ b/Engine/Plugins/Animation/ControlRig/Source/ControlRigEditor/Private/Sequencer/SelectionSets/SelectionSets.h
@@ -14,13 +14,15 @@ class UControlRig;
 class UAIESelectionSets;
 
 
-//Type the selection set item
+//Type  of the selection set item
 UENUM()
 enum EAIESelectionSetItemType : int32
 {
+	/*Control Rig Type*/
 	ControlRig = 0x0,
+	/*Actor Type*/
 	Actor = 0x1,
-//caussing mac issues not used yet so removing	Component = 0x2,
+//causing mac issues not used yet so removing	Component = 0x2,
 };
 ENUM_CLASS_FLAGS(EAIESelectionSetItemType)
 
@@ -28,32 +30,48 @@ USTRUCT(BlueprintType)
 struct FAIESelectionSetItemName
 {
 	GENERATED_BODY();
+
+	/* The Name of the selection set item*, maybe an Actor Label or a ControlRig Control Name */
 	UPROPERTY(BlueprintReadWrite, EditAnywhere, Category = "Data")
 	FString Name;
+
+	/* The Name of the selection set mirrored item for Control Rig Controls that have them */
 	UPROPERTY(BlueprintReadWrite, EditAnywhere, Category = "Data")
 	FString MirrorName;
-	UPROPERTY(BlueprintReadOnly, EditAnywhere, Category = "Data", meta = (Bitmask, BitmaskEnum = "/Script/ControlRigEditor.EAIESelectionSetItemType"))
-	int32 Type;
+
+	/* The EAIESelectionSetItemType, 0 is a Control Rig Control, 1 is an Actor */
 	UPROPERTY(BlueprintReadOnly, EditAnywhere, Category = "Data")
+	int32 Type = 1;
+
+	//currently not used but maybe in future
+	UPROPERTY()
 	int32 Duplicates = 0;
+
+	/* The Name of the Owner Actor for Multi-Sets containing multiple rigs, may not be set*/
 	UPROPERTY(BlueprintReadWrite, EditAnywhere, Category = "Data")
 	FString OwnerActorName;
 };
 
+
 USTRUCT(BlueprintType)
 struct FAIESelectionSetItemViewData
 {
 	GENERATED_BODY();
 
+	/* The Color of the Set*/
 	UPROPERTY(BlueprintReadWrite, EditAnywhere, Category = "Data")
-	FLinearColor Color;
+	FLinearColor Color = FLinearColor::White;
 
+	/* The logical row the set should be in, basically it's order. This is an internal order for just one set*/
+	UPROPERTY(BlueprintReadWrite, EditAnywhere, Category = "Data")
+	int32 Row = 0;
+	
+	//currently not used but maybe in future
 	UPROPERTY()
-	int32 Row;
-	UPROPERTY()
-	int32 Column;
+	int32 Column = 0;
 };
 
+
 USTRUCT(BlueprintType)
 struct FAIESelectionSetItem
 {
@@ -63,16 +81,22 @@ struct FAIESelectionSetItem
 	UPROPERTY()
 	FGuid Guid;
 	
+	/* The Name of the Set*/
 	UPROPERTY(BlueprintReadWrite, EditAnywhere, Category = "Data")
 	FText ItemName;
 
+	/* The collection of set items that make up the set, basically an array of names with some metadata*/
 	UPROPERTY(BlueprintReadWrite, EditAnywhere, Category = "Data")
 	TArray<FAIESelectionSetItemName> Names;
+
+	/* Viewing Data for the set, like location and row color*/
 	UPROPERTY(BlueprintReadWrite, EditAnywhere, Category = "Data")
 	FAIESelectionSetItemViewData ViewData;
-	UPROPERTY(BlueprintReadWrite, EditAnywhere, Category = "Hierarchy")
+
+	//currently not used but maybe in future
+	UPROPERTY()
 	FGuid Parent;
-	UPROPERTY(BlueprintReadWrite, EditAnywhere, Category = "Hierarchy")
+	//currently not used but maybe in future
 	TArray<FGuid> Children;
 
 	
@@ -84,7 +108,6 @@ struct FAIESelectionSetItem
 	static FString GetMirror(const FString& InString);
 };
 
-
 struct FActorWithSelectionSet
 {
 	TSet<TWeakObjectPtr<UObject>> Actors; //actor or crs
