Control Rig: Unify size of "Add Constraint" context menus to be 350. This fixes some of the menus being too narrow to read any of the property texts.

diff --git a/Engine/Plugins/Animation/ControlRig/Source/ControlRigEditor/Private/Editor/Constraints/SConstraintsWidget.cpp b/Engine/Plugins/Animation/ControlRig/Source/ControlRigEditor/Private/Editor/Constraints/SConstraintsWidget.cpp
--- a/Engine/Plugins/Animation/ControlRig/Source/ControlRigEditor/Private/Editor/Constraints/SConstraintsWidget.cpp
+++ b/Engine/Plugins/Animation/ControlRig/Source/ControlRigEditor/Private/Editor/Constraints/SConstraintsWidget.cpp
@@ -480,7 +480,11 @@ TSharedRef<SWidget> SConstraintMenuEntry::GenerateConstraintDefaultWidget()
 			.AutoHeight()
 			.VAlign( VAlign_Center )
 			[
-				DetailsView
+				SNew(SBox)
+				.WidthOverride(350.f) // Make all the menus have a consistent size & prevent details from being too narrow to read the property texts. 
+				[
+					DetailsView
+				]
 			]
 			
 			+ SVerticalBox::Slot()
