Control Rig: Fix Constraints tab inline tabs getting cut off when you make the tab very narrow.

diff --git a/Engine/Plugins/Animation/ControlRig/Source/ControlRigEditor/Private/Misc/Tab/SInlineTabPanel.cpp b/Engine/Plugins/Animation/ControlRig/Source/ControlRigEditor/Private/Misc/Tab/SInlineTabPanel.cpp
--- a/Engine/Plugins/Animation/ControlRig/Source/ControlRigEditor/Private/Misc/Tab/SInlineTabPanel.cpp
+++ b/Engine/Plugins/Animation/ControlRig/Source/ControlRigEditor/Private/Misc/Tab/SInlineTabPanel.cpp
@@ -28,8 +28,10 @@ static TArray<FTabEntry> Transform(TConstArrayView<FInlineTabArgs> InArgs, const
 
 		Result.Emplace();
 		FTabEntry& Entry = Result[Index];
-		
-		Entry.ButtonContent.Widget = SNew(SHorizontalBox)
+
+		Entry.MinButtonSlotWidth = 50.f; // When user resizes the button to be quite narrow, leave enough space to show icon & text. 
+		Entry.ButtonContent.Widget =
+			SNew(SHorizontalBox)
 			.ToolTipText(CurrentArgs.ToolTipText)
 			
 			+SHorizontalBox::Slot()
@@ -41,12 +43,15 @@ static TArray<FTabEntry> Transform(TConstArrayView<FInlineTabArgs> InArgs, const
 			]
 			
 			+SHorizontalBox::Slot()
-			.AutoWidth()
+			.FillContentWidth(1.f) // Reduce text desired size when user resizes button to be narrow...
+			.MinWidth(20.f) // ... but force enough space to show the first letter ...
 			.Padding(4.f, 0.f, 0.f, 0.f)
 			[
 				SNew(STextBlock)
 				.Text(CurrentArgs.Label)
+				.OverflowPolicy(ETextOverflowPolicy::Ellipsis) // ... and cut off the rest of the text when there's not enough space
 			];
+		
 		Entry.OnTabSelected.BindLambda([Index, WeakOwnerPanel = This.ToWeakPtr(), Delegate = Entry.OnTabSelected]()
 		{
 			if (const TSharedPtr<SInlineTabPanel> OwnerPin = WeakOwnerPanel.Pin())
diff --git a/Engine/Plugins/Animation/ControlRig/Source/ControlRigEditor/Private/Misc/Tab/STabArea.cpp b/Engine/Plugins/Animation/ControlRig/Source/ControlRigEditor/Private/Misc/Tab/STabArea.cpp
--- a/Engine/Plugins/Animation/ControlRig/Source/ControlRigEditor/Private/Misc/Tab/STabArea.cpp
+++ b/Engine/Plugins/Animation/ControlRig/Source/ControlRigEditor/Private/Misc/Tab/STabArea.cpp
@@ -18,7 +18,8 @@ void STabArea::Construct(const FArguments& InArgs)
 		TSharedPtr<STabButton> Button;
 		Content->AddSlot()
 			.Padding(InArgs._Padding)
-			.AutoWidth()
+			.FillContentWidth(1.f) // This causes all buttons have the same width and grow until all their desired sizes have been reached.
+			.MinWidth(Entry.MinButtonSlotWidth) 
 			[
 				SAssignNew(Button, STabButton)
 				.OnActivated_Lambda([this, OnTabSelected = Entry.OnTabSelected, ButtonIndex]()
diff --git a/Engine/Plugins/Animation/ControlRig/Source/ControlRigEditor/Private/Misc/Tab/STabArea.h b/Engine/Plugins/Animation/ControlRig/Source/ControlRigEditor/Private/Misc/Tab/STabArea.h
--- a/Engine/Plugins/Animation/ControlRig/Source/ControlRigEditor/Private/Misc/Tab/STabArea.h
+++ b/Engine/Plugins/Animation/ControlRig/Source/ControlRigEditor/Private/Misc/Tab/STabArea.h
@@ -11,7 +11,13 @@ class STabButton;
 	
 struct FTabEntry
 {
+	/** Content to display in the button */
 	TAlwaysValidWidget ButtonContent;
+
+	/** Min widhth you want the button slot to have when the panel is resized */
+	TAttribute<float> MinButtonSlotWidth = 60.f;
+
+	/** Invoked when the tab is opened */
 	FSimpleDelegate OnTabSelected;
 };
 
