Control Rig: Fix Constraining being duplicated into the Window and Tools context menus while in Animation Mode

diff --git a/Engine/Plugins/Animation/ControlRig/Source/ControlRigEditor/Private/EditMode/Constrain/ConstrainToolsManager.cpp b/Engine/Plugins/Animation/ControlRig/Source/ControlRigEditor/Private/EditMode/Constrain/ConstrainToolsManager.cpp
--- a/Engine/Plugins/Animation/ControlRig/Source/ControlRigEditor/Private/EditMode/Constrain/ConstrainToolsManager.cpp
+++ b/Engine/Plugins/Animation/ControlRig/Source/ControlRigEditor/Private/EditMode/Constrain/ConstrainToolsManager.cpp
@@ -13,6 +13,7 @@ namespace UE::ControlRigEditor
 {
 FConstrainToolsManager::FConstrainToolsManager(
 	const TSharedRef<FTabManager>& InToolkitTabManager, 
+	const TSharedRef<FWorkspaceItem>& InWorkspaceMenuGroup,
 	FControlRigEditMode& InOwningEditMode,
 	const TSharedRef<FRigSelectionViewModel>& InSelectionViewModel
 	)
@@ -20,7 +21,7 @@ FConstrainToolsManager::FConstrainToolsManager(
 	, OwningEditMode(InOwningEditMode)
 	, SelectionViewModel(InSelectionViewModel)
 {
-	RegisterTabSpawner(TabManager);
+	RegisterTabSpawner(TabManager, InWorkspaceMenuGroup);
 }
 
 FConstrainToolsManager::~FConstrainToolsManager()
@@ -59,14 +60,15 @@ bool FConstrainToolsManager::IsShowingWidget() const
 	return SpawnedTab.IsValid();
 }
 
-void FConstrainToolsManager::RegisterTabSpawner(const TSharedRef<FTabManager>& InTabManager)
+void FConstrainToolsManager::RegisterTabSpawner(const TSharedRef<FTabManager>& InTabManager, const TSharedRef<FWorkspaceItem>& InWorkspaceMenuGroup)
 {
 	InTabManager->RegisterTabSpawner(
 		FControlRigEditModeToolkit::ConstrainingTabName, FOnSpawnTab::CreateRaw(this, &FConstrainToolsManager::SpawnSnapperTab)
 		)
 		.SetDisplayName(LOCTEXT("Tab.Title", "Constraining"))
 		.SetTooltipText(LOCTEXT("Tab.ToolTip", "Snap child objects to a parent object over a set of frames."))
-		.SetIcon(FSlateIcon(TEXT("ControlRigEditorStyle"), TEXT("ControlRig.SnapperTool")));
+		.SetIcon(FSlateIcon(TEXT("ControlRigEditorStyle"), TEXT("ControlRig.SnapperTool")))
+		.SetGroup(InWorkspaceMenuGroup);
 	InTabManager->RegisterDefaultTabWindowSize(FControlRigEditModeToolkit::ConstrainingTabName, FVector2D(380, 500));
 }
 
diff --git a/Engine/Plugins/Animation/ControlRig/Source/ControlRigEditor/Private/EditMode/Constrain/ConstrainToolsManager.h b/Engine/Plugins/Animation/ControlRig/Source/ControlRigEditor/Private/EditMode/Constrain/ConstrainToolsManager.h
--- a/Engine/Plugins/Animation/ControlRig/Source/ControlRigEditor/Private/EditMode/Constrain/ConstrainToolsManager.h
+++ b/Engine/Plugins/Animation/ControlRig/Source/ControlRigEditor/Private/EditMode/Constrain/ConstrainToolsManager.h
@@ -7,9 +7,10 @@
 
 class FControlRigEditMode;
 class FSpawnTabArgs;
-class SDockTab;
 class FTabManager;
 class FToolBarBuilder;
+class FWorkspaceItem;
+class SDockTab;
 
 namespace UE::ControlRigEditor
 {
@@ -29,6 +30,7 @@ public:
 
 	explicit FConstrainToolsManager(
 		const TSharedRef<FTabManager>& InToolkitTabManager,
+		const TSharedRef<FWorkspaceItem>& InWorkspaceMenuGroup,
 		FControlRigEditMode& InOwningEditMode UE_LIFETIMEBOUND,
 		const TSharedRef<FRigSelectionViewModel>& InSelectionViewModel
 		);
@@ -53,7 +55,7 @@ private:
 	TWeakPtr<SDockTab> SpawnedTab;
 
 	// Tab spawning
-	void RegisterTabSpawner(const TSharedRef<FTabManager>& InTabManager);
+	void RegisterTabSpawner(const TSharedRef<FTabManager>& InTabManager, const TSharedRef<FWorkspaceItem>& InWorkspaceMenuGroup);
 	void UnregisterTabSpawner(const TSharedRef<FTabManager>& InTabManager);
 	TSharedRef<SDockTab> SpawnSnapperTab(const FSpawnTabArgs& Args);
 };
diff --git a/Engine/Plugins/Animation/ControlRig/Source/ControlRigEditor/Private/EditMode/ControlRigEditModeToolkit.cpp b/Engine/Plugins/Animation/ControlRig/Source/ControlRigEditor/Private/EditMode/ControlRigEditModeToolkit.cpp
--- a/Engine/Plugins/Animation/ControlRig/Source/ControlRigEditor/Private/EditMode/ControlRigEditModeToolkit.cpp
+++ b/Engine/Plugins/Animation/ControlRig/Source/ControlRigEditor/Private/EditMode/ControlRigEditModeToolkit.cpp
@@ -374,8 +374,9 @@ void FControlRigEditModeToolkit::RequestModeUITabs()
 	FModeToolkit::RequestModeUITabs();
 	if (ModeUILayer.IsValid())
 	{
-		TSharedPtr<FAssetEditorModeUILayer> ModeUILayerPtr = ModeUILayer.Pin();
-		TSharedRef<FWorkspaceItem> MenuGroup = ModeUILayerPtr->GetModeMenuCategory().ToSharedRef();
+		const TSharedPtr<FAssetEditorModeUILayer> ModeUILayerPtr = ModeUILayer.Pin();
+		const TSharedRef<FWorkspaceItem> MenuGroup = ModeUILayerPtr->GetModeMenuCategory().ToSharedRef();
+		WorkspaceMenuCategory = MenuGroup;
 
 		FMinorTabConfig DetailTabInfo;
 		DetailTabInfo.OnSpawnTab = FOnSpawnTab::CreateStatic(&SpawnDetailsTab, &EditMode);
@@ -466,8 +467,10 @@ void FControlRigEditModeToolkit::InvokeUI()
 	TweenOverlayManager = MakeUnique<UE::ControlRigEditor::FTweenOverlayManager>(GetToolkitHost(), GetToolkitCommands(), SharedThis(&EditMode));
 
 	SelectionViewModel = MakeShared<UE::ControlRigEditor::FRigSelectionViewModel>();
+
+	checkf(WorkspaceMenuCategory, TEXT("WorkspaceMenuCategory should have been set in RequestModeUITabs by now."))
 	ConstrainToolsManager = MakeUnique<UE::ControlRigEditor::FConstrainToolsManager>(
-		ModeUILayerPtr->GetTabManager().ToSharedRef(), EditMode, SelectionViewModel.ToSharedRef()
+		ModeUILayerPtr->GetTabManager().ToSharedRef(), GetWorkspaceMenuCategory(), EditMode, SelectionViewModel.ToSharedRef()
 		);
 	
 	const FControlRigUIRestoreStates& RestoreStates = UControlRigEditModeSettings::Get()->LastUIStates;
