Control Rig: Fix failing ensure when clicking Add button in space picker when no control is selected. Disables the button when no control is selected.

diff --git a/Engine/Plugins/Animation/ControlRig/Source/ControlRigEditor/Private/Editor/SRigSpacePickerWidget.cpp b/Engine/Plugins/Animation/ControlRig/Source/ControlRigEditor/Private/Editor/SRigSpacePickerWidget.cpp
--- a/Engine/Plugins/Animation/ControlRig/Source/ControlRigEditor/Private/Editor/SRigSpacePickerWidget.cpp
+++ b/Engine/Plugins/Animation/ControlRig/Source/ControlRigEditor/Private/Editor/SRigSpacePickerWidget.cpp
@@ -700,9 +700,10 @@ TSharedRef<SWidget> SRigSpacePickerWidget::CreateAddButton()
 {
 	return SAssignNew(AddComboButton, SPositiveActionButton)
 		.OnGetMenuContent(this, &SRigSpacePickerWidget::OnGetAddButtonContent)
+		.IsEnabled(this, &SRigSpacePickerWidget::IsAddButtonEnabled)
 		.OnMenuOpenChanged(this, &SRigSpacePickerWidget::OnMenuOpenChanged)
 		.Cursor(EMouseCursor::Default)
-		.ToolTipText(LOCTEXT("AddSpace.ToolTip", "Add Space"))
+		.ToolTipText(this, &SRigSpacePickerWidget::GetAddButtonTooltipText)
 		.Text(LOCTEXT("AddSpace.Label", "Add"))
 		.Visibility_Lambda([this]()
 		{
@@ -825,6 +826,27 @@ void SRigSpacePickerWidget::HandleClickTreeItemInAddMenu(TSharedPtr<FRigTreeElem
 	AddComboButton->SetIsMenuOpen(false, false);
 }
 
+FText SRigSpacePickerWidget::GetAddButtonTooltipText() const
+{
+	FText Reason;
+	const bool bIsEnabled = IsAddButtonEnabledWithReason(&Reason);
+	return bIsEnabled
+		 ? LOCTEXT("AddSpace.ToolTip", "Add Space")
+		 : Reason;
+}
+
+#define SET_REASON(OutReason, Text) if (OutReason) { *OutReason = Text; }
+bool SRigSpacePickerWidget::IsAddButtonEnabledWithReason(FText* OutReason) const
+{
+	if (GetHierarchyConst() == nullptr)
+	{
+		SET_REASON(OutReason, LOCTEXT("Add.ToolTip.NoSelection", "Select a control to add space."));
+		return false;
+	}
+	return true;
+}
+#undef SET_REASON
+
 bool SRigSpacePickerWidget::IsRestricted() const
 {
 	if(URigHierarchy* CurrentHierarchy = GetHierarchy())
diff --git a/Engine/Plugins/Animation/ControlRig/Source/ControlRigEditor/Private/Editor/SRigSpacePickerWidget.h b/Engine/Plugins/Animation/ControlRig/Source/ControlRigEditor/Private/Editor/SRigSpacePickerWidget.h
--- a/Engine/Plugins/Animation/ControlRig/Source/ControlRigEditor/Private/Editor/SRigSpacePickerWidget.h
+++ b/Engine/Plugins/Animation/ControlRig/Source/ControlRigEditor/Private/Editor/SRigSpacePickerWidget.h
@@ -136,12 +136,30 @@ private:
 	/** The combo button for adding items. Can be nullptr. */
 	TSharedPtr<SPositiveActionButton> AddComboButton;
 
+	/**
+	 * Whether the add button menu is currently open.
+	 * 
+	 * Helps us decide whether to put the floating window that contains us, if any, to the front.
+	 * We don't want to put the window in front if a context menu of OUR UI is opened - only when the user clicks some other UE UI.
+	 */
 	bool bIsAddMenuOpen = false;
-	
+
+	/** @return Creates the button for adding a space to the selected control. */
 	TSharedRef<SWidget> CreateAddButton();
+	
+	/** @return Menu content for the add button */
 	TSharedRef<SWidget> OnGetAddButtonContent();
+	/** Handles the menu closing. Ensures the space picker widget is put on top of all UI if it is placed in the floating window.*/
 	void OnMenuOpenChanged(bool bIsOpen);
+	/** Assigns the clicked space to the current control rig selection */
 	void HandleClickTreeItemInAddMenu(TSharedPtr<FRigTreeElement> InItem);
+	
+	/** @return Whether the button is enabled. Disabled e.g. when no control rig is selected. */
+	bool IsAddButtonEnabled() const { return IsAddButtonEnabledWithReason(); }
+	/** @return Gives the user a reason when the button is disabled. */
+	FText GetAddButtonTooltipText() const;
+	/** @return Whether the button should be enabled. Optionally returns a call to action text to the user. */
+	bool IsAddButtonEnabledWithReason(FText* OutReason = nullptr) const;
 
 public:
 	
