Control Rig: Fix crash when undoing in sequencer while control rig editor is opened
#rnx

diff --git a/Engine/Plugins/Animation/ControlRig/Source/ControlRigEditor/Private/EditMode/ControlRigEditMode.cpp b/Engine/Plugins/Animation/ControlRig/Source/ControlRigEditor/Private/EditMode/ControlRigEditMode.cpp
--- a/Engine/Plugins/Animation/ControlRig/Source/ControlRigEditor/Private/EditMode/ControlRigEditMode.cpp
+++ b/Engine/Plugins/Animation/ControlRig/Source/ControlRigEditor/Private/EditMode/ControlRigEditMode.cpp
@@ -2510,11 +2510,15 @@ TArray<FRigElementKey> FControlRigEditMode::GetSelectedRigElements(UControlRig*
 	{
 		return SelectedKeys;
 	}
-	if (ControlRig->GetHierarchy())
+	
+	URigHierarchy* Hierarchy = ControlRig->GetHierarchy();
+	if (Hierarchy == nullptr)
 	{
-		SelectedKeys = ControlRig->GetHierarchy()->GetSelectedKeys();
+		return SelectedKeys;
 	}
 
+	SelectedKeys = Hierarchy->GetSelectedKeys();
+
 	// currently only 1 transient control is allowed at a time
 	// Transient Control's bSelected flag is never set to true, probably to avoid confusing other parts of the system
 	// But since Edit Mode directly deals with transient controls, its selection status is given special treatment here.
@@ -2525,7 +2529,7 @@ TArray<FRigElementKey> FControlRigEditMode::GetSelectedRigElements(UControlRig*
 		{
 			const FName ControlName = UControlRig::GetNameForTransientControl(SelectedKeys[0]);
 			const FRigElementKey TransientControlKey = FRigElementKey(ControlName, ERigElementType::Control);
-			if(ControlRig->GetHierarchy()->Contains(TransientControlKey))
+			if(Hierarchy->Contains(TransientControlKey))
 			{
 				SelectedKeys.Add(TransientControlKey);
 			}
@@ -2536,7 +2540,7 @@ TArray<FRigElementKey> FControlRigEditMode::GetSelectedRigElements(UControlRig*
 	{
 		// check if there is a pin value transient control active
 		// when a pin control is active, all existing selection should have been cleared
-		TArray<FRigControlElement*> TransientControls = ControlRig->GetHierarchy()->GetTransientControls();
+		TArray<FRigControlElement*> TransientControls = Hierarchy->GetTransientControls();
 
 		if (TransientControls.Num() > 0)
 		{
