Control Rig: Draw wires in dependency graph with arrows

#rnx

diff --git a/Engine/Plugins/Animation/ControlRig/Source/ControlRigEditor/Private/RigDependencyGraph/RigDependencyConnectionDrawingPolicy.h b/Engine/Plugins/Animation/ControlRig/Source/ControlRigEditor/Private/RigDependencyGraph/RigDependencyConnectionDrawingPolicy.h
--- a/Engine/Plugins/Animation/ControlRig/Source/ControlRigEditor/Private/RigDependencyGraph/RigDependencyConnectionDrawingPolicy.h
+++ b/Engine/Plugins/Animation/ControlRig/Source/ControlRigEditor/Private/RigDependencyGraph/RigDependencyConnectionDrawingPolicy.h
@@ -10,9 +10,6 @@ public:
 	FRigDependencyConnectionDrawingPolicy(int32 InBackLayerID, int32 InFrontLayerID, float InZoomFactor, const FSlateRect& InClippingRect, FSlateWindowElementList& InDrawElements)
 		: FConnectionDrawingPolicy(InBackLayerID, InFrontLayerID, InZoomFactor, InClippingRect, InDrawElements)
 	{
-		// We don't draw arrows here
-		ArrowImage = nullptr;
-		ArrowRadius = FVector2f::ZeroVector;
 	}
 
 	virtual FVector2f ComputeSplineTangent(const FVector2f& Start, const FVector2f& End) const override
