Control Rig: Extracted toolkit responsibilities from SControlRigEditModeTools to FControlRigEditModeToolkit.
- Namely building the toolbar, GetActiveToolDisplayName, and GetActiveToolMessage now live in FControlRigEditModeToolkit.
- Removed the members in SControlRigEditModeTools which became unused as result.

diff --git a/Engine/Plugins/Animation/ControlRig/Source/ControlRigEditor/Private/EditMode/ControlRigEditModeToolkit.cpp b/Engine/Plugins/Animation/ControlRig/Source/ControlRigEditor/Private/EditMode/ControlRigEditModeToolkit.cpp
--- a/Engine/Plugins/Animation/ControlRig/Source/ControlRigEditor/Private/EditMode/ControlRigEditModeToolkit.cpp
+++ b/Engine/Plugins/Animation/ControlRig/Source/ControlRigEditor/Private/EditMode/ControlRigEditModeToolkit.cpp
@@ -81,7 +81,7 @@ FControlRigEditModeToolkit::~FControlRigEditModeToolkit()
 
 void FControlRigEditModeToolkit::Init(const TSharedPtr<IToolkitHost>& InitToolkitHost)
 {
-	SAssignNew(ModeTools, SControlRigEditModeTools, SharedThis(this), EditMode);
+	SAssignNew(ModeTools, SControlRigEditModeTools);
 
 	FPropertyEditorModule& PropertyEditorModule = FModuleManager::GetModuleChecked<FPropertyEditorModule>("PropertyEditor");
 
@@ -119,15 +119,10 @@ void FControlRigEditModeToolkit::BuildToolPalette(FName PaletteName, class FTool
 {
 	if (PaletteName == AnimationName)
 	{
-		ModeTools->CustomizeToolBarPalette(ToolBarBuilder);
+		BuildToolBar(ToolBarBuilder);
 	}
 }
 
-void FControlRigEditModeToolkit::OnToolPaletteChanged(FName PaletteName)
-{
-
-}
-
 void FControlRigEditModeToolkit::TryToggleToolkitUI(const FName InName)
 {
 	if (IsToolkitUIActive(InName))
@@ -279,13 +274,12 @@ void FControlRigEditModeToolkit::OnControlsChanged(TConstArrayView<TWeakObjectPt
 
 FText FControlRigEditModeToolkit::GetActiveToolDisplayName() const
 {
-	return ModeTools->GetActiveToolName();
+	return FText::FromString(TEXT("Control Rig Editing"));
 }
 
 FText FControlRigEditModeToolkit::GetActiveToolMessage() const
 {
-
-	return ModeTools->GetActiveToolMessage();
+	return FText::GetEmpty();
 }
 
 TSharedRef<SDockTab> SpawnPoseTab(const FSpawnTabArgs& Args, TWeakPtr<FControlRigEditModeToolkit> SharedToolkit)
@@ -640,4 +634,145 @@ void FControlRigEditModeToolkit::SaveLayoutState() const
 	Settings->SaveConfig();
 }
 
+void FControlRigEditModeToolkit::BuildToolBar(FToolBarBuilder& InToolBarBuilder)
+{
+	//TOGGLE SELECTED RIG CONTROLS
+	InToolBarBuilder.AddToolBarButton(
+		FUIAction(
+			FExecuteAction::CreateLambda([this] { EditMode.SetOnlySelectRigControls(!EditMode.GetOnlySelectRigControls()); }),
+			FCanExecuteAction(),
+			FIsActionChecked::CreateLambda([this] { return EditMode.GetOnlySelectRigControls(); })
+		),
+		NAME_None,
+		LOCTEXT("OnlySelectControls", "Select"),
+		LOCTEXT("OnlySelectControlsTooltip", "Only Select Control Rig Controls"),
+		FSlateIcon(TEXT("ControlRigEditorStyle"), TEXT("ControlRig.OnlySelectControls")),
+		EUserInterfaceActionType::ToggleButton
+		);
+	
+	//POSES
+	InToolBarBuilder.AddToolBarButton(
+		FUIAction(
+			FExecuteAction::CreateSP(this, &FControlRigEditModeToolkit::TryToggleToolkitUI, FControlRigEditModeToolkit::PoseTabName),
+			FCanExecuteAction(),
+			FIsActionChecked::CreateSP(this, &FControlRigEditModeToolkit::IsToolkitUIActive, FControlRigEditModeToolkit::PoseTabName)
+			),
+		NAME_None,
+		LOCTEXT("Poses", "Poses"),
+		LOCTEXT("PosesTooltip", "Show Poses"),
+		FSlateIcon(TEXT("ControlRigEditorStyle"), TEXT("ControlRig.PoseTool")),
+		EUserInterfaceActionType::ToggleButton
+	);
+	
+	// Tweens
+	InToolBarBuilder.AddToolBarButton(
+		FUIAction(
+			FExecuteAction::CreateSP(this, &FControlRigEditModeToolkit::TryToggleToolkitUI, FControlRigEditModeToolkit::TweenOverlayName),
+			FCanExecuteAction(),
+			FIsActionChecked::CreateSP(this, &FControlRigEditModeToolkit::IsToolkitUIActive, FControlRigEditModeToolkit::TweenOverlayName)
+		),		
+		NAME_None,
+		LOCTEXT("Tweens", "Tweens"),
+		LOCTEXT("TweensTooltip", "Create Tweens.\n\nDefault bindings:\nU - While hidden, temporarily bring widget to your cursor\nU+LMB - While visible, indirectly move the slider\nAlt+U - Toggle visibility\nShift+U - Cycle function\nCtrl+U - Toggle Overshoot Mode"),
+		FSlateIcon(TEXT("ControlRigEditorStyle"), TEXT("ControlRig.TweenTool")),
+		EUserInterfaceActionType::ToggleButton
+	);
+	
+	InToolBarBuilder.AddToolBarButton(
+		FUIAction(
+			FExecuteAction::CreateSP(this, &FControlRigEditModeToolkit::TryToggleToolkitUI, FControlRigEditModeToolkit::ConstrainingTabName),
+			FCanExecuteAction(),
+			FIsActionChecked::CreateSP(this, &FControlRigEditModeToolkit::IsToolkitUIActive, FControlRigEditModeToolkit::ConstrainingTabName)
+		),
+		NAME_None,
+		LOCTEXT("Constrain", "Constrain"),
+		LOCTEXT("SnapperTooltip", "Tools for constraining, such as snapping, space picker, and constraints edition."),
+		FSlateIcon(TEXT("ControlRigEditorStyle"), TEXT("ControlRig.SnapperTool")),
+		EUserInterfaceActionType::ToggleButton
+	);
+	
+	// Selection Sets
+#if SELECTION_SETS_AS_TAB
+	InToolBarBuilder.AddToolBarButton(
+		FUIAction(
+			FExecuteAction::CreateRaw(this, &FControlRigEditModeToolkit::TryToggleToolkitUI, FControlRigEditModeToolkit::SelectionSetsTabName),
+			FCanExecuteAction(),
+			FIsActionChecked::CreateRaw(this, &FControlRigEditModeToolkit::IsToolkitUIActive, FControlRigEditModeToolkit::SelectionSetsTabName)
+		),
+		NAME_None,
+		LOCTEXT("SelectionSets", "Selection Sets"),
+		LOCTEXT("SelectionSetsTooltip", "Open Selection Sets"),
+		FSlateIcon(TEXT("ControlRigEditorStyle"), TEXT("ControlRig.SelectionSet")),
+		EUserInterfaceActionType::ToggleButton
+	);
+#else
+	InToolBarBuilder.AddToolBarButton(
+		FUIAction(
+			FExecuteAction::CreateRaw(this, &FControlRigEditModeToolkit::TryToggleToolkitUI, FControlRigEditModeToolkit::SelectionSetsOverlayName),
+			FCanExecuteAction(),
+			FIsActionChecked::CreateRaw(this, &FControlRigEditModeToolkit::IsToolkitUIActive, FControlRigEditModeToolkit::SelectionSetsOverlayName)
+		),
+		NAME_None,
+		LOCTEXT("SelectionSets", "Selection Sets"),
+		LOCTEXT("SelectionSetsTooltip", "Open Selection Sets"),
+		FSlateIcon(TEXT("ControlRigEditorStyle"), TEXT("ControlRig.SelectionSet")),
+		EUserInterfaceActionType::ToggleButton
+	);
+#endif
+	
+	/*  leaving this around for a bit to get animator feedback on the new menu system
+	// Motion Trail
+	ToolBarBuilder.AddToolBarButton(
+		FExecuteAction::CreateRaw(this, &FControlRigEditModeToolkit::TryInvokeToolkitUI, FControlRigEditModeToolkit::MotionTrailTabName),
+		NAME_None,
+		LOCTEXT("MotionTrails", "Trails"),
+		LOCTEXT("MotionTrailsTooltip", "Display motion trails for animated objects"),
+		FSlateIcon(TEXT("ControlRigEditorStyle"), TEXT("ControlRig.EditableMotionTrails")),
+		EUserInterfaceActionType::Button
+	);
+	*/
+	
+	// Anim Layer
+	InToolBarBuilder.AddToolBarButton(
+		FUIAction(
+			FExecuteAction::CreateSP(this, &FControlRigEditModeToolkit::TryToggleToolkitUI, FControlRigEditModeToolkit::AnimLayerTabName),
+			FCanExecuteAction(),
+			FIsActionChecked::CreateSP(this, &FControlRigEditModeToolkit::IsToolkitUIActive, FControlRigEditModeToolkit::AnimLayerTabName)
+		),
+		NAME_None,
+		LOCTEXT("Layers", "Layers"),
+		LOCTEXT("AnimLayersTooltip", "Display animation layers"),
+		FSlateIcon(TEXT("ControlRigEditorStyle"), TEXT("ControlRig.AnimLayers")),
+		EUserInterfaceActionType::ToggleButton
+	);
+	//Pivot
+	/** like motion trail leaving it in around for a bi
+	ToolBarBuilder.AddToolBarButton(
+		FUIAction(
+		FExecuteAction::CreateSP(this, &SControlRigEditModeTools::ToggleEditPivotMode),
+		FCanExecuteAction(),
+		FIsActionChecked::CreateLambda([this] {
+				if (TSharedPtr<IToolkit> Toolkit = OwningToolkit.Pin())
+				{
+					if (Toolkit.IsValid())
+					{
+						const FString ActiveToolName = Toolkit->GetToolkitHost()->GetEditorModeManager().GetInteractiveToolsContext()->ToolManager->GetActiveToolName(EToolSide::Left);
+						if (ActiveToolName == TEXT("SequencerPivotTool"))
+						{
+							return true;
+						}
+					}
+				}
+				return false;
+			})
+		),
+		NAME_None,
+		LOCTEXT("TempPivot", "Pivot"),
+		LOCTEXT("TempPivotTooltip", "Create a temporary pivot to rotate the selected Control"),
+		FSlateIcon(TEXT("ControlRigEditorStyle"), TEXT("ControlRig.TemporaryPivot")),
+		EUserInterfaceActionType::ToggleButton
+		);
+		*/
+}
+
 #undef LOCTEXT_NAMESPACE
diff --git a/Engine/Plugins/Animation/ControlRig/Source/ControlRigEditor/Private/EditMode/ControlRigEditModeToolkit.h b/Engine/Plugins/Animation/ControlRig/Source/ControlRigEditor/Private/EditMode/ControlRigEditModeToolkit.h
--- a/Engine/Plugins/Animation/ControlRig/Source/ControlRigEditor/Private/EditMode/ControlRigEditModeToolkit.h
+++ b/Engine/Plugins/Animation/ControlRig/Source/ControlRigEditor/Private/EditMode/ControlRigEditModeToolkit.h
@@ -63,7 +63,6 @@ public:
 	/** Modes Panel Header Information **/
 	virtual FText GetActiveToolDisplayName() const override;
 	virtual FText GetActiveToolMessage() const override;
-	virtual void OnToolPaletteChanged(FName PaletteName) override;
 
 	/** If the named UI is shown, hide it. If it's hidden, show it. */
 	void TryToggleToolkitUI(const FName InName);
@@ -127,4 +126,7 @@ private:
 
 	/** The tools widget */
 	TSharedPtr<SControlRigEditModeTools> ModeTools;
+
+	/** Adds the toolbar items to the toolbar builder. */
+	void BuildToolBar(FToolBarBuilder& ToolBarBuilder);
 };
diff --git a/Engine/Plugins/Animation/ControlRig/Source/ControlRigEditor/Private/EditMode/SControlRigEditModeTools.cpp b/Engine/Plugins/Animation/ControlRig/Source/ControlRigEditor/Private/EditMode/SControlRigEditModeTools.cpp
--- a/Engine/Plugins/Animation/ControlRig/Source/ControlRigEditor/Private/EditMode/SControlRigEditModeTools.cpp
+++ b/Engine/Plugins/Animation/ControlRig/Source/ControlRigEditor/Private/EditMode/SControlRigEditModeTools.cpp
@@ -44,9 +44,6 @@
 
 #define LOCTEXT_NAMESPACE "ControlRigEditModeTools"
 
-//statics to reuse in the UI
-FRigSpacePickerBakeSettings SControlRigEditModeTools::BakeSpaceSettings;
-
 void SControlRigEditModeTools::Cleanup()
 {
 	// This is required as these hold a shared pointer to THIS OBJECT and make this class not to be destroyed when the parent class releases the shared pointer of this object
@@ -65,21 +62,8 @@ void SControlRigEditModeTools::SetControlRigs(const TArrayView<TWeakObjectPtr<UC
 	UpdateOverridesDetailsView();
 }
 
-const URigHierarchy* SControlRigEditModeTools::GetHierarchy() const
-{
-	//mz todo handle multiple rigs
-	UControlRig* Rig = ControlRigs.Num() > 0 ? ControlRigs[0].Get() : nullptr;
-	if (Rig)
-	{
-		Rig->GetHierarchy();
-	}
-	return nullptr;
-}
-
-void SControlRigEditModeTools::Construct(const FArguments& InArgs, TSharedPtr<FControlRigEditModeToolkit> InOwningToolkit, FControlRigEditMode& InEditMode)
+void SControlRigEditModeTools::Construct(const FArguments& InArgs)
 {
-	bIsChangingRigHierarchy = false;
-	OwningToolkit = InOwningToolkit;
 	// initialize settings view
 	FDetailsViewArgs DetailsViewArgs;
 	{
@@ -96,25 +80,11 @@ void SControlRigEditModeTools::Construct(const FArguments& InArgs, TSharedPtr<FC
 		DetailsViewArgs.bAllowMultipleTopLevelObjects = false;
 		DetailsViewArgs.bShowScrollBar = false; // Don't need to show this, as we are putting it in a scroll box
 	}
-	
-	ModeTools = InEditMode.GetModeManager();
 
 	SettingsDetailsView = FModuleManager::GetModuleChecked<FPropertyEditorModule>("PropertyEditor").CreateDetailView(DetailsViewArgs);
 	OverridesDetailsView = FModuleManager::GetModuleChecked<FPropertyEditorModule>("PropertyEditor").CreateDetailView(DetailsViewArgs);
 	OverridesDetailsView->OnFinishedChangingProperties().AddSP(this, &SControlRigEditModeTools::OnOverrideOptionFinishedChange);
 	OverridesDetailsView->SetIsPropertyVisibleDelegate(FIsPropertyVisible::CreateSP(this, &SControlRigEditModeTools::ShouldShowOverrideProperty));
-
-	DisplaySettings.bShowBones = false;
-	DisplaySettings.bShowControls = true;
-	DisplaySettings.bShowNulls = false;
-	DisplaySettings.bShowReferences = false;
-	DisplaySettings.bShowSockets = false;
-	DisplaySettings.bHideParentsOnFilter = true;
-	DisplaySettings.bFlattenHierarchyOnFilter = true;
-	DisplaySettings.bShowIconColors = true;
-	DisplaySettings.bArrangeByModules = false;
-	DisplaySettings.bFlattenModules = false;
-	DisplaySettings.NameDisplayMode = EElementNameDisplayMode::AssetDefault;
 	
 	ChildSlot
 	[
@@ -333,193 +303,6 @@ bool SControlRigEditModeTools::IsReadOnlyPropertyOnDetailCustomization(const FPr
 		(InPropertyAndParent.ParentProperties.Num() > 0 && ShouldPropertyBeEnabled(*InPropertyAndParent.ParentProperties[0])));
 }
 
-void SControlRigEditModeTools::CustomizeToolBarPalette(FToolBarBuilder& ToolBarBuilder)
-{
-	const TSharedPtr<FControlRigEditModeToolkit> OwningToolkitPin = OwningToolkit.Pin();
-	if (!ensure(OwningToolkitPin))
-	{
-		return;
-	}
-	
-	//TOGGLE SELECTED RIG CONTROLS
-	ToolBarBuilder.AddToolBarButton(
-		FUIAction(
-			FExecuteAction::CreateLambda([this] {
-				FControlRigEditMode* ControlRigEditMode = static_cast<FControlRigEditMode*>(ModeTools->GetActiveMode(FControlRigEditMode::ModeName));
-				if (ControlRigEditMode)
-				{
-					ControlRigEditMode->SetOnlySelectRigControls(!ControlRigEditMode->GetOnlySelectRigControls());
-				}
-			}),
-			FCanExecuteAction(),
-			FIsActionChecked::CreateLambda([this] {
-				FControlRigEditMode* ControlRigEditMode = static_cast<FControlRigEditMode*>(ModeTools->GetActiveMode(FControlRigEditMode::ModeName));
-				return ControlRigEditMode && ControlRigEditMode->GetOnlySelectRigControls();
-			})
-
-		),
-		NAME_None,
-		LOCTEXT("OnlySelectControls", "Select"),
-		LOCTEXT("OnlySelectControlsTooltip", "Only Select Control Rig Controls"),
-		FSlateIcon(TEXT("ControlRigEditorStyle"), TEXT("ControlRig.OnlySelectControls")),
-		EUserInterfaceActionType::ToggleButton
-		);
-
-	ToolBarBuilder.AddSeparator();
-
-	//POSES
-	ToolBarBuilder.AddToolBarButton(
-		FUIAction(
-			FExecuteAction::CreateSP(OwningToolkitPin.Get(), &FControlRigEditModeToolkit::TryToggleToolkitUI, FControlRigEditModeToolkit::PoseTabName),
-			FCanExecuteAction(),
-			FIsActionChecked::CreateSP(OwningToolkitPin.Get(), &FControlRigEditModeToolkit::IsToolkitUIActive, FControlRigEditModeToolkit::PoseTabName)
-			),
-		NAME_None,
-		LOCTEXT("Poses", "Poses"),
-		LOCTEXT("PosesTooltip", "Show Poses"),
-		FSlateIcon(TEXT("ControlRigEditorStyle"), TEXT("ControlRig.PoseTool")),
-		EUserInterfaceActionType::ToggleButton
-	);
-	ToolBarBuilder.AddSeparator();
-
-	// Tweens
-	ToolBarBuilder.AddToolBarButton(
-		FUIAction(
-			FExecuteAction::CreateSP(OwningToolkitPin.Get(), &FControlRigEditModeToolkit::TryToggleToolkitUI, FControlRigEditModeToolkit::TweenOverlayName),
-			FCanExecuteAction(),
-			FIsActionChecked::CreateSP(OwningToolkitPin.Get(), &FControlRigEditModeToolkit::IsToolkitUIActive, FControlRigEditModeToolkit::TweenOverlayName)
-		),		
-		NAME_None,
-		LOCTEXT("Tweens", "Tweens"),
-		LOCTEXT("TweensTooltip", "Create Tweens.\n\nDefault bindings:\nU - While hidden, temporarily bring widget to your cursor\nU+LMB - While visible, indirectly move the slider\nAlt+U - Toggle visibility\nShift+U - Cycle function\nCtrl+U - Toggle Overshoot Mode"),
-		FSlateIcon(TEXT("ControlRigEditorStyle"), TEXT("ControlRig.TweenTool")),
-		EUserInterfaceActionType::ToggleButton
-	);
-	
-	ToolBarBuilder.AddToolBarButton(
-		FUIAction(
-			FExecuteAction::CreateSP(OwningToolkitPin.Get(), &FControlRigEditModeToolkit::TryToggleToolkitUI, FControlRigEditModeToolkit::ConstrainingTabName),
-			FCanExecuteAction(),
-			FIsActionChecked::CreateSP(OwningToolkitPin.Get(), &FControlRigEditModeToolkit::IsToolkitUIActive, FControlRigEditModeToolkit::ConstrainingTabName)
-		),
-		NAME_None,
-		LOCTEXT("Constrain", "Constrain"),
-		LOCTEXT("SnapperTooltip", "Tools for constraining, such as snapping, space picker, and constraints edition."),
-		FSlateIcon(TEXT("ControlRigEditorStyle"), TEXT("ControlRig.SnapperTool")),
-		EUserInterfaceActionType::ToggleButton
-	);
-
-	// Selection Sets
-#if SELECTION_SETS_AS_TAB
-
-	ToolBarBuilder.AddToolBarButton(
-		FUIAction(
-			FExecuteAction::CreateRaw(OwningToolkit.Pin().Get(), &FControlRigEditModeToolkit::TryToggleToolkitUI, FControlRigEditModeToolkit::SelectionSetsTabName),
-			FCanExecuteAction(),
-			FIsActionChecked::CreateRaw(OwningToolkit.Pin().Get(), &FControlRigEditModeToolkit::IsToolkitUIActive, FControlRigEditModeToolkit::SelectionSetsTabName)
-		),
-		NAME_None,
-		LOCTEXT("SelectionSets", "Selection Sets"),
-		LOCTEXT("SelectionSetsTooltip", "Open Selection Sets"),
-		FSlateIcon(TEXT("ControlRigEditorStyle"), TEXT("ControlRig.SelectionSet")),
-		EUserInterfaceActionType::ToggleButton
-	);
-#else
-
-	ToolBarBuilder.AddToolBarButton(
-		FUIAction(
-			FExecuteAction::CreateRaw(OwningToolkit.Pin().Get(), &FControlRigEditModeToolkit::TryToggleToolkitUI, FControlRigEditModeToolkit::SelectionSetsOverlayName),
-			FCanExecuteAction(),
-			FIsActionChecked::CreateRaw(OwningToolkit.Pin().Get(), &FControlRigEditModeToolkit::IsToolkitUIActive, FControlRigEditModeToolkit::SelectionSetsOverlayName)
-		),
-		NAME_None,
-		LOCTEXT("SelectionSets", "Selection Sets"),
-		LOCTEXT("SelectionSetsTooltip", "Open Selection Sets"),
-		FSlateIcon(TEXT("ControlRigEditorStyle"), TEXT("ControlRig.SelectionSet")),
-		EUserInterfaceActionType::ToggleButton
-	);
-#endif
-	
-
-	/*  leaving this around for a bit to get animator feedback on the new menu system
-	// Motion Trail
-	ToolBarBuilder.AddToolBarButton(
-		FExecuteAction::CreateRaw(OwningToolkit.Pin().Get(), &FControlRigEditModeToolkit::TryInvokeToolkitUI, FControlRigEditModeToolkit::MotionTrailTabName),
-		NAME_None,
-		LOCTEXT("MotionTrails", "Trails"),
-		LOCTEXT("MotionTrailsTooltip", "Display motion trails for animated objects"),
-		FSlateIcon(TEXT("ControlRigEditorStyle"), TEXT("ControlRig.EditableMotionTrails")),
-		EUserInterfaceActionType::Button
-	);
-	*/
-
-	// Anim Layer
-	ToolBarBuilder.AddToolBarButton(
-		FUIAction(
-			FExecuteAction::CreateSP(OwningToolkitPin.Get(), &FControlRigEditModeToolkit::TryToggleToolkitUI, FControlRigEditModeToolkit::AnimLayerTabName),
-			FCanExecuteAction(),
-			FIsActionChecked::CreateSP(OwningToolkitPin.Get(), &FControlRigEditModeToolkit::IsToolkitUIActive, FControlRigEditModeToolkit::AnimLayerTabName)
-		),
-		NAME_None,
-		LOCTEXT("Layers", "Layers"),
-		LOCTEXT("AnimLayersTooltip", "Display animation layers"),
-		FSlateIcon(TEXT("ControlRigEditorStyle"), TEXT("ControlRig.AnimLayers")),
-		EUserInterfaceActionType::ToggleButton
-	);
-
-	//Pivot
-	/** like motion trail leaving it in around for a bi
-	ToolBarBuilder.AddToolBarButton(
-		FUIAction(
-		FExecuteAction::CreateSP(this, &SControlRigEditModeTools::ToggleEditPivotMode),
-		FCanExecuteAction(),
-		FIsActionChecked::CreateLambda([this] {
-				if (TSharedPtr<IToolkit> Toolkit = OwningToolkit.Pin())
-				{
-					if (Toolkit.IsValid())
-					{
-						const FString ActiveToolName = Toolkit->GetToolkitHost()->GetEditorModeManager().GetInteractiveToolsContext()->ToolManager->GetActiveToolName(EToolSide::Left);
-						if (ActiveToolName == TEXT("SequencerPivotTool"))
-						{
-							return true;
-						}
-					}
-				}
-				return false;
-
-			})
-		),
-		NAME_None,
-		LOCTEXT("TempPivot", "Pivot"),
-		LOCTEXT("TempPivotTooltip", "Create a temporary pivot to rotate the selected Control"),
-		FSlateIcon(TEXT("ControlRigEditorStyle"), TEXT("ControlRig.TemporaryPivot")),
-		EUserInterfaceActionType::ToggleButton
-		);
-		*/
-}
-
-void SControlRigEditModeTools::ToggleEditPivotMode()
-{
-	FEditorModeID ModeID = TEXT("SequencerToolsEditMode");
-	if (const TSharedPtr<IToolkit> Toolkit = OwningToolkit.Pin())
-	{
-		if (Toolkit.IsValid())
-		{
-			const FEditorModeTools& Tools =Toolkit->GetToolkitHost()->GetEditorModeManager();
-			const FString ActiveToolName = Tools.GetInteractiveToolsContext()->ToolManager->GetActiveToolName(EToolSide::Left);
-			if (ActiveToolName == TEXT("SequencerPivotTool"))
-			{
-				Tools.GetInteractiveToolsContext()->ToolManager->DeactivateTool(EToolSide::Left, EToolShutdownType::Completed);
-			}
-			else
-			{
-				Tools.GetInteractiveToolsContext()->ToolManager->SelectActiveToolType(EToolSide::Left, TEXT("SequencerPivotTool"));
-				Tools.GetInteractiveToolsContext()->ToolManager->ActivateTool(EToolSide::Left);
-			}
-		}
-	}
-}
-
 
 /* MZ TODO
 void SControlRigEditModeTools::MakeSelectionSetDialog()
@@ -560,15 +343,6 @@ void SControlRigEditModeTools::MakeSelectionSetDialog()
 	}
 }
 */
-FText SControlRigEditModeTools::GetActiveToolName() const
-{
-	return FText::FromString(TEXT("Control Rig Editing"));
-}
-
-FText SControlRigEditModeTools::GetActiveToolMessage() const
-{
-	return  FText();
-}
 
 EVisibility SControlRigEditModeTools::GetOverridesExpanderVisibility() const
 {
diff --git a/Engine/Plugins/Animation/ControlRig/Source/ControlRigEditor/Private/EditMode/SControlRigEditModeTools.h b/Engine/Plugins/Animation/ControlRig/Source/ControlRigEditor/Private/EditMode/SControlRigEditModeTools.h
--- a/Engine/Plugins/Animation/ControlRig/Source/ControlRigEditor/Private/EditMode/SControlRigEditModeTools.h
+++ b/Engine/Plugins/Animation/ControlRig/Source/ControlRigEditor/Private/EditMode/SControlRigEditModeTools.h
@@ -29,7 +29,7 @@ public:
 	SLATE_BEGIN_ARGS(SControlRigEditModeTools) {}
 	SLATE_END_ARGS();
 
-	void Construct(const FArguments& InArgs, TSharedPtr<FControlRigEditModeToolkit> InOwningToolkit, FControlRigEditMode& InEditMode);
+	void Construct(const FArguments& InArgs);
 	void Cleanup();
 
 	/** Set the objects to be displayed in the details panel */
@@ -41,17 +41,12 @@ public:
 	/** Set The Control Rig we are using */
 	void SetControlRigs(const TArrayView<TWeakObjectPtr<UControlRig>>& InControlRigs);
 	
-	/** Returns the hierarchy currently being used */
-	const URigHierarchy* GetHierarchy() const;
 	// IDetailKeyframeHandler interface
 	virtual bool IsPropertyKeyable(const UClass* InObjectClass, const class IPropertyHandle& PropertyHandle) const override;
 	virtual bool IsPropertyKeyingEnabled() const override;
 	virtual void OnKeyPropertyClicked(const IPropertyHandle& KeyedPropertyHandle) override;
 	virtual bool IsPropertyAnimated(const class IPropertyHandle& PropertyHandle, UObject *ParentObject) const override;
 
-	//reuse settings for space and constraint baking baking
-	static FRigSpacePickerBakeSettings BakeSpaceSettings;
-
 private:
 
 	/** Sequencer we are currently bound to */
@@ -71,26 +66,5 @@ private:
 	void OnOverrideOptionFinishedChange(const FPropertyChangedEvent& PropertyChangedEvent);
 	bool ShouldShowOverrideProperty(const FPropertyAndParent& InPropertyAndParent) const;
 	void UpdateOverridesDetailsView();
-	
-private:
-	/** Toolbar functions and windows*/
-	void ToggleEditPivotMode();
-
-	//TODO may put back void MakeSelectionSetDialog();
-	//TWeakPtr<SWindow> SelectionSetWindow;
-
-	FEditorModeTools* ModeTools = nullptr;
-	FRigTreeDisplaySettings DisplaySettings;
-	const FRigTreeDisplaySettings& GetDisplaySettings() const { return DisplaySettings; }
-	bool bIsChangingRigHierarchy = false;
-
-	// The toolkit that created this UI
-	TWeakPtr<FControlRigEditModeToolkit> OwningToolkit;
-
-public:
-	/** Modes Panel Header Information **/
-	void CustomizeToolBarPalette(FToolBarBuilder& ToolBarBuilder);
-	FText GetActiveToolName() const;
-	FText GetActiveToolMessage() const;
 };
 
