fixed potential null dereference

#rnx

diff --git a/Engine/Plugins/Animation/ControlRig/Source/ControlRigEditor/Private/Sequencer/ControlRigParameterTrackEditor.cpp b/Engine/Plugins/Animation/ControlRig/Source/ControlRigEditor/Private/Sequencer/ControlRigParameterTrackEditor.cpp
--- a/Engine/Plugins/Animation/ControlRig/Source/ControlRigEditor/Private/Sequencer/ControlRigParameterTrackEditor.cpp
+++ b/Engine/Plugins/Animation/ControlRig/Source/ControlRigEditor/Private/Sequencer/ControlRigParameterTrackEditor.cpp
@@ -2704,16 +2704,19 @@ void FControlRigParameterTrackEditor::OnChannelChanged(const FMovieSceneChannelM
 void FControlRigParameterTrackEditor::AddConstraintToSequencer(const TSharedPtr<ISequencer>& InSequencer, UTickableTransformConstraint* InConstraint,
 	const UConstraintCreationOptions* InSequencerOptions)
 {
-	TGuardValue<bool> DisableTrackCreation(bAutoGenerateControlRigTrack, false);
-	if (InSequencerOptions)
+	if (InConstraint)
 	{
-		InConstraint->Evaluate();
-	}
+		TGuardValue<bool> DisableTrackCreation(bAutoGenerateControlRigTrack, false);
+		if (InSequencerOptions)
+		{
+			InConstraint->Evaluate();
+		}
 	
-	static const FSequencerCreationOptions DefaultOptions;
-	const FSequencerCreationOptions& Options = InSequencerOptions ? InSequencerOptions->SequencerOptions : DefaultOptions;
+		static const FSequencerCreationOptions DefaultOptions;
+		const FSequencerCreationOptions& Options = InSequencerOptions ? InSequencerOptions->SequencerOptions : DefaultOptions;
 	
-	FMovieSceneConstraintChannelHelper::AddConstraintToSequencer(InSequencer, InConstraint, Options);
+		FMovieSceneConstraintChannelHelper::AddConstraintToSequencer(InSequencer, InConstraint, Options);
+	}
 }
 
 void FControlRigParameterTrackEditor::AddTrackForComponent(USceneComponent* InComponent, FGuid InBinding) 
