Control Rig: Fix space panel list not updating when you press ESC

diff --git a/Engine/Plugins/Animation/ControlRig/Source/ControlRigEditor/Private/Editor/SRigSpacePickerWidget.cpp b/Engine/Plugins/Animation/ControlRig/Source/ControlRigEditor/Private/Editor/SRigSpacePickerWidget.cpp
--- a/Engine/Plugins/Animation/ControlRig/Source/ControlRigEditor/Private/Editor/SRigSpacePickerWidget.cpp
+++ b/Engine/Plugins/Animation/ControlRig/Source/ControlRigEditor/Private/Editor/SRigSpacePickerWidget.cpp
@@ -989,11 +989,14 @@ TArray<FRigElementKeyWithLabel> SRigSpacePickerWidget::GetCurrentParents_Private
 
 void SRigSpacePickerWidget::RepopulateItemSpaces()
 {
-	if(!ItemSpacesListBox.IsValid())
+	if (!ItemSpacesListBox.IsValid())
 	{
 		return;
 	}
-	if(!Hierarchy.IsValid())
+	
+	// If nothing was selected, Hierarchy will be nullptr. So it's important to clear all rows before returning.
+	ClearListBox(ItemSpacesListBox);
+	if (!Hierarchy.IsValid())
 	{
 		return;
 	}
@@ -1073,34 +1076,6 @@ void SRigSpacePickerWidget::RepopulateItemSpaces()
 		}
 	}
 
-	/*
-	struct FKeySortPredicate
-	{
-		bool operator()(const FRigElementKey& A, const FRigElementKey& B) const
-		{
-			static TMap<ERigElementType, int32> TypeOrder;
-			if(TypeOrder.IsEmpty())
-			{
-				TypeOrder.Add(ERigElementType::Control, 0);
-				TypeOrder.Add(ERigElementType::Reference, 1);
-				TypeOrder.Add(ERigElementType::Null, 2);
-				TypeOrder.Add(ERigElementType::Bone, 3);
-				TypeOrder.Add(ERigElementType::RigidBody, 4);
-			}
-
-			const int32 TypeIndexA = TypeOrder.FindChecked(A.Type);
-			const int32 TypeIndexB = TypeOrder.FindChecked(B.Type);
-			if(TypeIndexA != TypeIndexB)
-			{
-				return TypeIndexA < TypeIndexB;
-			}
-
-			return A.Name.Compare(B.Name) < 0; 
-		}
-	};
-	SpacesFromDelegate.Sort(FKeySortPredicate());
-	*/
-
 	TArray<FRigElementKeyWithLabel> CombinedSpaces = FavoriteKeys;
 	for(const FRigElementKeyWithLabel& Space : SpacesFromDelegate)
 	{
@@ -1114,10 +1089,8 @@ void SRigSpacePickerWidget::RepopulateItemSpaces()
 	{
 		return;
 	}
-
-	ClearListBox(ItemSpacesListBox);
-
-	for(const FRigElementKeyWithLabel& Space : CombinedSpaces)
+	
+	for (const FRigElementKeyWithLabel& Space : CombinedSpaces)
 	{
 		TPair<const FSlateBrush*, FSlateColor> IconAndColor = SRigHierarchyItem::GetBrushForElementType(StrongHierarchy, Space.Key);
 		
@@ -1138,6 +1111,7 @@ void SRigSpacePickerWidget::RepopulateItemSpaces()
 void SRigSpacePickerWidget::ClearListBox(TSharedPtr<SVerticalBox> InListBox)
 {
 	InListBox->ClearChildren();
+	CurrentSpaceKeys.Empty();
 }
 
 void SRigSpacePickerWidget::UpdateActiveSpaces()
