Control Rig: "ControlRig.Overrides 1" no longer shows any overrides in control rig's menu

diff --git a/Engine/Plugins/Animation/ControlRig/Source/ControlRigEditor/Private/EditMode/ControlRigEditModeToolkit.cpp b/Engine/Plugins/Animation/ControlRig/Source/ControlRigEditor/Private/EditMode/ControlRigEditModeToolkit.cpp
--- a/Engine/Plugins/Animation/ControlRig/Source/ControlRigEditor/Private/EditMode/ControlRigEditModeToolkit.cpp
+++ b/Engine/Plugins/Animation/ControlRig/Source/ControlRigEditor/Private/EditMode/ControlRigEditModeToolkit.cpp
@@ -71,6 +71,11 @@ const FName FControlRigEditModeToolkit::DetailsTabName = FName(TEXT("ControlRigD
 TSharedPtr<UE::ControlRigEditor::SAnimDetailsView> FControlRigEditModeToolkit::Details = nullptr;
 TSharedPtr<SControlRigOutliner> FControlRigEditModeToolkit::Outliner = nullptr;
 
+FControlRigEditModeToolkit::FControlRigEditModeToolkit(FControlRigEditMode& InEditMode)
+	: EditMode(InEditMode)
+	, SelectionViewModel(MakeShared<UE::ControlRigEditor::FRigSelectionViewModel>())
+{}
+
 FControlRigEditModeToolkit::~FControlRigEditModeToolkit()
 {
 	if (ModeTools)
@@ -81,7 +86,7 @@ FControlRigEditModeToolkit::~FControlRigEditModeToolkit()
 
 void FControlRigEditModeToolkit::Init(const TSharedPtr<IToolkitHost>& InitToolkitHost)
 {
-	SAssignNew(ModeTools, SControlRigEditModeTools);
+	SAssignNew(ModeTools, SControlRigEditModeTools, SelectionViewModel);
 
 	FPropertyEditorModule& PropertyEditorModule = FModuleManager::GetModuleChecked<FPropertyEditorModule>("PropertyEditor");
 
@@ -266,10 +271,7 @@ bool FControlRigEditModeToolkit::IsToolkitUIActive(const FName InName) const
 
 void FControlRigEditModeToolkit::OnControlsChanged(TConstArrayView<TWeakObjectPtr<UControlRig>> InControlRigs) const
 {
-	if (SelectionViewModel)
-	{
-		SelectionViewModel->RefreshSubscriptions(InControlRigs);
-	}
+	SelectionViewModel->SetControls(InControlRigs);
 }
 
 FText FControlRigEditModeToolkit::GetActiveToolDisplayName() const
@@ -460,11 +462,9 @@ void FControlRigEditModeToolkit::InvokeUI()
 	// TweenOverlayManager will automatically restore its last visibility state upon construction.
 	TweenOverlayManager = MakeUnique<UE::ControlRigEditor::FTweenOverlayManager>(GetToolkitHost(), GetToolkitCommands(), SharedThis(&EditMode));
 
-	SelectionViewModel = MakeShared<UE::ControlRigEditor::FRigSelectionViewModel>();
-
 	checkf(WorkspaceMenuCategory, TEXT("WorkspaceMenuCategory should have been set in RequestModeUITabs by now."))
 	ConstrainToolsManager = MakeUnique<UE::ControlRigEditor::FConstrainToolsManager>(
-		ModeUILayerPtr->GetTabManager().ToSharedRef(), GetWorkspaceMenuCategory(), EditMode, SelectionViewModel.ToSharedRef()
+		ModeUILayerPtr->GetTabManager().ToSharedRef(), GetWorkspaceMenuCategory(), EditMode, SelectionViewModel
 		);
 	
 	const FControlRigUIRestoreStates& RestoreStates = UControlRigEditModeSettings::Get()->LastUIStates;
@@ -562,8 +562,6 @@ void FControlRigEditModeToolkit::UnregisterAndRemoveFloatingTabs()
 	{
 		ConstrainToolsManager.Reset();
 	}
-	
-	SelectionViewModel.Reset();
 
 #if SELECTION_SETS_AS_TAB == 0
 	if (SelectionSetsOverlayManager)
diff --git a/Engine/Plugins/Animation/ControlRig/Source/ControlRigEditor/Private/EditMode/ControlRigEditModeToolkit.h b/Engine/Plugins/Animation/ControlRig/Source/ControlRigEditor/Private/EditMode/ControlRigEditModeToolkit.h
--- a/Engine/Plugins/Animation/ControlRig/Source/ControlRigEditor/Private/EditMode/ControlRigEditModeToolkit.h
+++ b/Engine/Plugins/Animation/ControlRig/Source/ControlRigEditor/Private/EditMode/ControlRigEditModeToolkit.h
@@ -36,13 +36,8 @@ class FControlRigEditModeToolkit : public FModeToolkit
 public:
 	friend UE::ControlRigEditor::SControlRigTweenWidget;
 
-	FControlRigEditModeToolkit(FControlRigEditMode& InEditMode)
-		: EditMode(InEditMode)
-	{
-		
-	}
-
-	~FControlRigEditModeToolkit();
+	explicit FControlRigEditModeToolkit(FControlRigEditMode& InEditMode);
+	virtual ~FControlRigEditModeToolkit() override;
 
 	/** IToolkit interface */
 	virtual FName GetToolkitFName() const override { return FName("AnimationMode"); }
@@ -115,8 +110,8 @@ private:
 	/** The edit mode we are bound to */
 	FControlRigEditMode& EditMode;
 
-	/** View model for selection changing. Used by UI that depends on selected control rig. */
-	TSharedPtr<UE::ControlRigEditor::FRigSelectionViewModel> SelectionViewModel;
+	/** View model for the controls rigs that the user is working on. Used by UI. */
+	const TSharedRef<UE::ControlRigEditor::FRigSelectionViewModel> SelectionViewModel;
 	
 	/** Manages the tween viewport UI overlay. */
 	TUniquePtr<UE::ControlRigEditor::FTweenOverlayManager> TweenOverlayManager;
diff --git a/Engine/Plugins/Animation/ControlRig/Source/ControlRigEditor/Private/EditMode/Models/RigSelectionViewModel.cpp b/Engine/Plugins/Animation/ControlRig/Source/ControlRigEditor/Private/EditMode/Models/RigSelectionViewModel.cpp
--- a/Engine/Plugins/Animation/ControlRig/Source/ControlRigEditor/Private/EditMode/Models/RigSelectionViewModel.cpp
+++ b/Engine/Plugins/Animation/ControlRig/Source/ControlRigEditor/Private/EditMode/Models/RigSelectionViewModel.cpp
@@ -11,7 +11,7 @@ FRigSelectionViewModel::~FRigSelectionViewModel()
 	CleanupSubscriptions();
 }
 
-void FRigSelectionViewModel::RefreshSubscriptions(TConstArrayView<TWeakObjectPtr<UControlRig>> InControlRigs)
+void FRigSelectionViewModel::SetControls(TConstArrayView<TWeakObjectPtr<UControlRig>> InControlRigs)
 {
 	CleanupSubscriptions();
 
@@ -19,15 +19,17 @@ void FRigSelectionViewModel::RefreshSubscriptions(TConstArrayView<TWeakObjectPtr
 	{
 		if (UControlRig* ControlRig = WeakControlRig.Get())
 		{
-			SubscribedToRigs.AddUnique(ControlRig);
+			RuntimeRigs.AddUnique(ControlRig);
 			ControlRig->ControlSelected().AddRaw(this, &FRigSelectionViewModel::BroadcastOnRigElementSelected);
 		}
 	}
+
+	OnControlsChangedDelegate.Broadcast();
 }
 
 void FRigSelectionViewModel::CleanupSubscriptions()
 {
-	for (const TWeakObjectPtr<UControlRig>& WeakControlRig : SubscribedToRigs)
+	for (const TWeakObjectPtr<UControlRig>& WeakControlRig : RuntimeRigs)
 	{
 		if (UControlRig* ControlRig = WeakControlRig.Get())
 		{
@@ -35,7 +37,7 @@ void FRigSelectionViewModel::CleanupSubscriptions()
 		}
 	}
 		
-	SubscribedToRigs.Reset();
+	RuntimeRigs.Reset();
 }
 
 void FRigSelectionViewModel::BroadcastOnRigElementSelected(
diff --git a/Engine/Plugins/Animation/ControlRig/Source/ControlRigEditor/Private/EditMode/Models/RigSelectionViewModel.h b/Engine/Plugins/Animation/ControlRig/Source/ControlRigEditor/Private/EditMode/Models/RigSelectionViewModel.h
--- a/Engine/Plugins/Animation/ControlRig/Source/ControlRigEditor/Private/EditMode/Models/RigSelectionViewModel.h
+++ b/Engine/Plugins/Animation/ControlRig/Source/ControlRigEditor/Private/EditMode/Models/RigSelectionViewModel.h
@@ -7,8 +7,8 @@
 
 namespace UE::ControlRigEditor
 {
-
 DECLARE_MULTICAST_DELEGATE_ThreeParams(FOnControlSelected, UControlRig* /*Subject*/, FRigControlElement* /*ControlElement*/, bool /*bSelected*/);
+using FOnControlsChanged = FSimpleMulticastDelegate;
 	
 /**
  * Responsible for invoking an event when control rig selection changes.
@@ -21,19 +21,27 @@ public:
 	~FRigSelectionViewModel();
 
 	/** Called when the global array of control rigs that the user is working on changes. */
-	void RefreshSubscriptions(TConstArrayView<TWeakObjectPtr<UControlRig>> InControlRigs);
+	void SetControls(TConstArrayView<TWeakObjectPtr<UControlRig>> InControlRigs);
+
+	/** The control rigs that the user is working on. */
+	TConstArrayView<TWeakObjectPtr<UControlRig>> GetControlRigs() const { return RuntimeRigs; }
 	
 	/** Invoked when a control rig is selected. */
 	FOnControlSelected& OnControlSelected() { return OnControlSelectedDelegate; }
+	/** Invoked when the controls rigs change. */
+	FOnControlsChanged& OnControlsChanged() { return OnControlsChangedDelegate; }
 
 private:
 
-	/** The control rigs we subscribed to in the past. */
-	TArray<TWeakObjectPtr<UControlRig>> SubscribedToRigs;
+	/** The control rigs that the user is working on. */
+	TArray<TWeakObjectPtr<UControlRig>> RuntimeRigs;
 
 	/** Invoked when a control rig is selected. */
 	FOnControlSelected OnControlSelectedDelegate;
 
+	/** Invoked when the controls rigs change. */
+	FOnControlsChanged OnControlsChangedDelegate;
+
 	/** Unbinds all subscriptions. */
 	void CleanupSubscriptions();
 
diff --git a/Engine/Plugins/Animation/ControlRig/Source/ControlRigEditor/Private/EditMode/SControlRigEditModeTools.cpp b/Engine/Plugins/Animation/ControlRig/Source/ControlRigEditor/Private/EditMode/SControlRigEditModeTools.cpp
--- a/Engine/Plugins/Animation/ControlRig/Source/ControlRigEditor/Private/EditMode/SControlRigEditModeTools.cpp
+++ b/Engine/Plugins/Animation/ControlRig/Source/ControlRigEditor/Private/EditMode/SControlRigEditModeTools.cpp
@@ -41,6 +41,7 @@
 #include "Transform/TransformConstraint.h"
 #include "EditMode/ControlRigEditModeToolkit.h"
 #include "Editor/Constraints/SConstraintsWidget.h"
+#include "Models/RigSelectionViewModel.h"
 
 #define LOCTEXT_NAMESPACE "ControlRigEditModeTools"
 
@@ -56,14 +57,11 @@ void SControlRigEditModeTools::Cleanup()
 	}
 }
 
-void SControlRigEditModeTools::SetControlRigs(const TArrayView<TWeakObjectPtr<UControlRig>>& InControlRigs)
-{
-	ControlRigs = InControlRigs;
-	UpdateOverridesDetailsView();
-}
-
-void SControlRigEditModeTools::Construct(const FArguments& InArgs)
+void SControlRigEditModeTools::Construct(const FArguments& InArgs, const TSharedRef<UE::ControlRigEditor::FRigSelectionViewModel>& InRigViewModel)
 {
+	RigViewModel = InRigViewModel;
+	RigViewModel->OnControlsChanged().AddSP(this, &SControlRigEditModeTools::UpdateOverridesDetailsView);
+	
 	// initialize settings view
 	FDetailsViewArgs DetailsViewArgs;
 	{
@@ -132,6 +130,11 @@ void SControlRigEditModeTools::Construct(const FArguments& InArgs)
 	];
 }
 
+SControlRigEditModeTools::~SControlRigEditModeTools()
+{
+	RigViewModel->OnControlsChanged().RemoveAll(this);
+}
+
 void SControlRigEditModeTools::SetSettingsDetailsObject(const TWeakObjectPtr<>& InObject)
 {
 	if (SettingsDetailsView)
@@ -395,7 +398,7 @@ void SControlRigEditModeTools::UpdateOverridesDetailsView()
 	if(WeakSequencer.IsValid())
 	{
 		ISequencer* Sequencer = WeakSequencer.Pin().Get();
-		for(const TWeakObjectPtr<UControlRig>& WeakControlRig : ControlRigs)
+		for(const TWeakObjectPtr<UControlRig>& WeakControlRig : RigViewModel->GetControlRigs())
 		{
 			if(const UControlRig* ControlRig = WeakControlRig.Get())
 			{
diff --git a/Engine/Plugins/Animation/ControlRig/Source/ControlRigEditor/Private/EditMode/SControlRigEditModeTools.h b/Engine/Plugins/Animation/ControlRig/Source/ControlRigEditor/Private/EditMode/SControlRigEditModeTools.h
--- a/Engine/Plugins/Animation/ControlRig/Source/ControlRigEditor/Private/EditMode/SControlRigEditModeTools.h
+++ b/Engine/Plugins/Animation/ControlRig/Source/ControlRigEditor/Private/EditMode/SControlRigEditModeTools.h
@@ -5,31 +5,28 @@
 #include "CoreMinimal.h"
 #include "Widgets/SCompoundWidget.h"
 #include "IDetailKeyframeHandler.h"
+#include "EditMode/ControlRigEditModeUtil.h"
 #include "RigVMModel/RigVMGraph.h"
 #include "Editor/SRigHierarchyTreeView.h"
 #include "Editor/SRigSpacePickerWidget.h"
 #include "Widgets/DeclarativeSyntaxSupport.h"
 
-class SConstraintsEditionWidget;
-class FControlRigEditMode;
+
 class IDetailsView;
 class ISequencer;
-class SControlPicker;
-class SExpandableArea;
 class SRigHierarchyTreeView;
-class UControlRig;
-class URigHierarchy;
-class FToolBarBuilder;
-class FEditorModeTools;
-class FControlRigEditModeToolkit;
+namespace UE::ControlRigEditor { class FRigSelectionViewModel; }
 
+/** Shows UI for adding override assets to control rig when the ControlRig.Overrides cvar is enabled. */
 class SControlRigEditModeTools : public SCompoundWidget, public IDetailKeyframeHandler
 {
 public:
 	SLATE_BEGIN_ARGS(SControlRigEditModeTools) {}
 	SLATE_END_ARGS();
 
-	void Construct(const FArguments& InArgs);
+	void Construct(const FArguments& InArgs, const TSharedRef<UE::ControlRigEditor::FRigSelectionViewModel>& InRigViewModel);
+	
+	virtual ~SControlRigEditModeTools() override;
 	void Cleanup();
 
 	/** Set the objects to be displayed in the details panel */
@@ -37,9 +34,6 @@ public:
 
 	/** Set the sequencer we are bound to */
 	void SetSequencer(TWeakPtr<ISequencer> InSequencer);
-
-	/** Set The Control Rig we are using */
-	void SetControlRigs(const TArrayView<TWeakObjectPtr<UControlRig>>& InControlRigs);
 	
 	// IDetailKeyframeHandler interface
 	virtual bool IsPropertyKeyable(const UClass* InObjectClass, const class IPropertyHandle& PropertyHandle) const override;
@@ -49,12 +43,12 @@ public:
 
 private:
 
+	/** Knows about the control rigs that are selected. */
+	TSharedPtr<UE::ControlRigEditor::FRigSelectionViewModel> RigViewModel;
+
 	/** Sequencer we are currently bound to */
 	TWeakPtr<ISequencer> WeakSequencer;
 	TSharedPtr<IDetailsView> SettingsDetailsView;
-	
-	/** Storage for control rigs */
-	TArray<TWeakObjectPtr<UControlRig>> ControlRigs;
 
 	TSharedPtr<IDetailsView> OverridesDetailsView;
 
