Control Rig: Fix constraints created not refreshing until you scrub if the constraint was created using actor picking instigated by the Create button in the sub-menu.

While it was undefined behaviour, effectively what happened was that the arly DismissAllMenus() call called ~SConstraintMenuEntry which caused OnConstraintCreated to become destroyed, too, which makes it unbound. That was no problem for the Args.Parent.Object.IsValid() case because the FCreationArgs (which I removed) had copied the callback already.
However, CreateSelectionPicker would also try to copy it and it would be unset. So when actor picking finished, it wouldn't invoke the callback for updating the UI.

diff --git a/Engine/Plugins/Animation/ControlRig/Source/ControlRigEditor/Private/Editor/Constraints/SConstraintsWidget.cpp b/Engine/Plugins/Animation/ControlRig/Source/ControlRigEditor/Private/Editor/Constraints/SConstraintsWidget.cpp
--- a/Engine/Plugins/Animation/ControlRig/Source/ControlRigEditor/Private/Editor/Constraints/SConstraintsWidget.cpp
+++ b/Engine/Plugins/Animation/ControlRig/Source/ControlRigEditor/Private/Editor/Constraints/SConstraintsWidget.cpp
@@ -14,6 +14,7 @@
 #include "LevelEditor.h"
 #include "SSocketChooser.h"
 #include "Modules/ModuleManager.h"
+#include "Misc/ScopeExit.h"
 #include "Widgets/Input/SButton.h"
 #include "Framework/MultiBox/MultiBoxBuilder.h"
 #include "PropertyEditorModule.h"
@@ -578,35 +579,25 @@ void SConstraintMenuEntry::NotifyPostChange(const FPropertyChangedEvent& InPrope
 
 FReply SConstraintMenuEntry::CreateFromSelectionOrPicker(const bool bUseDefault) const
 {
-	struct FCreationArgs
-	{
-		FOnConstraintCreated CreationDelegate;
-		ETransformConstraintType ConstraintType = ETransformConstraintType::Parent;
-		FConstrainable Parent;
-	};
-
-	FCreationArgs Args;
-	Args.CreationDelegate = OnConstraintCreated;
-	Args.ConstraintType = ConstraintType;
-	Args.Parent = OnGetParent.IsBound() ? OnGetParent.Execute() : FConstrainable();
-		
-	// hide the creation menu
-	FSlateApplication::Get().DismissAllMenus();
+	// We need to hide the creation menu. This destroys the menu that is keeping this SConstraintMenuEntry alive.
+	// So close only at the end because we're still accessing members in this call (undefined behaviour).
+	ON_SCOPE_EXIT { 	FSlateApplication::Get().DismissAllMenus(); };
 	
 	// use the selected parent if any
-	if (Args.Parent.Object.IsValid())
+	const FConstrainable Parent = OnGetParent.IsBound() ? OnGetParent.Execute() : FConstrainable();
+	if (Parent.Object.IsValid())
 	{
 		const TArray<AActor*> Children = UE::Private::GetCurrentSelection();
-		const bool bCreated = UE::Private::CreateConstraint(Children, Args.Parent, Args.ConstraintType, bUseDefault);
+		const bool bCreated = UE::Private::CreateConstraint(Children, Parent, ConstraintType, bUseDefault);
 		if (!bCreated)
 		{
 			return FReply::Unhandled();
 		}
 
 		// update list
-		if (Args.CreationDelegate.IsBound())
+		if (OnConstraintCreated.IsBound())
 		{
-			Args.CreationDelegate.Execute();
+			OnConstraintCreated.Execute();
 		}
 
 		return FReply::Handled();
