Control Rig: Re-introduce dependency checker

#rnx

diff --git a/Engine/Plugins/Animation/ControlRig/Source/ControlRig/Private/Constraints/ControlRigTransformableHandle.cpp b/Engine/Plugins/Animation/ControlRig/Source/ControlRig/Private/Constraints/ControlRigTransformableHandle.cpp
--- a/Engine/Plugins/Animation/ControlRig/Source/ControlRig/Private/Constraints/ControlRigTransformableHandle.cpp
+++ b/Engine/Plugins/Animation/ControlRig/Source/ControlRig/Private/Constraints/ControlRigTransformableHandle.cpp
@@ -407,11 +407,17 @@ bool UTransformableControlHandle::HasDirectDependencyWith(const UTransformableHa
 		// at this point, we know that both handles belong to the same rig
 		FRigControlElement* PossibleParent = AllControls[IndexOfPossibleParent];
 
-		auto IsParentedTo = [ControlElement, PossibleParent](UControlRig* InControlRig, const URigHierarchy* InRigHierarchy)
+		auto CanCauseCycle = [ControlElement, PossibleParent](UControlRig* InControlRig, const URigHierarchy* InRigHierarchy)
 		{
 			if (InControlRig && InRigHierarchy)
 			{
-				return InRigHierarchy->IsParentedTo(ControlElement, PossibleParent);
+			#if WITH_EDITOR
+				FRigDependenciesProviderForControlRig DependencyProvider(InControlRig);
+				DependencyProvider.SetInteractiveDialogEnabled(true);
+				return InRigHierarchy->CanCauseCycle(ControlElement, PossibleParent, DependencyProvider);
+            #else
+				return InRigHierarchy->CanCauseCycle(ControlElement, PossibleParent);
+            #endif	
 			}
 			return false;
 		};
@@ -441,7 +447,7 @@ bool UTransformableControlHandle::HasDirectDependencyWith(const UTransformableHa
 					}
 
 					// NOTE: we'd like to call this here but read URigHierarchy::GetDependenciesForVM about using this function with modular rigs
-					// if (IsParentedTo(ModuleRig, ModuleHierarchy))
+					// if (CanCauseCycle(ModuleRig, ModuleHierarchy))
 					// {
 					// 	return true;
 					// }
@@ -462,7 +468,7 @@ bool UTransformableControlHandle::HasDirectDependencyWith(const UTransformableHa
 		}
 
 		// default control rig test
-		if (IsParentedTo(ControlRigPtr, Hierarchy))
+		if (CanCauseCycle(ControlRigPtr, Hierarchy))
 		{
 			return true;
 		}
diff --git a/Engine/Plugins/Animation/ControlRig/Source/ControlRig/Private/Rigs/RigHierarchy.cpp b/Engine/Plugins/Animation/ControlRig/Source/ControlRig/Private/Rigs/RigHierarchy.cpp
--- a/Engine/Plugins/Animation/ControlRig/Source/ControlRig/Private/Rigs/RigHierarchy.cpp
+++ b/Engine/Plugins/Animation/ControlRig/Source/ControlRig/Private/Rigs/RigHierarchy.cpp
@@ -58,6 +58,8 @@ static TAutoConsoleVariable<int32> CVarControlRigHierarchyTracePrecision(TEXT("C
 static TAutoConsoleVariable<int32> CVarControlRigHierarchyTraceOnSpawn(TEXT("ControlRig.Hierarchy.TraceOnSpawn"), 0, TEXT("sets the number of frames to trace when a new hierarchy is spawned"));
 TAutoConsoleVariable<bool> CVarControlRigHierarchyEnableRotationOrder(TEXT("ControlRig.Hierarchy.EnableRotationOrder"), true, TEXT("enables the rotation order for controls"));
 TAutoConsoleVariable<bool> CVarControlRigHierarchyEnableModules(TEXT("ControlRig.Hierarchy.Modules"), true, TEXT("enables the modular rigging functionality"));
+static TAutoConsoleVariable<bool> CVarControlRigHierarchyTreatAvailableSpaceSafe(TEXT("ControlRig.Hierarchy.TreatAvailableSpacesAsSafeForSpaceSwitching"), true, TEXT("Controls will be able to space switch to the list of available spaces, even if the space switch would cause a cycle."));
+static TAutoConsoleVariable<bool> CVarControlRigHierarchyAllowToIgnoreCycles(TEXT("ControlRig.Hierarchy.AllowToIgnoreCycles"), true, TEXT("With this option turned on users can ignore cycles by clicking ignore in a modal dialog during the user interaction."));
 static int32 sRigHierarchyLastTrace = INDEX_NONE;
 
 // A console command to trace a single frame / single execution for a control rig anim node / control rig component
@@ -3893,11 +3895,12 @@ bool URigHierarchy::SetParentWeightArray(FRigBaseElement* InChild,  const TArray
 	return false;
 }
 
-bool URigHierarchy::CanSwitchToParent(FRigElementKey InChild, FRigElementKey InParent, const IRigDependenciesProvider& InDependencyProvider, FString* OutFailureReason)
+bool URigHierarchy::CanCauseCycle(FRigElementKey InChild, FRigElementKey InParent, const IRigDependenciesProvider& InDependencyProvider,
+	FString* OutFailureReason) const
 {
 	InParent = PreprocessParentElementKeyForSpaceSwitching(InChild, InParent);
 
-	FRigBaseElement* Child = Find(InChild);
+	const FRigBaseElement* Child = Find(InChild);
 	if(Child == nullptr)
 	{
 		if(OutFailureReason)
@@ -3907,14 +3910,14 @@ bool URigHierarchy::CanSwitchToParent(FRigElementKey InChild, FRigElementKey InP
 		return false;
 	}
 
-	FRigBaseElement* Parent = Find(InParent);
+	const FRigBaseElement* Parent = Find(InParent);
 	if(Parent == nullptr)
 	{
 		// if we don't specify anything and the element is parented directly to the world,
-		// perfomring this switch means unparenting it from world (since there is no default parent)
+		// performing this switch means unparenting it from world (since there is no default parent)
 		if(!InParent.IsValid() && GetFirstParent(InChild) == GetWorldSpaceReferenceKey())
 		{
-			return true;
+			return false;
 		}
 		
 		if(OutFailureReason)
@@ -3924,162 +3927,148 @@ bool URigHierarchy::CanSwitchToParent(FRigElementKey InChild, FRigElementKey InP
 		return false;
 	}
 
-	// see if this is already parented to the target parent
-	if(GetFirstParent(Child) == Parent)
-	{
-		return true;
-	}
-
-	const FRigMultiParentElement* MultiParentChild = Cast<FRigMultiParentElement>(Child);
-	if(MultiParentChild == nullptr)
+	const FRigTransformElement* TransformChild = Cast<FRigTransformElement>(Child);
+	if(TransformChild == nullptr)
 	{
 		if(OutFailureReason)
 		{
-			OutFailureReason->Appendf(TEXT("Child Element %s does not allow space switching (it's not a multi parent element)."), *InChild.ToString());
+			OutFailureReason->Appendf(TEXT("Child Element %s is not a transform element."), *InChild.ToString());
 		}
+		return false;
 	}
 
-	const FRigTransformElement* TransformParent = Cast<FRigMultiParentElement>(Parent);
+	const FRigTransformElement* TransformParent = Cast<FRigTransformElement>(Parent);
 	if(TransformParent == nullptr)
 	{
 		if(OutFailureReason)
 		{
 			OutFailureReason->Appendf(TEXT("Parent Element %s is not a transform element"), *InParent.ToString());
 		}
+		return false;
+	}
+
+	return CanCauseCycle(TransformChild, TransformParent, InDependencyProvider, OutFailureReason);
+}
+
+bool URigHierarchy::CanCauseCycle(const FRigTransformElement* InChild, const FRigTransformElement* InParent,
+	const IRigDependenciesProvider& InDependencyProvider, FString* OutFailureReason) const
+{
+	if (!InChild || !InParent)
+	{
+		return false;
+	}
+	
+	// see if this is already parented to the target parent
+	if(GetParents(InChild).Contains(InParent))
+	{
+		return false;
 	}
 
-	TArray<TPair<FRigHierarchyRecord,FRigHierarchyRecord>> DependencyChain;
-	if(IsParentedTo(Parent, Child, InDependencyProvider, &DependencyChain))
+	FRigHierarchyDependencyChain DependencyChain;
+	if(IsParentedTo(const_cast<FRigTransformElement*>(InParent), const_cast<FRigTransformElement*>(InChild), InDependencyProvider, &DependencyChain))
 	{
 #if WITH_EDITOR
-		
-		if(OutFailureReason)
+		if (!DependencyChain.IsEmpty() && InDependencyProvider.IsInteractiveDialogEnabled() && DismissDependencyDelegate.IsBound())
 		{
-			if (DependencyChain.IsEmpty())
+			// for dependencies which are not just pure elements - we should ask the user to potentially dismiss it 
+			if (DependencyChain.ContainsByPredicate([](const TPair<FRigHierarchyRecord,FRigHierarchyRecord>& Dependency) -> bool
 			{
-				OutFailureReason->Appendf(TEXT("Cannot switch '%s' to '%s' - would cause a cycle."), *InChild.ToString(), *InParent.ToString());
-			}
-			else
+				return !Dependency.Key.IsElement(); 
+			}))
 			{
-				TArray<FString> Messages;
-
-				for (const TPair<FRigHierarchyRecord,FRigHierarchyRecord>& DependencyPair : DependencyChain)
+				if (CVarControlRigHierarchyTreatAvailableSpaceSafe.GetValueOnAnyThread())
 				{
-					auto GetLabelForRecord = [this](const FRigHierarchyRecord& InRecord) -> FString
+					// if this is a defined space for something - allow it even if there's a cycle.
+					// we assume that spaces defined by riggers are safe to use.
+					if (const FRigControlElement* Control = Cast<FRigControlElement>(InParent))
 					{
-						FString Label;
-							
-						switch (InRecord.Type)
+						if (Control->Settings.Customization.AvailableSpaces.ContainsByPredicate([InChild](const FRigElementKeyWithLabel& Space) -> bool
 						{
-							case FRigHierarchyRecord::EType_RigElement:
-							case FRigHierarchyRecord::EType_Metadata:
-							{
-								Label = GetKey(InRecord.Index).ToString();
-								if (InRecord.IsMetadata())
-								{
-									Label += FString::Printf(TEXT(" Metadata(%s)"), *InRecord.Name.ToString());
-								}
-								return Label;
-							}
-							case FRigHierarchyRecord::EType_Variable:
-							case FRigHierarchyRecord::EType_Instruction:
-							{
-								if (InRecord.IsVariable())
-								{
-									Label = TEXT("Variable");
-								}
-								else
-								{
-									Label = TEXT("Node");
-								}
-
-								UControlRig* ControlRig = Cast<UControlRig>(GetOuter());
-								if (!ControlRig)
-								{
-									break;
-								}
+							return Space.Key == InChild->GetKey();
+						}))
+						{
+							return false;
+						}
+					}
+				}			
+				if (CVarControlRigHierarchyAllowToIgnoreCycles.GetValueOnAnyThread() && IsInGameThread())
+				{
+					// if the user decides to dismiss this dependency - we'll ignore the cycle and will still allow the action.
+					if (DismissDependencyDelegate.Execute(this, InChild->GetKey(), InParent->GetKey(), DependencyChain))
+					{
+						return false;
+					}
+				}
+			}
+		}
 
-								UModularRig* ModularRig = Cast<UModularRig>(ControlRig);
+		if(OutFailureReason)
+		{
+			if (DependencyChain.IsEmpty())
+			{
+				OutFailureReason->Appendf(TEXT("Relating '%s' to '%s' would cause a cycle."), *InChild->GetKey().ToString(), *InParent->GetKey().ToString());
+			}
+			else
+			{
+				const FString Message = GetMessageFromDependencyChain(DependencyChain);
+				OutFailureReason->Appendf(TEXT("Relating '%s' to '%s' would cause a cycle.\n\n%s"), *InChild->GetKey().ToString(), *InParent->GetKey().ToString(), *Message);
+			}
+		}
+#endif
+		return true;
+	}
 
-								const URigVM* VM = nullptr;
-								if (InRecord.Name.IsNone() || !ModularRig)
-								{
-									VM = ControlRig->GetVM();
-								}
-								else
-								{
-									const FRigModuleInstance* RigModule = ModularRig->FindModule(InRecord.Name);
-									if (!RigModule)
-									{
-										break;
-									}
-									UControlRig* ModuleRig = RigModule->GetRig();
-									if (!ModuleRig)
-									{
-										break;
-									}
-									VM = ModuleRig->GetVM();
-								}
+	return false;
+}
 
-								check(VM);
+bool URigHierarchy::CanSwitchToParent(FRigElementKey InChild, FRigElementKey InParent, const IRigDependenciesProvider& InDependencyProvider, FString* OutFailureReason) const
+{
+	const FRigBaseElement* Child = Find(InChild);
+	if(Child == nullptr)
+	{
+		if(OutFailureReason)
+		{
+			OutFailureReason->Appendf(TEXT("Child Element %s cannot be found."), *InChild.ToString());
+		}
+		return false;
+	}
 
-								if (InRecord.IsVariable())
-								{
-									const TArray<FRigVMExternalVariableDef>& ExternalVariableDefs = VM->GetExternalVariableDefs();
-									if (ExternalVariableDefs.IsValidIndex(InRecord.Index))
-									{
-										Label += FString::Printf(TEXT("(%s)"), *ExternalVariableDefs[InRecord.Index].Name.ToString());
-									}
-								}
-								else
-								{
-									if (!VM->GetByteCode().GetInstructions().IsValidIndex(InRecord.Index))
-									{
-										break;
-									}
+	const FRigBaseElement* Parent = Find(InParent);
+	if(Parent == nullptr)
+	{
+		// if we don't specify anything and the element is parented directly to the world,
+		// performing this switch means unparenting it from world (since there is no default parent)
+		if(!InParent.IsValid() && GetFirstParent(InChild) == GetWorldSpaceReferenceKey())
+		{
+			return true;
+		}
 		
-									const FRigVMInstruction Instruction = VM->GetByteCode().GetInstructions()[InRecord.Index];
-									if (Instruction.OpCode != ERigVMOpCode::Execute)
-									{
-										break;
-									}
-
-									URigVMNode* Node = Cast<URigVMNode>(VM->GetByteCode().GetSubjectForInstruction(InRecord.Index));
-									if (Node == nullptr)
-									{
-										break;
-									}
-
-									Label += FString::Printf(TEXT("(%s)"), *Node->GetPathName());
-								}
-
-								break;
-							}
-							default:
-							{
-								Label = TEXT("Unknown");
-								break;
-							}
-						}
-
-						return Label;
-					};
+		if(OutFailureReason)
+		{
+			OutFailureReason->Appendf(TEXT("Parent Element %s cannot be found."), *InParent.ToString());
+		}
+		return false;
+	}
 
-					const FString TargetLabel = GetLabelForRecord(DependencyPair.Value);
-					const FString SourceLabel = GetLabelForRecord(DependencyPair.Key);
-					Messages.Add(FString::Printf(TEXT("'%s' depends on '%s'."), *TargetLabel, *SourceLabel));
-				}
+	const FRigMultiParentElement* MultiParentChild = Cast<FRigMultiParentElement>(Child);
+	if(MultiParentChild == nullptr)
+	{
+		if(OutFailureReason)
+		{
+			OutFailureReason->Appendf(TEXT("Child Element %s does not allow space switching (it's not a multi parent element)."), *InChild.ToString());
+		}
+	}
 
-				const FString JoinedMessages = FString::Join(Messages, TEXT("\n"));
-				
-				OutFailureReason->Appendf(TEXT("Cannot switch '%s' to '%s' - would cause a cycle.\n\n%s"), *InChild.ToString(), *InParent.ToString(), *JoinedMessages);
-			}
+	const FRigTransformElement* TransformParent = Cast<FRigTransformElement>(Parent);
+	if(TransformParent == nullptr)
+	{
+		if(OutFailureReason)
+		{
+			OutFailureReason->Appendf(TEXT("Parent Element %s is not a transform element"), *InParent.ToString());
 		}
-#endif
-		return false;
 	}
 
-	return true;
+	return !CanCauseCycle(MultiParentChild, TransformParent, InDependencyProvider, OutFailureReason);
 }
 
 bool URigHierarchy::SwitchToParent(FRigElementKey InChild, FRigElementKey InParent, bool bInitial, bool bAffectChildren, const IRigDependenciesProvider& InDependencyProvider, FString* OutFailureReason)
@@ -6468,9 +6457,9 @@ FRigHierarchyKey URigHierarchy::GetPreviousHierarchyParent(const FRigHierarchyKe
 	return FRigHierarchyKey();
 }
 
-bool URigHierarchy::IsParentedTo(FRigBaseElement* InChild, FRigBaseElement* InParent, const IRigDependenciesProvider& InDependencyProvider, TArray<TPair<FRigHierarchyRecord,FRigHierarchyRecord>>* OutDependencyChain) const
+bool URigHierarchy::IsParentedTo(FRigBaseElement* InChild, FRigBaseElement* InParent, const IRigDependenciesProvider& InDependencyProvider, FRigHierarchyDependencyChain* OutDependencyChain) const
 {
-	TArray<TPair<FRigHierarchyRecord,FRigHierarchyRecord>> DependencyChain;
+	FRigHierarchyDependencyChain DependencyChain;
 
 	bool bActualTopology = false;
 	if (OutDependencyChain == nullptr)
@@ -6494,7 +6483,7 @@ bool URigHierarchy::IsParentedTo(FRigBaseElement* InChild, FRigBaseElement* InPa
 	return true;
 }
 
-bool URigHierarchy::IsDependentOn(FRigBaseElement* InDependent, FRigBaseElement* InDependency, const IRigDependenciesProvider& InDependencyProvider, bool bIsOnActualTopology, TArray<TPair<FRigHierarchyRecord,FRigHierarchyRecord>>& OutDependencyChain) const
+bool URigHierarchy::IsDependentOn(FRigBaseElement* InDependent, FRigBaseElement* InDependency, const IRigDependenciesProvider& InDependencyProvider, bool bIsOnActualTopology, FRigHierarchyDependencyChain& OutDependencyChain) const
 {
 	if((InDependent == nullptr) || (InDependency == nullptr))
 	{
@@ -6602,7 +6591,7 @@ bool URigHierarchy::IsDependentOn(FRigBaseElement* InDependent, FRigBaseElement*
 			bool bIsOnActualTopology,
 			const TTuple<int32,int32>& CacheKey,
 			const FRigHierarchyRecord& InDependentRecord,
-			TArray<TPair<FRigHierarchyRecord,FRigHierarchyRecord>>& OutDependencyChain)
+			FRigHierarchyDependencyChain& OutDependencyChain)
 		{
 			const TRigHierarchyDependencyMap& DependencyMap = InDependencyProvider.GetRigHierarchyDependencies();
 			
@@ -6971,13 +6960,13 @@ void URigHierarchy::UpdateCachedChildren()
 }
 
 	
-FRigElementKey URigHierarchy::PreprocessParentElementKeyForSpaceSwitching(const FRigElementKey& InChildKey, const FRigElementKey& InParentKey)
+FRigElementKey URigHierarchy::PreprocessParentElementKeyForSpaceSwitching(const FRigElementKey& InChildKey, const FRigElementKey& InParentKey) const
 {
 	LLM_SCOPE_BYNAME(TEXT("Animation/ControlRig"));
 	
 	if(InParentKey == GetWorldSpaceReferenceKey())
 	{
-		return GetOrAddWorldSpaceReference();
+		return const_cast<URigHierarchy*>(this)->GetOrAddWorldSpaceReference();
 	}
 	else if(InParentKey == GetDefaultParentKey())
 	{
@@ -8856,6 +8845,120 @@ void URigHierarchy::NotifyPostUndoSelectionChanges()
 
 	SelectedKeysBeforeUndo.Reset();
 }
+
+FString URigHierarchy::GetMessageFromDependencyChain(const FRigHierarchyDependencyChain& InDependencyChain) const
+{
+	TArray<FString> Messages;
+
+	for (const TPair<FRigHierarchyRecord,FRigHierarchyRecord>& DependencyPair : InDependencyChain)
+	{
+		auto GetLabelForRecord = [this](const FRigHierarchyRecord& InRecord) -> FString
+		{
+			FString Label;
+				
+			switch (InRecord.Type)
+			{
+				case FRigHierarchyRecord::EType_RigElement:
+				case FRigHierarchyRecord::EType_Metadata:
+				{
+					Label = GetKey(InRecord.Index).ToString();
+					if (InRecord.IsMetadata())
+					{
+						Label += FString::Printf(TEXT(" Metadata(%s)"), *InRecord.Name.ToString());
+					}
+					return Label;
+				}
+				case FRigHierarchyRecord::EType_Variable:
+				case FRigHierarchyRecord::EType_Instruction:
+				{
+					if (InRecord.IsVariable())
+					{
+						Label = TEXT("Variable");
+					}
+					else
+					{
+						Label = TEXT("Node");
+					}
+
+					UControlRig* ControlRig = Cast<UControlRig>(GetOuter());
+					if (!ControlRig)
+					{
+						break;
+					}
+
+					UModularRig* ModularRig = Cast<UModularRig>(ControlRig);
+
+					const URigVM* VM = nullptr;
+					if (InRecord.Name.IsNone() || !ModularRig)
+					{
+						VM = ControlRig->GetVM();
+					}
+					else
+					{
+						const FRigModuleInstance* RigModule = ModularRig->FindModule(InRecord.Name);
+						if (!RigModule)
+						{
+							break;
+						}
+						UControlRig* ModuleRig = RigModule->GetRig();
+						if (!ModuleRig)
+						{
+							break;
+						}
+						VM = ModuleRig->GetVM();
+					}
+
+					check(VM);
+
+					if (InRecord.IsVariable())
+					{
+						const TArray<FRigVMExternalVariableDef>& ExternalVariableDefs = VM->GetExternalVariableDefs();
+						if (ExternalVariableDefs.IsValidIndex(InRecord.Index))
+						{
+							Label += FString::Printf(TEXT("(%s)"), *ExternalVariableDefs[InRecord.Index].Name.ToString());
+						}
+					}
+					else
+					{
+						if (!VM->GetByteCode().GetInstructions().IsValidIndex(InRecord.Index))
+						{
+							break;
+						}
+
+						const FRigVMInstruction Instruction = VM->GetByteCode().GetInstructions()[InRecord.Index];
+						if (Instruction.OpCode != ERigVMOpCode::Execute)
+						{
+							break;
+						}
+
+						URigVMNode* Node = Cast<URigVMNode>(VM->GetByteCode().GetSubjectForInstruction(InRecord.Index));
+						if (Node == nullptr)
+						{
+							break;
+						}
+
+						Label += FString::Printf(TEXT("(Instruction %03d, %s)"), InRecord.Index, *Node->GetNodeTitle());
+					}
+
+					break;
+				}
+				default:
+				{
+					Label = TEXT("Unknown");
+					break;
+				}
+			}
+
+			return Label;
+		};
+
+		const FString TargetLabel = GetLabelForRecord(DependencyPair.Value);
+		const FString SourceLabel = GetLabelForRecord(DependencyPair.Key);
+		Messages.Add(FString::Printf(TEXT("'%s' depends on '%s'."), *TargetLabel, *SourceLabel));
+	}
+
+	return FString::Join(Messages, TEXT("\n"));
+}
 #endif
 
 bool URigHierarchy::IsAnimatable(const FRigElementKey& InKey) const
diff --git a/Engine/Plugins/Animation/ControlRig/Source/ControlRig/Public/Rigs/RigDependency.h b/Engine/Plugins/Animation/ControlRig/Source/ControlRig/Public/Rigs/RigDependency.h
--- a/Engine/Plugins/Animation/ControlRig/Source/ControlRig/Public/Rigs/RigDependency.h
+++ b/Engine/Plugins/Animation/ControlRig/Source/ControlRig/Public/Rigs/RigDependency.h
@@ -5,6 +5,7 @@
 #include "UObject/NameTypes.h"
 #include "UObject/WeakObjectPtr.h"
 #include "UObject/WeakObjectPtrTemplates.h"
+#include "RigHierarchyDefines.h"
 #include "RigDependencyRecords.h"
 #include "Units/RigUnitContext.h"
 
@@ -25,6 +26,13 @@ struct IRigDependenciesProvider
 	virtual const TRigHierarchyDependencyMap& GetReverseRigHierarchyDependencies() const = 0;
 	virtual void InvalidateCache() = 0;
 	virtual uint32 GetHash() const = 0;
+
+	bool IsInteractiveDialogEnabled() const { return bInteractiveDialogEnabled; } 
+	void SetInteractiveDialogEnabled(bool InEnabled) { bInteractiveDialogEnabled = InEnabled; }
+
+private:
+
+	bool bInteractiveDialogEnabled = false;
 };
 
 /**
@@ -78,6 +86,8 @@ protected:
 
 struct FRigDependenciesProviderForControlRig : public FRigDependenciesProviderForVM
 {
+public:
+	
 	FRigDependenciesProviderForControlRig() = default;
 	UE_API FRigDependenciesProviderForControlRig(UControlRig* InControlRig, const FName& InEventName = NAME_None, bool InFollowVariables = false);
 	virtual ~FRigDependenciesProviderForControlRig() override = default;
diff --git a/Engine/Plugins/Animation/ControlRig/Source/ControlRig/Public/Rigs/RigDependencyRecords.h b/Engine/Plugins/Animation/ControlRig/Source/ControlRig/Public/Rigs/RigDependencyRecords.h
--- a/Engine/Plugins/Animation/ControlRig/Source/ControlRig/Public/Rigs/RigDependencyRecords.h
+++ b/Engine/Plugins/Animation/ControlRig/Source/ControlRig/Public/Rigs/RigDependencyRecords.h
@@ -112,6 +112,7 @@ struct FRigHierarchyRecord
 
 typedef TMap<FRigHierarchyRecord, TSet<FRigHierarchyRecord>> TRigHierarchyDependencyMap;
 typedef TPair<FRigHierarchyRecord, TSet<FRigHierarchyRecord>> TRigHierarchyDependencyMapPair;
+typedef TArray<TPair<FRigHierarchyRecord, FRigHierarchyRecord>> FRigHierarchyDependencyChain;
 
 inline uint32 GetTypeHash(const FRigHierarchyRecord& InRecord)
 {
diff --git a/Engine/Plugins/Animation/ControlRig/Source/ControlRig/Public/Rigs/RigHierarchy.h b/Engine/Plugins/Animation/ControlRig/Source/ControlRig/Public/Rigs/RigHierarchy.h
--- a/Engine/Plugins/Animation/ControlRig/Source/ControlRig/Public/Rigs/RigHierarchy.h
+++ b/Engine/Plugins/Animation/ControlRig/Source/ControlRig/Public/Rigs/RigHierarchy.h
@@ -33,6 +33,7 @@ DECLARE_DYNAMIC_MULTICAST_DELEGATE_ThreeParams(FRigHierarchyModifiedDynamicEvent
 DECLARE_EVENT_FiveParams(URigHierarchy, FRigHierarchyUndoRedoTransformEvent, URigHierarchy*, const FRigElementKey&, ERigTransformType::Type, const FTransform&, bool /* bUndo */);
 DECLARE_MULTICAST_DELEGATE_TwoParams(FRigHierarchyMetadataChangedDelegate, const FRigElementKey& /* Key */, const FName& /* Name */);
 DECLARE_MULTICAST_DELEGATE_ThreeParams(FRigHierarchyMetadataTagChangedDelegate, const FRigElementKey& /* Key */, const FName& /* Tag */, bool /* AddedOrRemoved */);
+DECLARE_DELEGATE_RetVal_FourParams(bool, FRigHierarchyDismissDependencyDelegate, const URigHierarchy* /* InHierarchy */, const FRigElementKey& /* InChild */, const FRigElementKey& /* InParent */, const FRigHierarchyDependencyChain& /* InDependencyChain */);
 
 extern CONTROLRIG_API TAutoConsoleVariable<bool> CVarControlRigHierarchyEnableRotationOrder;
 extern CONTROLRIG_API TAutoConsoleVariable<bool> CVarControlRigHierarchyEnableModules;
@@ -293,7 +294,7 @@ public:
 	 * Resets all curves to 0.0
 	 */
 	UFUNCTION(BlueprintCallable, Category = URigHierarchy)
-    CONTROLRIG_API void ResetCurveValues();
+	CONTROLRIG_API void ResetCurveValues();
 
 	/**
 	 * Resets all curves to be unset (defaulting to 0.0)
@@ -330,7 +331,7 @@ public:
 	 * @param InElementType The type filter to apply
 	 * @return The number of elements in the Hierarchy
 	 */
-    CONTROLRIG_API int32 Num(ERigElementType InElementType) const;
+	CONTROLRIG_API int32 Num(ERigElementType InElementType) const;
 
 	// iterators
 	TArray<FRigBaseElement*>::RangedForIteratorType      begin() { return Elements.begin(); }
@@ -542,7 +543,7 @@ public:
 	 * @return The key of an element given its index
 	 */
 	UFUNCTION(BlueprintCallable, Category = URigHierarchy)
-    FRigElementKey GetKey(int32 InElementIndex) const
+	FRigElementKey GetKey(int32 InElementIndex) const
 	{
 		if(Elements.IsValidIndex(InElementIndex))
 		{
@@ -642,16 +643,16 @@ public:
 	 * @return The retrieved element or nullptr.
 	 */
 	template<typename T>
-    T* GetChecked(int32 InIndex)
+	T* GetChecked(int32 InIndex)
 	{
 		return CastChecked<T>(Get(InIndex));
 	}
 
 	/**
-     * Returns a handle to an existing element
-     * @param InKey The key of the handle to retrieve.
-     * @return The retrieved handle (may be invalid)
-     */
+	 * Returns a handle to an existing element
+	 * @param InKey The key of the handle to retrieve.
+	 * @return The retrieved handle (may be invalid)
+	 */
 	FRigElementHandle GetHandle(const FRigElementKey& InKey) const
 	{
 		if(Contains(InKey))
@@ -815,7 +816,7 @@ public:
 	 * @return The retrieved component or nullptr.
 	 */
 	template<typename T>
-    T* GetComponentChecked(int32 InIndex)
+	T* GetComponentChecked(int32 InIndex)
 	{
 		return CastChecked<T>(GetComponent(InIndex));
 	}
@@ -896,7 +897,7 @@ private:
 	* @param InKey The key of the bone element to retrieve. 
 	*/
 	UFUNCTION(BlueprintCallable, Category = URigHierarchy, meta = (DisplayName = "Find Bone", ScriptName = "FindBone"))
-    FRigBoneElement FindBone_ForBlueprintOnly(const FRigElementKey& InKey) const
+	FRigBoneElement FindBone_ForBlueprintOnly(const FRigElementKey& InKey) const
 	{
 		if (const FRigBoneElement* Bone = Find<FRigBoneElement>(InKey))
 		{
@@ -910,7 +911,7 @@ private:
 	* @param InKey The key of the control element to retrieve. 
 	*/
 	UFUNCTION(BlueprintCallable, Category = URigHierarchy, meta = (DisplayName = "Find Control", ScriptName = "FindControl"))
-    FRigControlElement FindControl_ForBlueprintOnly(const FRigElementKey& InKey) const
+	FRigControlElement FindControl_ForBlueprintOnly(const FRigElementKey& InKey) const
 	{
 		if (const FRigControlElement* Control = Find<FRigControlElement>(InKey))
 		{
@@ -924,7 +925,7 @@ private:
 	* @param InKey The key of the null element to retrieve. 
 	*/
 	UFUNCTION(BlueprintCallable, Category = URigHierarchy, meta = (DisplayName = "Find Null", ScriptName = "FindNull"))
-    FRigNullElement FindNull_ForBlueprintOnly(const FRigElementKey& InKey) const
+	FRigNullElement FindNull_ForBlueprintOnly(const FRigElementKey& InKey) const
 	{
 		if (const FRigNullElement* Null = Find<FRigNullElement>(InKey))
 		{
@@ -955,7 +956,7 @@ public:
 	 * @return The retrieved element or nullptr.
 	 */
 	template<typename T>
-    T* FindChecked(const FRigElementKey& InKey)
+	T* FindChecked(const FRigElementKey& InKey)
 	{
 		return GetChecked<T>(GetIndex(InKey));
 	}
@@ -976,16 +977,16 @@ public:
 
 			Traverse([&ElementVisited, &Results](FRigBaseElement* InElement, bool& bContinue)
 			{
-			    bContinue = !ElementVisited[InElement->GetIndex()];
-
-			    if(bContinue)
-			    {
-			        if(T* CastElement = Cast<T>(InElement))
-			        {
-			            Results.Add(CastElement);
-			        }
-			        ElementVisited[InElement->GetIndex()] = true;
-			    }
+				bContinue = !ElementVisited[InElement->GetIndex()];
+
+				if(bContinue)
+				{
+					if(T* CastElement = Cast<T>(InElement))
+					{
+						Results.Add(CastElement);
+					}
+					ElementVisited[InElement->GetIndex()] = true;
+				}
 			});
 		}
 		else
@@ -1006,7 +1007,7 @@ public:
 	 * @param bTraverse Returns the element keys in order of a depth first traversal
 	 */
 	template<typename T>
-    TArray<FRigElementKey> GetKeysOfType(bool bTraverse = false) const
+	TArray<FRigElementKey> GetKeysOfType(bool bTraverse = false) const
 	{
 		TArray<FRigElementKey> Keys;
 		TArray<T*> Results = GetElementsOfType<T>(bTraverse);
@@ -1023,7 +1024,7 @@ public:
 	 * @param bTraverse Returns the elements in order of a depth first traversal
 	 */
 	template<typename T>
-    TArray<T*> GetFilteredElements(TFunction<bool(T*)> InKeepElementFunction, bool bTraverse = false) const
+	TArray<T*> GetFilteredElements(TFunction<bool(T*)> InKeepElementFunction, bool bTraverse = false) const
 	{
 		TArray<T*> Results;
 
@@ -1033,21 +1034,21 @@ public:
 			ElementVisited.AddZeroed(Elements.Num());
 		
 			Traverse([&ElementVisited, &Results, InKeepElementFunction](FRigBaseElement* InElement, bool& bContinue)
-            {
-                bContinue = !ElementVisited[InElement->GetIndex()];
+			{
+				bContinue = !ElementVisited[InElement->GetIndex()];
 
-                if(bContinue)
-                {
-                    if(T* CastElement = Cast<T>(InElement))
-                    {
+				if(bContinue)
+				{
+					if(T* CastElement = Cast<T>(InElement))
+					{
 						if(InKeepElementFunction(CastElement))
 						{
 							Results.Add(CastElement);
 						}
-                    }
-                    ElementVisited[InElement->GetIndex()] = true;
-                }
-            });
+					}
+					ElementVisited[InElement->GetIndex()] = true;
+				}
+			});
 		}
 		else
 		{
@@ -1075,12 +1076,12 @@ public:
 	}
 	
 	/**
-   	 * Returns all Bone elements without traversing the hierarchy
-   	 */
-   	TArray<FRigBaseElement*>& GetBonesFast() const
-   	{
-   		return ElementsPerType[RigElementTypeToFlatIndex(ERigElementType::Bone)];
-   	}
+		* Returns all Bone elements without traversing the hierarchy
+		*/
+	TArray<FRigBaseElement*>& GetBonesFast() const
+	{
+		return ElementsPerType[RigElementTypeToFlatIndex(ERigElementType::Bone)];
+	}
 
 	/**
 	 * Returns all Bone elements
@@ -1180,7 +1181,7 @@ public:
 	 * @param bTraverse Returns the elements in order of a depth first traversal
 	 */
 	UFUNCTION(BlueprintCallable, Category = URigHierarchy, meta = (DisplayName = "Get References", ScriptName = "GetReferences"))
-    TArray<FRigElementKey> GetReferenceKeys(bool bTraverse = true) const
+	TArray<FRigElementKey> GetReferenceKeys(bool bTraverse = true) const
 	{
 		return GetKeysOfType<FRigReferenceElement>(bTraverse);
 	}
@@ -1294,7 +1295,7 @@ public:
 	CONTROLRIG_API bool RemoveMetadata(FRigElementKey InItem, FName InMetadataName);
 
 	/**
-     * Removes all of the metadata under a given item 
+	 * Removes all of the metadata under a given item 
 	 * @param InItem The element key to search under
 	 */
 	UFUNCTION(BlueprintCallable, Category = URigHierarchy)
@@ -1692,12 +1693,12 @@ public:
 		return SetArrayMetadata<FTransform>(InItem, ERigMetadataType::TransformArray, InMetadataName, InValue);
 	}
 
-		/**
-	 * Queries and returns the value of FLinearColor metadata
-	 * @param InItem The element key to return the metadata for
-	 * @param InMetadataName The name of the metadata to query
-	 * @param DefaultValue The default value to fall back on
-	 */
+	/**
+ * Queries and returns the value of FLinearColor metadata
+ * @param InItem The element key to return the metadata for
+ * @param InMetadataName The name of the metadata to query
+ * @param DefaultValue The default value to fall back on
+ */
 	UFUNCTION(BlueprintPure, Category = URigHierarchy)
 	FLinearColor GetLinearColorMetadata(FRigElementKey InItem, FName InMetadataName, FLinearColor DefaultValue) const
 	{
@@ -1861,10 +1862,10 @@ public:
 	}
 
 	/*
-     * Sets a tag on an element in the hierarchy
-     * @param InItem The item to set the tag for
-     * @param InTag The tag to set
-     */
+	 * Sets a tag on an element in the hierarchy
+	 * @param InItem The item to set the tag for
+	 * @param InTag The tag to set
+	 */
 	UFUNCTION(BlueprintCallable, Category = URigHierarchy)
 	bool SetTag(FRigElementKey InItem, FName InTag)
 	{
@@ -1961,7 +1962,7 @@ public:
 	 * @return true if a given element is selected
 	 */
 	UFUNCTION(BlueprintCallable, Category = URigHierarchy)
-    bool IsSelectedByIndex(int32 InIndex) const
+	bool IsSelectedByIndex(int32 InIndex) const
 	{
 		return IsSelected(Get(InIndex));
 	}
@@ -1981,13 +1982,13 @@ public:
 	{
 		TArray<FRigElementKey> Result;
 		Traverse([InKeys, &Result](FRigBaseElement* Element, bool& bContinue)
-        {
-            const FRigElementKey& Key = Element->GetKey();
-            if(InKeys.Contains(Key))
-            {
-                Result.AddUnique(Key);
-            }
-        });
+		{
+			const FRigElementKey& Key = Element->GetKey();
+			if(InKeys.Contains(Key))
+			{
+				Result.AddUnique(Key);
+			}
+		});
 		return Result;
 	}
 
@@ -2034,7 +2035,7 @@ public:
 	 * Sanitizes a name by removing invalid characters.
 	 * @param InName The name to sanitize.
 	 * @return The sanitized name.
- 	 */
+	  */
 	static CONTROLRIG_API FRigName GetSanitizedName(const FRigName& InName, bool bAllowNameSpaces = true, int32 InMaxNameLength = UE_RIGHIERARCHY_DEFAULT_MAX_NAME_LENGTH);
 
 	/**
@@ -2116,6 +2117,12 @@ public:
 	 */
 	FRigHierarchyMetadataTagChangedDelegate& OnMetadataTagChanged() { return MetadataTagChangedDelegate; }
 
+	/**
+	 * Returns the DependencyDismissed delegate. This is used by editors to dismiss dependencies
+	 * which may be blocked certain workflows such as space switching or establishing constraints.
+	 */
+	FRigHierarchyDismissDependencyDelegate& OnDependencyDismissed() { return DismissDependencyDelegate; }
+
 	/**
 	 * Returns the local current or initial value for a given key.
 	 * If the key is invalid FTransform::Identity will be returned.
@@ -3577,19 +3584,47 @@ public:
 	*/
 	CONTROLRIG_API bool SetParentWeightArray(FRigBaseElement* InChild,  const TArrayView<const FRigElementWeight>& InWeights, bool bInitial = false, bool bAffectChildren = true);
 
+	/**
+	 * Determines if relating the child element to the given parent would cause a cycle 
+	 * @param InChild The key of the child for the new potential relationship
+	 * @param InParent The key of the parent for the new potential relationship
+	 * @param InDependencyProvider An additional object providing dependencies to respect
+	 * @param OutFailureReason An optional pointer to retrieve the reason for failure
+	 * @return Returns true if adding a new relationship from child to parent would cause a cycle
+	 */
+	CONTROLRIG_API bool CanCauseCycle(
+		FRigElementKey InChild,
+		FRigElementKey InParent,
+		const IRigDependenciesProvider& InDependencyProvider = FEmptyRigDependenciesProvider(),
+		FString* OutFailureReason = nullptr) const;
+
+	/**
+	 * Determines if relating the child element to the given parent would cause a cycle 
+	 * @param InChild The key of the child for the new potential relationship
+	 * @param InParent The key of the parent for the new potential relationship
+	 * @param InDependencyProvider An additional object providing dependencies to respect
+	 * @param OutFailureReason An optional pointer to retrieve the reason for failure
+	 * @return Returns true if adding a new relationship from child to parent would cause a cycle
+	 */
+	CONTROLRIG_API bool CanCauseCycle(
+		const FRigTransformElement* InChild,
+		const FRigTransformElement* InParent,
+		const IRigDependenciesProvider& InDependencyProvider = FEmptyRigDependenciesProvider(),
+		FString* OutFailureReason = nullptr) const;
+	
 	/**
 	* Determines if the element can be switched to a provided parent
 	* @param InChild The key of the multi parented element
 	* @param InParent The key of the parent to look up the weight for
     * @param InDependencyProvider An additional object providing dependencies to respect
 	* @param OutFailureReason An optional pointer to retrieve the reason for failure
-	* @return Returns true if changing the weight was successful
+	* @return Returns true if the child can be switched to the desired parent
 	*/
 	CONTROLRIG_API bool CanSwitchToParent(
 		FRigElementKey InChild,
 		FRigElementKey InParent,
 		const IRigDependenciesProvider& InDependencyProvider = FEmptyRigDependenciesProvider(),
-		FString* OutFailureReason = nullptr);
+		FString* OutFailureReason = nullptr) const;
 
 	/**
 	 * Switches a multi parent element to a single parent.
@@ -3717,7 +3752,7 @@ public:
 	 * @param OutDependencyChain An optional output error message indicating the parent cycle
 	 * @return True if the given parent and given child have a parent-child relationship
 	 */
-    bool IsParentedTo(int32 InChildIndex, int32 InParentIndex, const IRigDependenciesProvider& InDependencyProvider, TArray<TPair<FRigHierarchyRecord,FRigHierarchyRecord>>* OutDependencyChain = nullptr) const
+    bool IsParentedTo(int32 InChildIndex, int32 InParentIndex, const IRigDependenciesProvider& InDependencyProvider, FRigHierarchyDependencyChain* OutDependencyChain = nullptr) const
 	{
 		if(Elements.IsValidIndex(InChildIndex) && Elements.IsValidIndex(InParentIndex))
 		{
@@ -4246,6 +4281,7 @@ private:
 	
 	FRigHierarchyMetadataChangedDelegate MetadataChangedDelegate;
 	FRigHierarchyMetadataTagChangedDelegate MetadataTagChangedDelegate;
+	FRigHierarchyDismissDependencyDelegate DismissDependencyDelegate;
 	FRigEventDelegate EventDelegate;
 
 	TSharedPtr<FRigHierarchyPoseAdapter> PoseAdapter;
@@ -4461,7 +4497,7 @@ public:
 	CONTROLRIG_API bool IsParentedTo(FRigBaseElement* InChild,
 		FRigBaseElement* InParent,
 		const IRigDependenciesProvider& InDependencyProvider = FEmptyRigDependenciesProvider(),
-		TArray<TPair<FRigHierarchyRecord,FRigHierarchyRecord>>* OutDependencyChain = nullptr) const;
+		FRigHierarchyDependencyChain* OutDependencyChain = nullptr) const;
 
 private:
 	/**
@@ -4477,7 +4513,7 @@ private:
 		FRigBaseElement* InDependency,
 		const IRigDependenciesProvider& InDependencyProvider,
 		bool bIsOnActualTopology,
-		TArray<TPair<FRigHierarchyRecord,FRigHierarchyRecord>>& OutDependencyChain) const;
+		FRigHierarchyDependencyChain& OutDependencyChain) const;
 
 public:
 	/**
@@ -4575,7 +4611,7 @@ private:
 	/**
 	 * Corrects a parent element key for space switching
 	 */
-	CONTROLRIG_API FRigElementKey PreprocessParentElementKeyForSpaceSwitching(const FRigElementKey& InChildKey, const FRigElementKey& InParentKey);
+	CONTROLRIG_API FRigElementKey PreprocessParentElementKeyForSpaceSwitching(const FRigElementKey& InChildKey, const FRigElementKey& InParentKey) const;
 
 	/*
 	 * Helper function to create an element for a given type
@@ -4664,6 +4700,11 @@ public:
 	 */
 	CONTROLRIG_API void CleanupInvalidCaches();
 
+#if WITH_EDITOR
+	// helper function to convert a dependency chain into a message string.
+	CONTROLRIG_API FString GetMessageFromDependencyChain(const FRigHierarchyDependencyChain& InDependencyChain) const;
+#endif
+
 private:
 	
 	/**
diff --git a/Engine/Plugins/Animation/ControlRig/Source/ControlRigEditor/Private/EditMode/Picker/SFloatingSpacePicker_AssetEditor.cpp b/Engine/Plugins/Animation/ControlRig/Source/ControlRigEditor/Private/EditMode/Picker/SFloatingSpacePicker_AssetEditor.cpp
--- a/Engine/Plugins/Animation/ControlRig/Source/ControlRigEditor/Private/EditMode/Picker/SFloatingSpacePicker_AssetEditor.cpp
+++ b/Engine/Plugins/Animation/ControlRig/Source/ControlRigEditor/Private/EditMode/Picker/SFloatingSpacePicker_AssetEditor.cpp
@@ -6,6 +6,7 @@
 #include "ScopedTransaction.h"
 #include "Editor/SRigSpacePickerWidget.h"
 #include "Rigs/RigHierarchyController.h"
+#include "SControlRigDismissDependencyDialog.h"
 
 #define LOCTEXT_NAMESPACE "SFloatingSpacePicker"
 
@@ -60,8 +61,11 @@ void SFloatingSpacePicker_AssetEditor::OnActiveSpaceChanged(
 	{
 		constexpr bool bInitial = false, bAffectChildren = true;
 		const FTransform Transform = InHierarchy->GetGlobalTransform(InControlKey);
+		FRigDependenciesProviderForControlRig DependenciesProvider(RuntimeRig);
+		DependenciesProvider.SetInteractiveDialogEnabled(true);
+		FControlRigDismissDependencyDialogGuard DependencyDialogGuard(InHierarchy);
+
 		FString OutFailureReason;
-		const FEmptyRigDependenciesProvider DependenciesProvider;
 		if (InHierarchy->SwitchToParent(InControlKey, InSpaceKey, bInitial, bAffectChildren, DependenciesProvider, &OutFailureReason))
 		{
 			InHierarchy->SetGlobalTransform(InControlKey, Transform);
diff --git a/Engine/Plugins/Animation/ControlRig/Source/ControlRigEditor/Private/EditMode/SControlRigSpacePicker.cpp b/Engine/Plugins/Animation/ControlRig/Source/ControlRigEditor/Private/EditMode/SControlRigSpacePicker.cpp
--- a/Engine/Plugins/Animation/ControlRig/Source/ControlRigEditor/Private/EditMode/SControlRigSpacePicker.cpp
+++ b/Engine/Plugins/Animation/ControlRig/Source/ControlRigEditor/Private/EditMode/SControlRigSpacePicker.cpp
@@ -19,6 +19,7 @@
 #include "MovieScene.h"
 #include "ControlRigSpaceChannelEditors.h"
 #include "Models/RigSelectionViewModel.h"
+#include "SControlRigDismissDependencyDialog.h"
 
 #define LOCTEXT_NAMESPACE "ControlRigSpacePicker"
 
@@ -111,7 +112,10 @@ void SControlRigSpacePicker::HandleActiveSpaceChanged(URigHierarchy* InHierarchy
 		if (UControlRig* ControlRig = GetControlRig())
 		{
 			FString FailureReason;
-			const FEmptyRigDependenciesProvider DependencyProvider;
+			FRigDependenciesProviderForControlRig DependencyProvider(ControlRig);
+			DependencyProvider.SetInteractiveDialogEnabled(true);
+			FControlRigDismissDependencyDialogGuard DependencyDialogGuard(InHierarchy);
+			
 			if (!InHierarchy->CanSwitchToParent(InControlKey, InSpaceKey, DependencyProvider, &FailureReason))
 			{
 				// notification
diff --git a/Engine/Plugins/Animation/ControlRig/Source/ControlRigEditor/Private/Editor/Constraints/SConstraintsWidget.cpp b/Engine/Plugins/Animation/ControlRig/Source/ControlRigEditor/Private/Editor/Constraints/SConstraintsWidget.cpp
--- a/Engine/Plugins/Animation/ControlRig/Source/ControlRigEditor/Private/Editor/Constraints/SConstraintsWidget.cpp
+++ b/Engine/Plugins/Animation/ControlRig/Source/ControlRigEditor/Private/Editor/Constraints/SConstraintsWidget.cpp
@@ -23,11 +23,13 @@
 #include "ScopedTransaction.h"
 #include "ISequencer.h"
 #include "ConstraintsManager.h"
+#include "ControlRigGizmoActor.h"
 #include "Styling/SlateIconFinder.h"
 #include "Widgets/Input/NumericTypeInterface.h"
 #include "FrameNumberDetailsCustomization.h"
 #include "SActionButton.h"
 #include "SAssetDropTarget.h"
+#include "SControlRigDismissDependencyDialog.h"
 #include "Sequencer/ControlRigParameterTrackEditor.h"
 #include "Sequencer/EditModeAnimationUtil.h"
 #include "Transform/TransformConstraint.h"
@@ -272,6 +274,17 @@ namespace UE::Private
 		{
 			if (Child != InParent.Object)
 			{
+				// enable the dependency dialog when creating a constraint
+				// for children that are control rig controls.
+				TSharedPtr<FControlRigDismissDependencyDialogGuard> DependencyDialogGuard;
+				if (AControlRigShapeActor* ChildActor = Cast<AControlRigShapeActor>(Child))
+				{
+					if (UControlRig* ControlRig = ChildActor->ControlRig.Get())
+					{
+						DependencyDialogGuard = MakeShared<FControlRigDismissDependencyDialogGuard>(ControlRig->GetHierarchy());
+					}
+				}
+				
 				FScopedTransaction Transaction(LOCTEXT("CreateConstraintKey", "Create Constraint Key"));
 				UTickableTransformConstraint* Constraint =
 					UE::TransformConstraintUtil::CreateAndAddFromObjects(World, Parent, InParent.Socket, Child, NAME_None, InConstraintType, true, bUseDefault);
diff --git a/Engine/Plugins/Animation/ControlRig/Source/ControlRigEditor/Private/Editor/SRigSpacePickerWidget.cpp b/Engine/Plugins/Animation/ControlRig/Source/ControlRigEditor/Private/Editor/SRigSpacePickerWidget.cpp
--- a/Engine/Plugins/Animation/ControlRig/Source/ControlRigEditor/Private/Editor/SRigSpacePickerWidget.cpp
+++ b/Engine/Plugins/Animation/ControlRig/Source/ControlRigEditor/Private/Editor/SRigSpacePickerWidget.cpp
@@ -22,6 +22,7 @@
 #include "RigVMBlueprintGeneratedClass.h"
 #include "SActionButton.h"
 #include "SPositiveActionButton.h"
+#include "SControlRigDismissDependencyDialog.h"
 
 #define LOCTEXT_NAMESPACE "SRigSpacePickerWidget"
 
@@ -787,7 +788,10 @@ void SRigSpacePickerWidget::HandleClickTreeItemInAddMenu(TSharedPtr<FRigTreeElem
 				for(const FRigElementKey& ControlKey : ControlKeys)
 				{
 					FString FailureReason;
-					const FEmptyRigDependenciesProvider DependencyProvider;
+					FRigDependenciesProviderForControlRig DependencyProvider(Cast<UControlRig>(StrongHierarchy->GetOuter()));
+					DependencyProvider.SetInteractiveDialogEnabled(true);
+					FControlRigDismissDependencyDialogGuard DependencyDialogGuard(StrongHierarchy);
+					
 					if (!StrongHierarchy->CanSwitchToParent(ControlKey, Key, DependencyProvider, &FailureReason))
 					{
 						// notification
diff --git a/Engine/Plugins/Animation/ControlRig/Source/ControlRigEditor/Private/SControlRigDismissDependencyDialog.cpp b/Engine/Plugins/Animation/ControlRig/Source/ControlRigEditor/Private/SControlRigDismissDependencyDialog.cpp
new file mode 100644
--- /dev/null
+++ b/Engine/Plugins/Animation/ControlRig/Source/ControlRigEditor/Private/SControlRigDismissDependencyDialog.cpp
@@ -0,0 +1,167 @@
+// Copyright Epic Games, Inc. All Rights Reserved.
+
+#include "SControlRigDismissDependencyDialog.h"
+#include "CoreMinimal.h"
+#include "Dialogs/Dialogs.h"
+#include "Widgets/SWindow.h"
+#include "Widgets/Layout/SScrollBox.h"
+#include "Widgets/Text/STextBlock.h"
+#include "Widgets/Input/SButton.h"
+#include "HAL/PlatformApplicationMisc.h"
+#include "Widgets/Notifications/SNotificationList.h"
+#include "Framework/Notifications/NotificationManager.h"
+#include "ControlRig.h"
+
+class SControlRigDismissDependencyDialog : public SModalEditorDialog<bool>
+{
+	SLATE_BEGIN_ARGS(SControlRigDismissDependencyDialog) {}
+	SLATE_END_ARGS()
+	
+	~SControlRigDismissDependencyDialog()
+	{
+	}
+	
+	void Construct(const FArguments& InArgs, const URigHierarchy* InHierarchy, const FRigElementKey& InChild, const FRigElementKey& InParent, const FRigHierarchyDependencyChain& InDependencyChain)
+	{
+		check(InHierarchy);
+
+		FString ClassName = InHierarchy->GetPathName();
+		if (const UControlRig* ControlRig = Cast<UControlRig>(InHierarchy->GetOuter()))
+		{
+			ClassName = ControlRig->GetClass()->GetPackage()->GetName();
+		}
+
+		const FText TitleFormat = NSLOCTEXT("ControlRig", "DismissDependencyTitleFormat", "Potential Cycle relating {1} to {0}");
+		const FText Title = FText::Format(TitleFormat, FText::FromString(InChild.ToString()), FText::FromString(InParent.ToString()));
+		const FText Format = NSLOCTEXT("ControlRig", "DismissDependencyMessageFormat", "{0}\n\n{1}\n\n{2}");
+		Message = FText::Format(
+			Format,
+			FText::FromString(ClassName),
+			Title,
+			FText::FromString(InHierarchy->GetMessageFromDependencyChain(InDependencyChain))
+		);
+		
+		ChildSlot
+		[
+			SNew(SBorder)
+			.BorderImage(FAppStyle::GetBrush("AssetDeleteDialog.Background"))
+			.Padding(10.0f)
+			[
+				SNew(SVerticalBox)
+
+				+ SVerticalBox::Slot()
+				.FillHeight(1)
+				.HAlign(HAlign_Fill)
+				.Padding(5.f)
+				[
+					SNew(SScrollBox)
+					.Orientation(Orient_Vertical)
+					+ SScrollBox::Slot()
+					.HAlign(HAlign_Fill)
+					.FillSize(1.0)
+					[
+						SNew(STextBlock)
+						.Font(FAppStyle::GetFontStyle("NormalFont"))
+						.HighlightText(Title)
+						.Text(Message)
+						.OnDoubleClicked(this, &SControlRigDismissDependencyDialog::OnMessageDoubleClicked)
+					]
+					+ SScrollBox::Slot()
+					.HAlign(HAlign_Fill)
+					.AutoSize()
+					[
+						SNew(STextBlock)
+						.Font(FAppStyle::GetFontStyle("SmallFont"))
+						.Text(NSLOCTEXT("ControlRig", "DoubleClickThisMessageToCopy", "Ignoring a cycle may cause unexpected behaviour in the rig.\n\nDouble-Click the message above to copy it.\n\nThis dialog can be disabled with 'ControlRig.Hierarchy.AllowToIgnoreCycles'."))
+						.OnDoubleClicked(this, &SControlRigDismissDependencyDialog::OnMessageDoubleClicked)
+					]
+				]
+
+				+ SVerticalBox::Slot()
+				.AutoHeight()
+				.HAlign(HAlign_Right)
+				.Padding(5.f)
+				[
+					SNew(SHorizontalBox)
+
+					+ SHorizontalBox::Slot()
+					.AutoWidth()
+					.HAlign(HAlign_Right)
+					.VAlign(VAlign_Center)
+					[
+						SNew(SButton)
+						.ContentPadding(FMargin(10, 5))
+						.Text(NSLOCTEXT("ControlRig", "Cancel", "Cancel"))
+						.ToolTipText(NSLOCTEXT("ControlRig", "ActionCancelDismissDependency", "Aborts the action and respects the dependency."))
+						.OnClicked(this, &SControlRigDismissDependencyDialog::OnCancel)
+					]
+
+					+ SHorizontalBox::Slot()
+					.AutoWidth()
+					.HAlign(HAlign_Right)
+					.VAlign(VAlign_Center)
+					[
+						SNew(SButton)
+						.ButtonStyle(FAppStyle::Get(), "FlatButton.Danger")
+						.ContentPadding(FMargin(10, 5))
+						.Text(NSLOCTEXT("ControlRig", "Ignore", "Ignore"))
+						.ToolTipText(NSLOCTEXT("ControlRig", "ActionDismissDependency", "Ignore the dependency and proceed with the action."))
+						.OnClicked(this, &SControlRigDismissDependencyDialog::OnDismissDependency)
+					]
+				]
+			]
+		];
+	}
+
+	virtual FReply OnKeyDown(const FGeometry& MyGeometry, const FKeyEvent& InKeyEvent) override
+	{
+		if (InKeyEvent.GetKey() == EKeys::Escape)
+		{
+			return OnCancel();
+		}
+		if (InKeyEvent.GetKey() == EKeys::C && InKeyEvent.IsControlDown())
+		{
+			OnMessageDoubleClicked(MyGeometry, FPointerEvent());
+		}
+		return SModalEditorDialog<bool>::OnKeyDown(MyGeometry, InKeyEvent);
+	}
+
+private:
+
+	FReply OnCancel()
+	{
+		ProvideResult(false);
+		return FReply::Handled();
+	}
+
+	FReply OnDismissDependency()
+	{
+		ProvideResult(true);
+		return FReply::Handled();
+	}
+
+	FReply OnMessageDoubleClicked(const FGeometry& MyGeometry, const FPointerEvent& PointerEvent) const
+	{
+		FPlatformApplicationMisc::ClipboardCopy(*Message.ToString());
+
+		FNotificationInfo Info(NSLOCTEXT("ControlRig", "DependencyChainMessageCopiedToClipboard", "Message has been copied to the clipboard!"));
+		Info.bFireAndForget = true;
+		Info.FadeOutDuration = 0.5f;
+		Info.ExpireDuration = 2.0f;
+		TSharedPtr<SNotificationItem> NotificationPtr = FSlateNotificationManager::Get().AddNotification(Info);
+		NotificationPtr->SetCompletionState(SNotificationItem::CS_Success);
+
+		return FReply::Handled();
+	}
+
+	FText Message;
+};
+
+
+bool ControlRigDismissDependencyDialog::LaunchDismissDependencyDialog(const URigHierarchy* InHierarchy, const FRigElementKey& InChild, const FRigElementKey& InParent, const FRigHierarchyDependencyChain& InDependencyChain)
+{
+	// todo?: compute a key to remember this setting and don't ask in the future
+	TSharedPtr<SControlRigDismissDependencyDialog> Dialog = SNew(SControlRigDismissDependencyDialog, InHierarchy, InChild, InParent, InDependencyChain);
+	const bool bDismissed = Dialog->ShowModalDialog(NSLOCTEXT("ControlRig", "DismissDependencyDialogTitle","Cyclic Dependency Found"));
+	return bDismissed;
+}
diff --git a/Engine/Plugins/Animation/ControlRig/Source/ControlRigEditor/Private/SControlRigDismissDependencyDialog.h b/Engine/Plugins/Animation/ControlRig/Source/ControlRigEditor/Private/SControlRigDismissDependencyDialog.h
new file mode 100644
--- /dev/null
+++ b/Engine/Plugins/Animation/ControlRig/Source/ControlRigEditor/Private/SControlRigDismissDependencyDialog.h
@@ -0,0 +1,35 @@
+// Copyright Epic Games, Inc. All Rights Reserved.
+#pragma once
+
+#include "CoreMinimal.h"
+#include "Rigs/RigHierarchy.h"
+
+struct ControlRigDismissDependencyDialog
+{
+	static bool LaunchDismissDependencyDialog(const URigHierarchy* InHierarchy, const FRigElementKey& InChild, const FRigElementKey& InParent, const FRigHierarchyDependencyChain& InDependencyChain);
+};
+
+class FControlRigDismissDependencyDialogGuard
+{
+public:
+	FControlRigDismissDependencyDialogGuard(URigHierarchy* InHierarchy)
+	: Hierarchy(InHierarchy)
+	{
+		if (Hierarchy.IsValid())
+		{
+			Hierarchy->OnDependencyDismissed().BindStatic(&ControlRigDismissDependencyDialog::LaunchDismissDependencyDialog);
+		}
+	}
+
+	~FControlRigDismissDependencyDialogGuard()
+	{
+		if (Hierarchy.IsValid())
+		{
+			Hierarchy->OnDependencyDismissed().Unbind();
+		}
+	}
+
+private:
+
+	TWeakObjectPtr<URigHierarchy> Hierarchy;
+};
\ No newline at end of file
