Control Rig: Fix viewport tweener not inserting keys on curves when tweening CRU nodes.

The issue is not related to CRU nodes specifically. The issue was that the code was not considering the fact that there may be more than one control rig track under the control rig object being animated.

diff --git a/Engine/Plugins/Animation/ControlRig/Source/ControlRigEditor/Private/EditMode/Tween/ControlRigObjectSelection.cpp b/Engine/Plugins/Animation/ControlRig/Source/ControlRigEditor/Private/EditMode/Tween/ControlRigObjectSelection.cpp
--- a/Engine/Plugins/Animation/ControlRig/Source/ControlRigEditor/Private/EditMode/Tween/ControlRigObjectSelection.cpp
+++ b/Engine/Plugins/Animation/ControlRig/Source/ControlRigEditor/Private/EditMode/Tween/ControlRigObjectSelection.cpp
@@ -187,10 +187,17 @@ bool FControlRigObjectSelection::Setup(const TArray<UControlRig*>& SelectedContr
 				}
 			}
 		}
-		const UMovieSceneControlRigParameterTrack* Track = Cast<UMovieSceneControlRigParameterTrack>(MovieScene->FindTrack(UMovieSceneControlRigParameterTrack::StaticClass(), Binding.GetObjectGuid(), NAME_None));
-		if (Track && Track->GetControlRig() && SelectedControlRigs.Contains(Track->GetControlRig()))
+
+		const TArray<UMovieSceneTrack*> Tracks = MovieScene->FindTracks(
+			UMovieSceneControlRigParameterTrack::StaticClass(), Binding.GetObjectGuid(), NAME_None
+			);
+		for (const UMovieSceneTrack* Track : Tracks)
 		{
-			SetupControlRigTrackChannels(Track);
+			const UMovieSceneControlRigParameterTrack* ControlRigTrack = Cast<UMovieSceneControlRigParameterTrack>(Track);
+			if (ControlRigTrack && ControlRigTrack->GetControlRig() && SelectedControlRigs.Contains(ControlRigTrack->GetControlRig()))
+			{
+				SetupControlRigTrackChannels(ControlRigTrack);
+			}
 		}
 	}
 	
@@ -298,4 +305,4 @@ TArray<UControlRig*> FControlRigObjectSelection::GetControlRigs(const TWeakPtr<F
 	}
 	return ControlRigs;
 }
-}
+}
\ No newline at end of file
