Anim In Engine: Regression from 'PR #13645: Make generic vector channel data display more robust for position/rotation/scale control types. ' The lambdas for getting values was cached for just the first control of a vector type, so other vector types wouldn't work with key area operations

#DisabledValidations Preflight

diff --git a/Engine/Plugins/Animation/ControlRig/Source/ControlRig/Private/Sequencer/MovieSceneControlRigParameterSection.cpp b/Engine/Plugins/Animation/ControlRig/Source/ControlRig/Private/Sequencer/MovieSceneControlRigParameterSection.cpp
--- a/Engine/Plugins/Animation/ControlRig/Source/ControlRig/Private/Sequencer/MovieSceneControlRigParameterSection.cpp
+++ b/Engine/Plugins/Animation/ControlRig/Source/ControlRig/Private/Sequencer/MovieSceneControlRigParameterSection.cpp
@@ -243,15 +243,6 @@ struct FParameterVectorChannelEditorData
 		};
 		static_assert(UE_ARRAY_COUNT(ScaleDisplayTexts) == UE_ARRAY_COUNT(MetaData));
 		
-		static const TFunction<TOptional<float>(UObject&, FTrackInstancePropertyBindings*)> GetExternalValueFuncs[] =
-		{
-			[InControlRig, InName, NumChannels](UObject& InObject, FTrackInstancePropertyBindings*) { return ExtractChannelX(InObject, InControlRig, InName, NumChannels); },
-			[InControlRig, InName, NumChannels](UObject& InObject, FTrackInstancePropertyBindings*) { return ExtractChannelY(InObject, InControlRig, InName, NumChannels); },
-			[InControlRig, InName, NumChannels](UObject& InObject, FTrackInstancePropertyBindings*) { return ExtractChannelZ(InObject, InControlRig, InName, NumChannels); },
-			[InControlRig, InName, NumChannels](UObject& InObject, FTrackInstancePropertyBindings*) { return ExtractChannelW(InObject, InControlRig, InName, NumChannels); },
-		};
-		static_assert(UE_ARRAY_COUNT(GetExternalValueFuncs) == UE_ARRAY_COUNT(ExternalValues));
-
 		const FString NameAsString = InName.ToString();
 		
 		check(NumChannels <= UE_ARRAY_COUNT(MetaData));
@@ -292,14 +283,22 @@ struct FParameterVectorChannelEditorData
 			MetaData[ChannelIdx].bEnabled = bEnabledOverride;
 			MetaData[ChannelIdx].SortOrder = SortStartIndex++;
 			MetaData[ChannelIdx].bCanCollapseToTrack = true;
-			
-			ExternalValues[ChannelIdx].OnGetExternalValue = GetExternalValueFuncs[ChannelIdx];
-			ExternalValues[ChannelIdx].OnGetCurrentValueAndWeight = [InName, NumChannels, ChannelIdx](UObject* Object, UMovieSceneSection* SectionToKey, FFrameNumber KeyTime,
-				FFrameRate TickResolution, FMovieSceneRootEvaluationTemplateInstance& RootTemplate, float& OutValue, float& OutWeight)
-			{
-				GetChannelValueAndWeight(InName, NumChannels, ChannelIdx, Object, SectionToKey, KeyTime, TickResolution, RootTemplate, OutValue, OutWeight);
-			};
 		}
+
+		ExternalValues[0].OnGetExternalValue = [InControlRig, InName, NumChannels](UObject& InObject, FTrackInstancePropertyBindings* Bindings) { return ExtractChannelX(InObject, InControlRig, InName, NumChannels); };
+		ExternalValues[1].OnGetExternalValue = [InControlRig, InName, NumChannels](UObject& InObject, FTrackInstancePropertyBindings* Bindings) { return ExtractChannelY(InObject, InControlRig, InName, NumChannels); };
+		ExternalValues[2].OnGetExternalValue = [InControlRig, InName, NumChannels](UObject& InObject, FTrackInstancePropertyBindings* Bindings) { return ExtractChannelZ(InObject, InControlRig, InName, NumChannels); };
+		ExternalValues[3].OnGetExternalValue = [InControlRig, InName, NumChannels](UObject& InObject, FTrackInstancePropertyBindings* Bindings) { return ExtractChannelW(InObject, InControlRig, InName, NumChannels); };
+
+		ExternalValues[0].OnGetCurrentValueAndWeight = [InName, NumChannels](UObject* Object, UMovieSceneSection* SectionToKey, FFrameNumber KeyTime, FFrameRate TickResolution, FMovieSceneRootEvaluationTemplateInstance& RootTemplate,
+			float& OutValue, float& OutWeight) { GetChannelValueAndWeight(InName, NumChannels, 0, Object, SectionToKey, KeyTime, TickResolution, RootTemplate, OutValue, OutWeight); };
+		ExternalValues[1].OnGetCurrentValueAndWeight = [InName, NumChannels](UObject* Object, UMovieSceneSection* SectionToKey, FFrameNumber KeyTime, FFrameRate TickResolution, FMovieSceneRootEvaluationTemplateInstance& RootTemplate,
+			float& OutValue, float& OutWeight) { GetChannelValueAndWeight(InName, NumChannels, 1, Object, SectionToKey, KeyTime, TickResolution, RootTemplate, OutValue, OutWeight); };
+		ExternalValues[2].OnGetCurrentValueAndWeight = [InName, NumChannels](UObject* Object, UMovieSceneSection* SectionToKey, FFrameNumber KeyTime, FFrameRate TickResolution, FMovieSceneRootEvaluationTemplateInstance& RootTemplate,
+			float& OutValue, float& OutWeight) { GetChannelValueAndWeight(InName, NumChannels, 2, Object, SectionToKey, KeyTime, TickResolution, RootTemplate, OutValue, OutWeight); };
+		ExternalValues[3].OnGetCurrentValueAndWeight = [InName, NumChannels](UObject* Object, UMovieSceneSection* SectionToKey, FFrameNumber KeyTime, FFrameRate TickResolution, FMovieSceneRootEvaluationTemplateInstance& RootTemplate,
+			float& OutValue, float& OutWeight) { GetChannelValueAndWeight(InName, NumChannels, 3, Object, SectionToKey, KeyTime, TickResolution, RootTemplate, OutValue, OutWeight); };
+
 	}
 
 	static FVector4 GetPropertyValue(UControlRig* ControlRig, FName ParameterName, UObject& InObject,int32 NumChannels)
