Anim Layers: Adjustment Blending needs to sort controls in parent to leaf order since we do operations in world space

diff --git a/Engine/Plugins/Animation/ControlRig/Source/ControlRigEditor/Private/Sequencer/AnimLayers/AnimLayers.cpp b/Engine/Plugins/Animation/ControlRig/Source/ControlRigEditor/Private/Sequencer/AnimLayers/AnimLayers.cpp
--- a/Engine/Plugins/Animation/ControlRig/Source/ControlRigEditor/Private/Sequencer/AnimLayers/AnimLayers.cpp
+++ b/Engine/Plugins/Animation/ControlRig/Source/ControlRigEditor/Private/Sequencer/AnimLayers/AnimLayers.cpp
@@ -3369,16 +3369,16 @@ FTransformAsArray operator+(const FTransformAsArray& LHS, const FTransformAsArra
 	return ReturnVal;
 }
 
-void GetWorldTransformPercentageOfChange(int32 FrameStartIndex, int32 FrameEndIndex, const TPair<FName, TSharedPtr<UE::AIE::FArrayOfTransforms>>& Pair, TArray<FTransformAsArray>& ValueDifferences, TArray<TArray<double>>& PercentageDifferences)
+void GetWorldTransformPercentageOfChange(int32 FrameStartIndex, int32 FrameEndIndex, const TSharedPtr<UE::AIE::FArrayOfTransforms>& ArrayOfTransforms, TArray<FTransformAsArray>& ValueDifferences, TArray<TArray<double>>& PercentageDifferences)
 {
 	const int32 NumKeys = FrameEndIndex - FrameStartIndex + 1;
 	FTransformAsArray CurrentArray;
-	FTransformAsArray NextArray(FTransformAsArray((*Pair.Value)[FrameStartIndex]));
+	FTransformAsArray NextArray(FTransformAsArray(ArrayOfTransforms->Transforms[FrameStartIndex]));
 	FTransformAsArray TotalDifferences,TotalChange;
 	for (int32 Index = FrameStartIndex; Index < FrameEndIndex; ++Index)
 	{
 		CurrentArray = NextArray;
-		FTransform& NextTransform = (*Pair.Value)[Index + 1];
+		FTransform& NextTransform = (ArrayOfTransforms->Transforms)[Index + 1];
 		NextArray = FTransformAsArray(NextTransform);
 		ValueDifferences[Index] = (NextArray - CurrentArray).Abs();
 		TotalChange += ValueDifferences[Index];
@@ -3520,11 +3520,17 @@ void AdjustmentBlendWorldSpace(ISequencer* Sequencer,  FAnimLayerItem& BaseAnimL
 	TArray<TArray<double>> PercentageDifferences;
 	FTransformAsArray StartValue, EndValue;
 
+
 	for (const TPair< UControlRig*, FControlRigAndWorldTransforms>& BaseControlRigAndWorld : BaseEvaluatingControlRigs)
 	{
 		UControlRig* ControlRig = BaseControlRigAndWorld.Key;
 		if (FControlRigAndWorldTransforms* EvalCRandWT = EvaluatingControlRigs.Find(ControlRig))
 		{
+			URigHierarchy* RigHierarchy = ControlRig->GetHierarchy();
+			if (!RigHierarchy)
+			{
+				continue;
+			}
 			for (int32 KeyIndex = 0; KeyIndex < AdditiveKeyIndices.Num() - 1; ++KeyIndex)
 			{
 				int32 FrameStartIndex = AdditiveKeyIndices[KeyIndex];
@@ -3533,84 +3539,98 @@ void AdjustmentBlendWorldSpace(ISequencer* Sequencer,  FAnimLayerItem& BaseAnimL
 				{
 					continue;  //should never happen
 				}
-				for (const TPair<FName, TSharedPtr<FArrayOfTransforms>>& BasePair : BaseControlRigAndWorld.Value.ControlAndWorldTransforms)
+				//needed controls sorted from root to leaves since we are setting things in global/world space
+				TArray <FName> ControlNames;
+				BaseControlRigAndWorld.Value.ControlAndWorldTransforms.GenerateKeyArray(ControlNames);
+				TArray<FRigElementKey> ControlKeys;
+				for (const FName& ControlName : ControlNames)
+				{
+					ControlKeys.Add(FRigElementKey(ControlName, ERigElementType::Control));
+				}
+				ControlKeys = RigHierarchy->SortKeys(ControlKeys);
+
+				//for (const TPair<FName, TSharedPtr<FArrayOfTransforms>>& BasePair : BaseControlRigAndWorld.Value.ControlAndWorldTransforms)
+				for(const FRigElementKey& Key: ControlKeys)
 				{
-					if (FRigControlElement* ControlElement = ControlRig->FindControl(BasePair.Key))
+					if (const TSharedPtr<FArrayOfTransforms>* Transforms = BaseControlRigAndWorld.Value.ControlAndWorldTransforms.Find(Key.Name))
 					{
-						if (ControlTypeValidForAdjustment(ControlElement) == false)
+						if (FRigControlElement* ControlElement = ControlRig->FindControl(Key.Name))
 						{
-							continue;
-						}
-						if (TSharedPtr<FArrayOfTransforms>* EvalTransforms = EvalCRandWT->ControlAndWorldTransforms.Find(BasePair.Key))
-						{
-							ValueDifferences.Reset();
-							ValueDifferences.SetNum(CurrentFrameTimes.NumFrames);
-							//get base layer values as diffs and percentages
-							GetWorldTransformPercentageOfChange(FrameStartIndex, FrameEndIndex, BasePair, ValueDifferences, PercentageDifferences);
-							//get new layer change
-							FTransform STBase = BasePair.Value->Transforms[FrameStartIndex];
-							FTransformAsArray StartTransformBase(STBase);
-							FTransformAsArray StartTransformSubtractBase((*EvalTransforms)->Transforms[FrameStartIndex]);
-							StartTransformSubtractBase = StartTransformSubtractBase - StartTransformBase;
-							
-							FTransform ETBase = BasePair.Value->Transforms[FrameEndIndex];
-							FTransformAsArray EndTransformBase(ETBase);
-							FTransformAsArray EndTransformSubtractBase((*EvalTransforms)->Transforms[FrameEndIndex]);
-							EndTransformSubtractBase = EndTransformSubtractBase - EndTransformBase;
-				
-							FTransformAsArray CurrentValue;
-							FTransformAsArray PreviousValue;
-							FTransformAsArray TotalPoseLayerChange;
-
-							for (int32 ChannelIndex = 0; ChannelIndex < 9; ++ChannelIndex)
+							if (ControlTypeValidForAdjustment(ControlElement) == false)
 							{
-								TotalPoseLayerChange[ChannelIndex] = FMath::Abs(StartTransformSubtractBase[ChannelIndex] - EndTransformSubtractBase[ChannelIndex]);
-								PreviousValue[ChannelIndex] = StartTransformSubtractBase[ChannelIndex];
+								continue;
 							}
-							FTransform CT,NT;
-							FTransformAsArray CurrentTransform, NewTransform;
-							for (int32 FrameIndex = FrameStartIndex + 1; FrameIndex <= FrameEndIndex; ++FrameIndex)
+							if (TSharedPtr<FArrayOfTransforms>* EvalTransforms = EvalCRandWT->ControlAndWorldTransforms.Find(Key.Name))
 							{
+								ValueDifferences.Reset();
+								ValueDifferences.SetNum(CurrentFrameTimes.NumFrames);
+								//get base layer values as diffs and percentages
+								GetWorldTransformPercentageOfChange(FrameStartIndex, FrameEndIndex, *Transforms, ValueDifferences, PercentageDifferences);
+								//get new layer change
+								FTransform STBase = (*Transforms)->Transforms[FrameStartIndex];
+								FTransformAsArray StartTransformBase(STBase);
+								FTransformAsArray StartTransformSubtractBase((*EvalTransforms)->Transforms[FrameStartIndex]);
+								StartTransformSubtractBase = StartTransformSubtractBase - StartTransformBase;
+
+								FTransform ETBase = (*Transforms)->Transforms[FrameEndIndex];
+								FTransformAsArray EndTransformBase(ETBase);
+								FTransformAsArray EndTransformSubtractBase((*EvalTransforms)->Transforms[FrameEndIndex]);
+								EndTransformSubtractBase = EndTransformSubtractBase - EndTransformBase;
+
+								FTransformAsArray CurrentValue;
+								FTransformAsArray PreviousValue;
+								FTransformAsArray TotalPoseLayerChange;
 
-								const FFrameNumber FrameNumber = CurrentFrameTimes.CalculateFrame(FrameIndex);
-								FFrameTime GlobalTime(FrameNumber);
-								GlobalTime = LocalToRootTransform.TryTransformTime(GlobalTime).Get(GlobalTime); //player evals in root time so need to go back to it.
-
-								FMovieSceneContext MovieSceneContext = FMovieSceneContext(FMovieSceneEvaluationRange(GlobalTime, TickResolution), Sequencer->GetPlaybackStatus()).SetHasJumped(true);
-								Sequencer->GetEvaluationTemplate().EvaluateSynchronousBlocking(MovieSceneContext);
-								ControlRig->Evaluate_AnyThread();
 								for (int32 ChannelIndex = 0; ChannelIndex < 9; ++ChannelIndex)
 								{
-									if (FMath::IsNearlyZero(TotalPoseLayerChange[ChannelIndex], (1.e-4f)))
-									{
-										CurrentValue[ChannelIndex] =PreviousValue[ChannelIndex];
-									}
-									else if (FrameIndex < PercentageDifferences[ChannelIndex].Num())
-									{
-										double Value = PercentageDifferences[ChannelIndex][FrameIndex - 1];
-										const double ValueDelta = (TotalPoseLayerChange[ChannelIndex] / 100.0) * Value;
-										CurrentValue[ChannelIndex] = (EndTransformSubtractBase[ChannelIndex] > StartTransformSubtractBase[ChannelIndex]) ?
-											PreviousValue[ChannelIndex] + ValueDelta : PreviousValue[ChannelIndex] - ValueDelta;
-									}
-									else
+									TotalPoseLayerChange[ChannelIndex] = FMath::Abs(StartTransformSubtractBase[ChannelIndex] - EndTransformSubtractBase[ChannelIndex]);
+									PreviousValue[ChannelIndex] = StartTransformSubtractBase[ChannelIndex];
+								}
+								FTransform CT, NT;
+								FTransformAsArray CurrentTransform, NewTransform;
+								for (int32 FrameIndex = FrameStartIndex + 1; FrameIndex <= FrameEndIndex; ++FrameIndex)
+								{
+
+									const FFrameNumber FrameNumber = CurrentFrameTimes.CalculateFrame(FrameIndex);
+									FFrameTime GlobalTime(FrameNumber);
+									GlobalTime = LocalToRootTransform.TryTransformTime(GlobalTime).Get(GlobalTime); //player evals in root time so need to go back to it.
+
+									FMovieSceneContext MovieSceneContext = FMovieSceneContext(FMovieSceneEvaluationRange(GlobalTime, TickResolution), Sequencer->GetPlaybackStatus()).SetHasJumped(true);
+									Sequencer->GetEvaluationTemplate().EvaluateSynchronousBlocking(MovieSceneContext);
+									ControlRig->Evaluate_AnyThread();
+									for (int32 ChannelIndex = 0; ChannelIndex < 9; ++ChannelIndex)
 									{
-										CurrentValue[ChannelIndex] = StartTransformSubtractBase[ChannelIndex] + (EndTransformSubtractBase[ChannelIndex] - StartTransformSubtractBase[ChannelIndex]) * (FrameIndex - FrameStartIndex) / (FrameEndIndex - FrameStartIndex);
+										if (FMath::IsNearlyZero(TotalPoseLayerChange[ChannelIndex], (1.e-4f)))
+										{
+											CurrentValue[ChannelIndex] = PreviousValue[ChannelIndex];
+										}
+										else if (FrameIndex < PercentageDifferences[ChannelIndex].Num())
+										{
+											double Value = PercentageDifferences[ChannelIndex][FrameIndex - 1];
+											const double ValueDelta = (TotalPoseLayerChange[ChannelIndex] / 100.0) * Value;
+											CurrentValue[ChannelIndex] = (EndTransformSubtractBase[ChannelIndex] > StartTransformSubtractBase[ChannelIndex]) ?
+												PreviousValue[ChannelIndex] + ValueDelta : PreviousValue[ChannelIndex] - ValueDelta;
+										}
+										else
+										{
+											CurrentValue[ChannelIndex] = StartTransformSubtractBase[ChannelIndex] + (EndTransformSubtractBase[ChannelIndex] - StartTransformSubtractBase[ChannelIndex]) * (FrameIndex - FrameStartIndex) / (FrameEndIndex - FrameStartIndex);
+										}
+										PreviousValue[ChannelIndex] = CurrentValue[ChannelIndex];
+
 									}
-									PreviousValue[ChannelIndex] = CurrentValue[ChannelIndex];
+									//get current transform of the base
+									CT = (*Transforms)->Transforms[FrameIndex];
+									CT = CT.GetRelativeTransform((*EvalCRandWT->ParentTransforms)[FrameIndex]);
+									CurrentTransform = FTransformAsArray(CT);
+									NewTransform = CurrentTransform + CurrentValue;
+									NT = NewTransform.AsTransform();;
+									Context.LocalTime = TickResolution.AsSeconds(FFrameTime(FrameNumber));
+									FRigControlValue Value = ControlRig->GetControlValueFromGlobalTransform(ControlElement->GetKey().Name, NT, TransformType);
+									NT = Value.GetAsTransform(ControlElement->Settings.ControlType, ControlElement->Settings.PrimaryAxis);
+									FEulerTransform EulerTransform(NT);;
+									UE::AIE::FSetTransformHelpers::SetControlTransform(EulerTransform, ControlRig, ControlElement, Context);
 
 								}
-								//get current transform of the base
-								CT = BasePair.Value->Transforms[FrameIndex];
-								CT = CT.GetRelativeTransform((*EvalCRandWT->ParentTransforms)[FrameIndex]);
-								CurrentTransform = FTransformAsArray(CT);
-								NewTransform = CurrentTransform + CurrentValue;
-								NT = NewTransform.AsTransform();;
-								Context.LocalTime = TickResolution.AsSeconds(FFrameTime(FrameNumber));
-								FRigControlValue Value = ControlRig->GetControlValueFromGlobalTransform(ControlElement->GetKey().Name, NT, TransformType);
-								NT = Value.GetAsTransform(ControlElement->Settings.ControlType, ControlElement->Settings.PrimaryAxis);
-								FEulerTransform EulerTransform(NT);;
-								UE::AIE::FSetTransformHelpers::SetControlTransform(EulerTransform, ControlRig, ControlElement, Context);
-
 							}
 						}
 					}
