ControlRig & dependency checker: only record instructions for forward solve events

to avoid keeping instructions from other ones (such as construction, for example) that disrupt the construction of dependencies.
a more sophisticated version would be to store the instructions by event and thus be able to query them according to the context in which they are used.

#rnx

diff --git a/Engine/Plugins/Animation/ControlRig/Source/ControlRig/Private/ControlRig.cpp b/Engine/Plugins/Animation/ControlRig/Source/ControlRig/Private/ControlRig.cpp
--- a/Engine/Plugins/Animation/ControlRig/Source/ControlRig/Private/ControlRig.cpp
+++ b/Engine/Plugins/Animation/ControlRig/Source/ControlRig/Private/ControlRig.cpp
@@ -1103,6 +1103,12 @@ bool UControlRig::Execute(const FName& InEventName)
 			bRecordInstructionsAtRuntime = false;
 		}
 		
+		// only record instructions when running one of the forward solve events for now
+		if (!bIsPreForwardSolve && !bIsForwardSolve && !bIsPostForwardSolve)
+		{
+			bRecordInstructionsAtRuntime = false;
+		}
+		
 		RecordTransformsPerInstructionGuard = MakeShared<TGuardValue<bool>>(Hierarchy->bRecordInstructionsAtRuntime, bRecordInstructionsAtRuntime);
 		if (!bIsEventInQueue || bIsEventFirstInQueue)
 		{
