ControlRig Edit Mode: ability to manipulate controls and actors simultaneously

#rnx

diff --git a/Engine/Plugins/Animation/ControlRig/Source/ControlRigEditor/Private/EditMode/ControlRigEditMode.cpp b/Engine/Plugins/Animation/ControlRig/Source/ControlRigEditor/Private/EditMode/ControlRigEditMode.cpp
--- a/Engine/Plugins/Animation/ControlRig/Source/ControlRigEditor/Private/EditMode/ControlRigEditMode.cpp
+++ b/Engine/Plugins/Animation/ControlRig/Source/ControlRigEditor/Private/EditMode/ControlRigEditMode.cpp
@@ -1141,6 +1141,8 @@ bool FControlRigEditMode::EndTransform(const FGizmoState& InState)
 
 bool FControlRigEditMode::HandleBeginTransform(const FEditorViewportClient* InViewportClient)
 {
+	bOtherSelectedObjects = false;
+	
 	if (!InViewportClient)
 	{
 		return false;
@@ -1278,7 +1280,14 @@ bool FControlRigEditMode::HandleBeginTransform(const FEditorViewportClient* InVi
 	{
 		bManipulatorMadeChange = false;
 	}
-	return InteractionScopes.Num() != 0;
+
+	const bool bInteractingWithControls = !InteractionScopes.IsEmpty();
+	
+	// do not stop if there are other selected objects that are not related to the control rig mode
+	// note that we are only interested in this value if controls are also being manipulated
+	bOtherSelectedObjects = bInteractingWithControls ? HasAnyOtherSelectedObjects() : false;
+	
+	return bOtherSelectedObjects ? false : bInteractingWithControls;
 }
 
 bool FControlRigEditMode::HandleEndTransform(FEditorViewportClient* InViewportClient)
@@ -1290,10 +1299,12 @@ bool FControlRigEditMode::HandleEndTransform(FEditorViewportClient* InViewportCl
 
 	const FControlRigInteractionTransformContext TransformContext(InViewportClient->GetWidgetMode());
 	const bool bWasInteracting = bManipulatorMadeChange && InteractionType != (uint8)EControlRigInteractionType::None;
+	const bool bHadOtherSelectedObjects = bOtherSelectedObjects;
 	
 	InteractionType = (uint8)EControlRigInteractionType::None;
 	bIsTracking = false;
 	InteractionDependencies.Reset();
+	bOtherSelectedObjects = false;
 	
 	if (InteractionScopes.Num() > 0)
 	{		
@@ -1352,8 +1363,9 @@ bool FControlRigEditMode::HandleEndTransform(FEditorViewportClient* InViewportCl
 			static constexpr bool bInvalidateHitProxies = true;
 			InViewportClient->Invalidate(bInvalidateChildViews, bInvalidateHitProxies);
 		}
-		
-		return true;
+
+		// do not stop if there were other selected objects that were not related to the control rig mode
+		return !bHadOtherSelectedObjects;
 	}
 
 	bManipulatorMadeChange = false;
@@ -2325,9 +2337,9 @@ bool FControlRigEditMode::InputDelta(FEditorViewportClient* InViewportClient, FV
 		
 		TickManipulatableObjects(RigsToTick);
 	}
-	
-	//if in level editor we want to move other things also
-	return IsInLevelEditor() ? false :bManipulatorMadeChange;
+
+	// do not stop if there are other selected objects that are not related to the control rig mode
+	return bOtherSelectedObjects ? false : bManipulatorMadeChange;
 }
 
 bool FControlRigEditMode::ShouldDrawWidget() const
@@ -7022,4 +7034,40 @@ TOptional<FTransform> FControlRigEditMode::GetConstraintParentTransform(const UC
 	return Dummy;
 }
 
+bool FControlRigEditMode::HasAnyOtherSelectedObjects() const
+{
+	if (IsInLevelEditor())
+	{
+		// check actor selection (minus control shapes)
+		TArray< TWeakObjectPtr<UObject> > Selection;
+		GetModeManager()->GetSelectedActors()->GetSelectedObjects(Selection);
+		Selection.RemoveAll([](const TWeakObjectPtr<UObject>& WeakActor)
+		{
+			return !WeakActor.IsValid() || WeakActor->IsA<AControlRigShapeActor>();
+		});
+
+		if (!Selection.IsEmpty())
+		{
+			return true;
+		}
+		
+		// check component selection (minus bound components)
+		GetModeManager()->GetSelectedComponents()->GetSelectedObjects(Selection);
+		for (const TWeakObjectPtr<UControlRig>& ControlRigPtr : RuntimeControlRigs)
+		{
+			if (const TSharedPtr<IControlRigObjectBinding>& Binding = ControlRigPtr.IsValid() ? ControlRigPtr->GetObjectBinding() : nullptr)
+			{
+				if (UObject* BoundObject = Binding->GetBoundObject())
+				{
+					Selection.Remove(BoundObject);
+				}
+			}
+		}
+		
+		return !Selection.IsEmpty();		
+	}
+	
+	return false;
+}
+
 #undef LOCTEXT_NAMESPACE
diff --git a/Engine/Plugins/Animation/ControlRig/Source/ControlRigEditor/Public/EditMode/ControlRigEditMode.h b/Engine/Plugins/Animation/ControlRig/Source/ControlRigEditor/Public/EditMode/ControlRigEditMode.h
--- a/Engine/Plugins/Animation/ControlRig/Source/ControlRigEditor/Public/EditMode/ControlRigEditMode.h
+++ b/Engine/Plugins/Animation/ControlRig/Source/ControlRigEditor/Public/EditMode/ControlRigEditMode.h
@@ -783,6 +783,11 @@ private:
 	// Used to store and apply keyframes (if deferred)  
 	mutable UE::AnimationEditMode::FControlRigKeyframer Keyframer;
 
+	// whether there are other objects selected in the world that are not shape actors or components bound to the rigs
+	bool bOtherSelectedObjects = false;
+	
+	bool HasAnyOtherSelectedObjects() const;
+
 	friend class FControlRigEditorModule;
 	friend class FControlRigBaseEditor;
 #if WITH_RIGVMLEGACYEDITOR
