Constraints: update default channel values if there were no previous keys

#rnx

diff --git a/Engine/Plugins/Animation/ControlRig/Source/ControlRigEditor/Private/Sequencer/ControlConstraintChannelInterface.cpp b/Engine/Plugins/Animation/ControlRig/Source/ControlRigEditor/Private/Sequencer/ControlConstraintChannelInterface.cpp
--- a/Engine/Plugins/Animation/ControlRig/Source/ControlRigEditor/Private/Sequencer/ControlConstraintChannelInterface.cpp
+++ b/Engine/Plugins/Animation/ControlRig/Source/ControlRigEditor/Private/Sequencer/ControlConstraintChannelInterface.cpp
@@ -6,6 +6,7 @@
 #include "ISequencer.h"
 
 #include "Constraints/ControlRigTransformableHandle.h"
+#include "Transform/TransformConstraintUtil.h"
 #include "Sequencer/MovieSceneControlRigParameterTrack.h"
 #include "LevelSequence.h"
 #include "MovieSceneToolHelpers.h"
@@ -194,6 +195,8 @@ bool FControlConstraintChannelInterface::SmartConstraintKey(
 			FCompensationEvaluator Evaluator(InConstraint);
 			Evaluator.ComputeLocalTransforms(ControlRig->GetWorld(), EvalParams);
 			TArray<FTransform>& ChildLocals = Evaluator.ChildLocals;
+
+			const bool bUpdateDefault = bNeedsCompensation && ChildLocals.IsEmpty() && bInfiniteInput;
 			
 			// store tangents at this time
 			TArray<FMovieSceneTangentData> Tangents;
@@ -261,8 +264,28 @@ bool FControlConstraintChannelInterface::SmartConstraintKey(
 				}
 
 				// add keys
-				ControlHandle->AddTransformKeys(FramesToCompensate,
-					ChildLocals, ChannelsToKey, TickResolution, nullptr,true);
+				constexpr bool bLocal = true;
+				ControlHandle->AddTransformKeys(FramesToCompensate, ChildLocals, ChannelsToKey, TickResolution, nullptr, bLocal);
+
+				// update default if needed
+				if (bUpdateDefault)
+				{
+					const FTransform Local = UE::TransformConstraintUtil::ComputeRelativeTransform(
+						ControlHandle->GetLocalTransform(),
+						ControlHandle->GetGlobalTransform(),
+						InConstraint->GetParentGlobalTransform(),
+						InConstraint);
+
+					FRigControlModifiedContext Context;
+					Context.KeyMask = static_cast<uint32>(ChannelsToKey);
+					
+					static constexpr bool bNotify = true, bUndo = false, bFixEuler = true;
+					ControlRig->SetControlLocalTransform(ControlName, Local, bNotify, Context, bUndo, bFixEuler);
+					if (ControlRig->IsAdditive())
+					{
+						ControlRig->Evaluate_AnyThread();
+					}
+				}
 
 				// set tangents at Time
 				if (bNeedsCompensation && !bInfiniteInput && NumChannels > 0)
