Selection SetsL Fix multiple CR issue with selection sets, remove uneeded settings,  tweaked layout and more BP exposure

diff --git a/Engine/Plugins/Animation/ControlRig/Source/ControlRigEditor/Private/Sequencer/SelectionSets/SSelectionSets.cpp b/Engine/Plugins/Animation/ControlRig/Source/ControlRigEditor/Private/Sequencer/SelectionSets/SSelectionSets.cpp
--- a/Engine/Plugins/Animation/ControlRig/Source/ControlRigEditor/Private/Sequencer/SelectionSets/SSelectionSets.cpp
+++ b/Engine/Plugins/Animation/ControlRig/Source/ControlRigEditor/Private/Sequencer/SelectionSets/SSelectionSets.cpp
@@ -459,7 +459,7 @@ void SSelectionSets::Construct(const FArguments& InArgs, FControlRigEditMode& In
 										{
 											if (UAIESelectionSets* SelectionSets = UAIESelectionSets::GetSelectionSets(SequencerPtr))
 											{
-												bEnabled = SelectionSets->bShowSelectedOnly;
+												bEnabled = SelectionSets->GetShowAndSetSelectedOnly();
 											}
 										}
 										return bEnabled ? ECheckBoxState::Checked : ECheckBoxState::Unchecked;
@@ -471,7 +471,7 @@ void SSelectionSets::Construct(const FArguments& InArgs, FControlRigEditMode& In
 										{
 											if (UAIESelectionSets* SelectionSets = UAIESelectionSets::GetSelectionSets(SequencerPtr))
 											{
-												SelectionSets->bShowSelectedOnly = bIsChecked;
+												SelectionSets->SetShowAndSetSelectedOnly(bIsChecked);
 											}
 										}
 										bPopulateUI = true;
@@ -491,7 +491,7 @@ void SSelectionSets::Construct(const FArguments& InArgs, FControlRigEditMode& In
 												{
 													if (UAIESelectionSets* SelectionSets = UAIESelectionSets::GetSelectionSets(SequencerPtr))
 													{
-														bool bValue = SelectionSets->bShowSelectedOnly;
+														bool bValue = SelectionSets->GetShowAndSetSelectedOnly();
 														if (bValue == true)
 														{
 															SlateColor = FAppStyle::GetSlateColor("SelectionColor");
@@ -532,7 +532,10 @@ void SSelectionSets::Construct(const FArguments& InArgs, FControlRigEditMode& In
 								SNew(SScrollBox)
 									+ SScrollBox::Slot()
 									[
-										SAssignNew(ActiveActorsWidget, SGridPanel)
+										SAssignNew(ActiveActorsWidget, SWrapBox)
+											.UseAllottedSize(true) // This is the key property
+											.InnerSlotPadding(FVector2D(10.0f, 10.0f))
+											.Orientation(Orient_Horizontal) // Default, but good to be explicit
 									]
 							]
 					]
@@ -622,16 +625,11 @@ void SSelectionSets::Construct(const FArguments& InArgs, FControlRigEditMode& In
 								.UseAllottedSize(true) // This is the key property
 								.InnerSlotPadding(FVector2D(10.0f, 10.0f))
 								.Orientation(Orient_Horizontal) // Default, but good to be explicit
-
 						]
 				]
 		];
 
-	bool bIsVertical = true;
-	if (const USelectionSetsSettings* Settings = GetDefault<USelectionSetsSettings>())
-	{
-		bIsVertical = Settings->GetSelectionSetsLayout() == ESelectionSetsLayout::Horizontal ? false : true;
-	}
+	constexpr bool bIsVertical = true;
 	if (bIsVertical)
 	{
 		ChildSlot
@@ -798,7 +796,7 @@ TSharedRef<SWidget> SSelectionSets::CreateActorMenu(FActorWithSelectionSet& Acto
 	FUIAction Action = FUIAction(
 		FExecuteAction::CreateLambda([&ActorWithSelectionSet]()
 			{
-				ActorWithSelectionSet.ShowAll();
+				ActorWithSelectionSet.ShowAllControlsOnThisActor();
 			}));
 
 	const FText Label = LOCTEXT("ShowAll", "Show All");
@@ -822,10 +820,11 @@ void SSelectionSets::PopulateActiveActors()
 	{
 		const float MenuIconSize = FAppStyle::Get().GetFloat("Menu.MenuIconSize");
 
-		auto AddActor = [this, WeakSequencer = SequencerPtr->AsWeak(), MenuIconSize](SGridPanel& Grid, int32 Row, int32 Column, TWeakObjectPtr<AActor>& ActorPtr, FActorWithSelectionSet& ActorWithSeletionSet)
+		auto AddActor = [this, WeakSequencer = SequencerPtr->AsWeak(), MenuIconSize](SWrapBox& Grid, int32 Row, int32 Column, TWeakObjectPtr<AActor>& ActorPtr, FActorWithSelectionSet& ActorWithSeletionSet)
 			{
-				Grid.AddSlot(Column, Row)
-					.Padding(4.0, 4.0, 2.0, 2.0)
+				Grid.AddSlot()
+					.HAlign(HAlign_Center)
+					.VAlign(VAlign_Center)
 					[
 						SNew(SCheckBox)
 							.OnGetMenuContent_Lambda([WeakSequencer, ActorPtr,this]()
@@ -884,7 +883,7 @@ void SSelectionSets::PopulateActiveActors()
 
 									PopulateSelectionSets();
 								})
-							.CheckBoxContentUsesAutoWidth(false)
+							//.CheckBoxContentUsesAutoWidth(false)
 							//.OnGetMenuContent(this, &SFilter::GetRightClickMenuContent)
 							[
 								SNew(SBorder)
diff --git a/Engine/Plugins/Animation/ControlRig/Source/ControlRigEditor/Private/Sequencer/SelectionSets/SSelectionSets.h b/Engine/Plugins/Animation/ControlRig/Source/ControlRigEditor/Private/Sequencer/SelectionSets/SSelectionSets.h
--- a/Engine/Plugins/Animation/ControlRig/Source/ControlRigEditor/Private/Sequencer/SelectionSets/SSelectionSets.h
+++ b/Engine/Plugins/Animation/ControlRig/Source/ControlRigEditor/Private/Sequencer/SelectionSets/SSelectionSets.h
@@ -18,7 +18,6 @@ class SBox;
 class SBorder;
 class SScrollBox;
 class SWrapBox;
-class SGridPanel;
 class SInlineEditableTextBlock;
 
 class SSelectionSetButton: public SButton
@@ -118,7 +117,7 @@ private:
 	TSharedPtr<SBorder> SearchSetsSlot;
 	TSharedPtr<SBox> ActiveSetsSlot;
 
-	TSharedPtr<SGridPanel> ActiveActorsWidget;
+	TSharedPtr<SWrapBox> ActiveActorsWidget;
 	TSharedPtr<SWrapBox> SelectionSetsWidget;
 
 	TMap<FGuid, TSharedPtr<SSelectionSetButton>> SelectionSetsButton;
diff --git a/Engine/Plugins/Animation/ControlRig/Source/ControlRigEditor/Private/Sequencer/SelectionSets/SelectionSets.cpp b/Engine/Plugins/Animation/ControlRig/Source/ControlRigEditor/Private/Sequencer/SelectionSets/SelectionSets.cpp
--- a/Engine/Plugins/Animation/ControlRig/Source/ControlRigEditor/Private/Sequencer/SelectionSets/SelectionSets.cpp
+++ b/Engine/Plugins/Animation/ControlRig/Source/ControlRigEditor/Private/Sequencer/SelectionSets/SelectionSets.cpp
@@ -235,9 +235,9 @@ void FActorsWithSelectionSets::SetUpFromSequencer(TSharedPtr<ISequencer>& Sequen
 								if (bExists == false)
 								{
 									Set.Color = SelectionSets->GetNextActorColor();
-									Set.Actors.Add(ControlRig);
 									Set.SetSelected(true);
 								}
+								Set.Actors.Add(ControlRig);
 								Set.SetsThatMatch.Add(Pair.Key);
 								if (MakeSelectedIfPresent.Contains(ControlRig))
 								{
@@ -354,25 +354,18 @@ bool FActorWithSelectionSet::IsControlRig() const
 
 bool FActorWithSelectionSet::IsActive(bool bShowSelectedOnly) const
 {
-	if (const USelectionSetsSettings* Settings = GetDefault<USelectionSetsSettings>())
+	if (IsControlRig() == false)
 	{
-		if (Settings->GetSelectionSetsActorSetVisibility() == ESelectionSetsActorSetVisibility::MultipleControRigs)
-		{
-			if (IsControlRig() == false)
-			{
-				return true;
-			}
-		}
+		return true;
 	}
 	if (bShowSelectedOnly)
 	{
 		return IsSelectedInViewport();
 	}
-	
 	return bIsSelected;
 }
 
-void FActorWithSelectionSet::ShowAll() const
+void FActorWithSelectionSet::ShowAllControlsOnThisActor() const
 {
 	FControlRigEditMode* EditMode = static_cast<FControlRigEditMode*>(GLevelEditorModeTools().GetActiveMode(FControlRigEditMode::ModeName));
 	if (EditMode)
@@ -452,38 +445,22 @@ UAIESelectionSets* UAIESelectionSets::GetSelectionSets(ULevelSequence* LevelSequ
 
 bool UAIESelectionSets::IsVisible(const FActorWithSelectionSet& ActorWithSelectionSet) const
 {
-	if (const USelectionSetsSettings* Settings = GetDefault<USelectionSetsSettings>())
+
+	if (bShowSelectedOnly)
+	{
+		return false;
+	}
+	if (ActorWithSelectionSet.IsControlRig() == true)
 	{
-		if (Settings->GetSelectionSetsActorSetVisibility() == ESelectionSetsActorSetVisibility::MultipleControRigs)
+		for (const TPair<TWeakObjectPtr<AActor>, FActorWithSelectionSet>& Pair : ActorsWithSelectionSets.ActorsWithSet)
 		{
-			if (bShowSelectedOnly)
-			{
-				return false;
-			}
-			if (ActorWithSelectionSet.IsControlRig() == true)
+			if (Pair.Value.IsControlRig() == true)
 			{
-				for (const TPair<TWeakObjectPtr<AActor>, FActorWithSelectionSet>& Pair : ActorsWithSelectionSets.ActorsWithSet)
+				if (&ActorWithSelectionSet != &Pair.Value)
 				{
-					if (Pair.Value.IsControlRig() == true)
-					{
-						if (&ActorWithSelectionSet != &Pair.Value)
-						{
-							return true;
-						}
-					}
+					return true;
 				}
 			}
-			else
-			{
-				return false;
-			}
-		}
-	}
-	if (bShowSelectedOnly)
-	{
-		if (ActorWithSelectionSet.IsSelectedInViewport())
-		{
-			return true;
 		}
 	}
 	return false;
@@ -575,6 +552,19 @@ TArray<FGuid> UAIESelectionSets::SortSelectionSetsByRow(TArray<FGuid>& InSelecti
 	return OutSelectionSets;
 }
 
+FGuid UAIESelectionSets::CreateSetItemFromSelection()
+{
+	if (TSharedPtr<ISequencer> SequencerPtr = GetSequencerFromAsset())
+	{
+		return CreateSetItemFromSelection(SequencerPtr);
+	}
+	else
+	{
+		UE_LOG(LogControlRig, Error, TEXT("Selection Sets: Need open Sequencer"));
+	}
+	return FGuid();
+}
+
 FGuid UAIESelectionSets::CreateSetItemFromSelection(TSharedPtr<ISequencer>& InSequencer)
 {
 	using namespace UE::AIE;
@@ -591,21 +581,9 @@ FGuid UAIESelectionSets::CreateSetItemFromSelection(TSharedPtr<ISequencer>& InSe
 		Item.ViewData.Color = GetNextSelectionSetColor();
 		FSequencerSelected::GetSelectedControlRigsAndBoundObjects(InSequencer.Get(), SelectedCRs, SelectedBoundObjects);
 
-	
 		if (SelectedCRs.Num() <= 0 && SelectedBoundObjects.Num() <= 0)
 		{
-			/* no longer do empty sets for now
-			FString Name = FString::Printf(TEXT("Empty Set %d"), SelectionSets.Num());
-			FText Text;
-			Text = Text.FromString(Name);
-			Item.ItemName = Text;
-			SelectionSets.Add(Item.Guid, Item);
-			SelectionSetsChangedBroadcast();
-
-			return Item.Guid;
-			*/
 			return FGuid();
-
 		}
 
 		FString Name = FString::Printf(TEXT("Set %d"), SelectionSets.Num());
@@ -683,6 +661,19 @@ FGuid UAIESelectionSets::CreateSetItemFromSelection(TSharedPtr<ISequencer>& InSe
 	}
 }
 
+FGuid UAIESelectionSets::CreateMirror(const FGuid& InGuid)
+{
+	if (TSharedPtr<ISequencer> SequencerPtr = GetSequencerFromAsset())
+	{
+		return CreateMirror(SequencerPtr, InGuid);
+	}
+	else
+	{
+		UE_LOG(LogControlRig, Error, TEXT("Selection Sets: Need open Sequencer"));
+	}
+	return FGuid();
+}
+
 FGuid UAIESelectionSets::CreateMirror(TSharedPtr<ISequencer>& InSequencer, const FGuid& InGuid)
 {
 	using namespace UE::Sequencer;
@@ -737,6 +728,11 @@ FGuid UAIESelectionSets::CreateMirror(TSharedPtr<ISequencer>& InSequencer, const
 	return NewGuid;
 }
 
+bool SetActorAsActive(AActor* InActor, bool bSetActive);
+
+/* Get the actors which are selectable*/
+UFUNCTION(BlueprintCallable, Category = "Actor")
+TArray<AActor> GetActiveActors(AActor* InActore);
 
 void UAIESelectionSets::SequencerBindingsAdded(TSharedPtr<ISequencer>& InSequencer)
 {
@@ -746,6 +742,19 @@ void UAIESelectionSets::SequencerBindingsAdded(TSharedPtr<ISequencer>& InSequenc
 
 }
 
+bool UAIESelectionSets::AddSelectionToSetItem(const FGuid& InGuid)
+{
+	if (TSharedPtr<ISequencer> SequencerPtr = GetSequencerFromAsset())
+	{
+		return AddSelectionToSetItem(InGuid,SequencerPtr);
+	}
+	else
+	{
+		UE_LOG(LogControlRig, Error, TEXT("Selection Sets: Need open Sequencer"));
+	}
+	return false;
+}
+
 bool UAIESelectionSets::AddSelectionToSetItem(const FGuid& InGuid, TSharedPtr<ISequencer>& InSequencer)
 {
 	using namespace UE::AIE;
@@ -809,6 +818,20 @@ bool UAIESelectionSets::AddSelectionToSetItem(const FGuid& InGuid, TSharedPtr<IS
 	}
 	return false;
 }
+
+bool UAIESelectionSets::RemoveSelectionFromSetItem(const FGuid& InGuid)
+{
+	if (TSharedPtr<ISequencer> SequencerPtr = GetSequencerFromAsset())
+	{
+		return RemoveSelectionFromSetItem(InGuid, SequencerPtr);
+	}
+	else
+	{
+		UE_LOG(LogControlRig, Error, TEXT("Selection Sets: Need open Sequencer"));
+	}
+	return false;
+}
+
 bool UAIESelectionSets::RemoveSelectionFromSetItem(const FGuid& InGuid, TSharedPtr<ISequencer>& InSequencer)
 {
 	using namespace UE::AIE;
@@ -849,6 +872,17 @@ bool UAIESelectionSets::RemoveSelectionFromSetItem(const FGuid& InGuid, TSharedP
 	return false;
 }
 
+void UAIESelectionSets::GetItemGuids(const FText& ItemName, TArray<FGuid>& OutGuids) const
+{
+	for (const TPair<FGuid,FAIESelectionSetItem>& Pair : SelectionSets)
+	{
+		if (Pair.Value.ItemName.IdenticalTo(ItemName))
+		{
+			OutGuids.Add(Pair.Key);
+		}
+	}
+}
+
 bool UAIESelectionSets::GetItemName(const FGuid& InGuid, FText& OutName) const
 {
 	if (const FAIESelectionSetItem* SetItem = SelectionSets.Find(InGuid))
@@ -861,32 +895,107 @@ bool UAIESelectionSets::GetItemName(const FGuid& InGuid, FText& OutName) const
 
 bool UAIESelectionSets::RenameSetItem(const FGuid& InGuid, const  FText& NewName)
 {
-	const FScopedTransaction Transaction(LOCTEXT("RenameSetItem_transaction", "Rename Selection Set"), !GIsTransacting);
-	Modify();
+	FScopedTransaction Transaction(LOCTEXT("RenameSetItem_transaction", "Rename Selection Set"), !GIsTransacting);
 	if (FAIESelectionSetItem* SetItem = SelectionSets.Find(InGuid))
 	{
+		Modify();
 		SetItem->ItemName = NewName;
 		return true;
 	}
+	Transaction.Cancel();
 	return false;
 }
+
+bool UAIESelectionSets::DeleteSetItem(const FGuid& InGuid) 
+{
+	if (TSharedPtr<ISequencer> SequencerPtr = GetSequencerFromAsset())
+	{
+		return DeleteSetItem(InGuid, SequencerPtr);
+	}
+	else
+	{
+		UE_LOG(LogControlRig, Error, TEXT("Selection Sets: Need open Sequencer"));
+	}
+	return false;
+}
+
+void UAIESelectionSets::UpdateRowValues(int32 AtRow, bool bIncrease)
+{
+	for (TPair<FGuid, FAIESelectionSetItem>& Pair : SelectionSets)
+	{
+		if (Pair.Value.ViewData.Row >= AtRow)
+		{
+			if (bIncrease)
+			{
+				Pair.Value.ViewData.Row++;
+			}
+			else
+			{
+				Pair.Value.ViewData.Row--;
+			}
+		}
+	}
+}
+
 bool UAIESelectionSets::DeleteSetItem(const FGuid& InGuid, TSharedPtr<ISequencer>& InSequencer)
 {
-	const FScopedTransaction Transaction(LOCTEXT("DeleteSetItem_transaction", "Delete Selection Set"), !GIsTransacting);
-	Modify();
+	FScopedTransaction Transaction(LOCTEXT("DeleteSetItem_transaction", "Delete Selection Set"), !GIsTransacting);
 	if (const FAIESelectionSetItem* SetItem = SelectionSets.Find(InGuid))
 	{
+		Modify();
 		ShowOrHideControls(InGuid, true, false);
 		SelectionSets.Remove(InGuid);
+		int32 Row = SetItem->ViewData.Row;
+		constexpr bool bIncrease = false;
+		UpdateRowValues(Row, bIncrease);
 		TSet<TWeakObjectPtr<UObject>> NoObjectsAdded;
 		ActorsWithSelectionSets.SetUpFromSequencer(InSequencer, this, NoObjectsAdded);
 		SelectionSetsChangedBroadcast();
 
 		return true;
 	}
+	Transaction.Cancel();
+	return false;
+}
+
+
+bool UAIESelectionSets::SetActorAsActive(AActor* InActor, bool bSetActive)
+{
+	if (FActorWithSelectionSet* ActorWithSeletionSet = ActorsWithSelectionSets.ActorsWithSet.Find(InActor))
+	{
+		ActorWithSeletionSet->SetSelected(bSetActive);
+		return true;
+	}
 	return false;
 }
 
+TArray<AActor*> UAIESelectionSets::GetAllActors() const
+{
+	TArray<AActor*>  Actors;
+	for (const TPair<TWeakObjectPtr<AActor>, FActorWithSelectionSet>& Pair : ActorsWithSelectionSets.ActorsWithSet)
+	{
+		if (Pair.Key.IsValid())
+		{
+			Actors.Add(Pair.Key.Get());
+		}
+	}
+	return Actors;
+}
+
+
+TArray<AActor*> UAIESelectionSets::GetActiveActors() const
+{
+	TArray<AActor*>  Actors;
+	for (const TPair<TWeakObjectPtr<AActor>, FActorWithSelectionSet>& Pair : ActorsWithSelectionSets.ActorsWithSet)
+	{
+		if (Pair.Key.IsValid() && Pair.Value.IsActive(bShowSelectedOnly))
+		{
+			Actors.Add(Pair.Key.Get());
+		}
+	}
+	return Actors;
+}
+
 bool UAIESelectionSets::GetItemColor(const FGuid& InGuid, FLinearColor& OutColor) const
 {
 	if (const FAIESelectionSetItem* SetItem = SelectionSets.Find(InGuid))
@@ -899,13 +1008,15 @@ bool UAIESelectionSets::GetItemColor(const FGuid& InGuid, FLinearColor& OutColor
 
 bool UAIESelectionSets::SetItemColor(const FGuid& InGuid, const FLinearColor& InColor)
 {
-	const FScopedTransaction Transaction(LOCTEXT("ColorSetItem_transaction", "Color Selection Set"), !GIsTransacting);
-	Modify();
+	FScopedTransaction Transaction(LOCTEXT("ColorSetItem_transaction", "Color Selection Set"), !GIsTransacting);
 	if (FAIESelectionSetItem* SetItem = SelectionSets.Find(InGuid))
 	{
+		Modify();
 		SetItem->ViewData.Color = InColor;
+		SelectionSetsChangedBroadcast();
 		return true;
 	}
+	Transaction.Cancel();
 	return false;
 }
 
@@ -1234,7 +1345,7 @@ bool UAIESelectionSets::ShowAllControls(const FGuid& InGuid) const
 						{
 							if (Pair.Value.SetsThatMatch.Contains(Guid))
 							{
-								Pair.Value.ShowAll();
+								Pair.Value.ShowAllControlsOnThisActor();
 							}
 						}
 					}
@@ -1313,6 +1424,24 @@ static UMovieScene3DTransformTrack* Get3DTransformTrack(TSharedPtr<ISequencer>&
 	return nullptr;
 }
 
+void UAIESelectionSets::SetShowAndSetSelectedOnly(bool bInShowSelectedOnly)
+{
+	bShowSelectedOnly = bInShowSelectedOnly;
+	SelectionSetsChangedBroadcast();
+}
+
+bool UAIESelectionSets::KeyAll(const FGuid& InGuid) const
+{
+	if (TSharedPtr<ISequencer> SequencerPtr = GetSequencerFromAsset())
+	{
+		return KeyAll(SequencerPtr,InGuid);
+	}
+	else
+	{
+		UE_LOG(LogControlRig, Error, TEXT("Selection Sets: Need open Sequencer"));
+	}
+	return false;
+}
 
 bool UAIESelectionSets::KeyAll(TSharedPtr<ISequencer>& InSequencer, const FGuid& InGuid) const
 {
@@ -1418,6 +1547,7 @@ bool UAIESelectionSets::RemoveChildFromItem(const FGuid& InParentGuid, const FGu
 
 void UAIESelectionSets::SelectionSetsChangedBroadcast() 
 {
+	GetActiveSelectionSets(); //this will sett up the sets
 	OnSelectionSetsChanged.Broadcast(this);
 }
 
@@ -1472,6 +1602,19 @@ FLinearColor UAIESelectionSets::GetNextActorColor()
 	return UE::AIE::GetNextRandomColor(NextIndex, false);
 }
 
+bool UAIESelectionSets::LoadFromJsonFile(const FFilePath& JsonFilePath)
+{
+	if (TSharedPtr<ISequencer> SequencerPtr = GetSequencerFromAsset())
+	{
+		return LoadFromJsonFile(SequencerPtr, JsonFilePath);
+	}
+	else
+	{
+		UE_LOG(LogControlRig, Error, TEXT("Selection Sets: Need open Sequencer"));
+	}
+	return false;
+}
+
 bool UAIESelectionSets::LoadFromJsonFile(TSharedPtr<ISequencer>& InSequencer, const FFilePath& JsonFilePath)
 {
 	FString JsonAsString;
@@ -1496,6 +1639,19 @@ bool UAIESelectionSets::ExportAsJsonFile(const FFilePath& JsonFilePath) const
 	return false;
 }
 
+bool UAIESelectionSets::LoadFromJsonString(const FString& JsonString)
+{
+	if (TSharedPtr<ISequencer> SequencerPtr = GetSequencerFromAsset())
+	{
+		return LoadFromJsonString(SequencerPtr, JsonString);
+	}
+	else
+	{
+		UE_LOG(LogControlRig, Error, TEXT("Selection Sets: Need open Sequencer"));
+	}
+	return false;
+}
+
 bool UAIESelectionSets::LoadFromJsonString(TSharedPtr<ISequencer>& InSequencer, const FString& JsonString)
 {
 	TSharedPtr<FJsonObject> JsonObject = MakeShareable(new FJsonObject());
diff --git a/Engine/Plugins/Animation/ControlRig/Source/ControlRigEditor/Private/Sequencer/SelectionSets/SelectionSets.h b/Engine/Plugins/Animation/ControlRig/Source/ControlRigEditor/Private/Sequencer/SelectionSets/SelectionSets.h
--- a/Engine/Plugins/Animation/ControlRig/Source/ControlRigEditor/Private/Sequencer/SelectionSets/SelectionSets.h
+++ b/Engine/Plugins/Animation/ControlRig/Source/ControlRigEditor/Private/Sequencer/SelectionSets/SelectionSets.h
@@ -26,17 +26,18 @@ enum EAIESelectionSetItemType : int32
 };
 ENUM_CLASS_FLAGS(EAIESelectionSetItemType)
 
+/* Individual 'Name" within a Seletion Set, maybe Actor Labe or Control Rig Control Name*/
 USTRUCT(BlueprintType)
 struct FAIESelectionSetItemName
 {
 	GENERATED_BODY();
 
 	/* The Name of the selection set item*, maybe an Actor Label or a ControlRig Control Name */
-	UPROPERTY(BlueprintReadWrite, EditAnywhere, Category = "Data")
+	UPROPERTY(BlueprintReadOnly, EditAnywhere, Category = "Data")
 	FString Name;
 
 	/* The Name of the selection set mirrored item for Control Rig Controls that have them */
-	UPROPERTY(BlueprintReadWrite, EditAnywhere, Category = "Data")
+	UPROPERTY(BlueprintReadOnly, EditAnywhere, Category = "Data")
 	FString MirrorName;
 
 	/* The EAIESelectionSetItemType, 0 is a Control Rig Control, 1 is an Actor */
@@ -48,22 +49,22 @@ struct FAIESelectionSetItemName
 	int32 Duplicates = 0;
 
 	/* The Name of the Owner Actor for Multi-Sets containing multiple rigs, may not be set*/
-	UPROPERTY(BlueprintReadWrite, EditAnywhere, Category = "Data")
+	UPROPERTY(BlueprintReadOnly, EditAnywhere, Category = "Data")
 	FString OwnerActorName;
 };
 
-
+/* View Data for the set, contains location and color in the UI*/
 USTRUCT(BlueprintType)
 struct FAIESelectionSetItemViewData
 {
 	GENERATED_BODY();
 
 	/* The Color of the Set*/
-	UPROPERTY(BlueprintReadWrite, EditAnywhere, Category = "Data")
+	UPROPERTY(BlueprintReadOnly, EditAnywhere, Category = "Data")
 	FLinearColor Color = FLinearColor::White;
 
 	/* The logical row the set should be in, basically it's order. This is an internal order for just one set*/
-	UPROPERTY(BlueprintReadWrite, EditAnywhere, Category = "Data")
+	UPROPERTY(BlueprintReadOnly, EditAnywhere, Category = "Data")
 	int32 Row = 0;
 	
 	//currently not used but maybe in future
@@ -71,13 +72,13 @@ struct FAIESelectionSetItemViewData
 	int32 Column = 0;
 };
 
-
+/* Main selection set item, has it's name, it's view data, and array of names making up the set*/
 USTRUCT(BlueprintType)
 struct FAIESelectionSetItem
 {
 	GENERATED_BODY();
 
-	//unique ID
+	/* Unique ID*/
 	UPROPERTY()
 	FGuid Guid;
 	
@@ -89,7 +90,7 @@ struct FAIESelectionSetItem
 	UPROPERTY(BlueprintReadWrite, EditAnywhere, Category = "Data")
 	TArray<FAIESelectionSetItemName> Names;
 
-	/* Viewing Data for the set, like location and row color*/
+	/* Viewing Data for the set, like location and color*/
 	UPROPERTY(BlueprintReadWrite, EditAnywhere, Category = "Data")
 	FAIESelectionSetItemViewData ViewData;
 
@@ -108,27 +109,46 @@ struct FAIESelectionSetItem
 	static FString GetMirror(const FString& InString);
 };
 
+/*  Internal Struct to specify a set of actors/control rigs and the sets they operate upon*/
 struct FActorWithSelectionSet
 {
-	TSet<TWeakObjectPtr<UObject>> Actors; //actor or crs
+	/*Name, usually the Actor Label*/
+	FText Name;
+
+	/*The Color in the UI*/
+	FLinearColor Color;
+
+	/*Actors in the level sequence that may be operated on by a Set*/
+	TSet<TWeakObjectPtr<UObject>> Actors; 
+
+	/*Set of selection sets that match up with these Actors*/
 	TSet<FGuid> SetsThatMatch;
 
+	/** If this Actor group is selected in the viewport*/
 	bool IsSelectedInViewport(bool bOnlyInViewport = false) const;
+
+	/** If this Actor group is a Control Rig Actor*/
 	bool IsControlRig() const;
-	bool IsActive(bool bShowSelectedOnly) const;
-	void ShowAll() const;
-	//
-	FText Name;
-	FLinearColor Color;
+
+	/** Make all Controls visisable on this Control Rig Actor*/
+	void ShowAllControlsOnThisActor() const;
+
+	/** Make this Set Actor Selected, may not be Active though based upon UAIESelectionSets::bShowSelectedOnly filter etc..*/
 	void SetSelected(bool bInSelected) { bIsSelected = bInSelected; }
+	
+	/** Get if this Set Actor is Selected. Note may be Selected by not Active for a set based upon UAIESelectionSets::bShowSelectedOnly filter*/
 	bool GetSelected() const {return bIsSelected;}
 
+	/** Get if this Set Actor is Active and may be operated upon by a Set Selection, takes into account filters/settings etc*/
+	bool IsActive(bool bShowSelectedOnly) const;
+
 	//if the set contains owned actors but they came from a different control rig we mark that. This is done in SetUpFromSequencer
 	void SetUpOwnedActors(const TWeakObjectPtr<AActor>& ActorOwner, TSharedPtr<ISequencer>& InSequencer, UAIESelectionSets* SelectionSets);
 	TMap <FGuid,FString> OwnedActors;
 
 	// it was previously selected but no more
 	mutable bool bWasSelectedInViewport = false;
+
 private:
 	bool bIsSelected = false;
 public:
@@ -136,7 +156,7 @@ public:
 
 };
 
-//set of control rigs and actors that may be pasted onto
+/* Internal struct containing all ofthe actor sets*/
 struct FActorsWithSelectionSets
 {
 	TMap<TWeakObjectPtr<AActor>, FActorWithSelectionSet> ActorsWithSet;
@@ -146,7 +166,6 @@ struct FActorsWithSelectionSets
 
 
 /** Link To Set of Anim Sequences that we may be linked to.*/
-
 UCLASS(BlueprintType, MinimalAPI)
 class UAIESelectionSets : public UAssetUserData
 {
@@ -156,52 +175,139 @@ public:
 
 	virtual bool IsEditorOnly() const override { return true; }
 
-	UPROPERTY(BlueprintReadWrite, Category = SelectionSet)
+	UPROPERTY(BlueprintReadOnly, Category = SelectionSet)
 	TMap<FGuid, FAIESelectionSetItem> SelectionSets;
-	
-	UPROPERTY(BlueprintReadWrite, Category = SelectionSet)
-	bool bShowSelectedOnly = true;
 
-	FActorsWithSelectionSets  ActorsWithSelectionSets;
+	/* Whether we are only showing and setting selections on current selections*/
+	UFUNCTION(BlueprintPure, Category = "State")
+	bool GetShowAndSetSelectedOnly() const
+	{
+		return bShowSelectedOnly;
+	}
 
+	/* Set whether we are only showing and setting selections on current selections*/
+	UFUNCTION(BlueprintCallable, Category = "State")
+	void SetShowAndSetSelectedOnly(bool bInShowSelectedOnly);
 		
-	//set of selection sets that are active
+	/* Get the current set of active seletion sets based upon active actors. This function will set it up also if needed*/
+	UFUNCTION(BlueprintCallable, Category = "State")
 	TArray<FGuid> GetActiveSelectionSets();
 
+	/* Set the AActor as Active or not, so it will or won't be selected from a selection set. Will return false if it doesn't match any sets*.
+	Note if GetShowAndSetSelectedOnly() is true the selected state of the Actor will take precedence over this*/
+	UFUNCTION(BlueprintCallable, Category = "Actor")
+	bool SetActorAsActive(AActor* InActor, bool bSetActive);
+
+	/* Get all the actors which are present to be selected by a selection set*/
+	UFUNCTION(BlueprintPure, Category = "Actor")
+	TArray<AActor*> GetAllActors() const;
+
+	/* Get the actors which are selectable*/
+	UFUNCTION(BlueprintPure, Category = "Actor")
+	TArray<AActor*> GetActiveActors() const;
+
+	/* Create a selection set from current selection. Will return empty FGuid if no set created*/
+	UFUNCTION(BlueprintCallable, Category = "Sets")
+	FGuid CreateSetItemFromSelection();
 	FGuid CreateSetItemFromSelection(TSharedPtr<ISequencer>& InSequencer);
+
+
+	/* Create a mirrored selection set from an existing set. Will return empty FGuid if no set created*/
+	UFUNCTION(BlueprintCallable, Category = "Sets")
+	FGuid CreateMirror(const FGuid& InGuid);
 	FGuid CreateMirror(TSharedPtr<ISequencer>& InSequencer, const FGuid& InGuid);
-	bool LoadFromJsonFile(TSharedPtr<ISequencer>& InSequencer, const FFilePath& JsonFilePath);
-	bool ExportAsJsonFile(const FFilePath& JsonFilePath)  const;
-	bool LoadFromJsonString(TSharedPtr<ISequencer>& InSequencer, const FString& JsonString);
-	bool ExportAsJsonString(FString& OutJsonString) const;
 
-	void SequencerBindingsAdded(TSharedPtr<ISequencer>& InSequencer);
+	/* Add selection to specified selection set*/
+	UFUNCTION(BlueprintCallable, Category = "Set")
+	bool AddSelectionToSetItem(const FGuid& InGuid);
 	bool AddSelectionToSetItem(const FGuid& InGuid, TSharedPtr<ISequencer>& InSequencer);
+
+	/* Remove selection from specified selection set*/
+	UFUNCTION(BlueprintCallable, Category = "Set")
+	bool RemoveSelectionFromSetItem(const FGuid& InGuid);
 	bool RemoveSelectionFromSetItem(const FGuid& InGuid, TSharedPtr<ISequencer>& InSequencer);
+
+	/* Rename selection set*/
+	UFUNCTION(BlueprintCallable, Category = "Set")
 	bool RenameSetItem(const FGuid& InGuid,const  FText& NewName);
+
+	/* Get selection set name*/
+	UFUNCTION(BlueprintPure, Category = "Set")
 	bool GetItemName(const FGuid& InGuid, FText& OutName) const;
+
+	/* Get the guids of set items associated with this name*/
+	UFUNCTION(BlueprintPure, Category = "State")
+	void GetItemGuids(const FText& ItemName, TArray<FGuid>& OutGuids) const;
+
+	/* Delete selection set*/
+	UFUNCTION(BlueprintCallable, Category = "Set")
+	bool DeleteSetItem(const FGuid& InGuid);
 	bool DeleteSetItem(const FGuid& InGuid, TSharedPtr<ISequencer>& InSequencer);
+
+	/* Get selection set color*/
+	UFUNCTION(BlueprintPure, Category = "Set")
 	bool GetItemColor(const FGuid& InGuid, FLinearColor& OutColor) const;
+
+	/* Set selection set color*/
+	UFUNCTION(BlueprintCallable, Category = "Set")
 	bool SetItemColor(const FGuid& InGuid, const FLinearColor& InColor);
 
+	/* Select from selection set, will select on the active actors*/
+	UFUNCTION(BlueprintCallable, Category = "Set")
 	bool SelectItem(const FGuid& InGuid, bool bDoMirror, bool bAdd, bool bToggle,bool bSelect = true) const;
 
-	bool AddChildToItem(const FGuid& InParentGuid, const FGuid& InChildGuid);
-	bool RemoveChildFromItem(const FGuid& InParentGuid, const FGuid& InChildGuid);
-
+	/* Show or hide controls on the active actors from the specified set*/
+	UFUNCTION(BlueprintCallable, Category = "Visibility")
 	bool ShowOrHideControls(const FGuid& InGuid, bool bShow, bool bDoMirror) const;
+
+	/* Isolate controls from just this selection set on active actors*/
+	UFUNCTION(BlueprintCallable, Category = "Visibility")
 	bool IsolateControls(const FGuid& InGuid) const;
+
+	/* Show all controls from this selection set's active actors*/
+	UFUNCTION(BlueprintCallable, Category = "Visibility")
 	bool ShowAllControls(const FGuid& InGuid) const;
 
+	/* Key the selection sets on active actors*/
+	UFUNCTION(BlueprintCallable, Category = "Set")
+	bool KeyAll(const FGuid& InGuid) const;
 	bool KeyAll(TSharedPtr<ISequencer>& InSequencer, const FGuid& InGuid) const;
 
+	/* Load sets from a JSON file*/
+	UFUNCTION(BlueprintCallable, Category = "ImportExport")
+	bool LoadFromJsonFile(const FFilePath& JsonFilePath);
+	bool LoadFromJsonFile(TSharedPtr<ISequencer>& InSequencer, const FFilePath& JsonFilePath);
+
+	/* Export sets to a JSON file*/
+	UFUNCTION(BlueprintCallable, Category = "ImportExport")
+	bool ExportAsJsonFile(const FFilePath& JsonFilePath)  const;
+
+	/* Load sets from a string*/
+	UFUNCTION(BlueprintCallable, Category = "ImportExport")
+	bool LoadFromJsonString(const FString& JsonString);
+	bool LoadFromJsonString(TSharedPtr<ISequencer>& InSequencer, const FString& JsonString);
+
+	/* Export sets from a string*/
+	UFUNCTION(BlueprintCallable, Category = "ImportExport")
+	bool ExportAsJsonString(FString& OutJsonString) const;
+
 	//Get updated visibility for the actor selections, using settings and properties
 	bool IsVisible(const FActorWithSelectionSet& ActorWithSelectionSet) const;
 
+	//utility functions
+	void SequencerBindingsAdded(TSharedPtr<ISequencer>& InSequencer);
+	//tbd
+	bool AddChildToItem(const FGuid& InParentGuid, const FGuid& InChildGuid);
+	bool RemoveChildFromItem(const FGuid& InParentGuid, const FGuid& InChildGuid);
+
 	static UAIESelectionSets* GetSelectionSets(TSharedPtr<ISequencer>& InSequencer, bool bAddIfDoesNotExist = false);
 	static UAIESelectionSets* GetSelectionSets(ULevelSequence* InLevelSequence, bool bAddIfDoesNotExist = false);
 	static 	TSharedPtr<ISequencer> GetSequencerFromAsset();
 
+public:
+	//current set of actors we can work on
+	FActorsWithSelectionSets  ActorsWithSelectionSets;
+
 	
 public:
 	DECLARE_EVENT_OneParam(UAIESelectionSets, FSelectionSetsChanged, UAIESelectionSets*);
@@ -220,7 +326,15 @@ public:
 
 private:
 	TArray<FGuid> SortSelectionSetsByRow(TArray<FGuid>& InSelectionSets);
+
+	//starting from AtRow increase or decrease row values, used for when changing row value or deleting the item
+	void UpdateRowValues(int32 AtRow, bool bIncrease);
 private:
 	//to handle selecting last selected object when none are selected
 	mutable TSet<FGuid> LastActiveSelectionSetsSet;
+
+	//store it but don't expose it, use ufunctions
+
+	UPROPERTY()
+	bool bShowSelectedOnly = true;
 };
diff --git a/Engine/Plugins/Animation/ControlRig/Source/ControlRigEditor/Private/Sequencer/SelectionSets/SelectionSetsSettings.cpp b/Engine/Plugins/Animation/ControlRig/Source/ControlRigEditor/Private/Sequencer/SelectionSets/SelectionSetsSettings.cpp
--- a/Engine/Plugins/Animation/ControlRig/Source/ControlRigEditor/Private/Sequencer/SelectionSets/SelectionSetsSettings.cpp
+++ b/Engine/Plugins/Animation/ControlRig/Source/ControlRigEditor/Private/Sequencer/SelectionSets/SelectionSetsSettings.cpp
@@ -9,11 +9,6 @@ USelectionSetsSettings::USelectionSetsSettings()
 	CategoryName = TEXT("General");
 	SectionName = TEXT("Animation Selection Sets");
 
-	SelectionSetsLayout = ESelectionSetsLayout::Vertical;
-
-	SelectionSetsActorSetVisibility = ESelectionSetsActorSetVisibility::MultipleControRigs;
-
-
 	if (CustomColors.Num() == 0)
 	{
 		CustomColors.Add(FLinearColor(.904, .323, .539)); //pastel red
@@ -30,19 +25,3 @@ USelectionSetsSettings::USelectionSetsSettings()
 	}
 }
 
-void USelectionSetsSettings::PostEditChangeProperty(FPropertyChangedEvent& PropertyChangedEvent)
-{
-	FName PropertyName = (PropertyChangedEvent.Property != NULL) ? PropertyChangedEvent.Property->GetFName() : NAME_None;
-	FName MemberPropertyName = (PropertyChangedEvent.MemberProperty != NULL) ? PropertyChangedEvent.MemberProperty->GetFName() : NAME_None;
-	if (PropertyName == GET_MEMBER_NAME_CHECKED(USelectionSetsSettings, SelectionSetsLayout) ||
-		MemberPropertyName == GET_MEMBER_NAME_CHECKED(USelectionSetsSettings, SelectionSetsLayout))
-	{
-		OnLayoutChanged.Broadcast();
-	}
-	else if (PropertyName == GET_MEMBER_NAME_CHECKED(USelectionSetsSettings, SelectionSetsActorSetVisibility) ||
-		MemberPropertyName == GET_MEMBER_NAME_CHECKED(USelectionSetsSettings, SelectionSetsActorSetVisibility))
-	{
-		OnActorSetVisibilityChanged.Broadcast();
-	}
-	Super::PostEditChangeProperty(PropertyChangedEvent);
-}
\ No newline at end of file
diff --git a/Engine/Plugins/Animation/ControlRig/Source/ControlRigEditor/Private/Sequencer/SelectionSets/SelectionSetsSettings.h b/Engine/Plugins/Animation/ControlRig/Source/ControlRigEditor/Private/Sequencer/SelectionSets/SelectionSetsSettings.h
--- a/Engine/Plugins/Animation/ControlRig/Source/ControlRigEditor/Private/Sequencer/SelectionSets/SelectionSetsSettings.h
+++ b/Engine/Plugins/Animation/ControlRig/Source/ControlRigEditor/Private/Sequencer/SelectionSets/SelectionSetsSettings.h
@@ -15,20 +15,6 @@
 
 #define UE_API CONTROLRIGEDITOR_API
 
-UENUM()
-enum class ESelectionSetsLayout : uint8
-{
-	Vertical	  UMETA(DisplayName = "Vertical"),
-	Horizontal	  UMETA(DisplayName = "Horizontal"),
-};
-
-UENUM()
-enum class ESelectionSetsActorSetVisibility : uint8
-{
-	AllVisible			 UMETA(DisplayName = "All Visible"),
-	MultipleControRigs	 UMETA(DisplayName = "Multiple Control Rigs"),
-};
-
 
 /** Serializable options for animation mode selection sets. */
 UCLASS(MinimalAPI, config=EditorPerProjectUserSettings)
@@ -37,41 +23,18 @@ class USelectionSetsSettings : public UDeveloperSettings
 public:
 	GENERATED_BODY()
 
-	DECLARE_MULTICAST_DELEGATE(FOnLayoutChanged);
-	DECLARE_MULTICAST_DELEGATE(FOnActorSetVisibilityChanged);
-
 	UE_API USelectionSetsSettings();
 
 	const TArray<FLinearColor>& GetCustomColors() const { return CustomColors; }
 
-	ESelectionSetsLayout GetSelectionSetsLayout() const { return SelectionSetsLayout; }
-
-	ESelectionSetsActorSetVisibility GetSelectionSetsActorSetVisibility() const { return SelectionSetsActorSetVisibility;  }
-
-
-	FOnLayoutChanged& GetOnLayoutChanged() { return OnLayoutChanged; }
-	FOnActorSetVisibilityChanged& GetOActorSetVisibilityChanged() { return OnActorSetVisibilityChanged; }
-
 protected:
 
-	/* Whether layout is horizontal or vertical. Need to re-enter animation mode for change to take place */
-	UPROPERTY(config, EditAnywhere, Category = "Layout")
-	ESelectionSetsLayout SelectionSetsLayout;
-
-	/* Whether or not we always show valid actors, or only show them when multiple control rigs match sets. Need to re-enter animation mode for change to take place */
-	UPROPERTY(config, EditAnywhere, Category = "Layout")
-	ESelectionSetsActorSetVisibility SelectionSetsActorSetVisibility;
-
 	/** Set of colors to use for selection sets*/
 	UPROPERTY(config, EditAnywhere, Category="Colors")
 	TArray<FLinearColor> CustomColors;
 
-	UE_API virtual void PostEditChangeProperty(FPropertyChangedEvent& PropertyChangedEvent) override;
-
 private:
 
-	FOnLayoutChanged OnLayoutChanged;
-	FOnActorSetVisibilityChanged OnActorSetVisibilityChanged;
 
 };
 
