Control Rig: Nullptr checks for pose library work

#rnx

diff --git a/Engine/Plugins/Animation/ControlRig/Source/ControlRig/Private/Tools/ControlRigPose.cpp b/Engine/Plugins/Animation/ControlRig/Source/ControlRig/Private/Tools/ControlRigPose.cpp
--- a/Engine/Plugins/Animation/ControlRig/Source/ControlRig/Private/Tools/ControlRigPose.cpp
+++ b/Engine/Plugins/Animation/ControlRig/Source/ControlRig/Private/Tools/ControlRigPose.cpp
@@ -356,6 +356,10 @@ void UControlRigPoseAsset::PostLoad()
 
 void UControlRigPoseAsset::SavePose(UControlRig* InControlRig, bool bUseAll)
 {
+	if (!InControlRig)
+	{
+		return;
+	}
 	Pose.SavePose(InControlRig,bUseAll);
 }
 
@@ -382,6 +386,11 @@ void UControlRigPoseAsset::PastePose(UControlRig* InControlRig, bool bDoKey, boo
 
 void UControlRigPoseAsset::SelectControls(UControlRig* InControlRig, bool bDoMirror, bool bClearSelection)
 {
+	if (!InControlRig)
+	{
+		return;
+	}
+	
 #if WITH_EDITOR
 	FScopedTransaction ScopedTransaction(LOCTEXT("SelectControlTransaction", "Select Control"));
 	InControlRig->Modify();
@@ -420,6 +429,10 @@ void UControlRigPoseAsset::SelectControls(UControlRig* InControlRig, bool bDoMir
 
 void UControlRigPoseAsset::GetCurrentPose(UControlRig* InControlRig, FControlRigControlPose& OutPose)
 {
+	if (!InControlRig)
+	{
+		return;
+	}
 	OutPose.SavePose(InControlRig, true);
 }
 
@@ -427,7 +440,10 @@ void UControlRigPoseAsset::GetCurrentPose(UControlRig* InControlRig, FControlRig
 TArray<FRigControlCopy> UControlRigPoseAsset::GetCurrentPose(UControlRig* InControlRig) 
 {
 	FControlRigControlPose TempPose;
-	TempPose.SavePose(InControlRig,true);
+	if (InControlRig)
+	{
+		TempPose.SavePose(InControlRig,true);
+	}
 	return TempPose.GetPoses();
 }
 
diff --git a/Engine/Plugins/Animation/ControlRig/Source/ControlRigEditor/Private/Sequencer/SelectionSets/SelectionSets.cpp b/Engine/Plugins/Animation/ControlRig/Source/ControlRigEditor/Private/Sequencer/SelectionSets/SelectionSets.cpp
--- a/Engine/Plugins/Animation/ControlRig/Source/ControlRigEditor/Private/Sequencer/SelectionSets/SelectionSets.cpp
+++ b/Engine/Plugins/Animation/ControlRig/Source/ControlRigEditor/Private/Sequencer/SelectionSets/SelectionSets.cpp
@@ -599,29 +599,32 @@ FGuid UAIESelectionSets::CreateSetItemFromSelection(TSharedPtr<ISequencer>& InSe
 		for (FControlRigAndControlsAndTrack& CRControls : SelectedCRs)
 		{
 			UControlRig* ControlRig = CRControls.ControlRig;
-			LastOwnerControlRig = ControlRig;
-			TempPose.SavePose(ControlRig, true);
-			for (FName& ControlName : CRControls.Controls)
+			if (ControlRig)
 			{
-				FAIESelectionSetItemName ItemName;
-				ItemName.Name = ControlName.ToString();
-				if (FRigControlCopy* CopyRigControl = MirrorTable.GetControl(ControlRig, TempPose, ControlName, true))
+				LastOwnerControlRig = ControlRig;
+				TempPose.SavePose(ControlRig, true);
+				for (FName& ControlName : CRControls.Controls)
 				{
-					ItemName.MirrorName = CopyRigControl->Name.ToString();
-				}
-				else
-				{
-					ItemName.MirrorName = ItemName.Name;
-				}
+					FAIESelectionSetItemName ItemName;
+					ItemName.Name = ControlName.ToString();
+					if (FRigControlCopy* CopyRigControl = MirrorTable.GetControl(ControlRig, TempPose, ControlName, true))
+					{
+						ItemName.MirrorName = CopyRigControl->Name.ToString();
+					}
+					else
+					{
+						ItemName.MirrorName = ItemName.Name;
+					}
 				
-				ItemName.Type = (int32)EAIESelectionSetItemType::ControlRig;
-				Duplicates.SetNum(0);
-				ActiveNames.MultiFind(ItemName.Name, Duplicates);
-				ItemName.Duplicates = Duplicates.Num();
-				Item.Names.Add(ItemName);
-				ActiveNames.Add(ItemName.Name, ItemName.Name);
+					ItemName.Type = (int32)EAIESelectionSetItemType::ControlRig;
+					Duplicates.SetNum(0);
+					ActiveNames.MultiFind(ItemName.Name, Duplicates);
+					ItemName.Duplicates = Duplicates.Num();
+					Item.Names.Add(ItemName);
+					ActiveNames.Add(ItemName.Name, ItemName.Name);
+				}
+				ObjectsAdded.Add(CRControls.ControlRig);
 			}
-			ObjectsAdded.Add(CRControls.ControlRig);
 		}
 		TOptional<FString> OwnerActorName;
 		if (LastOwnerControlRig)
@@ -784,20 +787,23 @@ bool UAIESelectionSets::AddSelectionToSetItem(const FGuid& InGuid, TSharedPtr<IS
 		for (FControlRigAndControlsAndTrack& CRControls : SelectedCRs)
 		{
 			UControlRig* ControlRig = CRControls.ControlRig;
-			TempPose.SavePose(ControlRig, true);
-			for (FName& ControlName : CRControls.Controls)
+			if (ControlRig)
 			{
-				FAIESelectionSetItemName ItemName;
-				ItemName.Name = ControlName.ToString();
-				ItemName.Type = (int32)EAIESelectionSetItemType::ControlRig;
-				if (FRigControlCopy* CopyRigControl = MirrorTable.GetControl(ControlRig, TempPose, ControlName, true))
+				TempPose.SavePose(ControlRig, true);
+				for (FName& ControlName : CRControls.Controls)
 				{
-					ItemName.MirrorName = CopyRigControl->Name.ToString();
+					FAIESelectionSetItemName ItemName;
+					ItemName.Name = ControlName.ToString();
+					ItemName.Type = (int32)EAIESelectionSetItemType::ControlRig;
+					if (FRigControlCopy* CopyRigControl = MirrorTable.GetControl(ControlRig, TempPose, ControlName, true))
+					{
+						ItemName.MirrorName = CopyRigControl->Name.ToString();
+					}
+					Duplicates.SetNum(0);
+					ActiveNames.MultiFind(ItemName.Name, Duplicates);
+					ItemName.Duplicates = Duplicates.Num();
+					SetItem->Names.Add(ItemName);
 				}
-				Duplicates.SetNum(0);
-				ActiveNames.MultiFind(ItemName.Name, Duplicates);
-				ItemName.Duplicates = Duplicates.Num();
-				SetItem->Names.Add(ItemName);
 			}
 		}
 		for (FObjectAndTrack& ObjectAndTrack : SelectedBoundObjects)
