Control Rig: Defensive coding for an edit mode tick happening after a world has been destroyed.

#rnx

diff --git a/Engine/Plugins/Animation/ControlRig/Source/ControlRigEditor/Private/EditMode/ControlRigEditMode.cpp b/Engine/Plugins/Animation/ControlRig/Source/ControlRigEditor/Private/EditMode/ControlRigEditMode.cpp
--- a/Engine/Plugins/Animation/ControlRig/Source/ControlRigEditor/Private/EditMode/ControlRigEditMode.cpp
+++ b/Engine/Plugins/Animation/ControlRig/Source/ControlRigEditor/Private/EditMode/ControlRigEditMode.cpp
@@ -4501,6 +4501,15 @@ void FControlRigEditMode::RecreateControlShapeActors()
 
 void FControlRigEditMode::CreateShapeActors(UControlRig* InControlRig)
 {
+	if (URigVMHost::IsGarbageOrDestroyed(WorldPtr))
+	{
+		return;
+	}
+	if (URigVMHost::IsGarbageOrDestroyed(WorldPtr->GetCurrentLevel()))
+	{
+		return;
+	}
+	
 	using namespace ControlRigEditMode::Shapes;
 	
 	if(bShowControlsAsOverlay)
