Control Rig: Solidify DeleteModule functionality

#rnx

diff --git a/Engine/Plugins/Animation/ControlRig/Source/ControlRig/Private/ModularRigController.cpp b/Engine/Plugins/Animation/ControlRig/Source/ControlRig/Private/ModularRigController.cpp
--- a/Engine/Plugins/Animation/ControlRig/Source/ControlRig/Private/ModularRigController.cpp
+++ b/Engine/Plugins/Animation/ControlRig/Source/ControlRig/Private/ModularRigController.cpp
@@ -1334,64 +1334,83 @@ bool UModularRigController::DeleteModule(const FName& InModuleName, bool bSetupU
 
 	(void)DeselectModule(Module->Name);
 
-	// Delete children
-	TArray<FName> ChildrenNames;
-	Algo::Transform(Module->CachedChildren, ChildrenNames, [](const FRigModuleReference* Child){ return Child->Name; });
-	for (const FName& ChildName : ChildrenNames)
+	TArray<FName> ModuleNamesToDelete;
+	ModuleNamesToDelete.Add(Module->Name);
+
+	for (int32 Index = 0; Index < ModuleNamesToDelete.Num(); Index++)
 	{
-		DeleteModule(ChildName, bSetupUndo);
+		const FRigModuleReference* ModuleToDelete = Model->FindModule(ModuleNamesToDelete[Index]);
+		if (!ModuleToDelete)
+		{
+			continue;
+		}
+		Algo::Transform(ModuleToDelete->CachedChildren, ModuleNamesToDelete, [](const FRigModuleReference* Child){ return Child->Name; });
 	}
 
-	Model->DeletedModules.Add(*Module);
-	Model->Modules.RemoveSingle(*Module);
-	Model->UpdateCachedChildren();
+	Algo::Reverse(ModuleNamesToDelete);
 
-	const FString InModuleNameString = InModuleName.ToString();
+	for (const FName& ModuleName : ModuleNamesToDelete)
+	{
+		const FRigModuleReference* ModuleToDelete = Model->FindModule(ModuleName);
+		Model->DeletedModules.Add(*Module);
+		Model->DeletedModules.Last().CachedChildren.Reset();
+		Model->Modules.RemoveSingle(*ModuleToDelete);
+	}
+	
+	Model->UpdateCachedChildren();
 
-	// Fix connections
+	for (const FName& ModuleName : ModuleNamesToDelete)
 	{
-		TArray<FRigElementKey> ToRemove;
-		for (FModularRigSingleConnection& Connection : Model->Connections)
-		{
-			const FRigHierarchyModulePath ConnectionModulePath(Connection.Connector.Name);
-			if (ConnectionModulePath.HasModuleName(InModuleNameString))
-			{
-				ToRemove.Add(Connection.Connector);
-			}
+		const FString ModuleNameString = ModuleName.ToString();
 
-			for(const FRigElementKey& Target : Connection.Targets)
+		// Fix connections
+		{
+			TSet<FRigElementKey> ToRemove;
+			for (FModularRigSingleConnection& Connection : Model->Connections)
 			{
-				const FRigHierarchyModulePath TargetModulePath(Target.Name);
-				if (TargetModulePath.HasModuleName(InModuleNameString))
+				const FRigHierarchyModulePath ConnectionModulePath(Connection.Connector.Name);
+				if (ConnectionModulePath.HasModuleName(ModuleNameString))
 				{
 					ToRemove.Add(Connection.Connector);
-					break;
+				}
+
+				for(const FRigElementKey& Target : Connection.Targets)
+				{
+					const FRigHierarchyModulePath TargetModulePath(Target.Name);
+					if (TargetModulePath.HasModuleName(ModuleNameString))
+					{
+						ToRemove.Add(Connection.Connector);
+						break;
+					}
 				}
 			}
+			for (FRigElementKey& KeyToRemove : ToRemove)
+			{
+				Model->Connections.RemoveConnection(KeyToRemove);
+			}
+			Model->Connections.UpdateFromConnectionList();
 		}
-		for (FRigElementKey& KeyToRemove : ToRemove)
+
+		// Fix bindings
+		for (FRigModuleReference& Reference : Model->Modules)
 		{
-			Model->Connections.RemoveConnection(KeyToRemove);
+			Reference.Bindings = Reference.Bindings.FilterByPredicate([ModuleNameString](const TPair<FName, FString>& Binding)
+			{
+				const FRigHierarchyModulePath BindingModulePath(Binding.Value);
+				if (BindingModulePath.HasModuleName(ModuleNameString))
+				{
+					return false;
+				}
+				return true;
+			});
 		}
-		Model->Connections.UpdateFromConnectionList();
 	}
 
-	// Fix bindings
-	for (FRigModuleReference& Reference : Model->Modules)
+	for (const FRigModuleReference& DeletedModule : Model->DeletedModules)
 	{
-		Reference.Bindings = Reference.Bindings.FilterByPredicate([InModuleNameString](const TPair<FName, FString>& Binding)
-		{
-			const FRigHierarchyModulePath BindingModulePath(Binding.Value);
-			if (BindingModulePath.HasModuleName(InModuleNameString))
-			{
-				return false;
-			}
-			return true;
-		});
+		Notify(EModularRigNotification::ModuleRemoved, &DeletedModule);
 	}
 
-	Notify(EModularRigNotification::ModuleRemoved, &Model->DeletedModules.Last());
-
 	Model->DeletedModules.Reset();
 
 #if WITH_EDITOR
