Selection Sets:  Show Mirrored set when added even if actor isn't selected.
#trivial

diff --git a/Engine/Plugins/Animation/ControlRig/Source/ControlRigEditor/Private/Sequencer/SelectionSets/SelectionSets.cpp b/Engine/Plugins/Animation/ControlRig/Source/ControlRigEditor/Private/Sequencer/SelectionSets/SelectionSets.cpp
--- a/Engine/Plugins/Animation/ControlRig/Source/ControlRigEditor/Private/Sequencer/SelectionSets/SelectionSets.cpp
+++ b/Engine/Plugins/Animation/ControlRig/Source/ControlRigEditor/Private/Sequencer/SelectionSets/SelectionSets.cpp
@@ -723,6 +723,10 @@ FGuid UAIESelectionSets::CreateMirror(TSharedPtr<ISequencer>& InSequencer, const
 			Item.ViewData.Row = SelectionSets.Num();
 			SelectionSets.Add(Item.Guid, Item);
 			TSet<TWeakObjectPtr<UObject>> NoObjectsAdded;
+			if (LastActiveSelectionSetsSet.Contains(InGuid))
+			{
+				LastActiveSelectionSetsSet.Add(NewGuid);
+			}
 			ActorsWithSelectionSets.SetUpFromSequencer(InSequencer, this, NoObjectsAdded);
 
 			SelectionSetsChangedBroadcast();
