ControlRig dependencies: optimize the dependencies computation

- note: this can probably be even faster

#rnx

diff --git a/Engine/Plugins/Animation/ControlRig/Source/ControlRig/Private/Rigs/RigDependency.cpp b/Engine/Plugins/Animation/ControlRig/Source/ControlRig/Private/Rigs/RigDependency.cpp
--- a/Engine/Plugins/Animation/ControlRig/Source/ControlRig/Private/Rigs/RigDependency.cpp
+++ b/Engine/Plugins/Animation/ControlRig/Source/ControlRig/Private/Rigs/RigDependency.cpp
@@ -158,11 +158,77 @@ void FRigDependenciesProviderForVM::ComputeDependencies(const FControlRigExecute
 	}
 
 	// for each written transform / metadata on an instruction
-	TArray<TTuple<int32,FRigHierarchyRecord>> InstructionsToVisit;
 	TArray<FRigHierarchyRecord> WrittenVariables;
 	TMap<FRigVMOperand, TArray<int32>> InputOperandToInstructions;
 	TMap<FRigVMOperand, TArray<int32>> OutputOperandToInstructions;
 
+	// use a combination of a set and an array to ensure a faster uniqueness while being able to iterate over the last elements added
+	TArray<TTuple<int32, FRigHierarchyRecord>> InstructionsToVisit;
+	TMap<int32, TSet<FRigHierarchyRecord>> UniqueInstructionsToVisit;
+	
+	auto ResetLastInstruction = [&InstructionsToVisit, &UniqueInstructionsToVisit](const int32 InNewInstruction, const FRigHierarchyRecord& InNewRecord)
+	{
+		if (InstructionsToVisit.Num() == 1)
+		{
+			// avoid reallocation and set value directly
+			InstructionsToVisit[0] = {InNewInstruction, InNewRecord};
+		}
+		else
+		{
+			InstructionsToVisit.Reset();
+			InstructionsToVisit.Emplace(InNewInstruction, InNewRecord);
+		}
+
+		UniqueInstructionsToVisit.Reset();
+	};
+
+	auto FindOrAddInstruction = [&InstructionsToVisit, &UniqueInstructionsToVisit](TConstArrayView<int32> InInstructions, const FRigHierarchyRecord& InRecord)
+	{
+		if (InInstructions.IsEmpty())
+		{
+			return;
+		}
+
+		const uint32 NewRecordHash = GetTypeHash(InRecord);
+		
+		for (int32 NewInstruction : InInstructions)
+		{
+			if (InstructionsToVisit.Num() == 1)
+			{
+				const int32 LastInstruction = InstructionsToVisit[0].Get<0>();
+				const FRigHierarchyRecord LastRecord = InstructionsToVisit[0].Get<1>();
+
+				// check if the {Instruction, InRecord} is equal to the last instruction and continue if so, as there's nothing to do
+				if (NewInstruction == LastInstruction && InRecord == LastRecord)
+				{
+					continue;
+				}
+				else
+				{
+					// otherwise, add the last instruction / last record hash to maintain uniqueness and proceed
+					UniqueInstructionsToVisit.Reset();
+					UniqueInstructionsToVisit.Emplace(LastInstruction).Emplace(LastRecord);
+				}
+			}
+
+			TSet<FRigHierarchyRecord>& UniqueRecords = UniqueInstructionsToVisit.FindOrAdd(NewInstruction);
+			if (UniqueRecords.IsEmpty())
+			{
+				UniqueRecords.AddByHash(NewRecordHash, InRecord);
+				InstructionsToVisit.Emplace(NewInstruction, InRecord);
+			}
+			else
+			{
+				bool bAlreadyInSet = false;
+				UniqueRecords.FindOrAddByHash(NewRecordHash, InRecord, &bAlreadyInSet);
+				if (!bAlreadyInSet)
+				{
+					InstructionsToVisit.Emplace(NewInstruction, InRecord);
+				}
+			}
+	   }
+	};
+
 	for(int32 InstructionIndex = 0; InstructionIndex < Instructions.Num(); InstructionIndex++)
 	{
 		FRigVMOperandArray InputOperands = InVM->GetByteCode().GetInputOperands(InstructionIndex);
@@ -204,15 +270,14 @@ void FRigDependenciesProviderForVM::ComputeDependencies(const FControlRigExecute
 					ReadRecord.ElementIndex,
 					ReadRecord.MetadataName);
 
-				InstructionsToVisit.Reset();
-				InstructionsToVisit.Emplace(InstructionIndex, ReadElementRecord);
+				ResetLastInstruction(InstructionIndex, ReadElementRecord);
 
 				// follow the output operands through the bytecode to determine variable writes of this instruction
 				for(int32 InstructionToVisitIndex = 0; InstructionToVisitIndex < InstructionsToVisit.Num(); InstructionToVisitIndex++)
 				{
 					const int32 InstructionToVisit = InstructionsToVisit[InstructionToVisitIndex].Get<0>();
 
-					FRigVMOperandArray OutputOperands = InVM->GetByteCode().GetOutputOperands(InstructionToVisit);
+					const FRigVMOperandArray OutputOperands = InVM->GetByteCode().GetOutputOperands(InstructionToVisit);
 					for(const FRigVMOperand OutputOperand : OutputOperands)
 					{
 						FRigHierarchyRecord LastRecord = InstructionsToVisit[InstructionToVisitIndex].Get<1>();
@@ -232,11 +297,7 @@ void FRigDependenciesProviderForVM::ComputeDependencies(const FControlRigExecute
 
 						if(const TArray<int32>* InstructionsWithOutputOperand = InputOperandToInstructions.Find(OutputOperandNoRegisterOffset))
 						{
-							for(int32 InstructionWithOutputOperand : *InstructionsWithOutputOperand)
-							{
-								const TTuple<int32,FRigHierarchyRecord> Tuple(InstructionWithOutputOperand, LastRecord);
-								InstructionsToVisit.AddUnique(Tuple);
-							}
+							FindOrAddInstruction(*InstructionsWithOutputOperand, LastRecord);
 						}
 					}
 				}
@@ -272,8 +333,7 @@ void FRigDependenciesProviderForVM::ComputeDependencies(const FControlRigExecute
 
 				OutDependencies.FindOrAdd(TargetRecord).Add(InstructionRecord);
 
-				InstructionsToVisit.Reset();
-				InstructionsToVisit.Emplace(InstructionIndex, InstructionRecord);
+				ResetLastInstruction(InstructionIndex, InstructionRecord);
 
 				// follow the input operands through the bytecode to determine inputs of this instruction
 				for(int32 InstructionToVisitIndex = 0; InstructionToVisitIndex < InstructionsToVisit.Num(); InstructionToVisitIndex++)
@@ -294,7 +354,7 @@ void FRigDependenciesProviderForVM::ComputeDependencies(const FControlRigExecute
 						}
 					}
 
-					FRigVMOperandArray InputOperands = InVM->GetByteCode().GetInputOperands(InstructionToVisit);
+					const FRigVMOperandArray InputOperands = InVM->GetByteCode().GetInputOperands(InstructionToVisit);
 					for(const FRigVMOperand InputOperand : InputOperands)
 					{
 						FRigHierarchyRecord LastRecord = InstructionsToVisit[InstructionToVisitIndex].Get<1>();
@@ -314,11 +374,7 @@ void FRigDependenciesProviderForVM::ComputeDependencies(const FControlRigExecute
 
 						if(const TArray<int32>* InstructionsWithOutputOperand = OutputOperandToInstructions.Find(InputOperandNoRegisterOffset))
 						{
-							for(int32 InstructionWithOutputOperand : *InstructionsWithOutputOperand)
-							{
-								const TTuple<int32,FRigHierarchyRecord> Tuple(InstructionWithOutputOperand, LastRecord);
-								InstructionsToVisit.AddUnique(Tuple);
-							}
+							FindOrAddInstruction(*InstructionsWithOutputOperand, LastRecord);
 						}
 					}
 				}
