Control Rig: Avoid running dependency traverser during the run of the VM

#rnx

diff --git a/Engine/Plugins/Animation/ControlRig/Source/ControlRig/Private/Rigs/RigHierarchy.cpp b/Engine/Plugins/Animation/ControlRig/Source/ControlRig/Private/Rigs/RigHierarchy.cpp
--- a/Engine/Plugins/Animation/ControlRig/Source/ControlRig/Private/Rigs/RigHierarchy.cpp
+++ b/Engine/Plugins/Animation/ControlRig/Source/ControlRig/Private/Rigs/RigHierarchy.cpp
@@ -4095,31 +4095,9 @@ bool URigHierarchy::SwitchToParent(FRigBaseElement* InChild, FRigBaseElement* In
 		}
 	}
 
-	// rely on the VM's dependency map if there's currently an available context.
-	const IRigDependenciesProvider* DependencyProviderPtr = &InDependencyProvider;
-
-#if WITH_EDITOR
-	// Keep this in function scope, since we might be taking a pointer to it.
-	FRigDependenciesProviderForControlRig DependencyProviderFromVM;
-	{
-		UE::TScopeLock Lock(ExecuteContextLock);
-		if(ExecuteContext != nullptr && DependencyProviderPtr->GetRigHierarchyDependencies().IsEmpty())
-		{
-			if(ExecuteContext->VM)
-			{
-				if(UControlRig* Rig = Cast<UControlRig>(GetOuter()))
-				{
-					DependencyProviderFromVM = FRigDependenciesProviderForControlRig(Rig);
-					DependencyProviderPtr = &DependencyProviderFromVM;
-				}
-			}
-		}
-	}
-#endif
-	
 	if(InChild && InParent)
 	{
-		if(!CanSwitchToParent(InChild->GetKey(), InParent->GetKey(), *DependencyProviderPtr, OutFailureReason))
+		if(!CanSwitchToParent(InChild->GetKey(), InParent->GetKey(), InDependencyProvider, OutFailureReason))
 		{
 			return false;
 		}
