ControlRig: fixed URigHierarchy::CanSwitchToParent not preprocessing keys anymore

#rnx

diff --git a/Engine/Plugins/Animation/ControlRig/Source/ControlRig/Private/Rigs/RigHierarchy.cpp b/Engine/Plugins/Animation/ControlRig/Source/ControlRig/Private/Rigs/RigHierarchy.cpp
--- a/Engine/Plugins/Animation/ControlRig/Source/ControlRig/Private/Rigs/RigHierarchy.cpp
+++ b/Engine/Plugins/Animation/ControlRig/Source/ControlRig/Private/Rigs/RigHierarchy.cpp
@@ -4023,6 +4023,8 @@ bool URigHierarchy::CanCauseCycle(const FRigTransformElement* InChild, const FRi
 
 bool URigHierarchy::CanSwitchToParent(FRigElementKey InChild, FRigElementKey InParent, const IRigDependenciesProvider& InDependencyProvider, FString* OutFailureReason) const
 {
+	InParent = PreprocessParentElementKeyForSpaceSwitching(InChild, InParent);
+	
 	const FRigBaseElement* Child = Find(InChild);
 	if(Child == nullptr)
 	{
