Control Rig: Solidify override value serialization by introducing size + hash for the property check

#rnx

diff --git a/Engine/Plugins/Animation/ControlRig/Source/ControlRig/Private/ControlRigOverride.cpp b/Engine/Plugins/Animation/ControlRig/Source/ControlRig/Private/ControlRigOverride.cpp
--- a/Engine/Plugins/Animation/ControlRig/Source/ControlRig/Private/ControlRigOverride.cpp
+++ b/Engine/Plugins/Animation/ControlRig/Source/ControlRig/Private/ControlRigOverride.cpp
@@ -3,6 +3,7 @@
 #include "ControlRigOverride.h"
 #include "ControlRig.h"
 #include "RigVMCore/RigVMMemoryCommon.h"
+#include "ControlRigObjectVersion.h"
 
 #if WITH_EDITOR
 #include "PropertyPath.h"
@@ -32,6 +33,8 @@ TAutoConsoleVariable<bool> CVarControlRigEnableOverrides(TEXT("ControlRig.Overri
 
 FControlRigOverrideValue::FControlRigOverrideValue(const FString& InPath, const UObject* InSubject)
 	: SubjectKey(InSubject ? InSubject->GetFName() : NAME_None)
+	, Hash(0)
+	, DataPropertyHash(0)
 {
 	if(InSubject)
 	{
@@ -44,6 +47,8 @@ FControlRigOverrideValue::FControlRigOverrideValue(const FString& InPath, const
 
 FControlRigOverrideValue::FControlRigOverrideValue(const FString& InPath, const UStruct* InOwnerStruct, const void* InSubjectPtr, const FName& InSubjectKey)
 	: SubjectKey(InSubjectKey)
+	, Hash(0)
+	, DataPropertyHash(0)
 {
 	ensure(InOwnerStruct);
 	ensureMsgf(InPath.Contains(TEXT(".")) == false, TEXT("Paths for overrides should use the -> separator."));
@@ -59,6 +64,8 @@ FControlRigOverrideValue::FControlRigOverrideValue(const FString& InPath, const
 
 FControlRigOverrideValue::FControlRigOverrideValue(const FString& InPath, const UStruct* InOwnerStruct, const FString& InValueAsString, const FName& InSubjectKey, const FControlRigOverrideValueErrorPipe::TReportFunction& InReportFunction)
 	: SubjectKey(InSubjectKey)
+	, Hash(0)
+	, DataPropertyHash(0)
 {
 	if(SetPropertiesFromPath(InPath, InOwnerStruct))
 	{
@@ -85,7 +92,7 @@ FControlRigOverrideValue& FControlRigOverrideValue::operator=(const FControlRigO
 	Properties = InOther.Properties;
 	Hash = InOther.Hash;
 
-	if(GetLeafProperty())
+	if(GetLeafProperty() && InOther.IsDataValid())
 	{
 		if(const void* OtherDataPtr = InOther.GetData())
 		{
@@ -105,6 +112,8 @@ bool FControlRigOverrideValue::Serialize(FArchive& Ar)
 		Reset();
 	}
 
+	Ar.UsingCustomVersion(FControlRigObjectVersion::GUID);
+
 	Ar << SubjectKey;
 	
 	int32 NumProperties = 0;
@@ -176,16 +185,43 @@ bool FControlRigOverrideValue::Serialize(FArchive& Ar)
 
 	if(Ar.IsLoading())
 	{
+		bool bSkipData = true;
+		
 		if(const FProperty* LeafProperty = GetLeafProperty())
 		{
-			void* DataPtr = AllocateDataIfRequired();
-			LeafProperty->SerializeItem(FStructuredArchiveFromArchive(Ar).GetSlot(), DataPtr);
+			int32 PropertySize = 0;
+			uint32 PropertyHash = 0;
+
+			if(Ar.CustomVer(FControlRigObjectVersion::GUID) >= FControlRigObjectVersion::OverridesStoreTOCDataForProperties)
+			{
+				Ar << PropertySize;
+				Ar << PropertyHash;
+			}
+			else
+			{
+				PropertySize = LeafProperty->GetSize();
+				PropertyHash = GetPropertyHash();
+			}
+
+			if (LeafProperty->GetSize() == PropertySize && GetPropertyHash() == PropertyHash)
+			{
+				void* DataPtr = AllocateDataIfRequired();
+				LeafProperty->SerializeItem(FStructuredArchiveFromArchive(Ar).GetSlot(), DataPtr);
 
-			FString ValueAsString;
-			LeafProperty->ExportTextItem_Direct(ValueAsString, DataPtr, DataPtr, nullptr, PPF_None, nullptr);
-			CachedStringValue = ValueAsString;
+				FString ValueAsString;
+				LeafProperty->ExportTextItem_Direct(ValueAsString, DataPtr, DataPtr, nullptr, PPF_None, nullptr);
+				CachedStringValue = ValueAsString;
+
+				bSkipData = false;
+			}
+			else
+			{
+				DataPropertyHash = 0;
+				DataArray.Reset();
+			}
 		}
-		else
+
+		if (bSkipData)
 		{
 			// we can't deserialize the data - skip it
 			Ar.Seek(ArchivePositionBeforeData + OffsetForData);
@@ -197,6 +233,12 @@ bool FControlRigOverrideValue::Serialize(FArchive& Ar)
 	{
 		if(const FProperty* LeafProperty = GetLeafProperty())
 		{
+			int32 PropertySize = LeafProperty->GetSize();
+			Ar << PropertySize;
+
+			uint32 PropertyHash = DataPropertyHash;
+			Ar << PropertyHash;
+			
 			void *DataPtr = AllocateDataIfRequired();
 			LeafProperty->SerializeItem(FStructuredArchiveFromArchive(Ar).GetSlot(), DataPtr, nullptr);
 
@@ -215,7 +257,123 @@ bool FControlRigOverrideValue::Serialize(FArchive& Ar)
 bool FControlRigOverrideValue::IsValid() const
 {
 	const bool bIsSet = Path.IsSet() && CachedStringValue.IsSet() && (GetData() != nullptr);
-	return bIsSet && (GetLeafProperty() != nullptr);
+	return bIsSet && IsDataValid();
+}
+
+bool FControlRigOverrideValue::IsDataValid() const
+{
+	if(GetLeafProperty())
+	{
+		return IsDataValid(DataArray, DataPropertyHash);
+	}
+	return false;
+}
+
+bool FControlRigOverrideValue::IsDataValid(const TArray<uint8, TAlignedHeapAllocator<16>>& InDataArray, const uint32& InOutPropertyHash) const
+{
+	if (const FProperty* LeafProperty = GetLeafProperty())
+	{
+		if (InDataArray.Num() == LeafProperty->GetSize())
+		{
+			if (GetPropertyHash() == InOutPropertyHash)
+			{
+				return true;
+			}
+		}
+	}
+	return false;
+}
+
+uint32 FControlRigOverrideValue::GetPropertyHash() const
+{
+	uint32 PropertyHash = 0;
+	for (const FPropertyInfo& PropertyInfo : Properties)
+	{
+		PropertyHash = HashCombine(PropertyHash, GetTypeHash(PropertyInfo.ArrayIndex));
+		PropertyHash = HashCombine(PropertyHash, GetPropertyHash(PropertyInfo.Property.Get()));
+	}
+	return PropertyHash;
+}
+
+#define UE_CONTROLRIGOVERRIDE_PROPERTYHASHCASE(PropertyType) \
+if (InProperty->IsA<PropertyType>()) \
+{ \
+	static const FString PropertyString = TEXT(#PropertyType); \
+	static const uint32 PropertyHash = GetTypeHash(PropertyString); \
+	return PropertyHash; \
+}
+
+uint32 FControlRigOverrideValue::GetPropertyHash(const FProperty* InProperty)
+{
+	UE_CONTROLRIGOVERRIDE_PROPERTYHASHCASE(FBoolProperty);
+	UE_CONTROLRIGOVERRIDE_PROPERTYHASHCASE(FUInt16Property);
+	UE_CONTROLRIGOVERRIDE_PROPERTYHASHCASE(FUInt32Property);
+	UE_CONTROLRIGOVERRIDE_PROPERTYHASHCASE(FUInt64Property);
+	UE_CONTROLRIGOVERRIDE_PROPERTYHASHCASE(FInt8Property);
+	UE_CONTROLRIGOVERRIDE_PROPERTYHASHCASE(FInt16Property);
+	UE_CONTROLRIGOVERRIDE_PROPERTYHASHCASE(FIntProperty);
+	UE_CONTROLRIGOVERRIDE_PROPERTYHASHCASE(FInt64Property);
+	UE_CONTROLRIGOVERRIDE_PROPERTYHASHCASE(FFloatProperty);
+	UE_CONTROLRIGOVERRIDE_PROPERTYHASHCASE(FDoubleProperty);
+	UE_CONTROLRIGOVERRIDE_PROPERTYHASHCASE(FNameProperty);
+	UE_CONTROLRIGOVERRIDE_PROPERTYHASHCASE(FStrProperty);
+	
+	if (const FByteProperty* ByteProperty = CastField<FByteProperty>(InProperty))
+	{
+		static const FString PropertyString = TEXT("ByteProperty");
+		static const uint32 PropertyHash = GetTypeHash(PropertyString);
+		if (const UEnum* Enum = ByteProperty->Enum)
+		{
+			return HashCombine(PropertyHash, GetTypeHash(Enum->GetName()));
+		}
+		return PropertyHash;
+	}
+	if (const FEnumProperty* EnumProperty = CastField<FEnumProperty>(InProperty))
+	{
+		static const FString PropertyString = TEXT("EnumProperty");
+		static const uint32 PropertyHash = GetTypeHash(PropertyString);
+		if (const UEnum* Enum = EnumProperty->GetEnum())
+		{
+			return HashCombine(PropertyHash, GetTypeHash(Enum->GetName()));
+		}
+		return PropertyHash;
+	}
+	if (const FArrayProperty* ArrayProperty = CastField<FArrayProperty>(InProperty))
+	{
+		static const FString PropertyString = TEXT("ArrayProperty");
+		static const uint32 PropertyHash = GetTypeHash(PropertyString);
+		return HashCombine(PropertyHash, GetPropertyHash(ArrayProperty->Inner));
+	}
+	if (const FObjectPropertyBase* ObjectProperty = CastField<FObjectPropertyBase>(InProperty))
+	{
+		static const FString PropertyString = TEXT("ObjectProperty");
+		static const uint32 PropertyHash = GetTypeHash(PropertyString);
+		if (const UClass* Class = ObjectProperty->PropertyClass)
+		{
+			return HashCombine(PropertyHash, GetTypeHash(Class->GetName()));
+		}
+		return PropertyHash;
+	}
+	if (const FStructProperty* StructProperty = CastField<FStructProperty>(InProperty))
+	{
+		static const FString PropertyString = TEXT("StructProperty");
+		static const uint32 PropertyHash = GetTypeHash(PropertyString);
+		if (const UStruct* Struct = StructProperty->Struct)
+		{
+			uint32 CombinedHash = HashCombine(PropertyHash, GetTypeHash(Struct->GetName()));
+			for (TFieldIterator<FProperty> PropertyIt(Struct); PropertyIt; ++PropertyIt)
+			{
+				if(FProperty* PropertyOnStruct = *PropertyIt)
+				{
+					CombinedHash = HashCombine(CombinedHash, GetPropertyHash(PropertyOnStruct));
+				}
+			}
+			return CombinedHash;
+		}
+		return PropertyHash;
+	}
+
+	return 0;
 }
 
 void FControlRigOverrideValue::Reset()
@@ -617,49 +775,51 @@ bool FControlRigOverrideValue::SetPropertiesFromPath(const FString& InPath, cons
 
 void* FControlRigOverrideValue::AllocateDataIfRequired()
 {
-	return AllocateDataIfRequired(DataArray);
+	return AllocateDataIfRequired(DataArray, DataPropertyHash);
 }
 
 void FControlRigOverrideValue::FreeDataIfRequired()
 {
-	FreeDataIfRequired(DataArray);
+	FreeDataIfRequired(DataArray, DataPropertyHash);
 }
 
-void* FControlRigOverrideValue::AllocateDataIfRequired(TArray<uint8, TAlignedHeapAllocator<16>>& InOutDataArray) const
+void* FControlRigOverrideValue::AllocateDataIfRequired(TArray<uint8, TAlignedHeapAllocator<16>>& InOutDataArray, uint32& InOutPropertyHash) const
 {
 	if(const FProperty* LeafProperty = GetLeafProperty())
 	{
-		if(InOutDataArray.IsEmpty())
+		if(InOutDataArray.IsEmpty() || !IsDataValid(InOutDataArray, InOutPropertyHash))
 		{
+			InOutDataArray.Reset();
 			InOutDataArray.AddZeroed(LeafProperty->GetSize());
+			InOutPropertyHash = GetPropertyHash();
 			LeafProperty->InitializeValue(InOutDataArray.GetData());
 		}
-		else
-		{
-			check(InOutDataArray.Num() == LeafProperty->GetSize());
-		}
 	}
 	return InOutDataArray.GetData();
 }
 
-void FControlRigOverrideValue::FreeDataIfRequired(TArray<uint8, TAlignedHeapAllocator<16>>& InOutDataArray) const
+void FControlRigOverrideValue::FreeDataIfRequired(TArray<uint8, TAlignedHeapAllocator<16>>& InOutDataArray, uint32& InOutPropertyHash) const
 {
-	if(void* DataPtr = InOutDataArray.GetData())
+	if(IsDataValid())
 	{
-		if(const FProperty* LeafProperty = GetLeafProperty())
+		if(void* DataPtr = InOutDataArray.GetData())
 		{
-			if(const FArrayProperty* ArrayProperty = CastField<FArrayProperty>(LeafProperty))
-			{
-				FScriptArrayHelper ArrayHelper(ArrayProperty, DataPtr);
-				ArrayHelper.DestroyContainer_Unsafe();
-			}
-			else
+			if(const FProperty* LeafProperty = GetLeafProperty())
 			{
-				LeafProperty->DestroyValue(DataPtr);
+				if(const FArrayProperty* ArrayProperty = CastField<FArrayProperty>(LeafProperty))
+				{
+					FScriptArrayHelper ArrayHelper(ArrayProperty, DataPtr);
+					ArrayHelper.DestroyContainer_Unsafe();
+				}
+				else
+				{
+					LeafProperty->DestroyValue(DataPtr);
+				}
 			}
 		}
 	}
 	InOutDataArray.Reset();
+	InOutPropertyHash = 0;
 }
 
 void FControlRigOverrideValue::CopyValue(void* InDestPtr, const void* InSourcePtr) const
@@ -718,7 +878,8 @@ bool FControlRigOverrideValue::SetFromString(const FString& InValue, const FCont
 		if(const FProperty* LeafProperty = GetLeafProperty())
 		{
 			TArray<uint8, TAlignedHeapAllocator<16>> TempDataArray;
-			void* TempDataPtr = AllocateDataIfRequired(TempDataArray); 
+			uint32 TempPropertyHash = 0;
+			void* TempDataPtr = AllocateDataIfRequired(TempDataArray, TempPropertyHash); 
 
 			// first try to import the string value using the property.
 			// we need to react to any message coming in from the import since missing parentheses for
@@ -727,14 +888,15 @@ bool FControlRigOverrideValue::SetFromString(const FString& InValue, const FCont
 			LeafProperty->ImportText_Direct(*InValue, TempDataPtr, nullptr, PPF_None, &ErrorPipe);
 			if(ErrorPipe.GetNumErrors() == 0)
 			{
-				FreeDataIfRequired(DataArray);
+				FreeDataIfRequired(DataArray, DataPropertyHash);
 				Swap(DataArray, TempDataArray);
+				Swap(DataPropertyHash, TempPropertyHash);
 				
 				CachedStringValue = InValue;
 				UpdateHash();
 				return true;
 			}
-			FreeDataIfRequired(TempDataArray);
+			FreeDataIfRequired(TempDataArray, TempPropertyHash);
 
 			// we've hit a problem when importing the value
 			if(const FStructProperty* StructProperty = CastField<FStructProperty>(LeafProperty))
diff --git a/Engine/Plugins/Animation/ControlRig/Source/ControlRig/Public/ControlRigObjectVersion.h b/Engine/Plugins/Animation/ControlRig/Source/ControlRig/Public/ControlRigObjectVersion.h
--- a/Engine/Plugins/Animation/ControlRig/Source/ControlRig/Public/ControlRigObjectVersion.h
+++ b/Engine/Plugins/Animation/ControlRig/Source/ControlRig/Public/ControlRigObjectVersion.h
@@ -130,6 +130,9 @@ struct FControlRigObjectVersion
 		// New setting for connectors to optionally specify their use only during post construction
 		RigHierarchyPostConstructionConnectors,
 
+		// Overrides store TOC data for properties to solidify loading of data when the definition has changed
+		OverridesStoreTOCDataForProperties,
+
 		// -----<new versions can be added above this line>-------------------------------------------------
 		VersionPlusOne,
 		LatestVersion = VersionPlusOne - 1
diff --git a/Engine/Plugins/Animation/ControlRig/Source/ControlRig/Public/ControlRigOverride.h b/Engine/Plugins/Animation/ControlRig/Source/ControlRig/Public/ControlRigOverride.h
--- a/Engine/Plugins/Animation/ControlRig/Source/ControlRig/Public/ControlRigOverride.h
+++ b/Engine/Plugins/Animation/ControlRig/Source/ControlRig/Public/ControlRigOverride.h
@@ -201,10 +201,14 @@ private:
 	UE_API bool SetPropertiesFromPath(const FString& InPath, const UStruct* InOwnerStruct);
 	UE_API void* AllocateDataIfRequired();
 	UE_API void FreeDataIfRequired();
-	UE_API void* AllocateDataIfRequired(TArray<uint8, TAlignedHeapAllocator<16>>& InOutDataArray) const;
-	UE_API void FreeDataIfRequired(TArray<uint8, TAlignedHeapAllocator<16>>& InOutDataArray) const;
+	UE_API void* AllocateDataIfRequired(TArray<uint8, TAlignedHeapAllocator<16>>& InOutDataArray, uint32& InOutPropertyHash) const;
+	UE_API void FreeDataIfRequired(TArray<uint8, TAlignedHeapAllocator<16>>& InOutDataArray, uint32& InOutPropertyHash) const;
 	UE_API void CopyValue(void* InDestPtr, const void* InSourcePtr) const;
 	UE_API void UpdateHash();
+	UE_API bool IsDataValid() const;
+	UE_API bool IsDataValid(const TArray<uint8, TAlignedHeapAllocator<16>>& InDataArray, const uint32& InOutPropertyHash) const;
+	UE_API uint32 GetPropertyHash() const;
+	UE_API static uint32 GetPropertyHash(const FProperty* InProperty);
 
 	template<typename T>
 	bool InitFromLegacyString(const FString& InLegacyString, const T& InDefault)
@@ -234,6 +238,7 @@ private:
 
 	TArray<FPropertyInfo> Properties;
 	TArray<uint8, TAlignedHeapAllocator<16>> DataArray;
+	uint32 DataPropertyHash;
 	
 	friend class UModularRigController;
 };
