ControlRig & Sequencer: fix incorrect use of the DefaultParent key

#rnx

diff --git a/Engine/Plugins/Animation/ControlRig/Source/ControlRig/Private/Sequencer/MovieSceneControlRigParameterSection.cpp b/Engine/Plugins/Animation/ControlRig/Source/ControlRig/Private/Sequencer/MovieSceneControlRigParameterSection.cpp
--- a/Engine/Plugins/Animation/ControlRig/Source/ControlRig/Private/Sequencer/MovieSceneControlRigParameterSection.cpp
+++ b/Engine/Plugins/Animation/ControlRig/Source/ControlRig/Private/Sequencer/MovieSceneControlRigParameterSection.cpp
@@ -1733,7 +1733,11 @@ void UMovieSceneControlRigParameterSection::AddSpaceChannel(FName InControlName,
 				if (const FRigControlElement* Control = ControlRig ? ControlRig->FindControl(InControlName) : nullptr)
 				{
 					const URigHierarchy* Hierarchy = ControlRig->GetHierarchy();
-					const FRigElementKey CurrentParentKey = Hierarchy->GetActiveParent(Control->GetKey());
+					FRigElementKey CurrentParentKey = Hierarchy->GetActiveParent(Control->GetKey());
+					if (CurrentParentKey == URigHierarchy::GetDefaultParentKey())
+					{
+						CurrentParentKey = Hierarchy->GetDefaultParent(Control->GetKey());
+					}
 
 					// set the default channel value to the current parent if it differs from the one initially constructed
 					if (CurrentParentKey != Hierarchy->GetFirstParent(Control->GetKey()))
