Control Rig: Pull the data for connector target combo always out of the model

#rnx

diff --git a/Engine/Plugins/Animation/ControlRig/Source/ControlRigEditor/Private/ControlRigElementDetails.cpp b/Engine/Plugins/Animation/ControlRig/Source/ControlRigEditor/Private/ControlRigElementDetails.cpp
--- a/Engine/Plugins/Animation/ControlRig/Source/ControlRigEditor/Private/ControlRigElementDetails.cpp
+++ b/Engine/Plugins/Animation/ControlRig/Source/ControlRigEditor/Private/ControlRigElementDetails.cpp
@@ -6368,17 +6368,24 @@ void FRigConnectorElementDetails::CustomizeConnectorTargets(IDetailLayoutBuilder
 	});
 
 	TArray<FRigElementKey> CurrentTargets;
-	if(const UControlRig* ControlRig = Hierarchy->GetTypedOuter<UControlRig>())
+	if(UControlRig* ControlRig = Hierarchy->GetTypedOuter<UControlRig>())
 	{
 		if(const FRigConnectorElement* Connector = PerElementInfos[0].Element.Get<FRigConnectorElement>())
 		{
-			const FRigElementKeyRedirector Redirector = ControlRig->GetElementKeyRedirector();
-			const FRigElementKeyRedirector::FCachedKeyArray* Cache = Redirector.Find(Connector->GetKey());
-			if(Cache && !Cache->IsEmpty())
+			if (const UModularRig* ModularRig = Cast<UModularRig>(ControlRig))
 			{
-				for(int32 Index = 0; Index < Cache->Num(); Index++)
+				CurrentTargets = ModularRig->GetModularRigModel().Connections.FindTargetsFromConnector(Connector->GetKey());
+			}
+			else
+			{
+				const FRigElementKeyRedirector& Redirector = ControlRig->GetElementKeyRedirector();
+				const FRigElementKeyRedirector::FCachedKeyArray* Cache = Redirector.Find(Connector->GetKey());
+				if(Cache && !Cache->IsEmpty())
 				{
-					CurrentTargets.Add((*Cache)[Index].GetKey());
+					for(int32 Index = 0; Index < Cache->Num(); Index++)
+					{
+						CurrentTargets.Add((*Cache)[Index].GetKey());
+					}
 				}
 			}
 		}
diff --git a/Engine/Plugins/Animation/ControlRig/Source/ControlRigEditor/Private/Editor/SModularRigTreeView.cpp b/Engine/Plugins/Animation/ControlRig/Source/ControlRigEditor/Private/Editor/SModularRigTreeView.cpp
--- a/Engine/Plugins/Animation/ControlRig/Source/ControlRigEditor/Private/Editor/SModularRigTreeView.cpp
+++ b/Engine/Plugins/Animation/ControlRig/Source/ControlRigEditor/Private/Editor/SModularRigTreeView.cpp
@@ -486,15 +486,7 @@ TArray<FRigElementKey> SModularRigModelItem::GetTargetKeys() const
 	TArray<FRigElementKey> Result;
 	if (const UModularRig* ModularRig = Delegates.GetModularRig())
 	{
-		const FRigElementKeyRedirector ElementKeyRedirector = ModularRig->GetElementKeyRedirector();
-		if(const FRigElementKeyRedirector::FCachedKeyArray* Cache = ElementKeyRedirector.Find(ConnectorKey))
-		{
-			Result.Reserve(Cache->Num());
-			for(int32 Index = 0; Index < Cache->Num(); Index++)
-			{
-				Result.Add(Cache->operator[](Index).GetKey());
-			}
-		}
+		Result = ModularRig->GetModularRigModel().Connections.FindTargetsFromConnector(ConnectorKey);
 	}
 	return Result;
 }
