Anim in Engine: Fix an issue where Integer Controls could not be edited in Anim Details when their property value was 0

diff --git a/Engine/Plugins/Animation/ControlRig/Source/ControlRigEditor/Private/AnimDetails/Widgets/SAnimDetailsValueNumeric.cpp b/Engine/Plugins/Animation/ControlRig/Source/ControlRigEditor/Private/AnimDetails/Widgets/SAnimDetailsValueNumeric.cpp
--- a/Engine/Plugins/Animation/ControlRig/Source/ControlRigEditor/Private/AnimDetails/Widgets/SAnimDetailsValueNumeric.cpp
+++ b/Engine/Plugins/Animation/ControlRig/Source/ControlRigEditor/Private/AnimDetails/Widgets/SAnimDetailsValueNumeric.cpp
@@ -29,7 +29,7 @@ namespace UE::ControlRigEditor
 			TOptional<NumericType> SliderMinValue;
 			TOptional<NumericType> SliderMaxValue;
 			double SliderExponent = 1.0;
-			double Delta = 0.0;
+			double Delta = 1.0;
 			float ShiftMultiplier = 10.f;
 			float CtrlMultiplier = 0.1f;
 			bool SupportDynamicSliderMaxValue = false;
