Control Rig: Rely on spaces as valid space targets

#rnx

diff --git a/Engine/Plugins/Animation/ControlRig/Source/ControlRig/Private/Rigs/RigHierarchy.cpp b/Engine/Plugins/Animation/ControlRig/Source/ControlRig/Private/Rigs/RigHierarchy.cpp
--- a/Engine/Plugins/Animation/ControlRig/Source/ControlRig/Private/Rigs/RigHierarchy.cpp
+++ b/Engine/Plugins/Animation/ControlRig/Source/ControlRig/Private/Rigs/RigHierarchy.cpp
@@ -3980,11 +3980,11 @@ bool URigHierarchy::CanCauseCycle(const FRigTransformElement* InChild, const FRi
 				{
 					// if this is a defined space for something - allow it even if there's a cycle.
 					// we assume that spaces defined by riggers are safe to use.
-					if (const FRigControlElement* Control = Cast<FRigControlElement>(InParent))
+					if (const FRigControlElement* Control = Cast<FRigControlElement>(InChild))
 					{
-						if (Control->Settings.Customization.AvailableSpaces.ContainsByPredicate([InChild](const FRigElementKeyWithLabel& Space) -> bool
+						if (Control->Settings.Customization.AvailableSpaces.ContainsByPredicate([InParent](const FRigElementKeyWithLabel& Space) -> bool
 						{
-							return Space.Key == InChild->GetKey();
+							return Space.Key == InParent->GetKey();
 						}))
 						{
 							return false;
