[ControlRigOverride] Added a missing nullcheck when getting property hash

#DisabledValidations Preflight

diff --git a/Engine/Plugins/Animation/ControlRig/Source/ControlRig/Private/ControlRigOverride.cpp b/Engine/Plugins/Animation/ControlRig/Source/ControlRig/Private/ControlRigOverride.cpp
--- a/Engine/Plugins/Animation/ControlRig/Source/ControlRig/Private/ControlRigOverride.cpp
+++ b/Engine/Plugins/Animation/ControlRig/Source/ControlRig/Private/ControlRigOverride.cpp
@@ -305,6 +305,11 @@ if (InProperty->IsA<PropertyType>()) \
 
 uint32 FControlRigOverrideValue::GetPropertyHash(const FProperty* InProperty)
 {
+	if (InProperty == nullptr)
+	{
+		return 0;
+	}
+	
 	UE_CONTROLRIGOVERRIDE_PROPERTYHASHCASE(FBoolProperty);
 	UE_CONTROLRIGOVERRIDE_PROPERTYHASHCASE(FUInt16Property);
 	UE_CONTROLRIGOVERRIDE_PROPERTYHASHCASE(FUInt32Property);
