Constraints: fix autokey on release mixing constraints & explicit rotation mode

#rnx

diff --git a/Engine/Plugins/Animation/ControlRig/Source/ControlRigEditor/Private/EditMode/ControlRigEditMode.cpp b/Engine/Plugins/Animation/ControlRig/Source/ControlRigEditor/Private/EditMode/ControlRigEditMode.cpp
--- a/Engine/Plugins/Animation/ControlRig/Source/ControlRigEditor/Private/EditMode/ControlRigEditMode.cpp
+++ b/Engine/Plugins/Animation/ControlRig/Source/ControlRigEditor/Private/EditMode/ControlRigEditMode.cpp
@@ -5294,6 +5294,10 @@ void FControlRigEditMode::SetControlShapeTransform(
 			ControlRig->SetControlLocalTransform(InShapeActor->ControlName, LocalTransform, bNotify, Context, bUndo, bFixEulerFlips);
 			EvaluateRigIfAdditive();
 		}
+		else
+		{
+			KeyframeData.bConstraintSpace = true;
+		}
 
 		TGuardValue<bool> CompensateGuard(FMovieSceneConstraintChannelHelper::bDoNotCompensate, true);
 		const FConstraintsManagerController& Controller = FConstraintsManagerController::Get(World);
