Control Rig: Force the execution of construction inmediately in Update_AnyThread if needed, so that the PoseAdapter can be updated properly during Evaluate_AnyThread
#rnx

diff --git a/Engine/Plugins/Animation/ControlRig/Source/ControlRig/Private/AnimNode_ControlRigBase.cpp b/Engine/Plugins/Animation/ControlRig/Source/ControlRig/Private/AnimNode_ControlRigBase.cpp
--- a/Engine/Plugins/Animation/ControlRig/Source/ControlRig/Private/AnimNode_ControlRigBase.cpp
+++ b/Engine/Plugins/Animation/ControlRig/Source/ControlRig/Private/AnimNode_ControlRigBase.cpp
@@ -20,6 +20,7 @@
 
 DECLARE_CYCLE_STAT(TEXT("ControlRig_UpdateInput"), STAT_ControlRig_UpdateInput, STATGROUP_Anim);
 DECLARE_CYCLE_STAT(TEXT("ControlRig_Evaluate"), STAT_ControlRig_Evaluate, STATGROUP_Anim);
+DECLARE_CYCLE_STAT(TEXT("ControlRig_Evaluate_Construction"), STAT_ControlRig_Evaluate_Construction, STATGROUP_Anim);
 DECLARE_CYCLE_STAT(TEXT("ControlRig_UpdateOutput"), STAT_ControlRig_UpdateOutput, STATGROUP_Anim);
 
 #if ENABLE_ANIM_DEBUG
@@ -184,6 +185,8 @@ void FAnimNode_ControlRigBase::Evaluate_AnyThread(FPoseContext& Output)
 		SourcePose.ResetToRefPose();
 	}
 
+	ExecuteConstructionIfNeeded();
+
 	if(!ControlRigHierarchyMappings.CheckPoseAdapter())
 	{
 		Output = SourcePose;
@@ -327,6 +330,59 @@ void FAnimNode_ControlRigBase::ExecuteControlRig(FPoseContext& InOutput)
 	}
 }
 
+void FAnimNode_ControlRigBase::ExecuteConstructionIfNeeded()
+{
+	SCOPE_CYCLE_COUNTER(STAT_ControlRig_Evaluate_Construction);
+
+	if (UControlRig* ControlRig = GetControlRig())
+	{
+		if (!ControlRig->IsConstructionRequired())
+		{
+			return;
+		}
+		
+		// Before we start modifying the RigHierarchy, we need to lock the rig to avoid
+		// corrupted state
+		UE::TScopeLock LockRig(ControlRig->GetEvaluateMutex());
+
+		URigHierarchy* Hierarchy = ControlRig->GetHierarchy();
+		if(Hierarchy == nullptr)
+		{
+			return;
+		}
+
+		if (bExecute)
+		{
+			TGuardValue<bool> ResetCurrentTransfromsAfterConstructionGuard(ControlRig->bResetCurrentTransformsAfterConstruction, true);
+			
+#if WITH_EDITOR
+			if(Hierarchy->IsTracingChanges())
+			{
+				Hierarchy->StorePoseForTrace(TEXT("FAnimNode_ControlRigBase::BeforeEvaluate"));
+			}
+#endif
+
+			// empty queue (only leave construction and post construction)
+			const TArray<FName> CurrentEventQueue = ControlRig->GetEventQueue();
+			ControlRig->SetEventQueue({});
+			TGuardValue<bool> OnlyRunConstruction(ControlRig->bOnlyRunConstruction, true);
+
+			// evaluate control rig
+			UpdateGetAssetUserDataDelegate(ControlRig);
+			ControlRig->Evaluate_AnyThread();
+
+#if WITH_EDITOR
+			if(Hierarchy->IsTracingChanges())
+			{
+				Hierarchy->StorePoseForTrace(TEXT("FAnimNode_ControlRigBase::AfterEvaluate"));
+			}
+#endif
+
+			ControlRig->SetEventQueue(CurrentEventQueue);
+		}
+	}
+}
+
 struct FControlRigControlScope
 {
 	FControlRigControlScope(UControlRig* InControlRig)
diff --git a/Engine/Plugins/Animation/ControlRig/Source/ControlRig/Private/ControlRig.cpp b/Engine/Plugins/Animation/ControlRig/Source/ControlRig/Private/ControlRig.cpp
--- a/Engine/Plugins/Animation/ControlRig/Source/ControlRig/Private/ControlRig.cpp
+++ b/Engine/Plugins/Animation/ControlRig/Source/ControlRig/Private/ControlRig.cpp
@@ -95,6 +95,7 @@ UControlRig::UControlRig(const FObjectInitializer& ObjectInitializer)
 	, bEvaluationTriggeredFromInteraction(false)
 	, bInteractionJustBegan(false)
 	, bIsAdditive(false)
+	, bOnlyRunConstruction(false)
 	, DebugBoneRadiusMultiplier(1.f)
 #if WITH_EDITOR
 	, bRecordSelectionPoseForConstructionMode(true)
@@ -520,6 +521,11 @@ void UControlRig::Evaluate_AnyThread()
 		{
 			Super::Evaluate_AnyThread();
 		}
+
+		if (bOnlyRunConstruction)
+		{
+			return;
+		}
 		
 		// we can have other systems trying to poke into running instances of Control Rigs
 		// on the anim thread and query data, such as
diff --git a/Engine/Plugins/Animation/ControlRig/Source/ControlRig/Public/AnimNode_ControlRigBase.h b/Engine/Plugins/Animation/ControlRig/Source/ControlRig/Public/AnimNode_ControlRigBase.h
--- a/Engine/Plugins/Animation/ControlRig/Source/ControlRig/Public/AnimNode_ControlRigBase.h
+++ b/Engine/Plugins/Animation/ControlRig/Source/ControlRig/Public/AnimNode_ControlRigBase.h
@@ -143,6 +143,7 @@ protected:
 	
 	// execute control rig on the input pose and outputs the result
 	UE_API void ExecuteControlRig(FPoseContext& InOutput);
+	UE_API void ExecuteConstructionIfNeeded();
 
 	UE_API void QueueControlRigDrawInstructions(UControlRig* ControlRig, FAnimInstanceProxy* Proxy) const;
 
diff --git a/Engine/Plugins/Animation/ControlRig/Source/ControlRig/Public/ControlRig.h b/Engine/Plugins/Animation/ControlRig/Source/ControlRig/Public/ControlRig.h
--- a/Engine/Plugins/Animation/ControlRig/Source/ControlRig/Public/ControlRig.h
+++ b/Engine/Plugins/Animation/ControlRig/Source/ControlRig/Public/ControlRig.h
@@ -772,6 +772,7 @@ protected:
 	};
 	FRigPose PoseBeforeBackwardsSolve;
 	FRigPose ControlsAfterBackwardsSolve;
+	bool bOnlyRunConstruction = false;
 	TMap<FRigElementKey, FRigSetControlValueInfo> ControlValues; // Layered Rigs: Additive values in local space (to add after backwards solve)
 	TMap<FRigElementKey, FRigSwitchParentInfo> SwitchParentValues; // Layered Rigs: Parent switching values to perform after backwards solve
 
diff --git a/Engine/Plugins/Experimental/UAF/UAFControlRig/Source/UAFControlRig/Private/ControlRigTask.cpp b/Engine/Plugins/Experimental/UAF/UAFControlRig/Source/UAFControlRig/Private/ControlRigTask.cpp
--- a/Engine/Plugins/Experimental/UAF/UAFControlRig/Source/UAFControlRig/Private/ControlRigTask.cpp
+++ b/Engine/Plugins/Experimental/UAF/UAFControlRig/Source/UAFControlRig/Private/ControlRigTask.cpp
@@ -90,6 +90,7 @@ void FAnimNextControlRigTask::ExecuteControlRig(UE::UAF::FEvaluationVM& VM, UE::
 			{
 				//UpdateGetAssetUserDataDelegate(ControlRig);
 				ControlRig->Execute(FRigUnit_PrepareForExecution::EventName);
+				ControlRig->Execute(FRigUnit_PostPrepareForExecution::EventName);
 			}
 
 			// UpdateInputOutputMappingIfRequired was done in CacheBones, but there is no AnimNext equivalent
