ControlRig dependencies: optimize the dependencies computation when not tracking variables

#rnx

diff --git a/Engine/Plugins/Animation/ControlRig/Source/ControlRig/Private/Rigs/RigDependency.cpp b/Engine/Plugins/Animation/ControlRig/Source/ControlRig/Private/Rigs/RigDependency.cpp
--- a/Engine/Plugins/Animation/ControlRig/Source/ControlRig/Private/Rigs/RigDependency.cpp
+++ b/Engine/Plugins/Animation/ControlRig/Source/ControlRig/Private/Rigs/RigDependency.cpp
@@ -258,58 +258,57 @@ void FRigDependenciesProviderForVM::ComputeDependencies(const FControlRigExecute
 	{
 		WrittenVariables.Reset();
 
-		if (ReadRecordsPerInstruction.IsValidIndex(InstructionIndex))
+		if (bFollowVariables)
 		{
-			const TArray<FInstructionRecord>& ReadRecords = ReadRecordsPerInstruction[InstructionIndex];
-			for(int32 ReadRecordIndex = 0; ReadRecordIndex < ReadRecords.Num(); ReadRecordIndex++)
+			if (ReadRecordsPerInstruction.IsValidIndex(InstructionIndex))
 			{
-				const FInstructionRecord& ReadRecord = ReadRecords[ReadRecordIndex];
-
-				const FRigHierarchyRecord ReadElementRecord(
-					ReadRecord.MetadataName.IsNone() ? FRigHierarchyRecord::EType_RigElement : FRigHierarchyRecord::EType_Metadata,
-					ReadRecord.ElementIndex,
-					ReadRecord.MetadataName);
+				const TArray<FInstructionRecord>& ReadRecords = ReadRecordsPerInstruction[InstructionIndex];
+				for(int32 ReadRecordIndex = 0; ReadRecordIndex < ReadRecords.Num(); ReadRecordIndex++)
+				{
+					const FInstructionRecord& ReadRecord = ReadRecords[ReadRecordIndex];
 
-				ResetLastInstruction(InstructionIndex, ReadElementRecord);
+					const FRigHierarchyRecord ReadElementRecord(
+						ReadRecord.MetadataName.IsNone() ? FRigHierarchyRecord::EType_RigElement : FRigHierarchyRecord::EType_Metadata,
+						ReadRecord.ElementIndex,
+						ReadRecord.MetadataName);
 
-				// follow the output operands through the bytecode to determine variable writes of this instruction
-				for(int32 InstructionToVisitIndex = 0; InstructionToVisitIndex < InstructionsToVisit.Num(); InstructionToVisitIndex++)
-				{
-					const int32 InstructionToVisit = InstructionsToVisit[InstructionToVisitIndex].Get<0>();
+					ResetLastInstruction(InstructionIndex, ReadElementRecord);
 
-					const FRigVMOperandArray OutputOperands = InVM->GetByteCode().GetOutputOperands(InstructionToVisit);
-					for(const FRigVMOperand OutputOperand : OutputOperands)
+					// follow the output operands through the bytecode to determine variable writes of this instruction
+					for(int32 InstructionToVisitIndex = 0; InstructionToVisitIndex < InstructionsToVisit.Num(); InstructionToVisitIndex++)
 					{
-						FRigHierarchyRecord LastRecord = InstructionsToVisit[InstructionToVisitIndex].Get<1>();
+						const int32 InstructionToVisit = InstructionsToVisit[InstructionToVisitIndex].Get<0>();
 
-						if (OutputOperand.GetMemoryType() == ERigVMMemoryType::External)
+						const FRigVMOperandArray OutputOperands = InVM->GetByteCode().GetOutputOperands(InstructionToVisit);
+						TSet<FRigVMOperand> OutputOperandsNoRegisterOffset;
+						OutputOperandsNoRegisterOffset.Reserve(OutputOperands.Num());
+						for(const FRigVMOperand& OutputOperand : OutputOperands)
 						{
-							if (!bFollowVariables)
-							{
-								continue;
-							}
-							const FRigHierarchyRecord VariableRecord(FRigHierarchyRecord::EType_Variable, OutputOperand.GetRegisterIndex(), ReadRecord.VM);
-							WrittenVariables.AddUnique(VariableRecord);
-							LastRecord = VariableRecord;
+							OutputOperandsNoRegisterOffset.Emplace({OutputOperand.GetMemoryType(), OutputOperand.GetRegisterIndex()});
 						}
 						
-						const FRigVMOperand OutputOperandNoRegisterOffset(OutputOperand.GetMemoryType(), OutputOperand.GetRegisterIndex());
-
-						if(const TArray<int32>* InstructionsWithOutputOperand = InputOperandToInstructions.Find(OutputOperandNoRegisterOffset))
+						for(const FRigVMOperand& OutputOperandNoRegisterOffset : OutputOperandsNoRegisterOffset)
 						{
-							FindOrAddInstruction(*InstructionsWithOutputOperand, LastRecord);
+							FRigHierarchyRecord LastRecord = InstructionsToVisit[InstructionToVisitIndex].Get<1>();
+
+							if (OutputOperandNoRegisterOffset.GetMemoryType() == ERigVMMemoryType::External)
+							{
+								const FRigHierarchyRecord VariableRecord(FRigHierarchyRecord::EType_Variable, OutputOperandNoRegisterOffset.GetRegisterIndex(), ReadRecord.VM);
+								WrittenVariables.AddUnique(VariableRecord);
+								LastRecord = VariableRecord;
+							}
+
+							if(const TArray<int32>* InstructionsWithOutputOperand = InputOperandToInstructions.Find(OutputOperandNoRegisterOffset))
+							{
+								FindOrAddInstruction(*InstructionsWithOutputOperand, LastRecord);
+							}
 						}
 					}
-				}
-
-				const FRigHierarchyRecord ElementRecord(
-					ReadRecord.MetadataName.IsNone() ? FRigHierarchyRecord::EType_RigElement : FRigHierarchyRecord::EType_Metadata,
-					ReadRecord.ElementIndex,
-					ReadRecord.MetadataName);
 
-				for (const FRigHierarchyRecord& WrittenVariableRecord : WrittenVariables)
-				{
-					OutDependencies.FindOrAdd(WrittenVariableRecord).Add(ElementRecord);
+					for (const FRigHierarchyRecord& WrittenVariableRecord : WrittenVariables)
+					{
+						OutDependencies.FindOrAdd(WrittenVariableRecord).Add(ReadElementRecord);
+					}
 				}
 			}
 		}
