Control Rig: Fix crash when clearing CR class from the CR Component
#rnx

diff --git a/Engine/Plugins/Animation/ControlRig/Source/ControlRigEditor/Private/Sequencer/ControlRigParameterTrackEditor.cpp b/Engine/Plugins/Animation/ControlRig/Source/ControlRigEditor/Private/Sequencer/ControlRigParameterTrackEditor.cpp
--- a/Engine/Plugins/Animation/ControlRig/Source/ControlRigEditor/Private/Sequencer/ControlRigParameterTrackEditor.cpp
+++ b/Engine/Plugins/Animation/ControlRig/Source/ControlRigEditor/Private/Sequencer/ControlRigParameterTrackEditor.cpp
@@ -3154,78 +3154,81 @@ void FControlRigParameterTrackEditor::SelectRigsAndControls(UControlRig* Control
 		if (MovieSection)
 		{
 			ControlRig = MovieSection->GetControlRig();
-			//Only create the edit mode if we have a KeyAra selected and it's not set and we have some boundobjects.
-			if (!ControlRigEditMode)
+			if (ControlRig)
 			{
-				ControlRigEditMode = GetEditMode(true);
-				if (TSharedPtr<IControlRigObjectBinding> ObjectBinding = ControlRig->GetObjectBinding())
+				//Only create the edit mode if we have a KeyAra selected and it's not set and we have some boundobjects.
+				if (!ControlRigEditMode)
 				{
-					if (ControlRigEditMode)
+					ControlRigEditMode = GetEditMode(true);
+					if (TSharedPtr<IControlRigObjectBinding> ObjectBinding = ControlRig->GetObjectBinding())
 					{
-						ControlRigEditMode->AddControlRigObject(ControlRig, Sequencer);
+						if (ControlRigEditMode)
+						{
+							ControlRigEditMode->AddControlRigObject(ControlRig, Sequencer);
+						}
 					}
 				}
-			}
-			else
-			{
-				if (ControlRigEditMode->AddControlRigObject(ControlRig, Sequencer))
+				else
 				{
-					//force an evaluation, this will get the control rig setup so edit mode looks good.
-					if (Sequencer)
+					if (ControlRigEditMode->AddControlRigObject(ControlRig, Sequencer))
 					{
-						Sequencer->NotifyMovieSceneDataChanged(EMovieSceneDataChangeType::Unknown);
+						//force an evaluation, this will get the control rig setup so edit mode looks good.
+						if (Sequencer)
+						{
+							Sequencer->NotifyMovieSceneDataChanged(EMovieSceneDataChangeType::Unknown);
+						}
 					}
 				}
-			}
 
-			const FMovieSceneChannelMetaData* MetaData = KeyArea->GetChannel().GetMetaData();
-			if (MetaData)
-			{
-				StringArray.SetNum(0);
-				FString String = MetaData->Name.ToString();
-				String.ParseIntoArray(StringArray, TEXT("."));
-				if (StringArray.Num() > 0)
+				const FMovieSceneChannelMetaData* MetaData = KeyArea->GetChannel().GetMetaData();
+				if (MetaData)
 				{
-					const FName ControlName(*StringArray[0]);
-
-					// skip nested controls which have the shape enabled flag turned on
-					if(const URigHierarchy* Hierarchy = ControlRig->GetHierarchy())
+					StringArray.SetNum(0);
+					FString String = MetaData->Name.ToString();
+					String.ParseIntoArray(StringArray, TEXT("."));
+					if (StringArray.Num() > 0)
 					{
+						const FName ControlName(*StringArray[0]);
 
-						if(const FRigControlElement* ControlElement = Hierarchy->Find<FRigControlElement>(FRigElementKey(ControlName, ERigElementType::Control)))
+						// skip nested controls which have the shape enabled flag turned on
+						if(const URigHierarchy* Hierarchy = ControlRig->GetHierarchy())
 						{
-						
-							if (ControlElement->Settings.ControlType == ERigControlType::Bool ||
-								ControlElement->Settings.ControlType == ERigControlType::Float ||
-								ControlElement->Settings.ControlType == ERigControlType::ScaleFloat ||
-								ControlElement->Settings.ControlType == ERigControlType::Integer)
+
+							if(const FRigControlElement* ControlElement = Hierarchy->Find<FRigControlElement>(FRigElementKey(ControlName, ERigElementType::Control)))
 							{
-								if (ControlElement->Settings.SupportsShape() || !Hierarchy->IsAnimatable(ControlElement))
+						
+								if (ControlElement->Settings.ControlType == ERigControlType::Bool ||
+									ControlElement->Settings.ControlType == ERigControlType::Float ||
+									ControlElement->Settings.ControlType == ERigControlType::ScaleFloat ||
+									ControlElement->Settings.ControlType == ERigControlType::Integer)
 								{
-
-									if (const FRigControlElement* ParentControlElement = Cast<FRigControlElement>(Hierarchy->GetFirstParent(ControlElement)))
+									if (ControlElement->Settings.SupportsShape() || !Hierarchy->IsAnimatable(ControlElement))
 									{
-										if (const TSet<FName>* Controls = RigsAndControls.Find(ControlRig))
+
+										if (const FRigControlElement* ParentControlElement = Cast<FRigControlElement>(Hierarchy->GetFirstParent(ControlElement)))
 										{
-											if (Controls->Contains(ParentControlElement->GetFName()))
+											if (const TSet<FName>* Controls = RigsAndControls.Find(ControlRig))
 											{
-												continue;
+												if (Controls->Contains(ParentControlElement->GetFName()))
+												{
+													continue;
+												}
 											}
 										}
 									}
 								}
+								RigsAndControls.FindOrAdd(ControlRig).Add(ControlName);
 							}
-							RigsAndControls.FindOrAdd(ControlRig).Add(ControlName);
 						}
 					}
 				}
-			}
 			
-			if (bSelectedSectionSetsSectionToKey)
-			{
-				if (UMovieSceneTrack* Track = MovieSection->GetTypedOuter<UMovieSceneTrack>())
+				if (bSelectedSectionSetsSectionToKey)
 				{
-					TracksAndSections.FindOrAdd(Track).Add(MovieSection);
+					if (UMovieSceneTrack* Track = MovieSection->GetTypedOuter<UMovieSceneTrack>())
+					{
+						TracksAndSections.FindOrAdd(Track).Add(MovieSection);
+					}
 				}
 			}
 		}
