Control Rig: Fix layered rigs not initializing null elements before backwards solve
#rnx

diff --git a/Engine/Plugins/Animation/ControlRig/Source/ControlRig/Private/ControlRig.cpp b/Engine/Plugins/Animation/ControlRig/Source/ControlRig/Private/ControlRig.cpp
--- a/Engine/Plugins/Animation/ControlRig/Source/ControlRig/Private/ControlRig.cpp
+++ b/Engine/Plugins/Animation/ControlRig/Source/ControlRig/Private/ControlRig.cpp
@@ -552,6 +552,10 @@ void UControlRig::Evaluate_AnyThread()
 					Hierarchy->SetTransform(Control, Hierarchy->GetTransform(Control, ERigTransformType::InitialLocal), ERigTransformType::CurrentLocal, true);
 				}
 			}
+			else if (FRigNullElement* NullElement = Cast<FRigNullElement>(Element))
+			{
+				Hierarchy->SetTransform(NullElement, Hierarchy->GetTransform(NullElement, ERigTransformType::InitialLocal), ERigTransformType::CurrentLocal, true);
+			}
 			return true;
 		});
 
@@ -2439,7 +2443,8 @@ FRigControlValue UControlRig::GetControlValueFromGlobalTransform(const FName& In
 
 					// LocalTransform = InGlobal - ParentGlobalTransform - OffsetLocal
 					const FTransform OffsetLocalTransform = InGlobalTransform.GetRelativeTransform(ParentGlobalTransform);
-					const FTransform LocalTransform = OffsetLocalTransform.GetRelativeTransform(DynamicHierarchy->GetControlOffsetTransform(ControlElement, ERigTransformType::CurrentLocal));
+					const FTransform OffsetTransform = DynamicHierarchy->GetControlOffsetTransform(ControlElement, ERigTransformType::CurrentLocal);
+					const FTransform LocalTransform = OffsetLocalTransform.GetRelativeTransform(OffsetTransform);
 
 					// Additive = LocalTransform - AnimLocalTransform
 					const FTransform& AnimLocalTransform = ControlsAfterBackwardsSolve[ControlIndex].LocalTransform;
@@ -3666,7 +3671,7 @@ void UControlRig::PostInitInstance(URigVMHost* InCDO)
 	// set up the hierarchy
 	{
 		// If this is not a CDO, it should have never saved its hieararchy. However, we have found that some rigs in the past
-		// did save their hiearchies. If that's the case, let's mark them as garbage and rename them before creating our own hierarchy.
+		// did save their hierarchies. If that's the case, let's mark them as garbage and rename them before creating our own hierarchy.
 		if (!HasAnyFlags(RF_ClassDefaultObject))
 		{
 			FGCScopeGuard Guard;
