Anim Layers: Fix duplicate/merging control rig sections when the channel layout is different between sections.

#DisabledValidations Preflight

diff --git a/Engine/Plugins/Animation/ControlRig/Source/ControlRigEditor/Private/Sequencer/AnimLayers/AnimLayers.cpp b/Engine/Plugins/Animation/ControlRig/Source/ControlRigEditor/Private/Sequencer/AnimLayers/AnimLayers.cpp
--- a/Engine/Plugins/Animation/ControlRig/Source/ControlRigEditor/Private/Sequencer/AnimLayers/AnimLayers.cpp
+++ b/Engine/Plugins/Animation/ControlRig/Source/ControlRigEditor/Private/Sequencer/AnimLayers/AnimLayers.cpp
@@ -1612,8 +1612,493 @@
 	}
 }
 
+//get the set of sections active, or with this control name in it
+static void GetValidCRSections(UMovieSceneSection* InSection, const FName& ControlName, TArray<UMovieSceneSection*>& OutSections)
+{
+	if (UMovieSceneTrack* OwnerTrack = InSection->GetTypedOuter<UMovieSceneTrack>())
+	{
+		TArray<UMovieSceneSection*> TrackSections = OwnerTrack->GetAllSections();
+		for (UMovieSceneSection* Section : TrackSections)
+		{
+			if (UMovieSceneControlRigParameterSection* CRSection = Cast<UMovieSceneControlRigParameterSection>(Section))
+			{
+				if (Section && (Section->IsActive() || Section == InSection) &&
+					CRSection->GetControlNameMask(ControlName))
+				{
+					OutSections.Add(Section);
+				}
+			}
+		}
+	}
+}
+
+static void GetValidSections(UMovieSceneSection* InSection, TArray<UMovieSceneSection*>& OutSections)
+{
+	if (UMovieSceneTrack* OwnerTrack = InSection->GetTypedOuter<UMovieSceneTrack>())
+	{
+		TArray<UMovieSceneSection*> TrackSections = OwnerTrack->GetAllSections();
+		for (UMovieSceneSection* Section : TrackSections)
+		{
+			if (Section && (Section->IsActive() || Section == InSection))
+			{
+				OutSections.Add(Section);
+			}
+		}
+	}
+}
+static void RevertWeightChannelToOne(FMovieSceneFloatChannel* FloatChannel, const TRange<FFrameNumber>& FrameRange)
+{
+	if (FloatChannel)
+	{
+		//don't add key if there!
+		TArray<FFrameNumber> KeyTimes;
+		TArray<FKeyHandle> KeyHandles;
+		FloatChannel->GetKeys(FrameRange, &KeyTimes, &KeyHandles);
+		if (KeyTimes.Num() > 0)
+		{
+			FloatChannel->DeleteKeys(KeyHandles);
+			if (FloatChannel->GetNumKeys() > 0) //if it still has keys see if the weight is not one 
+			{
+				float StartValue = 0.0, EndValue = 0.0;
+				FFrameTime StartFrameTime = FrameRange.HasLowerBound() ? FrameRange.GetLowerBoundValue() : FFrameTime(KeyTimes[0]);
+				FFrameTime EndFrameTime = FrameRange.HasUpperBound() ? FrameRange.GetUpperBoundValue() : FFrameTime(KeyTimes[KeyTimes.Num() -1]);
+				FloatChannel->Evaluate(StartFrameTime, StartValue);
+				FloatChannel->Evaluate(EndFrameTime, EndValue);
+				//if not zero at boundaries set one keys there
+				if (FMath::IsNearlyEqual(StartValue, 1.0f) == false ||
+					FMath::IsNearlyEqual(EndValue, 1.0f) == false)
+				{
+					FloatChannel->AddCubicKey(StartFrameTime.FrameNumber, 1.0f, ERichCurveTangentMode::RCTM_SmartAuto);
+					FloatChannel->AddCubicKey(EndFrameTime.FrameNumber, 1.0f, ERichCurveTangentMode::RCTM_SmartAuto);
+				}
+			}
+			else
+			{
+				FloatChannel->SetDefault(1.0f);
+			}
+		}
+	}
+}
+
+namespace UE::AIE
+{ 
+// disable section if no control in the section and override, need to make sure we skip it's eval or it will incorrectly
+//override when merging
+static bool ShouldDisableSection(UMovieSceneControlRigParameterSection* CRSection, FRigControlElement* ControlElement)
+{
+	return (CRSection->GetControlNameMask(ControlElement->GetFName()) == false &&
+		CRSection->GetBlendType().IsValid() && CRSection->GetBlendType() == EMovieSceneBlendType::Override);
+}
+
+static void MergeControlRigSections(UMovieSceneControlRigParameterSection* BaseSection, UMovieSceneControlRigParameterSection* Section, const TRange<FFrameNumber>& Range,
+	const int32* Increment)
+{
+
+	if (!BaseSection || !Section)
+	{
+		return;
+	}
+	TArrayView<FMovieSceneFloatChannel*> BaseFloatChannels = BaseSection->GetChannelProxy().GetChannels<FMovieSceneFloatChannel>();
+	if (BaseFloatChannels.Num() > 0)
+	{
+		//need to go through each control and merge that
+		TArray<FRigControlElement*> Controls;
+		UControlRig* ControlRig = BaseSection->GetControlRig();
+		if (ControlRig == nullptr)
+		{
+			return;
+		}
+		ControlRig->GetControlsInOrder(Controls);
+		URigHierarchy* Hierarchy = ControlRig->GetHierarchy();
+		if (Hierarchy == nullptr)
+		{
+			return;
+		}
+
+		BaseSection->Modify();
+
+		const bool bIsOverride = (Section->GetBlendType().IsValid() && Section->GetBlendType() == EMovieSceneBlendType::Override) ||
+			(BaseSection->GetBlendType().IsValid() && BaseSection->GetBlendType() == EMovieSceneBlendType::Override);
+
+
+		for (int32 LocalControlIndex = 0; LocalControlIndex < Controls.Num(); ++LocalControlIndex)
+		{
+			FRigControlElement* ControlElement = Controls[LocalControlIndex];
+			check(ControlElement);
+			if (!Hierarchy->IsAnimatable(ControlElement))
+			{
+				continue;
+			}
+
+			const FName& ControlName = ControlElement->GetFName();
+
+			FChannelMapInfo* pTopChannelIndex = Section->ControlChannelMap.Find(ControlName);
+			if (!pTopChannelIndex)
+			{
+				continue;
+			}
+			if (FChannelMapInfo* pChannelIndex = BaseSection->ControlChannelMap.Find(ControlName))
+			{
+				//these may be different due to rig changes
+				const int32 ChannelIndex = pChannelIndex->ChannelIndex;
+				const int32 TopChannelIndex = pTopChannelIndex->ChannelIndex;
+
+				//if override we mask out if both are masked out, if not override we mask out if the top section is masked out
+				const bool bMaskedOutOfBase = BaseSection->GetControlNameMask(ControlName) == false;
+				const bool bMaskedOutOfSection = Section->GetControlNameMask(ControlName) == false;
+				const bool bMaskKeyOut = bIsOverride ? (bMaskedOutOfBase && bMaskedOutOfSection)
+					: bMaskedOutOfSection;
+				if (bMaskKeyOut)
+				{
+					continue;
+				}
+				TArray<UMovieSceneSection*> TrackSections;
+				GetValidCRSections(BaseSection, ControlName, TrackSections);
+
+				TOptional<bool> bBaseSectionResetActive;
+				TOptional<bool> bSectionResetActive;
+				if (ShouldDisableSection(BaseSection, ControlElement))
+				{
+					bBaseSectionResetActive = BaseSection->IsActive();
+					BaseSection->SetIsActive(false);
+				}
+
+				if (ShouldDisableSection(Section, ControlElement))
+				{
+					bSectionResetActive = Section->IsActive();
+					Section->SetIsActive(false);
+				}
+
+
+				switch (ControlElement->Settings.ControlType)
+				{
+
+				case ERigControlType::Float:
+				case ERigControlType::ScaleFloat:
+				{
+					int32 StartIndex = ChannelIndex;
+					int32 EndIndex = ChannelIndex;
+					int32 TopEndIndex = TopChannelIndex;
+					MovieSceneToolHelpers::FBaseTopSections BT(BaseSection,
+						Section, StartIndex, EndIndex, TopChannelIndex, TopEndIndex);
+					MovieSceneToolHelpers::MergeTwoSections<FMovieSceneFloatChannel>(BT, Range, TrackSections, Increment);
+					break;
+				}
+				case ERigControlType::Vector2D:
+				{
+					int32 StartIndex = ChannelIndex;
+					int32 EndIndex = ChannelIndex + 1;
+					int32 TopEndIndex = TopChannelIndex + 1;
+					MovieSceneToolHelpers::FBaseTopSections BT(BaseSection,
+						Section, StartIndex, EndIndex, TopChannelIndex, TopEndIndex);
+					MovieSceneToolHelpers::MergeTwoSections<FMovieSceneFloatChannel>(BT, Range, TrackSections, Increment);
+
+					break;
+				}
+				case ERigControlType::Position:
+				case ERigControlType::Scale:
+				case ERigControlType::Rotator:
+				{
+					int32 StartIndex = ChannelIndex;
+					int32 EndIndex = ChannelIndex + 2;
+					int32 TopEndIndex = TopChannelIndex + 2;
+					MovieSceneToolHelpers::FBaseTopSections BT(BaseSection,
+						Section, StartIndex, EndIndex, TopChannelIndex, TopEndIndex);
+					MovieSceneToolHelpers::MergeTwoSections<FMovieSceneFloatChannel>(BT, Range, TrackSections, Increment);
+
+					break;
+				}
+
+				case ERigControlType::Transform:
+				case ERigControlType::TransformNoScale:
+				case ERigControlType::EulerTransform:
+				{
+					EMovieSceneTransformChannel BaseChannelMask = BaseSection->GetTransformMask().GetChannels();
+					EMovieSceneTransformChannel ChannelMask = Section->GetTransformMask().GetChannels();
+
+					const bool bDoAllTransform = EnumHasAllFlags(ChannelMask, EMovieSceneTransformChannel::AllTransform);
+					if (bDoAllTransform)
+					{
+						if (ControlElement->Settings.ControlType == ERigControlType::TransformNoScale)
+						{
+							int32 StartIndex = ChannelIndex;
+							int32 EndIndex = ChannelIndex + 5;
+							int32 TopEndIndex = TopChannelIndex + 5;
+							MovieSceneToolHelpers::FBaseTopSections BT(BaseSection,
+								Section, StartIndex, EndIndex, TopChannelIndex, TopEndIndex);
+							MovieSceneToolHelpers::MergeTwoSections<FMovieSceneFloatChannel>(BT, Range, TrackSections, Increment);
+						}
+						else
+						{
+							int32 StartIndex = ChannelIndex;
+							int32 EndIndex = ChannelIndex + 8;
+							int32 TopEndIndex = TopChannelIndex + 8;
+							MovieSceneToolHelpers::FBaseTopSections BT(BaseSection,
+								Section, StartIndex, EndIndex, TopChannelIndex, TopEndIndex);
+							MovieSceneToolHelpers::MergeTwoSections<FMovieSceneFloatChannel>(BT, Range, TrackSections, Increment);
+						}
+					}
+					else
+					{
+						if (EnumHasAllFlags(BaseChannelMask, EMovieSceneTransformChannel::TranslationX)
+							&& EnumHasAllFlags(ChannelMask, EMovieSceneTransformChannel::TranslationX))
+						{
+							int32 StartIndex = ChannelIndex;
+							int32 EndIndex = StartIndex;
+							int32 TopStartIndex = TopChannelIndex;
+							int32 TopEndIndex = TopStartIndex;
+							MovieSceneToolHelpers::FBaseTopSections BT(BaseSection,
+								Section, StartIndex, EndIndex, TopStartIndex, TopEndIndex);
+							MovieSceneToolHelpers::MergeTwoSections<FMovieSceneFloatChannel>(BT, Range, TrackSections, Increment);
+
+						}
+						if (EnumHasAllFlags(BaseChannelMask, EMovieSceneTransformChannel::TranslationY)
+							&& EnumHasAllFlags(ChannelMask, EMovieSceneTransformChannel::TranslationY))
+						{
+							int32 StartIndex = ChannelIndex + 1;
+							int32 EndIndex = StartIndex;
+							int32 TopStartIndex = TopChannelIndex + 1;
+							int32 TopEndIndex = TopStartIndex;
+							MovieSceneToolHelpers::FBaseTopSections BT(BaseSection,
+								Section, StartIndex, EndIndex, TopStartIndex, TopEndIndex);
+							MovieSceneToolHelpers::MergeTwoSections<FMovieSceneFloatChannel>(BT, Range, TrackSections, Increment);
+
+						}
+						if (EnumHasAllFlags(BaseChannelMask, EMovieSceneTransformChannel::TranslationZ)
+							&& EnumHasAllFlags(ChannelMask, EMovieSceneTransformChannel::TranslationZ))
+						{
+							int32 StartIndex = ChannelIndex + 2;
+							int32 EndIndex = StartIndex;
+							int32 TopStartIndex = TopChannelIndex + 2;
+							int32 TopEndIndex = TopStartIndex;
+							MovieSceneToolHelpers::FBaseTopSections BT(BaseSection,
+								Section, StartIndex, EndIndex, TopStartIndex, TopEndIndex);
+							MovieSceneToolHelpers::MergeTwoSections<FMovieSceneFloatChannel>(BT, Range, TrackSections, Increment);
+
+						}
+						if (EnumHasAllFlags(BaseChannelMask, EMovieSceneTransformChannel::RotationX)
+							&& EnumHasAllFlags(ChannelMask, EMovieSceneTransformChannel::RotationX))
+						{
+							int32 StartIndex = ChannelIndex + 3;
+							int32 EndIndex = StartIndex;
+							int32 TopStartIndex = TopChannelIndex + 3;
+							int32 TopEndIndex = TopStartIndex;
+							MovieSceneToolHelpers::FBaseTopSections BT(BaseSection,
+								Section, StartIndex, EndIndex, TopStartIndex, TopEndIndex);
+							MovieSceneToolHelpers::MergeTwoSections<FMovieSceneFloatChannel>(BT, Range, TrackSections, Increment);
+
+						}
+						if (EnumHasAllFlags(BaseChannelMask, EMovieSceneTransformChannel::RotationY)
+							&& EnumHasAllFlags(ChannelMask, EMovieSceneTransformChannel::RotationY))
+						{
+							int32 StartIndex = ChannelIndex + 4;
+							int32 EndIndex = StartIndex;
+							int32 TopStartIndex = TopChannelIndex + 4;
+							int32 TopEndIndex = TopStartIndex;
+							MovieSceneToolHelpers::FBaseTopSections BT(BaseSection,
+								Section, StartIndex, EndIndex, TopStartIndex, TopEndIndex);
+							MovieSceneToolHelpers::MergeTwoSections<FMovieSceneFloatChannel>(BT, Range, TrackSections, Increment);
+
+						}
+						if (EnumHasAllFlags(BaseChannelMask, EMovieSceneTransformChannel::RotationZ)
+							&& EnumHasAllFlags(ChannelMask, EMovieSceneTransformChannel::RotationZ))
+						{
+							int32 StartIndex = ChannelIndex + 5;
+							int32 EndIndex = StartIndex;
+							int32 TopStartIndex = TopChannelIndex + 5;
+							int32 TopEndIndex = TopStartIndex;
+							MovieSceneToolHelpers::FBaseTopSections BT(BaseSection,
+								Section, StartIndex, EndIndex, TopStartIndex, TopEndIndex);
+							MovieSceneToolHelpers::MergeTwoSections<FMovieSceneFloatChannel>(BT, Range, TrackSections, Increment);
+
+						}
+						if (ControlElement->Settings.ControlType != ERigControlType::TransformNoScale)
+						{
+							if (EnumHasAllFlags(BaseChannelMask, EMovieSceneTransformChannel::ScaleX)
+								&& EnumHasAllFlags(ChannelMask, EMovieSceneTransformChannel::ScaleX))
+							{
+								int32 StartIndex = ChannelIndex + 6;
+								int32 EndIndex = StartIndex;
+								int32 TopStartIndex = TopChannelIndex + 6;
+								int32 TopEndIndex = TopStartIndex;
+								MovieSceneToolHelpers::FBaseTopSections BT(BaseSection,
+									Section, StartIndex, EndIndex, TopStartIndex, TopEndIndex);
+								MovieSceneToolHelpers::MergeTwoSections<FMovieSceneFloatChannel>(BT, Range, TrackSections, Increment);
+
+							}
+							if (EnumHasAllFlags(BaseChannelMask, EMovieSceneTransformChannel::ScaleY)
+								&& EnumHasAllFlags(ChannelMask, EMovieSceneTransformChannel::ScaleY))
+							{
+								int32 StartIndex = ChannelIndex + 7;
+								int32 EndIndex = StartIndex;
+								int32 TopStartIndex = TopChannelIndex + 7;
+								int32 TopEndIndex = TopStartIndex;
+								MovieSceneToolHelpers::FBaseTopSections BT(BaseSection,
+									Section, StartIndex, EndIndex, TopStartIndex, TopEndIndex);
+								MovieSceneToolHelpers::MergeTwoSections<FMovieSceneFloatChannel>(BT, Range, TrackSections, Increment);
+
+							}
+							if (EnumHasAllFlags(BaseChannelMask, EMovieSceneTransformChannel::ScaleZ)
+								&& EnumHasAllFlags(ChannelMask, EMovieSceneTransformChannel::ScaleZ))
+							{
+								int32 StartIndex = ChannelIndex + 8;
+								int32 EndIndex = StartIndex;
+								int32 TopStartIndex = TopChannelIndex + 8;
+								int32 TopEndIndex = TopStartIndex;
+								MovieSceneToolHelpers::FBaseTopSections BT(BaseSection,
+									Section, StartIndex, EndIndex, TopStartIndex, TopEndIndex);
+								MovieSceneToolHelpers::MergeTwoSections<FMovieSceneFloatChannel>(BT, Range, TrackSections, Increment);
+
+							}
+						}
+					}
+					break;
+				}
+				default:
+					break;
+				}
+				if (bBaseSectionResetActive.IsSet())
+				{
+					BaseSection->SetIsActive(bBaseSectionResetActive.GetValue());
+				}
+				if (bSectionResetActive.IsSet())
+				{
+					Section->SetIsActive(bSectionResetActive.GetValue());
+				}
+			}
+		}
+		if (BaseSection && BaseSection->GetBlendType() == EMovieSceneBlendType::Override)
+		{
+			if (FMovieSceneFloatChannel* FloatChannel = &BaseSection->Weight)
+			{
+				RevertWeightChannelToOne(FloatChannel, Range);
+			}
+		}
+	}
+}
+
+static void MergeTransformSections(UMovieScene3DTransformSection* BaseSection, UMovieScene3DTransformSection* Section, const TRange<FFrameNumber>& Range, const int32* Increment)
+{
+	if (!BaseSection || !Section)
+	{
+		return;
+	}
+	TArrayView<FMovieSceneDoubleChannel*> BaseDoubleChannels = BaseSection->GetChannelProxy().GetChannels<FMovieSceneDoubleChannel>();
+	if (BaseDoubleChannels.Num() > 0)
+	{
+		EMovieSceneTransformChannel BaseChannelMask = BaseSection->GetMask().GetChannels();
+		EMovieSceneTransformChannel ChannelMask = Section->GetMask().GetChannels();
+		BaseSection->Modify();
+		TArray<UMovieSceneSection*> TrackSections;
+		GetValidSections(BaseSection, TrackSections);
+		const bool bDoAllTransform = EnumHasAllFlags(ChannelMask, EMovieSceneTransformChannel::AllTransform);
+		if (bDoAllTransform)
+		{
+			int32 StartIndex = 0;
+			int32 EndIndex = BaseDoubleChannels.Num() - 1;
+			MovieSceneToolHelpers::FBaseTopSections BT(BaseSection,
+				Section, StartIndex, EndIndex, StartIndex, EndIndex);
+			MovieSceneToolHelpers::MergeTwoSections<FMovieSceneDoubleChannel>(BT, Range, TrackSections, Increment);
+		}
+		else
+		{
+			const int32 ChannelIndex = 0;
+			if (EnumHasAllFlags(BaseChannelMask, EMovieSceneTransformChannel::TranslationX)
+				&& EnumHasAllFlags(ChannelMask, EMovieSceneTransformChannel::TranslationX))
+			{
+				int32 StartIndex = ChannelIndex;
+				int32 EndIndex = StartIndex;
+				MovieSceneToolHelpers::FBaseTopSections BT(BaseSection,
+					Section, StartIndex, EndIndex, StartIndex, EndIndex);
+				MovieSceneToolHelpers::MergeTwoSections<FMovieSceneDoubleChannel>(BT, Range, TrackSections, Increment);
+			}
+			if (EnumHasAllFlags(BaseChannelMask, EMovieSceneTransformChannel::TranslationY)
+				&& EnumHasAllFlags(ChannelMask, EMovieSceneTransformChannel::TranslationY))
+			{
+				int32 StartIndex = ChannelIndex + 1;
+				int32 EndIndex = StartIndex;
+				MovieSceneToolHelpers::FBaseTopSections BT(BaseSection,
+					Section, StartIndex, EndIndex, StartIndex, EndIndex);
+				MovieSceneToolHelpers::MergeTwoSections<FMovieSceneDoubleChannel>(BT, Range, TrackSections, Increment);
+			}
+			if (EnumHasAllFlags(BaseChannelMask, EMovieSceneTransformChannel::TranslationZ)
+				&& EnumHasAllFlags(ChannelMask, EMovieSceneTransformChannel::TranslationZ))
+			{
+				int32 StartIndex = ChannelIndex + 2;
+				int32 EndIndex = StartIndex;
+				MovieSceneToolHelpers::FBaseTopSections BT(BaseSection,
+					Section, StartIndex, EndIndex, StartIndex, EndIndex);
+				MovieSceneToolHelpers::MergeTwoSections<FMovieSceneDoubleChannel>(BT, Range, TrackSections, Increment);
+			}
+			if (EnumHasAllFlags(BaseChannelMask, EMovieSceneTransformChannel::RotationX)
+				&& EnumHasAllFlags(ChannelMask, EMovieSceneTransformChannel::RotationX))
+			{
+				int32 StartIndex = ChannelIndex + 3;
+				int32 EndIndex = StartIndex;
+				MovieSceneToolHelpers::FBaseTopSections BT(BaseSection,
+					Section, StartIndex, EndIndex, StartIndex, EndIndex);
+				MovieSceneToolHelpers::MergeTwoSections<FMovieSceneDoubleChannel>(BT, Range, TrackSections, Increment);
+			}
+			if (EnumHasAllFlags(BaseChannelMask, EMovieSceneTransformChannel::RotationY)
+				&& EnumHasAllFlags(ChannelMask, EMovieSceneTransformChannel::RotationY))
+			{
+				int32 StartIndex = ChannelIndex + 4;
+				int32 EndIndex = StartIndex;
+				MovieSceneToolHelpers::FBaseTopSections BT(BaseSection,
+					Section, StartIndex, EndIndex, StartIndex, EndIndex);
+				MovieSceneToolHelpers::MergeTwoSections<FMovieSceneDoubleChannel>(BT, Range, TrackSections, Increment);
+			}
+			if (EnumHasAllFlags(BaseChannelMask, EMovieSceneTransformChannel::RotationZ)
+				&& EnumHasAllFlags(ChannelMask, EMovieSceneTransformChannel::RotationZ))
+			{
+				int32 StartIndex = ChannelIndex + 5;
+				int32 EndIndex = StartIndex;
+				MovieSceneToolHelpers::FBaseTopSections BT(BaseSection,
+					Section, StartIndex, EndIndex, StartIndex, EndIndex);
+				MovieSceneToolHelpers::MergeTwoSections<FMovieSceneDoubleChannel>(BT, Range, TrackSections, Increment);
+			}
+			if (EnumHasAllFlags(BaseChannelMask, EMovieSceneTransformChannel::ScaleX)
+				&& EnumHasAllFlags(ChannelMask, EMovieSceneTransformChannel::ScaleX))
+			{
+				int32 StartIndex = ChannelIndex + 6;
+				int32 EndIndex = StartIndex;
+				MovieSceneToolHelpers::FBaseTopSections BT(BaseSection,
+					Section, StartIndex, EndIndex, StartIndex, EndIndex);
+				MovieSceneToolHelpers::MergeTwoSections<FMovieSceneDoubleChannel>(BT, Range, TrackSections, Increment);
+			}
+			if (EnumHasAllFlags(BaseChannelMask, EMovieSceneTransformChannel::ScaleY)
+				&& EnumHasAllFlags(ChannelMask, EMovieSceneTransformChannel::ScaleY))
+			{
+				int32 StartIndex = ChannelIndex + 7;
+				int32 EndIndex = StartIndex;
+				MovieSceneToolHelpers::FBaseTopSections BT(BaseSection,
+					Section, StartIndex, EndIndex, StartIndex, EndIndex);
+				MovieSceneToolHelpers::MergeTwoSections<FMovieSceneDoubleChannel>(BT, Range, TrackSections, Increment);
+			}
+			if (EnumHasAllFlags(BaseChannelMask, EMovieSceneTransformChannel::ScaleZ)
+				&& EnumHasAllFlags(ChannelMask, EMovieSceneTransformChannel::ScaleZ))
+			{
+				int32 StartIndex = ChannelIndex + 8;
+				int32 EndIndex = StartIndex;
+				MovieSceneToolHelpers::FBaseTopSections BT(BaseSection,
+					Section, StartIndex, EndIndex, StartIndex, EndIndex);
+				MovieSceneToolHelpers::MergeTwoSections<FMovieSceneDoubleChannel>(BT, Range, TrackSections, Increment);
+			}
+		}
+		if (BaseSection && BaseSection->GetBlendType() == EMovieSceneBlendType::Override)
+		{
+			if (FMovieSceneFloatChannel* FloatChannel = BaseSection->GetWeightChannel())
+			{
+				RevertWeightChannelToOne(FloatChannel, Range);
+			}
+		}
+	}
+}
+}
+
 int32 UAnimLayers::DuplicateAnimLayer(ISequencer* SequencerPtr, int32 Index)
 {
+	using namespace UE::AIE;
 	int32 NewIndex = INDEX_NONE;
 	if (Index >= 1 && Index < AnimLayers.Num()) //don't duplicate base
 	{
@@ -1624,6 +2109,9 @@
 				UE_LOG(LogControlRig, Error, TEXT("Anim Layers: Can not duplicate empty layer"));
 				return INDEX_NONE;
 			}
+			TRangeBound<FFrameNumber> OpenBound;
+			TRange<FFrameNumber> InfiniteRange(OpenBound, OpenBound);
+
 			const FScopedTransaction Transaction(LOCTEXT("DuplicateAnimLayer_Transaction", "Duplicate Anim Layer"), !GIsTransacting);
 			Modify();
 			UAnimLayer* NewAnimLayer = NewObject<UAnimLayer>(this, NAME_None, RF_Transactional);
@@ -1673,7 +2161,8 @@
 												NewSectionItem.AnimLayerSet.Names.GetKeys(ControlNames);
 												SetUpControlRigSection(NewSection, ControlNames);
 												//current copy keys
-												CopySectionIntoAnother(NewSection, Section);
+												int32* Increment = nullptr;
+												MergeControlRigSections(NewSection, Section, InfiniteRange, Increment);
 											}
 										}
 									}
@@ -2047,424 +2536,12 @@
 	return true;
 }
 
-//get the set of sections active, or with this control name in it
-static void GetValidCRSections(UMovieSceneSection* InSection, const FName& ControlName, TArray<UMovieSceneSection*>& OutSections)
-{
-	if (UMovieSceneTrack* OwnerTrack = InSection->GetTypedOuter<UMovieSceneTrack>())
-	{
-		TArray<UMovieSceneSection*> TrackSections = OwnerTrack->GetAllSections();
-		for (UMovieSceneSection* Section : TrackSections)
-		{
-			if (UMovieSceneControlRigParameterSection* CRSection = Cast<UMovieSceneControlRigParameterSection>(Section))
-			{
-				if (Section && (Section->IsActive() || Section == InSection) &&
-					CRSection->GetControlNameMask(ControlName))
-				{
-					OutSections.Add(Section);
-				}
-			}
-		}
-	}
-}
 
 
-static void GetValidSections(UMovieSceneSection* InSection, TArray<UMovieSceneSection*>& OutSections)
-{
-	if (UMovieSceneTrack* OwnerTrack = InSection->GetTypedOuter<UMovieSceneTrack>())
-	{
-		TArray<UMovieSceneSection*> TrackSections = OwnerTrack->GetAllSections();
-		for (UMovieSceneSection* Section : TrackSections)
-		{
-			if (Section && (Section->IsActive() || Section == InSection))
-			{
-				OutSections.Add(Section);
-			}
-		}
-	}
-}
-static void RevertWeightChannelToOne(FMovieSceneFloatChannel* FloatChannel, const TRange<FFrameNumber>& FrameRange)
-{
-	if (FloatChannel)
-	{
-		//don't add key if there!
-		TArray<FFrameNumber> KeyTimes;
-		TArray<FKeyHandle> KeyHandles;
-		FloatChannel->GetKeys(FrameRange, &KeyTimes, &KeyHandles);
-		if (KeyTimes.Num() > 0)
-		{
-			FloatChannel->DeleteKeys(KeyHandles);
-			if (FloatChannel->GetNumKeys() > 0) //if it still has keys see if the weight is not one 
-			{
-				float StartValue = 0.0, EndValue = 0.0;
-				FFrameTime StartFrameTime(FrameRange.GetLowerBoundValue());
-				FFrameTime EndFrameTime(FrameRange.GetUpperBoundValue());
-				FloatChannel->Evaluate(StartFrameTime, StartValue);
-				FloatChannel->Evaluate(EndFrameTime, EndValue);
-				//if not zero at boundaries set one keys there
-				if (FMath::IsNearlyEqual(StartValue, 1.0f) == false ||
-					FMath::IsNearlyEqual(EndValue, 1.0f) == false)
-				{
-					FloatChannel->AddCubicKey(StartFrameTime.FrameNumber, 1.0f, ERichCurveTangentMode::RCTM_SmartAuto);
-					FloatChannel->AddCubicKey(EndFrameTime.FrameNumber, 1.0f, ERichCurveTangentMode::RCTM_SmartAuto);
-				}
-			}
-			else
-			{
-				FloatChannel->SetDefault(1.0f);
-			}
-		}
-	}
-}
-
-
-// disable section if no control in the section and override, need to make sure we skip it's eval or it will incorrectly
-//override when merging
-static bool ShouldDisableSection(UMovieSceneControlRigParameterSection* CRSection, FRigControlElement* ControlElement)
-{
-	return (CRSection->GetControlNameMask(ControlElement->GetFName()) == false &&
-		CRSection->GetBlendType().IsValid() && CRSection->GetBlendType() == EMovieSceneBlendType::Override);
-}
-
-static void MergeControlRigSections(UMovieSceneControlRigParameterSection* BaseSection, UMovieSceneControlRigParameterSection* Section, const TRange<FFrameNumber>& Range,
-	const int32* Increment)
-{
-
-	if (!BaseSection || !Section)
-	{
-		return;
-	}
-	TArrayView<FMovieSceneFloatChannel*> BaseFloatChannels = BaseSection->GetChannelProxy().GetChannels<FMovieSceneFloatChannel>();
-	if (BaseFloatChannels.Num() > 0)
-	{
-		//need to go through each control and merge that
-		TArray<FRigControlElement*> Controls;
-		UControlRig* ControlRig = BaseSection->GetControlRig();
-		if (ControlRig == nullptr)
-		{
-			return;
-		}
-		ControlRig->GetControlsInOrder(Controls);
-		URigHierarchy* Hierarchy = ControlRig->GetHierarchy();
-		if (Hierarchy == nullptr)
-		{
-			return;
-		}
-
-		BaseSection->Modify();
-
-		const bool bIsOverride = (Section->GetBlendType().IsValid() && Section->GetBlendType() == EMovieSceneBlendType::Override) ||
-			(BaseSection->GetBlendType().IsValid() && BaseSection->GetBlendType() == EMovieSceneBlendType::Override);
-
-
-		for (int32 LocalControlIndex = 0; LocalControlIndex < Controls.Num(); ++LocalControlIndex)
-		{
-			FRigControlElement* ControlElement = Controls[LocalControlIndex];
-			check(ControlElement);
-			if (!Hierarchy->IsAnimatable(ControlElement))
-			{
-				continue;
-			}
-
-			const FName& ControlName = ControlElement->GetFName();
-			if (FChannelMapInfo* pChannelIndex = BaseSection->ControlChannelMap.Find(ControlName))
-			{
-				const int32 ChannelIndex = pChannelIndex->ChannelIndex;			
-				//if override we mask out if both are masked out, if not override we mask out if the top section is masked out
-				const bool bMaskedOutOfBase = BaseSection->GetControlNameMask(ControlName) == false;
-				const bool bMaskedOutOfSection = Section->GetControlNameMask(ControlName) == false;
-				const bool bMaskKeyOut = bIsOverride ? (bMaskedOutOfBase && bMaskedOutOfSection)
-					: bMaskedOutOfSection;
-				if (bMaskKeyOut)
-				{
-					continue;
-				}
-				TArray<UMovieSceneSection*> TrackSections;
-				GetValidCRSections(BaseSection, ControlName, TrackSections);
-
-				TOptional<bool> bBaseSectionResetActive;
-				TOptional<bool> bSectionResetActive;
-				if (ShouldDisableSection(BaseSection, ControlElement))
-				{
-					bBaseSectionResetActive = BaseSection->IsActive();
-					BaseSection->SetIsActive(false);
-				}
-
-				if (ShouldDisableSection(Section, ControlElement))				
-				{
-					bSectionResetActive = Section->IsActive();
-					Section->SetIsActive(false);
-				}
-
-			
-				switch (ControlElement->Settings.ControlType)
-				{
-
-					case ERigControlType::Float:
-					case ERigControlType::ScaleFloat:
-					{
-						int32 StartIndex = ChannelIndex;
-						int32 EndIndex = ChannelIndex;
-						MovieSceneToolHelpers::MergeSections<FMovieSceneFloatChannel>(BaseSection,
-							Section, StartIndex, EndIndex, Range, TrackSections,  Increment);
-						break;
-					}
-					case ERigControlType::Vector2D:
-					{
-						int32 StartIndex = ChannelIndex;
-						int32 EndIndex = ChannelIndex + 1;
-						MovieSceneToolHelpers::MergeSections<FMovieSceneFloatChannel>(BaseSection,
-							Section, StartIndex, EndIndex, Range, TrackSections, Increment);
-						break;
-					}
-					case ERigControlType::Position:
-					case ERigControlType::Scale:
-					case ERigControlType::Rotator:
-					{
-						int32 StartIndex = ChannelIndex;
-						int32 EndIndex = ChannelIndex + 2;
-						MovieSceneToolHelpers::MergeSections<FMovieSceneFloatChannel>(BaseSection,
-							Section, StartIndex, EndIndex, Range, TrackSections,  Increment);
-						break;
-					}
-
-					case ERigControlType::Transform:
-					case ERigControlType::TransformNoScale:
-					case ERigControlType::EulerTransform:
-					{
-						EMovieSceneTransformChannel BaseChannelMask = BaseSection->GetTransformMask().GetChannels();
-						EMovieSceneTransformChannel ChannelMask = Section->GetTransformMask().GetChannels();
-
-						const bool bDoAllTransform = EnumHasAllFlags(ChannelMask, EMovieSceneTransformChannel::AllTransform);
-						if (bDoAllTransform)
-						{
-							if (ControlElement->Settings.ControlType == ERigControlType::TransformNoScale)
-							{
-								int32 StartIndex = ChannelIndex;
-								int32 EndIndex = ChannelIndex + 5;
-								MovieSceneToolHelpers::MergeSections<FMovieSceneFloatChannel>(BaseSection,
-									Section, StartIndex, EndIndex, Range, TrackSections, Increment);
-							}
-							else
-							{
-								int32 StartIndex = ChannelIndex;
-								int32 EndIndex = ChannelIndex + 8;
-								MovieSceneToolHelpers::MergeSections<FMovieSceneFloatChannel>(BaseSection,
-									Section, StartIndex, EndIndex, Range, TrackSections, Increment);
-							}
-						}
-						else
-						{
-							if (EnumHasAllFlags(BaseChannelMask, EMovieSceneTransformChannel::TranslationX)
-								 && EnumHasAllFlags(ChannelMask, EMovieSceneTransformChannel::TranslationX))
-							{
-								int32 StartIndex = ChannelIndex;
-								int32 EndIndex = StartIndex;
-								MovieSceneToolHelpers::MergeSections<FMovieSceneFloatChannel>(BaseSection,
-									Section, StartIndex, EndIndex, Range, TrackSections, Increment);
-							}
-							if (EnumHasAllFlags(BaseChannelMask, EMovieSceneTransformChannel::TranslationY)
-								&& EnumHasAllFlags(ChannelMask, EMovieSceneTransformChannel::TranslationY))
-							{
-								int32 StartIndex = ChannelIndex + 1;
-								int32 EndIndex = StartIndex;
-								MovieSceneToolHelpers::MergeSections<FMovieSceneFloatChannel>(BaseSection,
-									Section, StartIndex, EndIndex, Range, TrackSections, Increment);
-							}
-							if (EnumHasAllFlags(BaseChannelMask, EMovieSceneTransformChannel::TranslationZ)
-								&& EnumHasAllFlags(ChannelMask, EMovieSceneTransformChannel::TranslationZ))
-							{
-								int32 StartIndex = ChannelIndex + 2;
-								int32 EndIndex = StartIndex;
-								MovieSceneToolHelpers::MergeSections<FMovieSceneFloatChannel>(BaseSection,
-									Section, StartIndex, EndIndex, Range, TrackSections, Increment);
-							}
-							if (EnumHasAllFlags(BaseChannelMask, EMovieSceneTransformChannel::RotationX)
-								&& EnumHasAllFlags(ChannelMask, EMovieSceneTransformChannel::RotationX))
-							{
-								int32 StartIndex = ChannelIndex + 3;
-								int32 EndIndex = StartIndex;
-								MovieSceneToolHelpers::MergeSections<FMovieSceneFloatChannel>(BaseSection,
-									Section, StartIndex, EndIndex, Range, TrackSections, Increment);
-							}
-							if (EnumHasAllFlags(BaseChannelMask, EMovieSceneTransformChannel::RotationY)
-								&& EnumHasAllFlags(ChannelMask, EMovieSceneTransformChannel::RotationY))
-							{
-								int32 StartIndex = ChannelIndex + 4;
-								int32 EndIndex = StartIndex;
-								MovieSceneToolHelpers::MergeSections<FMovieSceneFloatChannel>(BaseSection,
-									Section, StartIndex, EndIndex, Range, TrackSections, Increment);
-							}
-							if (EnumHasAllFlags(BaseChannelMask, EMovieSceneTransformChannel::RotationZ)
-								&& EnumHasAllFlags(ChannelMask, EMovieSceneTransformChannel::RotationZ))
-							{
-								int32 StartIndex = ChannelIndex + 5;
-								int32 EndIndex = StartIndex;
-								MovieSceneToolHelpers::MergeSections<FMovieSceneFloatChannel>(BaseSection,
-									Section, StartIndex, EndIndex, Range, TrackSections, Increment);
-							}
-							if (ControlElement->Settings.ControlType != ERigControlType::TransformNoScale)
-							{
-								if (EnumHasAllFlags(BaseChannelMask, EMovieSceneTransformChannel::ScaleX)
-									&& EnumHasAllFlags(ChannelMask, EMovieSceneTransformChannel::ScaleX))
-								{
-									int32 StartIndex = ChannelIndex + 6;
-									int32 EndIndex = StartIndex;
-									MovieSceneToolHelpers::MergeSections<FMovieSceneFloatChannel>(BaseSection,
-										Section, StartIndex, EndIndex, Range, TrackSections, Increment);
-								}
-								if (EnumHasAllFlags(BaseChannelMask, EMovieSceneTransformChannel::ScaleY)
-									&& EnumHasAllFlags(ChannelMask, EMovieSceneTransformChannel::ScaleY))
-								{
-									int32 StartIndex = ChannelIndex + 7;
-									int32 EndIndex = StartIndex;
-									MovieSceneToolHelpers::MergeSections<FMovieSceneFloatChannel>(BaseSection,
-										Section, StartIndex, EndIndex, Range, TrackSections, Increment);
-								}
-								if (EnumHasAllFlags(BaseChannelMask, EMovieSceneTransformChannel::ScaleZ)
-									&& EnumHasAllFlags(ChannelMask, EMovieSceneTransformChannel::ScaleZ))
-								{
-									int32 StartIndex = ChannelIndex + 8;
-									int32 EndIndex = StartIndex;
-									MovieSceneToolHelpers::MergeSections<FMovieSceneFloatChannel>(BaseSection,
-										Section, StartIndex, EndIndex, Range, TrackSections, Increment);
-								}
-							}
-						}
-						break;
-					}
-					default:
-						break;
-				}
-				if (bBaseSectionResetActive.IsSet())
-				{
-					BaseSection->SetIsActive(bBaseSectionResetActive.GetValue());
-				}
-				if (bSectionResetActive.IsSet())
-				{
-					Section->SetIsActive(bSectionResetActive.GetValue());
-				} 
-			}
-		}
-		if (BaseSection && BaseSection->GetBlendType() == EMovieSceneBlendType::Override)
-		{
-			if (FMovieSceneFloatChannel * FloatChannel = &BaseSection->Weight)
-			{
-				RevertWeightChannelToOne(FloatChannel,Range);
-			}
-		}
-	}
-}
-
-static void MergeTransformSections(UMovieScene3DTransformSection* BaseSection, UMovieScene3DTransformSection* Section, const TRange<FFrameNumber>& Range, const int32* Increment)
-{
-	if (!BaseSection || !Section)
-	{
-		return;
-	}
-	TArrayView<FMovieSceneDoubleChannel*> BaseDoubleChannels = BaseSection->GetChannelProxy().GetChannels<FMovieSceneDoubleChannel>();
-	if (BaseDoubleChannels.Num() > 0)
-	{
-		EMovieSceneTransformChannel BaseChannelMask = BaseSection->GetMask().GetChannels();
-		EMovieSceneTransformChannel ChannelMask = Section->GetMask().GetChannels();
-		BaseSection->Modify();
-		TArray<UMovieSceneSection*> TrackSections;
-		GetValidSections(BaseSection, TrackSections);
-		const bool bDoAllTransform = EnumHasAllFlags(ChannelMask, EMovieSceneTransformChannel::AllTransform);
-		if (bDoAllTransform)
-		{
-			int32 StartIndex = 0;
-			int32 EndIndex = BaseDoubleChannels.Num() - 1;
-			MovieSceneToolHelpers::MergeSections<FMovieSceneDoubleChannel>(BaseSection,
-				Section, StartIndex, EndIndex, Range, TrackSections, Increment);
-		}
-		else
-		{
-			const int32 ChannelIndex = 0;
-			if (EnumHasAllFlags(BaseChannelMask, EMovieSceneTransformChannel::TranslationX)
-				&& EnumHasAllFlags(ChannelMask, EMovieSceneTransformChannel::TranslationX))
-			{
-				int32 StartIndex = ChannelIndex;
-				int32 EndIndex = StartIndex;
-				MovieSceneToolHelpers::MergeSections<FMovieSceneDoubleChannel>(BaseSection,
-					Section, StartIndex, EndIndex, Range, TrackSections, Increment);
-			}
-			if (EnumHasAllFlags(BaseChannelMask, EMovieSceneTransformChannel::TranslationY)
-				&& EnumHasAllFlags(ChannelMask, EMovieSceneTransformChannel::TranslationY))
-			{
-				int32 StartIndex = ChannelIndex + 1;
-				int32 EndIndex = StartIndex;
-				MovieSceneToolHelpers::MergeSections<FMovieSceneDoubleChannel>(BaseSection,
-					Section, StartIndex, EndIndex, Range, TrackSections, Increment);
-			}
-			if (EnumHasAllFlags(BaseChannelMask, EMovieSceneTransformChannel::TranslationZ)
-				&& EnumHasAllFlags(ChannelMask, EMovieSceneTransformChannel::TranslationZ))
-			{
-				int32 StartIndex = ChannelIndex + 2;
-				int32 EndIndex = StartIndex;
-				MovieSceneToolHelpers::MergeSections<FMovieSceneDoubleChannel>(BaseSection,
-					Section, StartIndex, EndIndex, Range, TrackSections, Increment);
-			}
-			if (EnumHasAllFlags(BaseChannelMask, EMovieSceneTransformChannel::RotationX)
-				&& EnumHasAllFlags(ChannelMask, EMovieSceneTransformChannel::RotationX))
-			{
-				int32 StartIndex = ChannelIndex + 3;
-				int32 EndIndex = StartIndex;
-				MovieSceneToolHelpers::MergeSections<FMovieSceneDoubleChannel>(BaseSection,
-					Section, StartIndex, EndIndex, Range, TrackSections, Increment);
-			}
-			if (EnumHasAllFlags(BaseChannelMask, EMovieSceneTransformChannel::RotationY)
-				&& EnumHasAllFlags(ChannelMask, EMovieSceneTransformChannel::RotationY))
-			{
-				int32 StartIndex = ChannelIndex + 4;
-				int32 EndIndex = StartIndex;
-				MovieSceneToolHelpers::MergeSections<FMovieSceneDoubleChannel>(BaseSection,
-					Section, StartIndex, EndIndex, Range, TrackSections, Increment);
-			}
-			if (EnumHasAllFlags(BaseChannelMask, EMovieSceneTransformChannel::RotationZ)
-				&& EnumHasAllFlags(ChannelMask, EMovieSceneTransformChannel::RotationZ))
-			{
-				int32 StartIndex = ChannelIndex +  5;
-				int32 EndIndex = StartIndex;
-				MovieSceneToolHelpers::MergeSections<FMovieSceneDoubleChannel>(BaseSection,
-					Section, StartIndex, EndIndex, Range, TrackSections, Increment);
-			}
-			if (EnumHasAllFlags(BaseChannelMask, EMovieSceneTransformChannel::ScaleX)
-				&& EnumHasAllFlags(ChannelMask, EMovieSceneTransformChannel::ScaleX))
-			{
-				int32 StartIndex = ChannelIndex + 6;
-				int32 EndIndex = StartIndex;
-				MovieSceneToolHelpers::MergeSections<FMovieSceneFloatChannel>(BaseSection,
-					Section, StartIndex, EndIndex, Range, TrackSections, Increment);
-			}
-			if (EnumHasAllFlags(BaseChannelMask, EMovieSceneTransformChannel::ScaleY)
-				&& EnumHasAllFlags(ChannelMask, EMovieSceneTransformChannel::ScaleY))
-			{
-				int32 StartIndex = ChannelIndex + 7;
-				int32 EndIndex = StartIndex;
-				MovieSceneToolHelpers::MergeSections<FMovieSceneDoubleChannel>(BaseSection,
-					Section, StartIndex, EndIndex, Range, TrackSections, Increment);
-			}
-			if (EnumHasAllFlags(BaseChannelMask, EMovieSceneTransformChannel::ScaleZ)
-				&& EnumHasAllFlags(ChannelMask, EMovieSceneTransformChannel::ScaleZ))
-			{
-				int32 StartIndex = ChannelIndex + 8;
-				int32 EndIndex = StartIndex;
-				MovieSceneToolHelpers::MergeSections<FMovieSceneFloatChannel>(BaseSection,
-					Section, StartIndex, EndIndex, Range, TrackSections, Increment);
-			}
-		}
-		if (BaseSection && BaseSection->GetBlendType() == EMovieSceneBlendType::Override)
-		{
-			if (FMovieSceneFloatChannel* FloatChannel = BaseSection->GetWeightChannel())
-			{
-				RevertWeightChannelToOne(FloatChannel,Range);
-			}
-		}
-	}
-}
-
 bool UAnimLayers::MergeAnimLayers(TSharedPtr<ISequencer>& InSequencerPtr, const TArray<int32>& Indices, const FMergeAnimLayerSettings* InSettings)
 {
+	using namespace UE::AIE;
+
 	if (InSequencerPtr.IsValid() == false)
 	{
 		return false;
@@ -2550,16 +2627,19 @@
 									{
 										int32 StartIndex = 0;
 										int32 EndIndex = BaseDoubleChannels.Num() - 1;
-										MovieSceneToolHelpers::MergeSections<FMovieSceneDoubleChannel>(BaseSectionItem->Section.Get(),
-											SectionItem.Section.Get(), StartIndex, EndIndex, Range, TrackSections, Increment);
+
+										MovieSceneToolHelpers::FBaseTopSections BT(BaseSectionItem->Section.Get(),
+											SectionItem.Section.Get(), StartIndex, EndIndex, StartIndex,EndIndex);
+										MovieSceneToolHelpers::MergeTwoSections<FMovieSceneDoubleChannel>(BT, Range, TrackSections, Increment);
 
 									}
 									else if (BaseFloatChannels.Num() > 0)
 									{
 										int32 StartIndex = 0;
 										int32 EndIndex = BaseFloatChannels.Num() - 1;
-										MovieSceneToolHelpers::MergeSections<FMovieSceneFloatChannel>(BaseSectionItem->Section.Get(),
-											SectionItem.Section.Get(), StartIndex, EndIndex, Range, TrackSections, Increment);
+										MovieSceneToolHelpers::FBaseTopSections BT(BaseSectionItem->Section.Get(),
+											SectionItem.Section.Get(), StartIndex, EndIndex,StartIndex, EndIndex);
+										MovieSceneToolHelpers::MergeTwoSections<FMovieSceneFloatChannel>(BT, Range, TrackSections, Increment);
 									}
 								}
 							}
@@ -3034,6 +3114,7 @@
 			{
 				ParameterSection->SetControlNameMask(RigName, true);
 			}
+			ParameterSection->CacheChannelProxy();
 		}
 	}
 }
@@ -3739,4 +3820,3 @@
 #undef LOCTEXT_NAMESPACE
 
 
-
diff --git a/Engine/Source/Editor/MovieSceneTools/Private/MovieSceneToolHelpers.cpp b/Engine/Source/Editor/MovieSceneTools/Private/MovieSceneToolHelpers.cpp
--- a/Engine/Source/Editor/MovieSceneTools/Private/MovieSceneToolHelpers.cpp
+++ b/Engine/Source/Editor/MovieSceneTools/Private/MovieSceneToolHelpers.cpp
@@ -5564,16 +5564,19 @@ bool MovieSceneToolHelpers::CollapseSection(TSharedPtr<ISequencer>& SequencerPtr
 					{
 						const int StartIndex = 0;
 						const int EndIndex = BaseDoubleChannels.Num() - 1;
-						MovieSceneToolHelpers::MergeSections<FMovieSceneDoubleChannel>(BaseSection,
-							Section, StartIndex, EndIndex, Range, TrackSections);
+
+						MovieSceneToolHelpers::FBaseTopSections BT(BaseSection,
+							Section, StartIndex, EndIndex, StartIndex, EndIndex);
+						MovieSceneToolHelpers::MergeTwoSections<FMovieSceneDoubleChannel>(BT, Range, TrackSections);
 
 					}
 					else if (BaseFloatChannels.Num() > 0)
 					{
 						const int StartIndex = 0;
 						const int EndIndex = BaseFloatChannels.Num() - 1;
-						MovieSceneToolHelpers::MergeSections<FMovieSceneFloatChannel>(BaseSection,
-							Section,StartIndex, EndIndex, Range, TrackSections);
+						MovieSceneToolHelpers::FBaseTopSections BT(BaseSection,
+							Section, StartIndex, EndIndex, StartIndex, EndIndex);
+						MovieSceneToolHelpers::MergeTwoSections<FMovieSceneFloatChannel>(BT, Range, TrackSections);
 					}
 				}
 
@@ -5606,6 +5609,12 @@ bool MovieSceneToolHelpers::CollapseSection(TSharedPtr<ISequencer>& SequencerPtr
 			{
 				Channel->Optimize(Params); //should also auto tangent
 			}
+			TArrayView<FMovieSceneDoubleChannel*> DoubleChannels = ChannelProxy.GetChannels<FMovieSceneDoubleChannel>();
+			for (FMovieSceneDoubleChannel* DChannel : DoubleChannels)
+			{
+				DChannel->Optimize(Params); //should also auto tangent
+			}
+
 		}
 		//reset everything back
 		SequencerPtr->NotifyMovieSceneDataChanged(EMovieSceneDataChangeType::MovieSceneStructureItemAdded);
diff --git a/Engine/Source/Editor/MovieSceneTools/Public/MovieSceneToolHelpers.h b/Engine/Source/Editor/MovieSceneTools/Public/MovieSceneToolHelpers.h
--- a/Engine/Source/Editor/MovieSceneTools/Public/MovieSceneToolHelpers.h
+++ b/Engine/Source/Editor/MovieSceneTools/Public/MovieSceneToolHelpers.h
@@ -976,30 +976,60 @@ public:
 		return true;
 	}
 
+	struct FBaseTopSections
+	{
+		FBaseTopSections(UMovieSceneSection* InBaseSection,
+						 UMovieSceneSection* InTopSection,
+						 int32 InBaseStartIndex,
+						 int32 InBaseEndIndex,
+						 int32 InTopStartIndex,
+						 int32 InTopEndIndex) :
+			BaseSection(InBaseSection),
+			TopSection(InTopSection),
+			BaseStartIndex(InBaseStartIndex),
+			BaseEndIndex(InBaseEndIndex),
+			TopStartIndex(InTopStartIndex),
+			TopEndIndex(InTopEndIndex)
+		{}
+		UMovieSceneSection* BaseSection;
+		UMovieSceneSection* TopSection;
+		int32 BaseStartIndex;
+		int32 BaseEndIndex;
+		int32 TopStartIndex; 
+		int32 TopEndIndex;
+	};
 	/*
 	* Merge the TopSection onto the BaseSection. This function should be used when merging
 	* a mix of Override/Additive and Absolute section Blend Types.
-	* @param BaseSection  The section to merge onto
-	* @param TopSection  The section we will merge from
-	* @param StartIndex Start Index of the channels we are merging
-	* @param EndIndex End Index of the channels we are merging, should be no greater than the number of channels -1
+	* @param BaseAndTopSections The sections we are merging with index information
 	* @param Range Range of time over which we will merge  Control Rig sections and don't want to merge the last weight float channel
 	* @param TrackSections Sections in this track
 	* @param Increment Optional increment index, if specfied we bake over the range using the increment value, otherwise by default we just merge onto the existing keys
 	* @return Returns true if successful
 	*/	
 	template<typename ChannelType>
-	static bool MergeSections(UMovieSceneSection* BaseSection, UMovieSceneSection* TopSection,
-		int32 StartIndex, int32 EndIndex, const TRange<FFrameNumber>& Range, const TArray<UMovieSceneSection*>& TrackSections, const int32* Increment = nullptr)
+	static bool MergeTwoSections(const FBaseTopSections& BT, const TRange<FFrameNumber>& Range, const TArray<UMovieSceneSection*>& TrackSections, const int32* Increment = nullptr)
 	{
-		TArrayView<ChannelType*> BaseChannels = BaseSection->GetChannelProxy().GetChannels<ChannelType>();
-		TArrayView<ChannelType*> TopChannels = TopSection->GetChannelProxy().GetChannels<ChannelType>();
+		TArrayView<ChannelType*> BaseChannels = BT.BaseSection->GetChannelProxy().GetChannels<ChannelType>();
+		TArrayView<ChannelType*> TopChannels = BT.TopSection->GetChannelProxy().GetChannels<ChannelType>();
 		if (TopChannels.Num() != BaseChannels.Num())
 		{
 			UE_LOG(LogMovieScene, Warning, TEXT("MergeSections:: Invalid number of channels"));
 			return false;
 		}
-		if (StartIndex < 0 || EndIndex >= BaseChannels.Num() || EndIndex < StartIndex)
+		if (BT.BaseStartIndex < 0 || BT.BaseEndIndex >= BaseChannels.Num() || BT.BaseEndIndex < BT.BaseStartIndex)
+		{
+			UE_LOG(LogMovieScene, Warning, TEXT("MergeSections:: Invalid Start/End indices"));
+			return false;
+		}
+
+		if (BT.TopStartIndex < 0 || BT.TopEndIndex >= BaseChannels.Num() || BT.TopEndIndex < BT.TopStartIndex)
+		{
+			UE_LOG(LogMovieScene, Warning, TEXT("MergeSections:: Invalid Start/End indices"));
+			return false;
+		}
+
+		if ((BT.BaseEndIndex - BT.BaseStartIndex) != (BT.TopEndIndex - BT.TopStartIndex))
 		{
 			UE_LOG(LogMovieScene, Warning, TEXT("MergeSections:: Invalid Start/End indices"));
 			return false;
@@ -1008,27 +1038,28 @@ public:
 		TArray<ChannelType*> Channels;
 		int32 ChannelIndex = 0;
 		TArray<UMovieSceneSection*> Sections;
-		Sections.Add(BaseSection);
-		Sections.Add(TopSection);
+		Sections.Add(BT.BaseSection);
+		Sections.Add(BT.TopSection);
 		FChannelMergeAlgorithm MergeAlgorithm = FChannelMergeAlgorithm::Add;
 		//if either setion is override we do an override blend since the result will be override
-		if ((TopSection->GetBlendType().IsValid() && TopSection->GetBlendType().Get() == EMovieSceneBlendType::Override) ||
-			(BaseSection->GetBlendType().IsValid() && BaseSection->GetBlendType().Get() == EMovieSceneBlendType::Override))
+		if ((BT.TopSection->GetBlendType().IsValid() && BT.TopSection->GetBlendType().Get() == EMovieSceneBlendType::Override) ||
+			(BT.BaseSection->GetBlendType().IsValid() && BT.BaseSection->GetBlendType().Get() == EMovieSceneBlendType::Override))
 		{
 			MergeAlgorithm = FChannelMergeAlgorithm::Override;
 		}
-		else if (TopSection->GetBlendType().IsValid() && TopSection->GetBlendType().Get() == EMovieSceneBlendType::Absolute)
+		else if (BT.TopSection->GetBlendType().IsValid() && BT.TopSection->GetBlendType().Get() == EMovieSceneBlendType::Absolute)
 
 		{
 			MergeAlgorithm = FChannelMergeAlgorithm::Average;
 		}
 
-		for (ChannelIndex = StartIndex; ChannelIndex <= EndIndex; ++ChannelIndex)
+		int32 TopChannelIndex = BT.TopStartIndex;
+		for (ChannelIndex = BT.BaseStartIndex; ChannelIndex <= BT.BaseEndIndex; ++ChannelIndex)
 		{
 			Channels.Reset();
 			Channels.Add(BaseChannels[ChannelIndex]);
-			Channels.Add(TopChannels[ChannelIndex]);
-
+			Channels.Add(TopChannels[TopChannelIndex]);
+			++TopChannelIndex;
 			if (MergeChannels(ChannelIndex, Channels, Sections, Range,
 				MergeAlgorithm, TrackSections, Increment) == false)
 			{
@@ -1045,6 +1076,15 @@ public:
 		return true;
 	}
 
+	template<typename ChannelType>
+	static bool MergeSections(UMovieSceneSection* BaseSection, UMovieSceneSection* TopSection,
+		int32 StartIndex, int32 EndIndex, const TRange<FFrameNumber>& Range, const TArray<UMovieSceneSection*>& TrackSections, const int32* Increment = nullptr)
+	{
+		FBaseTopSections BT(BaseSection, TopSection, StartIndex, EndIndex, StartIndex, EndIndex);
+		return MergeTwoSections<ChannelType>(BT, Range, TrackSections, Increment);
+	}
+
+	
 	/*
 	* Merge the following set of sections. Should be used when blending just Absolute with additive sections. If Merging with
 	* overrides, use the above MergeSections function
@@ -1226,3 +1266,4 @@ private:
 
 
 #undef UE_API
+
