Control Rig: The following UI can now be toggled with commands (no default bindings):
- Selection sets
- Pose library
- Anim layers
- Constrain window > Spaces
- Constrain window > Constraints
- Constrain window > Snapper

There is a known bug that the shortcut does not process if the viewport tab is not focused. By default e.g. Anim Layers spawns as a new window. Then the shortcut will not work. This will need a follow up.

diff --git a/Engine/Plugins/Animation/ControlRig/Source/ControlRigEditor/Private/EditMode/Constrain/ConstrainToolsManager.cpp b/Engine/Plugins/Animation/ControlRig/Source/ControlRigEditor/Private/EditMode/Constrain/ConstrainToolsManager.cpp
--- a/Engine/Plugins/Animation/ControlRig/Source/ControlRigEditor/Private/EditMode/Constrain/ConstrainToolsManager.cpp
+++ b/Engine/Plugins/Animation/ControlRig/Source/ControlRigEditor/Private/EditMode/Constrain/ConstrainToolsManager.cpp
@@ -5,6 +5,7 @@
 #include "EditMode/ControlRigEditModeToolkit.h"
 #include "Framework/Docking/TabManager.h"
 #include "SConstrainToolsRoot.h"
+#include "EditMode/ControlRigEditModeCommands.h"
 #include "Widgets/Docking/SDockTab.h"
 
 #define LOCTEXT_NAMESPACE "FConstrainToolsManager"
@@ -14,20 +15,24 @@ namespace UE::ControlRigEditor
 FConstrainToolsManager::FConstrainToolsManager(
 	const TSharedRef<FTabManager>& InToolkitTabManager, 
 	const TSharedRef<FWorkspaceItem>& InWorkspaceMenuGroup,
+	const TSharedRef<FUICommandList>& InToolkitCommandList,
 	FControlRigEditMode& InOwningEditMode,
 	const TSharedRef<FRigSelectionViewModel>& InSelectionViewModel
 	)
 	: TabManager(InToolkitTabManager)
+	, WeakCommandList(InToolkitCommandList)
 	, OwningEditMode(InOwningEditMode)
 	, SelectionViewModel(InSelectionViewModel)
 {
 	RegisterTabSpawner(TabManager, InWorkspaceMenuGroup);
+	BindCommands();
 }
 
 FConstrainToolsManager::~FConstrainToolsManager()
 {
 	HideWidget();
 	UnregisterTabSpawner(TabManager);
+	UnbindCommands();
 }
 
 void FConstrainToolsManager::ToggleVisibility() const
@@ -65,7 +70,7 @@ void FConstrainToolsManager::RegisterTabSpawner(const TSharedRef<FTabManager>& I
 	InTabManager->RegisterTabSpawner(
 		FControlRigEditModeToolkit::ConstrainingTabName, FOnSpawnTab::CreateRaw(this, &FConstrainToolsManager::SpawnSnapperTab)
 		)
-		.SetDisplayName(LOCTEXT("Tab.Title", "Constraining"))
+		.SetDisplayName(LOCTEXT("Tab.Title", "Constrain"))
 		.SetTooltipText(LOCTEXT("Tab.ToolTip", "Snap child objects to a parent object over a set of frames."))
 		.SetIcon(FSlateIcon(TEXT("ControlRigEditorStyle"), TEXT("ControlRig.SnapperTool")))
 		.SetGroup(InWorkspaceMenuGroup);
@@ -81,9 +86,50 @@ TSharedRef<SDockTab> FConstrainToolsManager::SpawnSnapperTab(const FSpawnTabArgs
 {
 	return SAssignNew(SpawnedTab, SDockTab)
 	[
-		SNew(SConstrainToolsRoot, OwningEditMode, SelectionViewModel)
+		SAssignNew(ConstrainWindowContent, SConstrainToolsRoot, OwningEditMode, SelectionViewModel)
 	];
 }
+
+void FConstrainToolsManager::BindCommands()
+{
+	const TSharedPtr<FUICommandList> CommandListPin = WeakCommandList.Pin();
+	if (!CommandListPin)
+	{
+		return;
+	}
+	
+	FControlRigEditModeCommands& Commands = FControlRigEditModeCommands::Get();
+	CommandListPin->MapAction(Commands.ToggleConstrainTab, FExecuteAction::CreateRaw(this, &FConstrainToolsManager::ToggleVisibility));
+	CommandListPin->MapAction(Commands.OpenSpacesTab, FExecuteAction::CreateRaw(this, &FConstrainToolsManager::HandleOpenTabCommand, EConstrainTab::Spaces));
+	CommandListPin->MapAction(Commands.OpenConstraintsTab, FExecuteAction::CreateRaw(this, &FConstrainToolsManager::HandleOpenTabCommand, EConstrainTab::Constraints));
+	CommandListPin->MapAction(Commands.OpenSnapperTab, FExecuteAction::CreateRaw(this, &FConstrainToolsManager::HandleOpenTabCommand, EConstrainTab::Snapper));
+}
+
+void FConstrainToolsManager::UnbindCommands()
+{
+	const TSharedPtr<FUICommandList> CommandListPin = WeakCommandList.Pin();
+	if (!CommandListPin || !FControlRigEditModeCommands::IsRegistered())
+	{
+		return;
+	}
+	
+	FControlRigEditModeCommands& Commands = FControlRigEditModeCommands::Get();
+	CommandListPin->UnmapAction(Commands.ToggleConstrainTab);
+	CommandListPin->UnmapAction(Commands.OpenSpacesTab);
+	CommandListPin->UnmapAction(Commands.OpenConstraintsTab);
+	CommandListPin->UnmapAction(Commands.OpenSnapperTab);
+}
+
+void FConstrainToolsManager::HandleOpenTabCommand(EConstrainTab InTab)
+{
+	// This will flash the tab if already visible, which is desirable.
+	ShowWidget();
+
+	if (const TSharedPtr<SConstrainToolsRoot> WidgetContentPin = ConstrainWindowContent.Pin())
+	{
+		WidgetContentPin->OpenTab(InTab);
+	}
+}
 }
 
 #undef LOCTEXT_NAMESPACE
\ No newline at end of file
diff --git a/Engine/Plugins/Animation/ControlRig/Source/ControlRigEditor/Private/EditMode/Constrain/ConstrainToolsManager.h b/Engine/Plugins/Animation/ControlRig/Source/ControlRigEditor/Private/EditMode/Constrain/ConstrainToolsManager.h
--- a/Engine/Plugins/Animation/ControlRig/Source/ControlRigEditor/Private/EditMode/Constrain/ConstrainToolsManager.h
+++ b/Engine/Plugins/Animation/ControlRig/Source/ControlRigEditor/Private/EditMode/Constrain/ConstrainToolsManager.h
@@ -9,12 +9,15 @@ class FControlRigEditMode;
 class FSpawnTabArgs;
 class FTabManager;
 class FToolBarBuilder;
+class FUICommandList;
 class FWorkspaceItem;
 class SDockTab;
 
 namespace UE::ControlRigEditor
 {
 class FRigSelectionViewModel;
+class SConstrainToolsRoot;
+enum class EConstrainTab : uint8;
 	
 /**
  * Manages UI that is used for constraining, namely:
@@ -31,6 +34,7 @@ public:
 	explicit FConstrainToolsManager(
 		const TSharedRef<FTabManager>& InToolkitTabManager,
 		const TSharedRef<FWorkspaceItem>& InWorkspaceMenuGroup,
+		const TSharedRef<FUICommandList>& InToolkitCommandList,
 		FControlRigEditMode& InOwningEditMode UE_LIFETIMEBOUND,
 		const TSharedRef<FRigSelectionViewModel>& InSelectionViewModel
 		);
@@ -45,7 +49,9 @@ private:
 
 	/** Used to show / hide the tab containing our widgets. */
 	const TSharedRef<FTabManager> TabManager;
-
+	/** The command list that we are bound to for switching between inline tabs. */
+	const TWeakPtr<FUICommandList> WeakCommandList;
+	
 	/** Needed to construct widget hierarchy.  */
 	FControlRigEditMode& OwningEditMode;
 	/** Handles selection changing. Needed to construct widget hierarchy. */
@@ -53,11 +59,18 @@ private:
 
 	/** The tab we spawned. Used to detect whether the tab is currently being shown. */
 	TWeakPtr<SDockTab> SpawnedTab;
+	/** The widget that contains the tabs. Only valid when SpawnedTab is valid. Weak as we don't keep a strong reference to the SDockTab either. */
+	TWeakPtr<SConstrainToolsRoot> ConstrainWindowContent;
 
 	// Tab spawning
 	void RegisterTabSpawner(const TSharedRef<FTabManager>& InTabManager, const TSharedRef<FWorkspaceItem>& InWorkspaceMenuGroup);
 	void UnregisterTabSpawner(const TSharedRef<FTabManager>& InTabManager);
 	TSharedRef<SDockTab> SpawnSnapperTab(const FSpawnTabArgs& Args);
+
+	// Commands
+	void BindCommands();
+	void UnbindCommands();
+	void HandleOpenTabCommand(EConstrainTab InTab);
 };
 }
 
diff --git a/Engine/Plugins/Animation/ControlRig/Source/ControlRigEditor/Private/EditMode/Constrain/SConstrainToolsRoot.cpp b/Engine/Plugins/Animation/ControlRig/Source/ControlRigEditor/Private/EditMode/Constrain/SConstrainToolsRoot.cpp
--- a/Engine/Plugins/Animation/ControlRig/Source/ControlRigEditor/Private/EditMode/Constrain/SConstrainToolsRoot.cpp
+++ b/Engine/Plugins/Animation/ControlRig/Source/ControlRigEditor/Private/EditMode/Constrain/SConstrainToolsRoot.cpp
@@ -19,9 +19,14 @@ void SConstrainToolsRoot::Construct(
 	const FArguments& InArgs, FControlRigEditMode& InOwningEditMode, const TSharedRef<FRigSelectionViewModel>& InSelectionViewModel
 	)
 {
+	static_assert(
+		static_cast<int32>(EConstrainTab::Spaces) == 0 && static_cast<int32>(EConstrainTab::Constraints) == 1 && static_cast<int32>(EConstrainTab::Snapper) == 2,
+		"Open tab relies on this order. Update the code accordingly."
+		);
+	
 	ChildSlot
 	[
-		SNew(SInlineTabPanel)
+		SAssignNew(InlineTabs, SInlineTabPanel)
 		.Tabs({
 			FInlineTabArgs()
 				.SetLabel(LOCTEXT("Spaces", "Spaces"))
@@ -42,6 +47,12 @@ void SConstrainToolsRoot::Construct(
 	];
 }
 
+void SConstrainToolsRoot::OpenTab(EConstrainTab InTab) const
+{
+	// The enumeration is declared in display order.
+	InlineTabs->SwitchToTab(static_cast<int32>(InTab));
+}
+
 TSharedRef<SWidget> SConstrainToolsRoot::CreateSpacesContent(
 	FControlRigEditMode& InOwningEditMode, const TSharedRef<FRigSelectionViewModel>& InSelectionViewModel
 	)
diff --git a/Engine/Plugins/Animation/ControlRig/Source/ControlRigEditor/Private/EditMode/Constrain/SConstrainToolsRoot.h b/Engine/Plugins/Animation/ControlRig/Source/ControlRigEditor/Private/EditMode/Constrain/SConstrainToolsRoot.h
--- a/Engine/Plugins/Animation/ControlRig/Source/ControlRigEditor/Private/EditMode/Constrain/SConstrainToolsRoot.h
+++ b/Engine/Plugins/Animation/ControlRig/Source/ControlRigEditor/Private/EditMode/Constrain/SConstrainToolsRoot.h
@@ -11,7 +11,15 @@ namespace UE::ControlRigEditor
 {
 class FRigSelectionViewModel;
 class SConstraintsEditionWidget_ConstrainPanel;
+class SInlineTabPanel;
 
+enum class EConstrainTab : uint8
+{
+	Spaces,
+	Constraints,
+	Snapper
+};
+	
 /**
  * Coordinator widget that contains
  * - SControlRigSnapper
@@ -28,8 +36,14 @@ public:
 
 	void Construct(const FArguments& InArgs, FControlRigEditMode& InOwningEditMode, const TSharedRef<FRigSelectionViewModel>& InSelectionViewModel);
 
+	/** Opens the specified tab. */
+	void OpenTab(EConstrainTab InTab) const;
+	
 private:
 
+	/** Manages the tabs. */
+	TSharedPtr<SInlineTabPanel> InlineTabs;
+	
 	/** The widget for editing constraints. */
 	TSharedPtr<SConstraintsEditionWidget_ConstrainPanel> ConstrainsEditionWidget;
 
diff --git a/Engine/Plugins/Animation/ControlRig/Source/ControlRigEditor/Private/EditMode/ControlRigEditModeCommands.cpp b/Engine/Plugins/Animation/ControlRig/Source/ControlRigEditor/Private/EditMode/ControlRigEditModeCommands.cpp
--- a/Engine/Plugins/Animation/ControlRig/Source/ControlRigEditor/Private/EditMode/ControlRigEditModeCommands.cpp
+++ b/Engine/Plugins/Animation/ControlRig/Source/ControlRigEditor/Private/EditMode/ControlRigEditModeCommands.cpp
@@ -49,11 +49,20 @@ void FControlRigEditModeCommands::RegisterCommands()
 	UI_COMMAND(OpenSpacePickerWidget, "Open the Space Picker", "Allows space switching on the control", EUserInterfaceActionType::Button, FInputChord(EKeys::Tab)); 
 
 
-	UI_COMMAND(SummonSelectionSetsWidget, "Move selection sets to cursor", "Places the selection sets overlay at your mouse position", EUserInterfaceActionType::Button, FInputChord());
-	UI_COMMAND(ToggleSelectionSetsWidget, "Show or hide selection sets", "Shows the selection sets overlay if it's hidden, and hides it if it's visible.", EUserInterfaceActionType::Button, FInputChord());
-
-	UI_COMMAND(SummonTweenWidget, "Move tween slider to cursor", "Places the tween slider at your mouse position", EUserInterfaceActionType::Button, FInputChord(EKeys::U)); 
-	UI_COMMAND(ToggleTweenWidget, "Show or hide tween slider", "Shows the tween slider if it's hidden, and hides it if it's visible.", EUserInterfaceActionType::Button, FInputChord(EKeys::U, EModifierKey::Alt)); 
+	UI_COMMAND(SummonSelectionSetsWidget, "Move Selection Sets to cursor", "Places the selection sets overlay at your mouse position", EUserInterfaceActionType::Button, FInputChord());
+	UI_COMMAND(ToggleSelectionSetsWidget, "Show or hide Selection Sets", "Shows the selection sets UI if it's hidden, and hides it if it's visible.", EUserInterfaceActionType::Button, FInputChord());
+
+	UI_COMMAND(SummonTweenWidget, "Move Tween Slider to cursor", "Places the tween slider at your mouse position", EUserInterfaceActionType::Button, FInputChord(EKeys::U)); 
+	UI_COMMAND(ToggleTweenWidget, "Show or hide Tween Slider", "Shows the tween slider if it's hidden, and hides it if it's visible.", EUserInterfaceActionType::Button, FInputChord(EKeys::U, EModifierKey::Alt));
+	
+	UI_COMMAND(ToggleAnimLayersTab, "Show or hide Anim Layers tab", "Toggles the visibility of the Anim Layers tab.", EUserInterfaceActionType::Button, FInputChord());
+	
+	UI_COMMAND(TogglePoseLibraryTab, "Show or hide Pose Library tab", "Toggles the visibility of the Pose Library tab.", EUserInterfaceActionType::Button, FInputChord());
+	
+	UI_COMMAND(ToggleConstrainTab, "Show or hide Constrain tab", "Toggles the visibility of the constrain tab.", EUserInterfaceActionType::Button, FInputChord()); 
+	UI_COMMAND(OpenSpacesTab, "Open Spaces tab", "Opens the spaces tab in the constrain tab.", EUserInterfaceActionType::Button, FInputChord()); 
+	UI_COMMAND(OpenConstraintsTab, "Open Constraints tab", "Opens the constraints tab in the constrain tab.", EUserInterfaceActionType::Button, FInputChord()); 
+	UI_COMMAND(OpenSnapperTab, "Open Snapper tab", "Opens the snapper tab in the constrain tab.", EUserInterfaceActionType::Button, FInputChord()); 
 }
 
 #undef LOCTEXT_NAMESPACE
diff --git a/Engine/Plugins/Animation/ControlRig/Source/ControlRigEditor/Private/EditMode/ControlRigEditModeCommands.h b/Engine/Plugins/Animation/ControlRig/Source/ControlRigEditor/Private/EditMode/ControlRigEditModeCommands.h
--- a/Engine/Plugins/Animation/ControlRig/Source/ControlRigEditor/Private/EditMode/ControlRigEditModeCommands.h
+++ b/Engine/Plugins/Animation/ControlRig/Source/ControlRigEditor/Private/EditMode/ControlRigEditModeCommands.h
@@ -122,6 +122,25 @@ public:
 	/** Toggles the visibility of the selection sets widget. */
 	TSharedPtr< FUICommandInfo > ToggleSelectionSetsWidget;
 
+	/** Toggles the visibility of the anim layers window. */
+	TSharedPtr< FUICommandInfo > ToggleAnimLayersTab;
+	
+	/** Toggles the visibility of the anim pose library window. */
+	TSharedPtr< FUICommandInfo > TogglePoseLibraryTab;
+	
+	/** Toggles the visibility of the constrain window. */
+	TSharedPtr< FUICommandInfo > ToggleConstrainTab;
+
+	/** Opens the spaces tab in the constrain window. */
+	TSharedPtr< FUICommandInfo > OpenSpacesTab;
+	
+	/** Opens the constraints tab in the constrain window. */
+	TSharedPtr< FUICommandInfo > OpenConstraintsTab;
+	
+	/** Opens the snapper tab in the constrain window. */
+	TSharedPtr< FUICommandInfo > OpenSnapperTab;
+
+	
 	/**
 	 * Initialize commands
 	 */
diff --git a/Engine/Plugins/Animation/ControlRig/Source/ControlRigEditor/Private/EditMode/ControlRigEditModeToolkit.cpp b/Engine/Plugins/Animation/ControlRig/Source/ControlRigEditor/Private/EditMode/ControlRigEditModeToolkit.cpp
--- a/Engine/Plugins/Animation/ControlRig/Source/ControlRigEditor/Private/EditMode/ControlRigEditModeToolkit.cpp
+++ b/Engine/Plugins/Animation/ControlRig/Source/ControlRigEditor/Private/EditMode/ControlRigEditModeToolkit.cpp
@@ -91,7 +91,6 @@ void FControlRigEditModeToolkit::Init(const TSharedPtr<IToolkitHost>& InitToolki
 	FPropertyEditorModule& PropertyEditorModule = FModuleManager::GetModuleChecked<FPropertyEditorModule>("PropertyEditor");
 
 	FDetailsViewArgs DetailsViewArgs;
-
 	DetailsViewArgs.bUpdatesFromSelection = false;
 	DetailsViewArgs.bLockable = false;
 	DetailsViewArgs.bAllowSearch = false;
@@ -101,8 +100,10 @@ void FControlRigEditModeToolkit::Init(const TSharedPtr<IToolkitHost>& InitToolki
 	DetailsViewArgs.DefaultsOnlyVisibility = EEditDefaultsOnlyNodeVisibility::Automatic;
 	DetailsViewArgs.bShowOptions = false;
 	DetailsViewArgs.bAllowMultipleTopLevelObjects = true;
-
 	DetailsView = PropertyEditorModule.CreateDetailView(DetailsViewArgs);
+	
+	BindCommands();
+	
 	FModeToolkit::Init(InitToolkitHost);
 }
 
@@ -464,7 +465,7 @@ void FControlRigEditModeToolkit::InvokeUI()
 
 	checkf(WorkspaceMenuCategory, TEXT("WorkspaceMenuCategory should have been set in RequestModeUITabs by now."))
 	ConstrainToolsManager = MakeUnique<UE::ControlRigEditor::FConstrainToolsManager>(
-		ModeUILayerPtr->GetTabManager().ToSharedRef(), GetWorkspaceMenuCategory(), EditMode, SelectionViewModel
+		ModeUILayerPtr->GetTabManager().ToSharedRef(), GetWorkspaceMenuCategory(), GetToolkitCommands(), EditMode, SelectionViewModel
 		);
 	
 	const FControlRigUIRestoreStates& RestoreStates = UControlRigEditModeSettings::Get()->LastUIStates;
@@ -773,4 +774,25 @@ void FControlRigEditModeToolkit::BuildToolBar(FToolBarBuilder& InToolBarBuilder)
 		*/
 }
 
+void FControlRigEditModeToolkit::BindCommands()
+{
+	FControlRigEditModeCommands& Comands = FControlRigEditModeCommands::Get();
+#if SELECTION_SETS_AS_TAB
+	ToolkitCommands->MapAction(
+		Comands.ToggleSelectionSetsWidget,
+		FExecuteAction::CreateSP(this, &FControlRigEditModeToolkit::TryToggleToolkitUI, SelectionSetsTabName)
+		);
+#endif
+
+	ToolkitCommands->MapAction(
+		Comands.ToggleAnimLayersTab,
+		FExecuteAction::CreateSP(this, &FControlRigEditModeToolkit::TryToggleToolkitUI, AnimLayerTabName)
+		);
+
+	ToolkitCommands->MapAction(
+		Comands.TogglePoseLibraryTab,
+		FExecuteAction::CreateSP(this, &FControlRigEditModeToolkit::TryToggleToolkitUI, PoseTabName)
+		);
+}
+
 #undef LOCTEXT_NAMESPACE
diff --git a/Engine/Plugins/Animation/ControlRig/Source/ControlRigEditor/Private/EditMode/ControlRigEditModeToolkit.h b/Engine/Plugins/Animation/ControlRig/Source/ControlRigEditor/Private/EditMode/ControlRigEditModeToolkit.h
--- a/Engine/Plugins/Animation/ControlRig/Source/ControlRigEditor/Private/EditMode/ControlRigEditModeToolkit.h
+++ b/Engine/Plugins/Animation/ControlRig/Source/ControlRigEditor/Private/EditMode/ControlRigEditModeToolkit.h
@@ -117,11 +117,15 @@ private:
 	TUniquePtr<UE::ControlRigEditor::FTweenOverlayManager> TweenOverlayManager;
 	/** Manages the constraining UI tools, which is a window containing tabs for spaces, constraining, and snapping. */
 	TUniquePtr<UE::ControlRigEditor::FConstrainToolsManager> ConstrainToolsManager;
+#if SELECTION_SETS_AS_TAB == 0
 	TUniquePtr<UE::AIE::FSelectionSetsOverlayManager> SelectionSetsOverlayManager;
+#endif
 
 	/** The tools widget */
 	TSharedPtr<SControlRigEditModeTools> ModeTools;
 
 	/** Adds the toolbar items to the toolbar builder. */
 	void BuildToolBar(FToolBarBuilder& ToolBarBuilder);
+
+	void BindCommands();
 };
diff --git a/Engine/Plugins/Animation/ControlRig/Source/ControlRigEditor/Private/Misc/Tab/SInlineTabPanel.cpp b/Engine/Plugins/Animation/ControlRig/Source/ControlRigEditor/Private/Misc/Tab/SInlineTabPanel.cpp
--- a/Engine/Plugins/Animation/ControlRig/Source/ControlRigEditor/Private/Misc/Tab/SInlineTabPanel.cpp
+++ b/Engine/Plugins/Animation/ControlRig/Source/ControlRigEditor/Private/Misc/Tab/SInlineTabPanel.cpp
@@ -80,7 +80,7 @@ void SInlineTabPanel::Construct(const FArguments& InArgs)
 			.BorderImage(FAppStyle::Get().GetBrush("ToolPanel.GroupBorder"))
 			.Padding(0.0f)
 			[
-				SNew(STabArea)
+				SAssignNew(TabArea, STabArea)
 					.Tabs(InlineTabPanel::Transform(InArgs._Tabs, SharedThis(this)))
 					.ActiveTabIndex(InArgs._ActiveTabIndex)
 					.Padding(InArgs._Padding)
@@ -119,6 +119,7 @@ void SInlineTabPanel::SwitchToTab(int32 InIndex) const
 {
 	if (ensure(InIndex >= 0 && InIndex < TabContentSwitcher->GetNumWidgets()))
 	{
+		TabArea->SetButtonActivated(InIndex);
 		TabContentSwitcher->SetActiveWidgetIndex(InIndex);
 	}
 }
diff --git a/Engine/Plugins/Animation/ControlRig/Source/ControlRigEditor/Private/Misc/Tab/SInlineTabPanel.h b/Engine/Plugins/Animation/ControlRig/Source/ControlRigEditor/Private/Misc/Tab/SInlineTabPanel.h
--- a/Engine/Plugins/Animation/ControlRig/Source/ControlRigEditor/Private/Misc/Tab/SInlineTabPanel.h
+++ b/Engine/Plugins/Animation/ControlRig/Source/ControlRigEditor/Private/Misc/Tab/SInlineTabPanel.h
@@ -69,6 +69,8 @@ public:
 
 private:
 
+	/** Displays the tab buttons. */
+	TSharedPtr<STabArea> TabArea;
 	/** Switches the content for the tabs */
 	TSharedPtr<SWidgetSwitcher> TabContentSwitcher;
 };
