Control Rig: Create a copy of metadata for the details panel. Relying on straight pointers from the hierarchy may cause crash due to pointer invalidity

#prefight tbd
#rnx

diff --git a/Engine/Plugins/Animation/ControlRig/Source/ControlRigEditor/Private/ControlRigElementDetails.cpp b/Engine/Plugins/Animation/ControlRig/Source/ControlRigEditor/Private/ControlRigElementDetails.cpp
--- a/Engine/Plugins/Animation/ControlRig/Source/ControlRigEditor/Private/ControlRigElementDetails.cpp
+++ b/Engine/Plugins/Animation/ControlRig/Source/ControlRigEditor/Private/ControlRigElementDetails.cpp
@@ -1844,7 +1844,9 @@ void FRigBaseElementDetails::CustomizeMetadata(IDetailLayoutBuilder& DetailBuild
 	for(FName MetadataName: MetadataNames)
 	{
 		FRigBaseMetadata* Metadata = Element->GetMetadata(MetadataName);
-		TSharedPtr<FStructOnScope> StructOnScope = MakeShareable(new FStructOnScope(Metadata->GetMetadataStruct(), reinterpret_cast<uint8*>(Metadata)));
+		
+		TSharedPtr<FStructOnScope> StructOnScope = MakeShareable(new FStructOnScope(Metadata->GetMetadataStruct()));
+		Metadata->GetMetadataStruct()->CopyScriptStruct(StructOnScope->GetStructMemory(), Metadata, 1);
 
 		FAddPropertyParams Params;
 		Params.CreateCategoryNodes(false);
