Selection Sets: Isolate just works on control rig with the set.
Make sure moving set items in the UI is scripable

diff --git a/Engine/Plugins/Animation/ControlRig/Source/ControlRigEditor/Private/ControlRigSequencerEditorLibrary.cpp b/Engine/Plugins/Animation/ControlRig/Source/ControlRigEditor/Private/ControlRigSequencerEditorLibrary.cpp
--- a/Engine/Plugins/Animation/ControlRig/Source/ControlRigEditor/Private/ControlRigSequencerEditorLibrary.cpp
+++ b/Engine/Plugins/Animation/ControlRig/Source/ControlRigEditor/Private/ControlRigSequencerEditorLibrary.cpp
@@ -54,6 +54,7 @@
 #include "Math/Abstraction/KeyBlendingAbstraction.h"
 #include "Tools/EvaluateSequencerTools.h"
 #include "Sequencer/AnimLayers/AnimLayers.h"
+#include "Sequencer/SelectionSets/SelectionSets.h"
 
 
 #include UE_INLINE_GENERATED_CPP_BY_NAME(ControlRigSequencerEditorLibrary)
@@ -3669,5 +3670,18 @@ int32 UControlRigSequencerEditorLibrary::GetAnimLayerIndex(UAnimLayer* AnimLayer
 	return INDEX_NONE;
 }
 
+
+UAIESelectionSets* UControlRigSequencerEditorLibrary::GetSelectionSets(bool bAddIfDoesNotExist)
+{
+	if (TSharedPtr<ISequencer> SequencerPtr = UAIESelectionSets::GetSequencerFromAsset())
+	{
+		if (UAIESelectionSets* SelectionSets = UAIESelectionSets::GetSelectionSets(SequencerPtr, bAddIfDoesNotExist))
+		{
+			return SelectionSets;
+		}
+	}
+	return nullptr;
+}
+
 #undef LOCTEXT_NAMESPACE
 
diff --git a/Engine/Plugins/Animation/ControlRig/Source/ControlRigEditor/Private/Sequencer/SelectionSets/SelectionSets.cpp b/Engine/Plugins/Animation/ControlRig/Source/ControlRigEditor/Private/Sequencer/SelectionSets/SelectionSets.cpp
--- a/Engine/Plugins/Animation/ControlRig/Source/ControlRigEditor/Private/Sequencer/SelectionSets/SelectionSets.cpp
+++ b/Engine/Plugins/Animation/ControlRig/Source/ControlRigEditor/Private/Sequencer/SelectionSets/SelectionSets.cpp
@@ -929,11 +929,17 @@ bool UAIESelectionSets::DeleteSetItem(const FGuid& InGuid)
 	return false;
 }
 
-void UAIESelectionSets::UpdateRowValues(int32 AtRow, bool bIncrease)
+void UAIESelectionSets::UpdateRowValues(int32 StartRow, int32 EndRow, bool bIncrease)
 {
+	if (StartRow > EndRow)
+	{
+		int32 Swap = EndRow;
+		EndRow = StartRow;
+		StartRow = Swap;
+	}
 	for (TPair<FGuid, FAIESelectionSetItem>& Pair : SelectionSets)
 	{
-		if (Pair.Value.ViewData.Row >= AtRow)
+		if (Pair.Value.ViewData.Row >= StartRow && Pair.Value.ViewData.Row <= EndRow)
 		{
 			if (bIncrease)
 			{
@@ -957,7 +963,9 @@ bool UAIESelectionSets::DeleteSetItem(const FGuid& InGuid, TSharedPtr<ISequencer
 		int32 Row = SetItem->ViewData.Row;
 		SelectionSets.Remove(InGuid);
 		constexpr bool bIncrease = false;
-		UpdateRowValues(Row, bIncrease);
+		const int32 StartRow = Row;
+		const int32 EndRow = SelectionSets.Num() - 1;
+		UpdateRowValues(StartRow, EndRow, bIncrease);
 		TSet<TWeakObjectPtr<UObject>> NoObjectsAdded;
 		ActorsWithSelectionSets.SetUpFromSequencer(InSequencer, this, NoObjectsAdded);
 		SelectionSetsChangedBroadcast();
@@ -1030,6 +1038,45 @@ bool UAIESelectionSets::SetItemColor(const FGuid& InGuid, const FLinearColor& In
 	return false;
 }
 
+bool UAIESelectionSets::GetItemRow(const FGuid& InGuid, int32& OutRow) const
+{
+	if (const FAIESelectionSetItem* SetItem = SelectionSets.Find(InGuid))
+	{
+		OutRow = SetItem->ViewData.Row;
+		return true;
+	}
+	return false;
+}
+
+bool UAIESelectionSets::SetItemRow(const FGuid& InGuid, int32 InRow)
+{
+	FScopedTransaction Transaction(LOCTEXT("SetItemRow_transaction", "Set Item Row"), !GIsTransacting);
+	if (FAIESelectionSetItem* SetItem = SelectionSets.Find(InGuid))
+	{
+		const int32 OldRow = SetItem->ViewData.Row;
+		if (OldRow != InRow)
+		{
+			Modify();
+			if (InRow < OldRow)
+			{
+				const bool bIncrease = true;
+				UpdateRowValues(InRow , OldRow, bIncrease);
+			}
+			else
+			{
+				const bool bIncrease = false;
+				UpdateRowValues(OldRow, InRow, bIncrease);
+			}
+			SetItem->ViewData.Row = InRow;
+
+			SelectionSetsChangedBroadcast();
+			return true;
+		}
+	}
+	Transaction.Cancel();
+	return false;
+}
+
 static void ClearSelection()
 {
 	FControlRigEditMode* EditMode = static_cast<FControlRigEditMode*>(GLevelEditorModeTools().GetActiveMode(FControlRigEditMode::ModeName));
@@ -1304,16 +1351,24 @@ bool UAIESelectionSets::IsolateControls(const FGuid& InGuid) const
 											}
 											TSet<FString> ControlsToHide;
 											TArray<FRigControlElement*> Controls = ControlRig->AvailableControls();
+											bool bHasControlSoWeIsolate = false;
 											for (const FRigControlElement* Control : Controls)
 											{
 												if (ControlNames.Contains(Control->GetName()) == false)
 												{
 													ControlsToHide.Add(Control->GetName());
 												}
+												else
+												{
+													bHasControlSoWeIsolate = true;
+												}
 											}
 											//not undoable yet
-											EditMode->ShowControlRigControls(ControlRig, ControlNames, true);
-											EditMode->ShowControlRigControls(ControlRig, ControlsToHide, false);
+											if (bHasControlSoWeIsolate)
+											{
+												EditMode->ShowControlRigControls(ControlRig, ControlNames, true);
+												EditMode->ShowControlRigControls(ControlRig, ControlsToHide, false);
+											}
 
 										}
 										else if (AActor* Actor = Cast<AActor>(Object.Get()))
diff --git a/Engine/Plugins/Animation/ControlRig/Source/ControlRigEditor/Private/Sequencer/SelectionSets/SelectionSets.h b/Engine/Plugins/Animation/ControlRig/Source/ControlRigEditor/Private/Sequencer/SelectionSets/SelectionSets.h
--- a/Engine/Plugins/Animation/ControlRig/Source/ControlRigEditor/Private/Sequencer/SelectionSets/SelectionSets.h
+++ b/Engine/Plugins/Animation/ControlRig/Source/ControlRigEditor/Private/Sequencer/SelectionSets/SelectionSets.h
@@ -83,15 +83,15 @@ struct FAIESelectionSetItem
 	FGuid Guid;
 	
 	/* The Name of the Set*/
-	UPROPERTY(BlueprintReadWrite, EditAnywhere, Category = "Data")
+	UPROPERTY(BlueprintReadOnly, EditAnywhere, Category = "Data")
 	FText ItemName;
 
 	/* The collection of set items that make up the set, basically an array of names with some metadata*/
-	UPROPERTY(BlueprintReadWrite, EditAnywhere, Category = "Data")
+	UPROPERTY(BlueprintReadOnly, EditAnywhere, Category = "Data")
 	TArray<FAIESelectionSetItemName> Names;
 
 	/* Viewing Data for the set, like location and color*/
-	UPROPERTY(BlueprintReadWrite, EditAnywhere, Category = "Data")
+	UPROPERTY(BlueprintReadOnly, EditAnywhere, Category = "Data")
 	FAIESelectionSetItemViewData ViewData;
 
 	//currently not used but maybe in future
@@ -189,7 +189,7 @@ public:
 	UFUNCTION(BlueprintCallable, Category = "State")
 	void SetShowAndSetSelectedOnly(bool bInShowSelectedOnly);
 		
-	/* Get the current set of active seletion sets based upon active actors. This function will set it up also if needed*/
+	/* Get the current set of active selection sets based upon active actors. This function will set it up also if needed*/
 	UFUNCTION(BlueprintCallable, Category = "State")
 	TArray<FGuid> GetActiveSelectionSets();
 
@@ -211,7 +211,6 @@ public:
 	FGuid CreateSetItemFromSelection();
 	FGuid CreateSetItemFromSelection(TSharedPtr<ISequencer>& InSequencer);
 
-
 	/* Create a mirrored selection set from an existing set. Will return empty FGuid if no set created*/
 	UFUNCTION(BlueprintCallable, Category = "Sets")
 	FGuid CreateMirror(const FGuid& InGuid);
@@ -252,6 +251,14 @@ public:
 	UFUNCTION(BlueprintCallable, Category = "Set")
 	bool SetItemColor(const FGuid& InGuid, const FLinearColor& InColor);
 
+	/* Get selection set row*/
+	UFUNCTION(BlueprintPure, Category = "Set")
+	bool GetItemRow(const FGuid& InGuid, int32& OutRow) const;
+
+	/* Set selection set row*/
+	UFUNCTION(BlueprintCallable, Category = "Set")
+	bool SetItemRow(const FGuid& InGuid, int32 InRow);
+
 	/* Select from selection set, will select on the active actors*/
 	UFUNCTION(BlueprintCallable, Category = "Set")
 	bool SelectItem(const FGuid& InGuid, bool bDoMirror, bool bAdd, bool bToggle,bool bSelect = true) const;
@@ -327,8 +334,8 @@ public:
 private:
 	TArray<FGuid> SortSelectionSetsByRow(TArray<FGuid>& InSelectionSets);
 
-	//starting from AtRow increase or decrease row values, used for when changing row value or deleting the item
-	void UpdateRowValues(int32 AtRow, bool bIncrease);
+	//starting from StartRow going to EndRow(inclusive) increase or decrease row values, used for when changing row value or deleting the item
+	void UpdateRowValues(int32 StartRow, int32 EndRow, bool bIncrease);
 private:
 	//to handle selecting last selected object when none are selected
 	mutable TSet<FGuid> LastActiveSelectionSetsSet;
diff --git a/Engine/Plugins/Animation/ControlRig/Source/ControlRigEditor/Public/ControlRigSequencerEditorLibrary.h b/Engine/Plugins/Animation/ControlRig/Source/ControlRigEditor/Public/ControlRigSequencerEditorLibrary.h
--- a/Engine/Plugins/Animation/ControlRig/Source/ControlRigEditor/Public/ControlRigSequencerEditorLibrary.h
+++ b/Engine/Plugins/Animation/ControlRig/Source/ControlRigEditor/Public/ControlRigSequencerEditorLibrary.h
@@ -30,6 +30,7 @@ struct FBakingAnimationKeySettings;
 struct FMergeAnimLayerSettings;
 class UControlRig;
 class UAnimLayer;
+class UAIESelectionSets;
 
 USTRUCT(BlueprintType)
 struct FControlRigSequencerBindingProxy
@@ -1212,7 +1213,7 @@ public:
 	static UE_API bool MergeAnimLayersWithSettings(const TArray<int32>& Indices, const FMergeAnimLayerSettings& MergeAnimLayerSettings);
 
 	/**
-	* Get the animation layer objects, will create it by default if it doesn't exist
+	* Get the animation layer objects, will create it bAddIfDoesNotExist is set to true
 	* @param bAddIfDoesNotExist If it doesn't exist create it
 	* @return Returns array of anim layer objects if they exist on active Sequencer
 	*/
@@ -1227,6 +1228,14 @@ public:
 	UFUNCTION(BlueprintCallable, Category = "Editor Scripting | Sequencer Tools | Control Rig | Animation Layers")
 	static UE_API int32 GetAnimLayerIndex(UAnimLayer* AnimLayer);
 
+	/**
+	* Get the selection set objects, will create it bAddIfDoesNotExist is set to true
+	* @param bAddIfDoesNotExist If it doesn't exist create it
+	* @return Returns main selection set object that contains operations you can perform to create and use selection sets
+	*/
+	UFUNCTION(BlueprintCallable, Category = "Editor Scripting | Sequencer Tools | Control Rig | Selection Sets")
+	static UE_API UAIESelectionSets* GetSelectionSets(bool bAddIfDoesNotExist = true);
+
 };
 
 #undef UE_API
