Constraints: fixed misuse of SetDefault

#rnx

diff --git a/Engine/Plugins/Animation/ControlRig/Source/ControlRig/Private/Sequencer/MovieSceneControlRigParameterSection.cpp b/Engine/Plugins/Animation/ControlRig/Source/ControlRig/Private/Sequencer/MovieSceneControlRigParameterSection.cpp
--- a/Engine/Plugins/Animation/ControlRig/Source/ControlRig/Private/Sequencer/MovieSceneControlRigParameterSection.cpp
+++ b/Engine/Plugins/Animation/ControlRig/Source/ControlRig/Private/Sequencer/MovieSceneControlRigParameterSection.cpp
@@ -1804,7 +1804,6 @@ void UMovieSceneControlRigParameterSection::AddConstraintChannel(UTickableConstr
 		const int32 NewIndex = ConstraintsChannels.Add(FConstraintAndActiveChannel(InConstraint));
 
 		FMovieSceneConstraintChannel* ExistingChannel = &ConstraintsChannels[NewIndex].ActiveChannel;
-		ExistingChannel->SetDefault(false);
 		
 		//make copy that we can spawn if it doesn't exist
 		//the rename changes the outer to this section (from any actor manager)
diff --git a/Engine/Source/Runtime/Experimental/Animation/Constraints/Public/ConstraintChannel.h b/Engine/Source/Runtime/Experimental/Animation/Constraints/Public/ConstraintChannel.h
--- a/Engine/Source/Runtime/Experimental/Animation/Constraints/Public/ConstraintChannel.h
+++ b/Engine/Source/Runtime/Experimental/Animation/Constraints/Public/ConstraintChannel.h
@@ -76,7 +76,7 @@ struct TStructOpsTypeTraits<FMovieSceneConstraintChannel> : public TStructOpsTyp
 template<>
 struct TMovieSceneChannelTraits<FMovieSceneConstraintChannel> : TMovieSceneChannelTraitsBase<FMovieSceneConstraintChannel>
 {
-	enum { SupportsDefaults = true };
+	enum { SupportsDefaults = false };
 
 #if WITH_EDITOR
 
diff --git a/Engine/Source/Runtime/MovieSceneTracks/Private/Sections/MovieScene3DTransformSection.cpp b/Engine/Source/Runtime/MovieSceneTracks/Private/Sections/MovieScene3DTransformSection.cpp
--- a/Engine/Source/Runtime/MovieSceneTracks/Private/Sections/MovieScene3DTransformSection.cpp
+++ b/Engine/Source/Runtime/MovieSceneTracks/Private/Sections/MovieScene3DTransformSection.cpp
@@ -1197,7 +1197,6 @@ void UMovieScene3DTransformSection::AddConstraintChannel(UTickableConstraint* In
 		const int32 NewIndex = Constraints->ConstraintsChannels.Add(FConstraintAndActiveChannel(InConstraint));
 
 		FMovieSceneConstraintChannel* ExistingChannel = &Constraints->ConstraintsChannels[NewIndex].ActiveChannel;
-		ExistingChannel->SetDefault(false);
 
 		//make copy that we can spawn if it doesn't exist
 		//the rename changes the outer to this section (from any actor manager)
