Control Rig: fixed dependency checker

#rnx

diff --git a/Engine/Plugins/Animation/ControlRig/Source/ControlRig/Private/Constraints/ControlRigTransformableHandle.cpp b/Engine/Plugins/Animation/ControlRig/Source/ControlRig/Private/Constraints/ControlRigTransformableHandle.cpp
--- a/Engine/Plugins/Animation/ControlRig/Source/ControlRig/Private/Constraints/ControlRigTransformableHandle.cpp
+++ b/Engine/Plugins/Animation/ControlRig/Source/ControlRig/Private/Constraints/ControlRigTransformableHandle.cpp
@@ -414,9 +414,9 @@ bool UTransformableControlHandle::HasDirectDependencyWith(const UTransformableHa
 			#if WITH_EDITOR
 				FRigDependenciesProviderForControlRig DependencyProvider(InControlRig);
 				DependencyProvider.SetInteractiveDialogEnabled(true);
-				return InRigHierarchy->CanCauseCycle(ControlElement, PossibleParent, DependencyProvider);
+				return InRigHierarchy->CanCauseCycle(PossibleParent, ControlElement, DependencyProvider);
             #else
-				return InRigHierarchy->CanCauseCycle(ControlElement, PossibleParent);
+				return InRigHierarchy->CanCauseCycle(PossibleParent, ControlElement);
             #endif	
 			}
 			return false;
