Control Rig: fixed dependency checker

#rnx

diff --git a/Engine/Plugins/Animation/ControlRig/Source/ControlRigEditor/Private/Editor/Constraints/SConstraintsWidget.cpp b/Engine/Plugins/Animation/ControlRig/Source/ControlRigEditor/Private/Editor/Constraints/SConstraintsWidget.cpp
--- a/Engine/Plugins/Animation/ControlRig/Source/ControlRigEditor/Private/Editor/Constraints/SConstraintsWidget.cpp
+++ b/Engine/Plugins/Animation/ControlRig/Source/ControlRigEditor/Private/Editor/Constraints/SConstraintsWidget.cpp
@@ -274,20 +274,24 @@ namespace UE::Private
 		{
 			if (Child != InParent.Object)
 			{
-				// enable the dependency dialog when creating a constraint
-				// for children that are control rig controls.
-				TSharedPtr<FControlRigDismissDependencyDialogGuard> DependencyDialogGuard;
-				if (AControlRigShapeActor* ChildActor = Cast<AControlRigShapeActor>(Child))
+
+				FScopedTransaction Transaction(LOCTEXT("CreateConstraintKey", "Create Constraint Key"));
+
+				TUniquePtr<FControlRigDismissDependencyDialogGuard> DependencyDialogGuard;
+				TFunction<void()> ValidDependency;
+				const AControlRigShapeActor* ChildActor = Cast<AControlRigShapeActor>(Child);
+				if (const UControlRig* ControlRig = ChildActor ? ChildActor->ControlRig.Get() : nullptr)
 				{
-					if (UControlRig* ControlRig = ChildActor->ControlRig.Get())
+					// enable the dependency dialog when creating a constraint for children that are control rig controls
+					DependencyDialogGuard = MakeUnique<FControlRigDismissDependencyDialogGuard>(ControlRig->GetHierarchy());
+					ValidDependency = [&DependencyDialogGuard]()
 					{
-						DependencyDialogGuard = MakeShared<FControlRigDismissDependencyDialogGuard>(ControlRig->GetHierarchy());
-					}
+						DependencyDialogGuard.Reset();
+					};
 				}
 				
-				FScopedTransaction Transaction(LOCTEXT("CreateConstraintKey", "Create Constraint Key"));
-				UTickableTransformConstraint* Constraint =
-					UE::TransformConstraintUtil::CreateAndAddFromObjects(World, Parent, InParent.Socket, Child, NAME_None, InConstraintType, true, bUseDefault);
+				UTickableTransformConstraint* Constraint = TransformConstraintUtil::CreateAndAddFromObjects(
+					World, Parent, InParent.Socket, Child, NAME_None, InConstraintType, true, bUseDefault, ValidDependency);
 
 				if (Constraint)
 				{
@@ -1237,11 +1241,14 @@ int32 FBaseConstraintListWidget::RefreshConstraintList()
 	const UEnum* ETransformConstraintTypeEnum = StaticEnum<ETransformConstraintType>();
 	for (const TWeakObjectPtr<UTickableConstraint>& Constraint : Constraints)
 	{
-		const int8 Type = FConstraintInfo::GetType(Constraint->GetClass());
-		if (ETransformConstraintTypeEnum->IsValidEnumValue(Type))
+		if (Constraint.IsValid())
 		{
-			const ETransformConstraintType ConstraintType = static_cast<ETransformConstraintType>(Type);
-			ListItems.Emplace(FEditableConstraintItem::Make(Constraint.Get(), ConstraintType));
+			const int8 Type = FConstraintInfo::GetType(Constraint->GetClass());
+			if (ETransformConstraintTypeEnum->IsValidEnumValue(Type))
+			{
+				const ETransformConstraintType ConstraintType = static_cast<ETransformConstraintType>(Type);
+				ListItems.Emplace(FEditableConstraintItem::Make(Constraint.Get(), ConstraintType));
+			}
 		}
 	}
 
diff --git a/Engine/Source/Runtime/Experimental/Animation/Constraints/Private/Transform/TransformConstraintUtil.cpp b/Engine/Source/Runtime/Experimental/Animation/Constraints/Private/Transform/TransformConstraintUtil.cpp
--- a/Engine/Source/Runtime/Experimental/Animation/Constraints/Private/Transform/TransformConstraintUtil.cpp
+++ b/Engine/Source/Runtime/Experimental/Animation/Constraints/Private/Transform/TransformConstraintUtil.cpp
@@ -327,7 +327,8 @@ UTickableTransformConstraint* CreateAndAddFromObjects(
 	UObject* InChild, const FName& InChildSocketName,
 	const ETransformConstraintType InType,
 	const bool bMaintainOffset,
-	const bool bUseDefault)
+	const bool bUseDefault,
+	const TFunction<void()>& InValidDependencyFunction)
 {
 	static const TCHAR* ErrorPrefix = TEXT("TransformConstraintUtil::CreateAndAddFromActors");
 	
@@ -357,6 +358,11 @@ UTickableTransformConstraint* CreateAndAddFromObjects(
 		ParentHandle->MarkAsGarbage();
 		return nullptr;
 	}
+
+	if (InValidDependencyFunction)
+	{
+		InValidDependencyFunction();
+	}
 	
 	UTickableTransformConstraint* Constraint = CreateFromType(InWorld, InType, bUseDefault);
 	if (Constraint && (ParentHandle->IsValid() && ChildHandle->IsValid()))
diff --git a/Engine/Source/Runtime/Experimental/Animation/Constraints/Public/Transform/TransformConstraintUtil.h b/Engine/Source/Runtime/Experimental/Animation/Constraints/Public/Transform/TransformConstraintUtil.h
--- a/Engine/Source/Runtime/Experimental/Animation/Constraints/Public/Transform/TransformConstraintUtil.h
+++ b/Engine/Source/Runtime/Experimental/Animation/Constraints/Public/Transform/TransformConstraintUtil.h
@@ -80,7 +80,8 @@ namespace UE::TransformConstraintUtil
 		UObject* InChild, const FName& InChildSocketName,
 		const ETransformConstraintType InType,
 		const bool bMaintainOffset = true,
-		const bool bUseDefault = false);
+		const bool bUseDefault = false,
+		const TFunction<void()>& InValidDependencyFunction = nullptr);
 
 	/** Registers a new transform constraint within the constraints manager. */	
 	CONSTRAINTS_API bool AddConstraint(
