[ControlRig & Sequencer] Fixed crash when adding a key to a sequence containing a rig that has a new control added since the first time the sequence was open

diff --git a/Engine/Plugins/Animation/ControlRig/Source/ControlRig/Private/Sequencer/MovieSceneControlRigParameterTrack.cpp b/Engine/Plugins/Animation/ControlRig/Source/ControlRig/Private/Sequencer/MovieSceneControlRigParameterTrack.cpp
--- a/Engine/Plugins/Animation/ControlRig/Source/ControlRig/Private/Sequencer/MovieSceneControlRigParameterTrack.cpp
+++ b/Engine/Plugins/Animation/ControlRig/Source/ControlRig/Private/Sequencer/MovieSceneControlRigParameterTrack.cpp
@@ -87,10 +87,7 @@ UMovieSceneControlRigParameterTrack::UMovieSceneControlRigParameterTrack(const F
 void UMovieSceneControlRigParameterTrack::BeginDestroy()
 {
 	Super::BeginDestroy();
-	if (IsValid(ControlRig))
-	{
-		ControlRig->OnPostConstruction_AnyThread().RemoveAll(this);
-	}
+	UnbindControlRigDelegates();
 }
 
 // UMovieSceneTrack interface
@@ -289,14 +286,11 @@ UMovieSceneSection* UMovieSceneControlRigParameterTrack::CreateControlRigSection
 		InControlRig->Rename(nullptr, this);
 	}
 
-	if (IsValid(ControlRig))
-	{
-		ControlRig->OnPostConstruction_AnyThread().RemoveAll(this);
-	}
+	UnbindControlRigDelegates();
 
 	ControlRig = InControlRig;
 
-	ControlRig->OnPostConstruction_AnyThread().AddUObject(this, &UMovieSceneControlRigParameterTrack::HandleOnPostConstructed);
+	BindControlRigDelegates();
 
 	UMovieSceneControlRigParameterSection* NewSection = Cast<UMovieSceneControlRigParameterSection>(CreateNewSection());
 
@@ -588,7 +582,6 @@ void UMovieSceneControlRigParameterTrack::ReconstructControlRig()
 	if (ControlRig && !ControlRig->HasAnyFlags(RF_NeedLoad | RF_NeedPostLoad | RF_NeedInitialization))
 	{
 		ControlRig->ConditionalPostLoad();
-		ApplyControlRigSettingsOverrides();
 		ControlRig->Initialize();
 		CheckForNameChanges();
 
@@ -634,10 +627,8 @@ void UMovieSceneControlRigParameterTrack::PostLoad()
 		ReconstructControlRig();
 	}
 
-	if (IsValid(ControlRig))
-	{
-		ControlRig->OnPostConstruction_AnyThread().AddUObject(this, &UMovieSceneControlRigParameterTrack::HandleOnPostConstructed);
-	}
+	BindControlRigDelegates();
+
 }
 
 //mz todo this is coming from BuildPatchServices in Runtime/Online/BuildPatchServices/Online/Core/AsyncHelper.h
@@ -703,6 +694,23 @@ namespace MovieSceneControlRigTrack
 	}
 }
 
+void UMovieSceneControlRigParameterTrack::HandleOnPreConstruct_GameThread()
+{
+	if (IsValid(ControlRig))
+	{
+		ApplyControlRigSettingsOverrides();
+	}
+}
+
+void UMovieSceneControlRigParameterTrack::HandleOnPreConstruct(UControlRig* Subject, const FName& InEventName)
+{
+	if(IsInGameThread())
+	{
+		HandleOnPreConstruct_GameThread();
+	}
+}
+
+
 void UMovieSceneControlRigParameterTrack::HandleOnPostConstructed_GameThread()
 {
 	if (IsValid(ControlRig))
@@ -746,6 +754,27 @@ void UMovieSceneControlRigParameterTrack::HandleOnPostConstructed(UControlRig* S
 	}
 }
 
+void UMovieSceneControlRigParameterTrack::BindControlRigDelegates()
+{
+	if (IsValid(ControlRig))
+	{
+		if (!ControlRig->OnPreConstruction_AnyThread().IsBoundToObject(this))
+		{
+			ControlRig->OnPreConstruction_AnyThread().AddUObject(this,  &UMovieSceneControlRigParameterTrack::HandleOnPreConstruct);
+			ControlRig->OnPostConstruction_AnyThread().AddUObject(this, &UMovieSceneControlRigParameterTrack::HandleOnPostConstructed);	
+		}
+	}
+}
+
+void UMovieSceneControlRigParameterTrack::UnbindControlRigDelegates()
+{
+	if (IsValid(ControlRig))
+	{
+		ControlRig->OnPreConstruction_AnyThread().RemoveAll(this);
+		ControlRig->OnPostConstruction_AnyThread().RemoveAll(this);
+	}
+}
+
 #if WITH_EDITORONLY_DATA
 void UMovieSceneControlRigParameterTrack::DeclareConstructClasses(TArray<FTopLevelAssetPath>& OutConstructClasses, const UClass* SpecificSubclass)
 {
@@ -804,10 +833,8 @@ void UMovieSceneControlRigParameterTrack::PostEditImport()
 	Super::PostEditImport();
 	if (ControlRig)
 	{
-		if (ControlRig->OnInitialized_AnyThread().IsBoundToObject(this) == false)
-		{
-			ControlRig->OnPostConstruction_AnyThread().AddUObject(this, &UMovieSceneControlRigParameterTrack::HandleOnPostConstructed);
-		}
+		BindControlRigDelegates();
+		
 		ControlRig->ClearFlags(RF_Transient); //when copied make sure it's no longer transient, sequencer does this for tracks/sections 
 											  //but not for all objects in them since the control rig itself has transient objects.
 	}
@@ -854,16 +881,14 @@ void UMovieSceneControlRigParameterTrack::ReplaceControlRig(UControlRig* NewCont
 		return;
 	}
 
-	if (IsValid(ControlRig))
-	{
-		ControlRig->OnPostConstruction_AnyThread().RemoveAll(this);
-	}
+	UnbindControlRigDelegates();
+
 
 	ControlRig = NewControlRig;
 
 	if (IsValid(ControlRig))
 	{
-		ControlRig->OnPostConstruction_AnyThread().AddUObject(this, &UMovieSceneControlRigParameterTrack::HandleOnPostConstructed);
+		BindControlRigDelegates();
 
 		if (ControlRig->GetOuter() != this)
 		{
@@ -910,11 +935,10 @@ TOptional<EEulerRotationOrder> UMovieSceneControlRigParameterTrack::GetControlRo
 	return Order;
 }
 
-void UMovieSceneControlRigParameterTrack::UpdateAndApplyControlRigSettingsOverrides(const FInstancedPropertyBag& InNewOverrides)
+void UMovieSceneControlRigParameterTrack::UpdateControlRigSettingsOverrides(const FInstancedPropertyBag& InNewOverrides)
 {
 	Modify();
 	ControlRigSettingsOverrides = InNewOverrides;
-	ApplyControlRigSettingsOverrides();
 }
 
 void UMovieSceneControlRigParameterTrack::ApplyControlRigSettingsOverrides()
diff --git a/Engine/Plugins/Animation/ControlRig/Source/ControlRig/Public/Sequencer/MovieSceneControlRigParameterTrack.h b/Engine/Plugins/Animation/ControlRig/Source/ControlRig/Public/Sequencer/MovieSceneControlRigParameterTrack.h
--- a/Engine/Plugins/Animation/ControlRig/Source/ControlRig/Public/Sequencer/MovieSceneControlRigParameterTrack.h
+++ b/Engine/Plugins/Animation/ControlRig/Source/ControlRig/Public/Sequencer/MovieSceneControlRigParameterTrack.h
@@ -85,9 +85,14 @@ public:
 #endif // WITH_EDITOR
 
 	//UControlRig Delegates
+	void HandleOnPreConstruct_GameThread();
+	void HandleOnPreConstruct(UControlRig* Subject, const FName& InEventName);
 	void HandleOnPostConstructed_GameThread();
 	void HandleOnPostConstructed(UControlRig* Subject, const FName& InEventName);
 
+	CONTROLRIG_API void BindControlRigDelegates();
+	void UnbindControlRigDelegates();
+
 #if WITH_EDITOR
 	void HandlePackageDone(const FEndLoadPackageContext& Context);
 	// control Rigs are ready only after its package is fully end-loaded
@@ -170,7 +175,7 @@ public:
 	UMovieSceneControlRigParameterSection::FSpaceChannelAddedEvent& SpaceChannelAdded() { return OnSpaceChannelAdded; }
 	IMovieSceneConstrainedSection::FConstraintChannelAddedEvent& ConstraintChannelAdded() { return OnConstraintChannelAdded; }
 
-	CONTROLRIG_API void UpdateAndApplyControlRigSettingsOverrides(const FInstancedPropertyBag& InNewOverrides);
+	CONTROLRIG_API void UpdateControlRigSettingsOverrides(const FInstancedPropertyBag& InNewOverrides);
 	CONTROLRIG_API void ApplyControlRigSettingsOverrides();
 	
 public:
diff --git a/Engine/Plugins/Animation/ControlRig/Source/ControlRigEditor/Private/Sequencer/ControlRigParameterTrackEditor.cpp b/Engine/Plugins/Animation/ControlRig/Source/ControlRigEditor/Private/Sequencer/ControlRigParameterTrackEditor.cpp
--- a/Engine/Plugins/Animation/ControlRig/Source/ControlRigEditor/Private/Sequencer/ControlRigParameterTrackEditor.cpp
+++ b/Engine/Plugins/Animation/ControlRig/Source/ControlRigEditor/Private/Sequencer/ControlRigParameterTrackEditor.cpp
@@ -509,8 +509,16 @@ void FControlRigParameterTrackEditor::OnInitialize()
 	// Register all modified/selections for control rigs
 	IterateTracks([this](UMovieSceneControlRigParameterTrack* Track)
 	{
+			
 		if (UControlRig* ControlRig = Track->GetControlRig())
 		{
+			// Make sure the delegates are bound before the control rig initializes
+			// Usually delegates should have been bound during Track PostLoad, and they stay bound even if the rig is recompiled via FCoreUObjectDelegates::OnObjectsReplaced above
+			// However, the logic in OnObjectsReplaced only runs if the sequencer is running. When the rig is recompiled while the sequencer is closed
+			// the new Control Rig received by the Track won't have its delegates bound to the Track, so we have to rebind them here such
+			// that the track can pull information from latest version of the rig after construction event
+			Track->BindControlRigDelegates();
+			
 			BindControlRig(ControlRig);
 		}
 
@@ -1729,7 +1737,7 @@ void FControlRigParameterTrackEditor::FRecreateRigOperator::OnRigSettingsChanged
 	Track->Modify();
 	ControlRig->Modify();
 
-	Track->UpdateAndApplyControlRigSettingsOverrides(SettingsForNewControlRig);
+	Track->UpdateControlRigSettingsOverrides(SettingsForNewControlRig);
 	ControlRig->Initialize();
 	ControlRig->Evaluate_AnyThread();
 
@@ -1741,7 +1749,6 @@ void FControlRigParameterTrackEditor::FRecreateRigOperator::OnRigSettingsChanged
 			if (CRSection)
 			{
 				Section->Modify();
-				CRSection->ClearAllParameters();
 				CRSection->RecreateWithThisControlRig(CRSection->GetControlRig(), true);
 			}
 		}
