Space Switching: When we create a channel Modify the section after we create the channel since on undo we want the channel to stick around so it can'e evaluate it's default value. Otherwise on undo it won't evaluate and will keep the last evaluated state, which would be the key before the undo.
Other solution would be to always set the default parent on every control during sequencer evaluation, that seems to heavy and too slow.

diff --git a/Engine/Plugins/Animation/ControlRig/Source/ControlRigEditor/Private/Sequencer/ControlRigSpaceChannelEditors.cpp b/Engine/Plugins/Animation/ControlRig/Source/ControlRigEditor/Private/Sequencer/ControlRigSpaceChannelEditors.cpp
--- a/Engine/Plugins/Animation/ControlRig/Source/ControlRigEditor/Private/Sequencer/ControlRigSpaceChannelEditors.cpp
+++ b/Engine/Plugins/Animation/ControlRig/Source/ControlRigEditor/Private/Sequencer/ControlRigSpaceChannelEditors.cpp
@@ -215,8 +215,6 @@ FSpaceChannelAndSection FControlRigSpaceChannelHelpers::FindSpaceChannelAndSecti
 				UMovieSceneControlRigParameterSection* ActiveSection = Cast<UMovieSceneControlRigParameterSection>(ControlRigParameterTrack->GetSectionToKey(ControlName));
 				if (ActiveSection)
 				{
-					ActiveSection->Modify();
-					ControlRig->Modify();
 					SpaceChannelAndSection.SectionToKey = ActiveSection;
 					FSpaceControlNameAndChannel* NameAndChannel = ActiveSection->GetSpaceChannel(ControlName);
 					if (NameAndChannel)
@@ -225,6 +223,7 @@ FSpaceChannelAndSection FControlRigSpaceChannelHelpers::FindSpaceChannelAndSecti
 					}
 					else if (bCreateIfNeeded)
 					{
+						ControlRig->Modify();
 						if (ControlRig->IsControlSelected(ControlName))
 						{
 							TPair<UControlRig*, FName> Pair;
@@ -233,6 +232,8 @@ FSpaceChannelAndSection FControlRigSpaceChannelHelpers::FindSpaceChannelAndSecti
 							ControlRigPairsToReselect.Add(Pair);
 						}
 						ActiveSection->AddSpaceChannel(ControlName, true /*ReconstructChannelProxy*/);
+						//important we don't want the adding of the space channel to be in the undo since we want the channel to stick around if we get undone
+						ActiveSection->Modify();
 						NameAndChannel = ActiveSection->GetSpaceChannel(ControlName);
 						if (NameAndChannel)
 						{
