Selection Sets:  Replace space with _ with names, Change UI to show Set button, and support for mult-asset sets, using OwnerActor string to specify which actor the set is associated with

#DisabledValidations Preflight

diff --git a/Engine/Plugins/Animation/ControlRig/Source/ControlRigEditor/Private/EditMode/ControlRigEditModeToolkit.cpp b/Engine/Plugins/Animation/ControlRig/Source/ControlRigEditor/Private/EditMode/ControlRigEditModeToolkit.cpp
--- a/Engine/Plugins/Animation/ControlRig/Source/ControlRigEditor/Private/EditMode/ControlRigEditModeToolkit.cpp
+++ b/Engine/Plugins/Animation/ControlRig/Source/ControlRigEditor/Private/EditMode/ControlRigEditModeToolkit.cpp
@@ -699,7 +699,7 @@ void FControlRigEditModeToolkit::BuildToolBar(FToolBarBuilder& InToolBarBuilder)
 			FIsActionChecked::CreateRaw(this, &FControlRigEditModeToolkit::IsToolkitUIActive, FControlRigEditModeToolkit::SelectionSetsTabName)
 		),
 		NAME_None,
-		LOCTEXT("SelectionSets", "Selection Sets"),
+		LOCTEXT("Sets", "Sets"),
 		LOCTEXT("SelectionSetsTooltip", "Open Selection Sets"),
 		FSlateIcon(TEXT("ControlRigEditorStyle"), TEXT("ControlRig.SelectionSet")),
 		EUserInterfaceActionType::ToggleButton
@@ -712,7 +712,7 @@ void FControlRigEditModeToolkit::BuildToolBar(FToolBarBuilder& InToolBarBuilder)
 			FIsActionChecked::CreateRaw(this, &FControlRigEditModeToolkit::IsToolkitUIActive, FControlRigEditModeToolkit::SelectionSetsOverlayName)
 		),
 		NAME_None,
-		LOCTEXT("SelectionSets", "Selection Sets"),
+		LOCTEXT("Sets", "Sets"),
 		LOCTEXT("SelectionSetsTooltip", "Open Selection Sets"),
 		FSlateIcon(TEXT("ControlRigEditorStyle"), TEXT("ControlRig.SelectionSet")),
 		EUserInterfaceActionType::ToggleButton
diff --git a/Engine/Plugins/Animation/ControlRig/Source/ControlRigEditor/Private/Sequencer/SelectionSets/SelectionSets.cpp b/Engine/Plugins/Animation/ControlRig/Source/ControlRigEditor/Private/Sequencer/SelectionSets/SelectionSets.cpp
--- a/Engine/Plugins/Animation/ControlRig/Source/ControlRigEditor/Private/Sequencer/SelectionSets/SelectionSets.cpp
+++ b/Engine/Plugins/Animation/ControlRig/Source/ControlRigEditor/Private/Sequencer/SelectionSets/SelectionSets.cpp
@@ -56,6 +56,31 @@
 /// FAIESelectionSetItem
 ///////////////////////////////////////////////////////////
 
+bool FAIESelectionSetItem::IsMultiAsset() const
+{
+	for (const FAIESelectionSetItemName& LocalName : Names)
+	{
+		if (LocalName.OwnerActorName.IsEmpty() == false)
+		{
+			return true;
+		}
+	}
+	return false;
+}
+
+
+bool FAIESelectionSetItem::ContainsOwnerActor(const FString& ActorLabel) const
+{
+	for (const FAIESelectionSetItemName& LocalName : Names)
+	{
+		if (LocalName.OwnerActorName == ActorLabel)
+		{
+			return true;
+		}
+	}
+	return false;
+}
+
 bool FAIESelectionSetItem::ContainsItem(const FString& InName) const
 {
 	for (const FAIESelectionSetItemName& LocalName : Names)
@@ -114,6 +139,10 @@ FString FAIESelectionSetItem::GetMirror(const FString& InString)
 	return NewString;
 }
 
+///////////////////////////////////////////////////////////
+/// FActorWithSelectionSet
+///////////////////////////////////////////////////////////
+
 AActor* FActorWithSelectionSet::GetControlRigActor(const UControlRig* InControlRig)
 {
 	TSharedPtr<IControlRigObjectBinding> ObjectBinding = InControlRig->GetObjectBinding();
@@ -148,6 +177,28 @@ void FActorWithSelectionSet::SetUpOwnedActors(const TWeakObjectPtr<AActor>& Acto
 	}
 }
 
+///////////////////////////////////////////////////////////
+/// FActorsWithSelectionSets
+///////////////////////////////////////////////////////////
+
+bool FActorsWithSelectionSets::ContainsActor(const FString& InString, bool bShowSelectedOnly) const
+{
+	for (const TPair<TWeakObjectPtr<AActor>, FActorWithSelectionSet>& Pair : ActorsWithSet)
+	{
+		if (Pair.Value.IsActive(bShowSelectedOnly))
+		{
+			if (const AActor* Actor = Pair.Key.Get())
+			{
+				const FString ActorLabel = Actor->GetActorLabel();
+				if (ActorLabel == InString)
+				{
+					return true;
+				}
+			}
+		}
+	}
+	return false;
+}
 
 void FActorsWithSelectionSets::SetUpFromSequencer(TSharedPtr<ISequencer>& Sequencer, UAIESelectionSets* SelectionSets, const TSet<TWeakObjectPtr<UObject>>& MakeSelectedIfPresent)
 {
@@ -521,12 +572,51 @@ TArray<FGuid> UAIESelectionSets::GetActiveSelectionSets()
 		}
 		LastActiveSelectionSetsSet = ActiveSelectionSetsSet;
 	}
+
+	//remove ones that aren't selected if multiasset
+	TArray<FGuid> SetsToRemove;
+	for (const FGuid& Guid : ActiveSelectionSetsSet)
+	{
+		if (FAIESelectionSetItem* SetItem = SelectionSets.Find(Guid))
+		{
+			if (SetItem->IsMultiAsset())
+			{
+				bool bHasOneActive = false;
+				for (const FAIESelectionSetItemName& Name : SetItem->Names)
+				{
+					if (Name.OwnerActorName.IsEmpty() == false)
+					{
+						if (ActorsWithSelectionSets.ContainsActor(Name.OwnerActorName,bShowSelectedOnly))
+						{
+							bHasOneActive = true;
+							break;
+						}
+					}
+				}
+				if (bHasOneActive == false)
+				{
+					SetsToRemove.Add(Guid);
+				}
+			}
+		}
+	}
+	if (SetsToRemove.Num() > 0)
+	{
+		for (const FGuid& Guid : SetsToRemove)
+		{
+			ActiveSelectionSetsSet.Remove(Guid);
+		}
+	}
+
 	TArray<FGuid> ActiveSelectionSetsArray;
 	if (ActiveSelectionSetsSet.Num() > 0)
 	{
 		ActiveSelectionSetsArray = ActiveSelectionSetsSet.Array();
 		ActiveSelectionSetsArray = SortSelectionSetsByRow(ActiveSelectionSetsArray);
 	}
+
+
+
 	return ActiveSelectionSetsArray;
 }
 
@@ -586,7 +676,7 @@ FGuid UAIESelectionSets::CreateSetItemFromSelection(TSharedPtr<ISequencer>& InSe
 			return FGuid();
 		}
 
-		FString Name = FString::Printf(TEXT("Set %d"), SelectionSets.Num());
+		FString Name = FString::Printf(TEXT("Set_%d"), SelectionSets.Num());
 		FText Text;
 		Text = Text.FromString(Name);
 		Item.ItemName = Text;
@@ -596,34 +686,45 @@ FGuid UAIESelectionSets::CreateSetItemFromSelection(TSharedPtr<ISequencer>& InSe
 		FControlRigPoseMirrorTable MirrorTable;
 		FControlRigControlPose TempPose;
 		UControlRig* LastOwnerControlRig = nullptr;
-		for (FControlRigAndControlsAndTrack& CRControls : SelectedCRs)
+		if (SelectedCRs.Num() > 0)
 		{
-			UControlRig* ControlRig = CRControls.ControlRig;
-			if (ControlRig)
+			const bool bOwnerActorForCRs = SelectedCRs.Num() > 1 || SelectedBoundObjects.Num() > 0;
+			for (FControlRigAndControlsAndTrack& CRControls : SelectedCRs)
 			{
-				LastOwnerControlRig = ControlRig;
-				TempPose.SavePose(ControlRig, true);
-				for (FName& ControlName : CRControls.Controls)
+				UControlRig* ControlRig = CRControls.ControlRig;
+				if (ControlRig)
 				{
-					FAIESelectionSetItemName ItemName;
-					ItemName.Name = ControlName.ToString();
-					if (FRigControlCopy* CopyRigControl = MirrorTable.GetControl(ControlRig, TempPose, ControlName, true))
+					LastOwnerControlRig = ControlRig;
+					TempPose.SavePose(ControlRig, true);
+					for (FName& ControlName : CRControls.Controls)
 					{
-						ItemName.MirrorName = CopyRigControl->Name.ToString();
-					}
-					else
-					{
-						ItemName.MirrorName = ItemName.Name;
+						FAIESelectionSetItemName ItemName;
+						ItemName.Name = ControlName.ToString();
+						if (FRigControlCopy* CopyRigControl = MirrorTable.GetControl(ControlRig, TempPose, ControlName, true))
+						{
+							ItemName.MirrorName = CopyRigControl->Name.ToString();
+						}
+						else
+						{
+							ItemName.MirrorName = ItemName.Name;
+						}
+
+						ItemName.Type = (int32)EAIESelectionSetItemType::ControlRig;
+						Duplicates.SetNum(0);
+						ActiveNames.MultiFind(ItemName.Name, Duplicates);
+						ItemName.Duplicates = Duplicates.Num();
+						if (bOwnerActorForCRs)
+						{
+							if (AActor* Actor = FActorWithSelectionSet::GetControlRigActor(ControlRig))
+							{
+								ItemName.OwnerActorName = Actor->GetActorLabel(false);
+							}
+						}
+						Item.Names.Add(ItemName);
+						ActiveNames.Add(ItemName.Name, ItemName.Name);
 					}
-				
-					ItemName.Type = (int32)EAIESelectionSetItemType::ControlRig;
-					Duplicates.SetNum(0);
-					ActiveNames.MultiFind(ItemName.Name, Duplicates);
-					ItemName.Duplicates = Duplicates.Num();
-					Item.Names.Add(ItemName);
-					ActiveNames.Add(ItemName.Name, ItemName.Name);
+					ObjectsAdded.Add(CRControls.ControlRig);
 				}
-				ObjectsAdded.Add(CRControls.ControlRig);
 			}
 		}
 		TOptional<FString> OwnerActorName;
@@ -893,6 +994,15 @@ void UAIESelectionSets::GetItemGuids(const FText& ItemName, TArray<FGuid>& OutGu
 	}
 }
 
+bool UAIESelectionSets::IsMultiAsset(const FGuid& InGuid) const
+{
+	if (const FAIESelectionSetItem* SetItem = SelectionSets.Find(InGuid))
+	{
+		return SetItem->IsMultiAsset();
+	}
+	return false;
+}
+
 bool UAIESelectionSets::GetItemName(const FGuid& InGuid, FText& OutName) const
 {
 	if (const FAIESelectionSetItem* SetItem = SelectionSets.Find(InGuid))
@@ -909,7 +1019,9 @@ bool UAIESelectionSets::RenameSetItem(const FGuid& InGuid, const  FText& NewName
 	if (FAIESelectionSetItem* SetItem = SelectionSets.Find(InGuid))
 	{
 		Modify();
-		SetItem->ItemName = NewName;
+		FString String = NewName.ToString();
+		String = String.Replace(TEXT(" "), TEXT("_"), ESearchCase::CaseSensitive);
+		SetItem->ItemName = FText::FromString(String);
 		return true;
 	}
 	Transaction.Cancel();
@@ -1134,13 +1246,20 @@ bool UAIESelectionSets::SelectItem(const FGuid& InGuid, bool bDoMirror, bool bAd
 			if (const FAIESelectionSetItem* CurrentSetItem = SelectionSets.Find(Guid))
 			{
 				for (const TPair<TWeakObjectPtr<AActor>, FActorWithSelectionSet>& Pair : ActorsWithSelectionSets.ActorsWithSet)
-				{
-					bool bIsActive = Pair.Value.IsActive(bShowSelectedOnly);
-					if (bIsActive)
+			{
+					if (Pair.Value.SetsThatMatch.Contains(Guid))
 					{
-						if (Pair.Value.SetsThatMatch.Contains(Guid))
+						if (AActor* Key = Cast<AActor>(Pair.Key.Get()))
 						{
-							if (AActor* Key = Cast<AActor>(Pair.Key.Get()))
+							bool bIsActive = Pair.Value.IsActive(bShowSelectedOnly);
+							if (!bIsActive && CurrentSetItem->IsMultiAsset()) //not active but it is a multiasset
+							{
+								if (CurrentSetItem->ContainsOwnerActor(Key->GetActorLabel()))
+								{
+									bIsActive = true;
+								}
+							}
+							if (bIsActive)
 							{
 								for (const TWeakObjectPtr<UObject>& Object : Pair.Value.Actors)
 								{
@@ -1231,7 +1350,21 @@ bool UAIESelectionSets::SelectItem(const FGuid& InGuid, bool bDoMirror, bool bAd
 				}
 				for (const FAIESelectionSetItemName& String : CachedControlRig.Value)
 				{
-
+					if (String.OwnerActorName.IsEmpty() == false)
+					{
+						if (AActor* Actor = FActorWithSelectionSet::GetControlRigActor(ControlRig))
+						{
+							FString ActorName = Actor->GetActorLabel(false);
+							if (ActorName != String.OwnerActorName)
+							{
+								continue;
+							}
+						}
+						else
+						{
+							continue;
+						}
+					}
 					if (bDoMirror)
 					{
 						const FName Name(*String.MirrorName);
diff --git a/Engine/Plugins/Animation/ControlRig/Source/ControlRigEditor/Private/Sequencer/SelectionSets/SelectionSets.h b/Engine/Plugins/Animation/ControlRig/Source/ControlRigEditor/Private/Sequencer/SelectionSets/SelectionSets.h
--- a/Engine/Plugins/Animation/ControlRig/Source/ControlRigEditor/Private/Sequencer/SelectionSets/SelectionSets.h
+++ b/Engine/Plugins/Animation/ControlRig/Source/ControlRigEditor/Private/Sequencer/SelectionSets/SelectionSets.h
@@ -49,7 +49,7 @@ struct FAIESelectionSetItemName
 	int32 Duplicates = 0;
 
 	/* The Name of the Owner Actor for Multi-Sets containing multiple rigs, may not be set*/
-	UPROPERTY(BlueprintReadOnly, EditAnywhere, Category = "Data")
+	UPROPERTY(BlueprintReadWrite, EditAnywhere, Category = "Data")
 	FString OwnerActorName;
 };
 
@@ -100,11 +100,12 @@ struct FAIESelectionSetItem
 	//currently not used but maybe in future
 	TArray<FGuid> Children;
 
-	
 	bool ContainsItem(const FString& Name) const;
 	FAIESelectionSetItemName* GetItem(const FString& Name);
 	bool RemoveItem(const FString& Name);
 
+	bool IsMultiAsset() const;
+	bool ContainsOwnerActor(const FString& ActorLabel) const;
 	void GetAllChildren(const UAIESelectionSets* SelectionSets, TArray<FGuid>& AllChildren) const;
 	static FString GetMirror(const FString& InString);
 };
@@ -160,7 +161,7 @@ public:
 struct FActorsWithSelectionSets
 {
 	TMap<TWeakObjectPtr<AActor>, FActorWithSelectionSet> ActorsWithSet;
-	
+	bool ContainsActor(const FString& ActorLabel, bool bShowSelectedOnly) const;
 	void SetUpFromSequencer(TSharedPtr<ISequencer>& InSequencer,  UAIESelectionSets* SelectionSets, const TSet<TWeakObjectPtr<UObject>>& MakeSelectedIfPresent);
 };
 
@@ -234,6 +235,10 @@ public:
 	UFUNCTION(BlueprintPure, Category = "Set")
 	bool GetItemName(const FGuid& InGuid, FText& OutName) const;
 
+	/* Rename selection set*/
+	UFUNCTION(BlueprintPure, Category = "Set")
+	bool IsMultiAsset(const FGuid& InGuid) const;
+
 	/* Get the guids of set items associated with this name*/
 	UFUNCTION(BlueprintPure, Category = "State")
 	void GetItemGuids(const FText& ItemName, TArray<FGuid>& OutGuids) const;
