Anim in Engine: Fix SAnimDetailsValueNumeric mixes (currently unused) min and max slider value

#rnx

diff --git a/Engine/Plugins/Animation/ControlRig/Source/ControlRigEditor/Private/AnimDetails/Widgets/SAnimDetailsValueNumeric.cpp b/Engine/Plugins/Animation/ControlRig/Source/ControlRigEditor/Private/AnimDetails/Widgets/SAnimDetailsValueNumeric.cpp
--- a/Engine/Plugins/Animation/ControlRig/Source/ControlRigEditor/Private/AnimDetails/Widgets/SAnimDetailsValueNumeric.cpp
+++ b/Engine/Plugins/Animation/ControlRig/Source/ControlRigEditor/Private/AnimDetails/Widgets/SAnimDetailsValueNumeric.cpp
@@ -257,8 +257,8 @@ namespace UE::ControlRigEditor
 			.SupportDynamicSliderMaxValue(PropertyMetaData.SupportDynamicSliderMaxValue)
 			.MinValue(PropertyMetaData.MinValue)
 			.MaxValue(PropertyMetaData.MaxValue)
-			.MaxSliderValue(PropertyMetaData.SliderMinValue)
-			.MinSliderValue(PropertyMetaData.SliderMaxValue)
+			.MinSliderValue(PropertyMetaData.SliderMinValue)
+			.MaxSliderValue(PropertyMetaData.SliderMaxValue)
 			.ShiftMultiplier(PropertyMetaData.ShiftMultiplier)
 			.CtrlMultiplier(PropertyMetaData.CtrlMultiplier)
 			.SliderExponent(PropertyMetaData.SliderExponent)
