ControlRig: Introduce cache for dependency traverse

#rnx

diff --git a/Engine/Plugins/Animation/ControlRig/Source/ControlRig/Private/Rigs/RigHierarchy.cpp b/Engine/Plugins/Animation/ControlRig/Source/ControlRig/Private/Rigs/RigHierarchy.cpp
--- a/Engine/Plugins/Animation/ControlRig/Source/ControlRig/Private/Rigs/RigHierarchy.cpp
+++ b/Engine/Plugins/Animation/ControlRig/Source/ControlRig/Private/Rigs/RigHierarchy.cpp
@@ -6441,12 +6441,14 @@ bool URigHierarchy::IsParentedTo(FRigBaseElement* InChild, FRigBaseElement* InPa
 {
 	FRigHierarchyDependencyChain DependencyChain;
 
-	bool bActualTopology = false;
+	THierarchyCache<TMap<TTuple<int32, int32>, bool>> TempDependencyCache;
+	THierarchyCache<TMap<TTuple<int32, int32>, bool>>* DependencyCachePtr = &TempDependencyCache;
+	
 	if (OutDependencyChain == nullptr)
 	{
 		DependencyChain.Reserve(64);
 		OutDependencyChain = &DependencyChain;
-		bActualTopology = true;
+		DependencyCachePtr = &ElementDependencyCache;
 	}
 
 	check(OutDependencyChain != nullptr);
@@ -6454,7 +6456,7 @@ bool URigHierarchy::IsParentedTo(FRigBaseElement* InChild, FRigBaseElement* InPa
 
 	{
 		UE::TScopeLock Lock(IsDependentOnLock);
-		if (!IsDependentOn(InChild, InParent, InDependencyProvider, bActualTopology, *OutDependencyChain))
+		if (!IsDependentOn(InChild, InParent, InDependencyProvider, DependencyCachePtr, *OutDependencyChain))
 		{
 			return false;
 		}
@@ -6463,7 +6465,12 @@ bool URigHierarchy::IsParentedTo(FRigBaseElement* InChild, FRigBaseElement* InPa
 	return true;
 }
 
-bool URigHierarchy::IsDependentOn(FRigBaseElement* InDependent, FRigBaseElement* InDependency, const IRigDependenciesProvider& InDependencyProvider, bool bIsOnActualTopology, FRigHierarchyDependencyChain& OutDependencyChain) const
+bool URigHierarchy::IsDependentOn(
+	FRigBaseElement* InDependent,
+	FRigBaseElement* InDependency,
+	const IRigDependenciesProvider& InDependencyProvider,
+	THierarchyCache<TMap<TTuple<int32, int32>, bool>>* InElementDependencyCache,
+	FRigHierarchyDependencyChain& OutDependencyChain) const
 {
 	if((InDependent == nullptr) || (InDependency == nullptr))
 	{
@@ -6489,16 +6496,16 @@ bool URigHierarchy::IsDependentOn(FRigBaseElement* InDependent, FRigBaseElement*
 		return false;
 	}
 	
-	if (bIsOnActualTopology)
+	if (InElementDependencyCache)
 	{
 		const uint32 ElementDependencyVersion = HashCombine(GetTopologyVersion(), GetParentWeightVersion()); 
-		if(!ElementDependencyCache.IsValid(ElementDependencyVersion))
+		if(!InElementDependencyCache->IsValid(ElementDependencyVersion))
 		{
-			ElementDependencyCache.Set(TMap<TTuple<int32, int32>, bool>(), ElementDependencyVersion);
+			InElementDependencyCache->Set(TMap<TTuple<int32, int32>, bool>(), ElementDependencyVersion);
 		}
 
 		// we'll only update the caches if we are following edges on the actual topology
-		if(const bool* bCachedResult = ElementDependencyCache.Get().Find(CacheKey))
+		if(const bool* bCachedResult = InElementDependencyCache->Get().Find(CacheKey))
 		{
 			return *bCachedResult;
 		}
@@ -6506,7 +6513,7 @@ bool URigHierarchy::IsDependentOn(FRigBaseElement* InDependent, FRigBaseElement*
 		// check if the reverse dependency check has been stored before - if the dependency is dependent
 		// then we don't need to recurse any further.
 		const TTuple<int32,int32> ReverseCacheKey(DependencyElementIndex, DependentElementIndex);
-		if(const bool* bReverseCachedResult = ElementDependencyCache.Get().Find(ReverseCacheKey))
+		if(const bool* bReverseCachedResult = InElementDependencyCache->Get().Find(ReverseCacheKey))
 		{
 			if(*bReverseCachedResult)
 			{
@@ -6523,12 +6530,12 @@ bool URigHierarchy::IsDependentOn(FRigBaseElement* InDependent, FRigBaseElement*
 			const FRigHierarchyRecord DependencyRecord(FRigHierarchyRecord::EType_RigElement, SingleParentElement->ParentElement->GetIndex(), NAME_None);
 			OutDependencyChain.Emplace(DependencyRecord, DependentRecord);
 		
-			if(IsDependentOn(SingleParentElement->ParentElement, InDependency, InDependencyProvider, bIsOnActualTopology, OutDependencyChain))
+			if(IsDependentOn(SingleParentElement->ParentElement, InDependency, InDependencyProvider, InElementDependencyCache, OutDependencyChain))
 			{
 				// we'll only update the caches if we are following edges on the actual topology
-				if(bIsOnActualTopology)
+				if(InElementDependencyCache)
 				{
-					ElementDependencyCache.Get().FindOrAdd(CacheKey, true);
+					InElementDependencyCache->Get().FindOrAdd(CacheKey, true);
 				}
 				return true;
 			}
@@ -6548,12 +6555,12 @@ bool URigHierarchy::IsDependentOn(FRigBaseElement* InDependent, FRigBaseElement*
 			const FRigHierarchyRecord DependencyRecord(FRigHierarchyRecord::EType_RigElement, ParentConstraint.ParentElement->GetIndex(), NAME_None);
 			OutDependencyChain.Emplace(DependencyRecord, DependentRecord);
 
-			if(IsDependentOn(ParentConstraint.ParentElement, InDependency, InDependencyProvider, bIsOnActualTopology, OutDependencyChain))
+			if(IsDependentOn(ParentConstraint.ParentElement, InDependency, InDependencyProvider, InElementDependencyCache, OutDependencyChain))
 			{
 				// we'll only update the caches if we are following edges on the actual topology
-				if(bIsOnActualTopology)
+				if(InElementDependencyCache)
 				{
-					ElementDependencyCache.Get().FindOrAdd(CacheKey, true);
+					InElementDependencyCache->Get().FindOrAdd(CacheKey, true);
 				}
 				return true;
 			}
@@ -6568,7 +6575,7 @@ bool URigHierarchy::IsDependentOn(FRigBaseElement* InDependent, FRigBaseElement*
 			const URigHierarchy* InHierarchy,
 			FRigBaseElement* InDependency,
 			const IRigDependenciesProvider& InDependencyProvider,
-			bool bIsOnActualTopology,
+			THierarchyCache<TMap<TTuple<int32, int32>, bool>>* InElementDependencyCache,
 			const TTuple<int32,int32>& CacheKey,
 			const FRigHierarchyRecord& InDependentRecord,
 			FRigHierarchyDependencyChain& OutDependencyChain)
@@ -6586,12 +6593,12 @@ bool URigHierarchy::IsDependentOn(FRigBaseElement* InDependent, FRigBaseElement*
 
 						OutDependencyChain.Emplace(DependencyRecord, InDependentRecord);
 
-						if(InHierarchy->IsDependentOn(InHierarchy->Elements[DependencyRecord.Index], InDependency, InDependencyProvider, bIsOnActualTopology, OutDependencyChain))
+						if(InHierarchy->IsDependentOn(InHierarchy->Elements[DependencyRecord.Index], InDependency, InDependencyProvider, InElementDependencyCache, OutDependencyChain))
 						{
 							// we'll only update the caches if we are following edges on the actual topology
-							if(bIsOnActualTopology)
+							if(InElementDependencyCache)
 							{
-								InHierarchy->ElementDependencyCache.Get().FindOrAdd(CacheKey, true);
+								InElementDependencyCache->Get().FindOrAdd(CacheKey, true);
 							}
 							return true;
 						}
@@ -6610,7 +6617,7 @@ bool URigHierarchy::IsDependentOn(FRigBaseElement* InDependent, FRigBaseElement*
 
 						OutDependencyChain.Emplace(DependencyRecord, InDependentRecord);
 
-						if (FoundDependencyInProvider(InHierarchy, InDependency, InDependencyProvider, bIsOnActualTopology, CacheKey, DependencyRecord, OutDependencyChain))
+						if (FoundDependencyInProvider(InHierarchy, InDependency, InDependencyProvider, InElementDependencyCache, CacheKey, DependencyRecord, OutDependencyChain))
 						{
 							return true;
 						}
@@ -6624,20 +6631,20 @@ bool URigHierarchy::IsDependentOn(FRigBaseElement* InDependent, FRigBaseElement*
 	};
 
 	// check the optional dependency map
-	if (Local::FoundDependencyInProvider(this, InDependency, InDependencyProvider, bIsOnActualTopology, CacheKey, DependentRecord, OutDependencyChain))
+	if (Local::FoundDependencyInProvider(this, InDependency, InDependencyProvider, InElementDependencyCache, CacheKey, DependentRecord, OutDependencyChain))
 	{
 		// we'll only update the caches if we are following edges on the actual topology
-		if(bIsOnActualTopology)
+		if(InElementDependencyCache)
 		{
-			ElementDependencyCache.Get().FindOrAdd(CacheKey, true);
+			InElementDependencyCache->Get().FindOrAdd(CacheKey, true);
 		}
 		return true;		
 	}
 
 	// we'll only update the caches if we are following edges on the actual topology
-	if(bIsOnActualTopology)
+	if(InElementDependencyCache)
 	{
-		ElementDependencyCache.Get().FindOrAdd(CacheKey, false);
+		InElementDependencyCache->Get().FindOrAdd(CacheKey, false);
 	}
 
 	return false;
diff --git a/Engine/Plugins/Animation/ControlRig/Source/ControlRig/Public/Rigs/RigHierarchy.h b/Engine/Plugins/Animation/ControlRig/Source/ControlRig/Public/Rigs/RigHierarchy.h
--- a/Engine/Plugins/Animation/ControlRig/Source/ControlRig/Public/Rigs/RigHierarchy.h
+++ b/Engine/Plugins/Animation/ControlRig/Source/ControlRig/Public/Rigs/RigHierarchy.h
@@ -4512,7 +4512,7 @@ private:
 	CONTROLRIG_API bool IsDependentOn(FRigBaseElement* InDependent,
 		FRigBaseElement* InDependency,
 		const IRigDependenciesProvider& InDependencyProvider,
-		bool bIsOnActualTopology,
+		THierarchyCache<TMap<TTuple<int32, int32>, bool>>* InElementDependencyCache,
 		FRigHierarchyDependencyChain& OutDependencyChain) const;
 
 public:
