Control Rig: Fix crash when closing selection sets tab by replacing AddRaw with AddSP.

The problem was that UAIESelectionSets::OnSelectionSetsChanged is invoked when you open the tab. The delegates from SSelectionSets were not being deregistered from UAIESelectionSets. Upon opening, the dangling callback would be invoked, which caused undefined behaviour. Ultimately, it would eventually cause a crash.

diff --git a/Engine/Plugins/Animation/ControlRig/Source/ControlRigEditor/Private/Sequencer/SelectionSets/SSelectionSets.cpp b/Engine/Plugins/Animation/ControlRig/Source/ControlRigEditor/Private/Sequencer/SelectionSets/SSelectionSets.cpp
--- a/Engine/Plugins/Animation/ControlRig/Source/ControlRigEditor/Private/Sequencer/SelectionSets/SSelectionSets.cpp
+++ b/Engine/Plugins/Animation/ControlRig/Source/ControlRigEditor/Private/Sequencer/SelectionSets/SSelectionSets.cpp
@@ -383,10 +383,10 @@ void SSelectionSets::Construct(const FArguments& InArgs, FControlRigEditMode& In
 	TSharedPtr<ISequencer> SequencerPtr = UAIESelectionSets::GetSequencerFromAsset();
 	if (SequencerPtr.IsValid())
 	{
-		SequencerPtr->OnActivateSequence().AddRaw(this, &SSelectionSets::OnActivateSequence);
-		SequencerPtr->OnMovieSceneDataChanged().AddRaw(this, &SSelectionSets::OnMovieSceneDataChanged);
-		SequencerPtr->OnCloseEvent().AddRaw(this, &SSelectionSets::OnCloseEvent);
-		SequencerPtr->OnMovieSceneBindingsChanged().AddRaw(this, &SSelectionSets::OnMovieSceneBindingsChanged);
+		SequencerPtr->OnActivateSequence().AddSP(this, &SSelectionSets::OnActivateSequence);
+		SequencerPtr->OnMovieSceneDataChanged().AddSP(this, &SSelectionSets::OnMovieSceneDataChanged);
+		SequencerPtr->OnCloseEvent().AddSP(this, &SSelectionSets::OnCloseEvent);
+		SequencerPtr->OnMovieSceneBindingsChanged().AddSP(this, &SSelectionSets::OnMovieSceneBindingsChanged);
 		if (UAIESelectionSets* SelectionSets = UAIESelectionSets::GetSelectionSets(SequencerPtr))
 		{
 			NewSelectionSet(SelectionSets);
@@ -1227,7 +1227,7 @@ void SSelectionSets::RegisterSelectionChanged()
 	}
 
 	// register
-	OnSelectionChangedHandle = ActorSelectionChangedEvent.AddRaw(this, &SSelectionSets::OnActorSelectionChanged);
+	OnSelectionChangedHandle = ActorSelectionChangedEvent.AddSP(this, &SSelectionSets::OnActorSelectionChanged);
 }
 
 void SSelectionSets::OnActorSelectionChanged(const TArray<UObject*>& NewSelection, bool bForceRefresh)
@@ -1258,7 +1258,7 @@ void SSelectionSets::UpdateBindings()
 		if (UAIESelectionSets* SelectionSets = UAIESelectionSets::GetSelectionSets(SequencerPtr, false /*Add if doesn't exist*/))
 		{
 			SelectionSets->SelectionListChanged().RemoveAll(this);
-			SelectionSets->SelectionListChanged().AddRaw(this, &SSelectionSets::HandleSelectionListChanged);
+			SelectionSets->SelectionListChanged().AddSP(this, &SSelectionSets::HandleSelectionListChanged);
 
 			SelectionSets->SequencerBindingsAdded(SequencerPtr);
 			PopulateUI();
@@ -1333,13 +1333,13 @@ void SSelectionSets::SetEditMode(FControlRigEditMode& InEditMode)
 			{
 				if (!ControlRig.Get()->ControlRigBound().IsBoundToObject(this))
 				{
-					ControlRig.Get()->ControlRigBound().AddRaw(this, &SSelectionSets::HandleOnControlRigBound);
+					ControlRig.Get()->ControlRigBound().AddSP(this, &SSelectionSets::HandleOnControlRigBound);
 					BoundControlRigs.Add(ControlRig);
 				}
 				const TSharedPtr<IControlRigObjectBinding> Binding = ControlRig.Get()->GetObjectBinding();
 				if (Binding && !Binding->OnControlRigBind().IsBoundToObject(this))
 				{
-					Binding->OnControlRigBind().AddRaw(this, &SSelectionSets::HandleOnObjectBoundToControlRig);
+					Binding->OnControlRigBind().AddSP(this, &SSelectionSets::HandleOnObjectBoundToControlRig);
 				}
 			}
 		}
@@ -1355,13 +1355,13 @@ void SSelectionSets::HandleControlAdded(UControlRig* ControlRig, bool bIsAdded)
 		{
 			if (!ControlRig->ControlRigBound().IsBoundToObject(this))
 			{
-				ControlRig->ControlRigBound().AddRaw(this, &SSelectionSets::HandleOnControlRigBound);
+				ControlRig->ControlRigBound().AddSP(this, &SSelectionSets::HandleOnControlRigBound);
 				BoundControlRigs.Add(ControlRig);
 			}
 			const TSharedPtr<IControlRigObjectBinding> Binding = ControlRig->GetObjectBinding();
 			if (Binding && !Binding->OnControlRigBind().IsBoundToObject(this))
 			{
-				Binding->OnControlRigBind().AddRaw(this, &SSelectionSets::HandleOnObjectBoundToControlRig);
+				Binding->OnControlRigBind().AddSP(this, &SSelectionSets::HandleOnObjectBoundToControlRig);
 			}
 		}
 		else
@@ -1388,7 +1388,7 @@ void SSelectionSets::HandleOnControlRigBound(UControlRig* InControlRig)
 
 	if (Binding && !Binding->OnControlRigBind().IsBoundToObject(this))
 	{
-		Binding->OnControlRigBind().AddRaw(this, &SSelectionSets::HandleOnObjectBoundToControlRig);
+		Binding->OnControlRigBind().AddSP(this, &SSelectionSets::HandleOnObjectBoundToControlRig);
 	}
 }
 
@@ -1403,7 +1403,7 @@ void SSelectionSets::NewSelectionSet(UAIESelectionSets* NewSelectionSet)
 	if (NewSelectionSet)
 	{
 		NewSelectionSet->SelectionListChanged().RemoveAll(this);
-		NewSelectionSet->SelectionListChanged().AddRaw(this, &SSelectionSets::HandleSelectionListChanged);
+		NewSelectionSet->SelectionListChanged().AddSP(this, &SSelectionSets::HandleSelectionListChanged);
 		if (NewSelectionSet->ExportAsJsonString(LastSelectionSetJSON) == false)
 		{
 			UE_LOG(LogControlRig, Error, TEXT("Selection Sets: Error fast JSON export"));
