Control Rig: Fix starting transaction while transacting.

diff --git a/Engine/Plugins/Animation/ControlRig/Source/ControlRigEditor/Private/Sequencer/ControlRigParameterTrackEditor.cpp b/Engine/Plugins/Animation/ControlRig/Source/ControlRigEditor/Private/Sequencer/ControlRigParameterTrackEditor.cpp
--- a/Engine/Plugins/Animation/ControlRig/Source/ControlRigEditor/Private/Sequencer/ControlRigParameterTrackEditor.cpp
+++ b/Engine/Plugins/Animation/ControlRig/Source/ControlRigEditor/Private/Sequencer/ControlRigParameterTrackEditor.cpp
@@ -4205,7 +4205,7 @@ void FControlRigParameterTrackEditor::HandleControlUndoBracket(UControlRig* Subj
 	if(IsInGameThread() && bOpenUndoBracket && ControlUndoBracket == 0)
 	{
 		FScopeLock ScopeLock(&ControlUndoTransactionMutex);
-		ControlUndoTransaction = MakeShareable(new FScopedTransaction(LOCTEXT("KeyMultipleControls", "Auto-Key multiple controls")));
+		ControlUndoTransaction = MakeShareable(new FScopedTransaction(LOCTEXT("KeyMultipleControls", "Auto-Key multiple controls"), !GIsTransacting));
 		ControlChangedDuringUndoBracket = 0;
 	}
 
