Control Rig: Fix dangling pointer in selection set UI.

Text block was keeping a pointer to TMap's value (UAIESelectionSets::ActorsWithSelectionSets::ActorsWithSet). When it reallocates, the UI's pointer becomes dangling.

Also simplified some of the surrounding code by extracting this look-up logic to a separate lambda.

diff --git a/Engine/Plugins/Animation/ControlRig/Source/ControlRigEditor/Private/Sequencer/SelectionSets/SSelectionSets.cpp b/Engine/Plugins/Animation/ControlRig/Source/ControlRigEditor/Private/Sequencer/SelectionSets/SSelectionSets.cpp
--- a/Engine/Plugins/Animation/ControlRig/Source/ControlRigEditor/Private/Sequencer/SelectionSets/SSelectionSets.cpp
+++ b/Engine/Plugins/Animation/ControlRig/Source/ControlRigEditor/Private/Sequencer/SelectionSets/SSelectionSets.cpp
@@ -786,7 +786,7 @@ TSharedRef<SWidget> SSelectionSets::MakeImportExportMenu()
 	return MenuBuilder.MakeWidget();
 }
 
-TSharedRef<SWidget> SSelectionSets::CreateActorMenu(FActorWithSelectionSet& ActorWithSelectionSet)
+TSharedRef<SWidget> SSelectionSets::CreateActorMenu(const FActorWithSelectionSet& ActorWithSelectionSet)
 {
 	constexpr bool bCloseMenuAfterSelection = true;
 	FMenuBuilder MenuBuilder(bCloseMenuAfterSelection, nullptr);
@@ -819,45 +819,35 @@ void SSelectionSets::PopulateActiveActors()
 	if (TSharedPtr<ISequencer> SequencerPtr = UAIESelectionSets::GetSequencerFromAsset())
 	{
 		const float MenuIconSize = FAppStyle::Get().GetFloat("Menu.MenuIconSize");
-
-		auto AddActor = [this, WeakSequencer = SequencerPtr->AsWeak(), MenuIconSize](SWrapBox& Grid, int32 Row, int32 Column, TWeakObjectPtr<AActor>& ActorPtr, FActorWithSelectionSet& ActorWithSeletionSet)
+		
+		
+		auto AddActor = [this, WeakSequencer = SequencerPtr->AsWeak(), MenuIconSize](SWrapBox& Grid, int32 Row, int32 Column, const TWeakObjectPtr<AActor>& ActorPtr, FActorWithSelectionSet& ActorWithSeletionSet)
 			{
+				const auto GetSelectionSetInfo = [WeakSequencer, ActorPtr]()
+				{
+					const TSharedPtr<ISequencer> SequencerPtr = WeakSequencer.Pin();
+					UAIESelectionSets* SelectionSets = SequencerPtr ? UAIESelectionSets::GetSelectionSets(SequencerPtr) : nullptr;
+					FActorWithSelectionSet* ActorWithSeletionSet = SelectionSets ? SelectionSets->ActorsWithSelectionSets.ActorsWithSet.Find(ActorPtr) : nullptr;
+					return ActorWithSeletionSet;
+				};
+				
 				Grid.AddSlot()
 					.HAlign(HAlign_Center)
 					.VAlign(VAlign_Center)
 					[
 						SNew(SCheckBox)
-							.OnGetMenuContent_Lambda([WeakSequencer, ActorPtr,this]()
-								{
-									if (TSharedPtr<ISequencer> SequencerPtr = WeakSequencer.Pin())
-									{
-										if (UAIESelectionSets* SelectionSets = UAIESelectionSets::GetSelectionSets(SequencerPtr))
-										{
-											if (FActorWithSelectionSet* ActorWithSeletionSet = SelectionSets->ActorsWithSelectionSets.ActorsWithSet.Find(ActorPtr))
-											{
-												return CreateActorMenu(*ActorWithSeletionSet);
-											}
-										}
-									}
-									return SNullWidget::NullWidget;
-								})
+							.OnGetMenuContent_Lambda([this, GetSelectionSetInfo]()
+							{
+								const FActorWithSelectionSet* ActorWithSeletionSet = GetSelectionSetInfo();
+								return ActorWithSeletionSet ? CreateActorMenu(*ActorWithSeletionSet) : SNullWidget::NullWidget;
+							})
 							.Style(FAppStyle::Get(), "FilterBar.FilterButton")
-							//	.ToolTipText(FilterToolTip)
-							.IsChecked_Lambda( [WeakSequencer, ActorPtr]() -> ECheckBoxState
-								{
-									if (TSharedPtr<ISequencer> SequencerPtr = WeakSequencer.Pin())
-									{
-										if (UAIESelectionSets* SelectionSets = UAIESelectionSets::GetSelectionSets(SequencerPtr))
-										{
-											if (FActorWithSelectionSet* ActorWithSeletionSet = SelectionSets->ActorsWithSelectionSets.ActorsWithSet.Find(ActorPtr))
-											{
-												bool bEnabled = ActorWithSeletionSet->GetSelected();
-												return bEnabled ? ECheckBoxState::Checked : ECheckBoxState::Unchecked;
-											}
-										}
-									}
-									return ECheckBoxState::Unchecked;
-								})
+							.IsChecked_Lambda( [GetSelectionSetInfo]() -> ECheckBoxState
+							{
+								const FActorWithSelectionSet* ActorWithSeletionSet = GetSelectionSetInfo();
+								const bool bEnabled = ActorWithSeletionSet && ActorWithSeletionSet->GetSelected();
+								return bEnabled ? ECheckBoxState::Checked : ECheckBoxState::Unchecked;
+							})
 							.OnCheckStateChanged_Lambda([this, ActorPtr, WeakSequencer](ECheckBoxState InCheckBoxState)
 								{
 									if (TSharedPtr<ISequencer> SequencerPtr = WeakSequencer.Pin())
@@ -883,8 +873,6 @@ void SSelectionSets::PopulateActiveActors()
 
 									PopulateSelectionSets();
 								})
-							//.CheckBoxContentUsesAutoWidth(false)
-							//.OnGetMenuContent(this, &SFilter::GetRightClickMenuContent)
 							[
 								SNew(SBorder)
 									.Padding(1.0f)
@@ -897,20 +885,16 @@ void SSelectionSets::PopulateActiveActors()
 											[
 												SNew(SImage)
 													.Image(FAppStyle::Get().GetBrush("FilterBar.FilterImage"))
-													.ColorAndOpacity_Lambda([ActorPtr, WeakSequencer]()
+													.ColorAndOpacity_Lambda([GetSelectionSetInfo]()
+													{
+														if (const FActorWithSelectionSet* ActorWithSeletionSet = GetSelectionSetInfo())
 														{
-															if (TSharedPtr<ISequencer> SequencerPtr = WeakSequencer.Pin())
-															{
-																if (UAIESelectionSets* SelectionSets = UAIESelectionSets::GetSelectionSets(SequencerPtr))
-																{
-																	if (FActorWithSelectionSet* ActorWithSeletionSet = SelectionSets->ActorsWithSelectionSets.ActorsWithSet.Find(ActorPtr))
-																	{
-																		return ActorWithSeletionSet->GetSelected() ? FSlateColor(ActorWithSeletionSet->Color) : FAppStyle::Get().GetSlateColor("Colors.Recessed");
-																	}
-																}
-															}
-															return FSlateColor(FLinearColor::White);
-														})
+															return ActorWithSeletionSet->GetSelected()
+																? FSlateColor(ActorWithSeletionSet->Color)
+																: FAppStyle::Get().GetSlateColor("Colors.Recessed");
+														}
+														return FSlateColor(FLinearColor::White);
+													})
 											]
 										+ SHorizontalBox::Slot()
 											.Padding(4.0f, 0.f, 0.f, 0.0f)
@@ -923,41 +907,36 @@ void SSelectionSets::PopulateActiveActors()
 													.HeightOverride(MenuIconSize)
 													[
 														SNew(SImage)
-															.Image_Lambda([ActorPtr, WeakSequencer]() -> const FSlateBrush*
+															.Image_Lambda([GetSelectionSetInfo]() -> const FSlateBrush*
+															{
+																const FActorWithSelectionSet* ActorWithSeletionSet = GetSelectionSetInfo();
+																for (const TWeakObjectPtr<UObject>& Object : ActorWithSeletionSet->Actors)
 																{
-																	if (TSharedPtr<ISequencer> SequencerPtr = WeakSequencer.Pin())
+																	if (UControlRig* ControlRig = Cast<UControlRig>(Object.Get()))
 																	{
-																		if (UAIESelectionSets* SelectionSets = UAIESelectionSets::GetSelectionSets(SequencerPtr))
-																		{
-																			if (FActorWithSelectionSet* ActorWithSeletionSet = SelectionSets->ActorsWithSelectionSets.ActorsWithSet.Find(ActorPtr))
-																			{
-																				for (const TWeakObjectPtr<UObject>& Object : ActorWithSeletionSet->Actors)
-																				{
-																					if (UControlRig* ControlRig = Cast<UControlRig>(Object.Get()))
-																					{
-																						return FSlateIconFinder::FindIconForClass(USkeleton::StaticClass()).GetIcon();
-																					}
-																				}
-																			}
-																		}
+																		return FSlateIconFinder::FindIconForClass(USkeleton::StaticClass()).GetIcon();
 																	}
-																	return FSlateIconFinder::FindIconForClass(AActor::StaticClass()).GetIcon();
-																})
+																}
+																return FSlateIconFinder::FindIconForClass(AActor::StaticClass()).GetIcon();
+															})
 															.ColorAndOpacity(FSlateColor::UseForeground())
 													]
 											]
 										+ SHorizontalBox::Slot()
-											//.Padding(TAttribute<FMargin>(this, &SFilter::GetFilterNamePadding))
 											.Padding(FMargin(4, 1, 4, 1))
 											.VAlign(VAlign_Center)
 											[
 												SNew(STextBlock)
-													.Text_Lambda([&ActorWithSeletionSet]()
-														{
-															return ActorWithSeletionSet.Name;
-														})
-													.IsEnabled_Lambda([&ActorWithSeletionSet] {return ActorWithSeletionSet.GetSelected(); })
-
+													.Text_Lambda([GetSelectionSetInfo]()
+													{
+														const FActorWithSelectionSet* SelectionSet = GetSelectionSetInfo();
+														return SelectionSet ? SelectionSet->Name : FText::GetEmpty();
+													})
+													.IsEnabled_Lambda([GetSelectionSetInfo]
+													{
+														const FActorWithSelectionSet* SelectionSet = GetSelectionSetInfo();
+														return SelectionSet && SelectionSet->GetSelected();
+													})
 											]
 									]
 							]
diff --git a/Engine/Plugins/Animation/ControlRig/Source/ControlRigEditor/Private/Sequencer/SelectionSets/SSelectionSets.h b/Engine/Plugins/Animation/ControlRig/Source/ControlRigEditor/Private/Sequencer/SelectionSets/SSelectionSets.h
--- a/Engine/Plugins/Animation/ControlRig/Source/ControlRigEditor/Private/Sequencer/SelectionSets/SSelectionSets.h
+++ b/Engine/Plugins/Animation/ControlRig/Source/ControlRigEditor/Private/Sequencer/SelectionSets/SSelectionSets.h
@@ -106,7 +106,7 @@ private:
 
 private:
 	
-	TSharedRef<SWidget> CreateActorMenu(FActorWithSelectionSet& ActorWithSelectionSet);
+	TSharedRef<SWidget> CreateActorMenu(const FActorWithSelectionSet& ActorWithSelectionSet);
 	TSharedRef<SWidget> MakeImportExportMenu();
 
 //widgets
diff --git a/Engine/Plugins/Animation/ControlRig/Source/ControlRigEditor/Private/Sequencer/SelectionSets/SelectionSets.cpp b/Engine/Plugins/Animation/ControlRig/Source/ControlRigEditor/Private/Sequencer/SelectionSets/SelectionSets.cpp
--- a/Engine/Plugins/Animation/ControlRig/Source/ControlRigEditor/Private/Sequencer/SelectionSets/SelectionSets.cpp
+++ b/Engine/Plugins/Animation/ControlRig/Source/ControlRigEditor/Private/Sequencer/SelectionSets/SelectionSets.cpp
@@ -462,7 +462,7 @@ TSharedPtr<ISequencer> UAIESelectionSets::GetSequencerFromAsset()
 	return SequencerPtr;
 }
 
-UAIESelectionSets* UAIESelectionSets::GetSelectionSets(TSharedPtr<ISequencer>& SequencerPtr, bool bAddIfDoesNotExist)
+UAIESelectionSets* UAIESelectionSets::GetSelectionSets(const TSharedPtr<ISequencer>& SequencerPtr, bool bAddIfDoesNotExist)
 {
 	if (SequencerPtr.IsValid() == false)
 	{
diff --git a/Engine/Plugins/Animation/ControlRig/Source/ControlRigEditor/Private/Sequencer/SelectionSets/SelectionSets.h b/Engine/Plugins/Animation/ControlRig/Source/ControlRigEditor/Private/Sequencer/SelectionSets/SelectionSets.h
--- a/Engine/Plugins/Animation/ControlRig/Source/ControlRigEditor/Private/Sequencer/SelectionSets/SelectionSets.h
+++ b/Engine/Plugins/Animation/ControlRig/Source/ControlRigEditor/Private/Sequencer/SelectionSets/SelectionSets.h
@@ -312,7 +312,7 @@ public:
 	bool AddChildToItem(const FGuid& InParentGuid, const FGuid& InChildGuid);
 	bool RemoveChildFromItem(const FGuid& InParentGuid, const FGuid& InChildGuid);
 
-	static UAIESelectionSets* GetSelectionSets(TSharedPtr<ISequencer>& InSequencer, bool bAddIfDoesNotExist = false);
+	static UAIESelectionSets* GetSelectionSets(const TSharedPtr<ISequencer>& InSequencer, bool bAddIfDoesNotExist = false);
 	static UAIESelectionSets* GetSelectionSets(ULevelSequence* InLevelSequence, bool bAddIfDoesNotExist = false);
 	static 	TSharedPtr<ISequencer> GetSequencerFromAsset();
 
