Control Rig: Avoid crash when right click on a GetAnimChannel node
#rnx

diff --git a/Engine/Plugins/Animation/ControlRig/Source/ControlRig/Private/Units/Hierarchy/RigUnit_ControlChannelFromItem.cpp b/Engine/Plugins/Animation/ControlRig/Source/ControlRig/Private/Units/Hierarchy/RigUnit_ControlChannelFromItem.cpp
--- a/Engine/Plugins/Animation/ControlRig/Source/ControlRig/Private/Units/Hierarchy/RigUnit_ControlChannelFromItem.cpp
+++ b/Engine/Plugins/Animation/ControlRig/Source/ControlRig/Private/Units/Hierarchy/RigUnit_ControlChannelFromItem.cpp
@@ -8,6 +8,11 @@
 FRigUnit_GetBoolAnimationChannelFromItem_Execute()
 {
 	Value = false;
+
+	if (!IsValid(ExecuteContext.Hierarchy))
+	{
+		return;
+	}
 	
 	if(const FRigControlElement* ChannelElement = ExecuteContext.Hierarchy->Find<FRigControlElement>(Item))
 	{
@@ -26,6 +31,11 @@ FRigUnit_GetBoolAnimationChannelFromItem_Execute()
 FRigUnit_GetFloatAnimationChannelFromItem_Execute()
 {
 	Value = 0.f;
+
+	if (!IsValid(ExecuteContext.Hierarchy))
+	{
+		return;
+	}
 	
 	if(const FRigControlElement* ChannelElement = ExecuteContext.Hierarchy->Find<FRigControlElement>(Item))
 	{
@@ -44,6 +54,11 @@ FRigUnit_GetFloatAnimationChannelFromItem_Execute()
 FRigUnit_GetIntAnimationChannelFromItem_Execute()
 {
 	Value = 0;
+
+	if (!IsValid(ExecuteContext.Hierarchy))
+	{
+		return;
+	}
 	
 	if(const FRigControlElement* ChannelElement = ExecuteContext.Hierarchy->Find<FRigControlElement>(Item))
 	{
@@ -62,6 +77,11 @@ FRigUnit_GetIntAnimationChannelFromItem_Execute()
 FRigUnit_GetVector2DAnimationChannelFromItem_Execute()
 {
 	Value = FVector2D::ZeroVector;
+
+	if (!IsValid(ExecuteContext.Hierarchy))
+	{
+		return;
+	}
 	
 	if(const FRigControlElement* ChannelElement = ExecuteContext.Hierarchy->Find<FRigControlElement>(Item))
 	{
@@ -80,6 +100,11 @@ FRigUnit_GetVector2DAnimationChannelFromItem_Execute()
 FRigUnit_GetVectorAnimationChannelFromItem_Execute()
 {
 	Value = FVector::ZeroVector;
+
+	if (!IsValid(ExecuteContext.Hierarchy))
+	{
+		return;
+	}
 	
 	if(const FRigControlElement* ChannelElement = ExecuteContext.Hierarchy->Find<FRigControlElement>(Item))
 	{
@@ -98,6 +123,11 @@ FRigUnit_GetVectorAnimationChannelFromItem_Execute()
 FRigUnit_GetRotatorAnimationChannelFromItem_Execute()
 {
 	Value = FRotator::ZeroRotator;
+
+	if (!IsValid(ExecuteContext.Hierarchy))
+	{
+		return;
+	}
 	
 	if(const FRigControlElement* ChannelElement = ExecuteContext.Hierarchy->Find<FRigControlElement>(Item))
 	{
@@ -116,6 +146,11 @@ FRigUnit_GetRotatorAnimationChannelFromItem_Execute()
 FRigUnit_GetTransformAnimationChannelFromItem_Execute()
 {
 	Value = FTransform::Identity;
+
+	if (!IsValid(ExecuteContext.Hierarchy))
+	{
+		return;
+	}
 	
 	if(const FRigControlElement* ChannelElement = ExecuteContext.Hierarchy->Find<FRigControlElement>(Item))
 	{
@@ -138,6 +173,11 @@ FRigUnit_GetTransformAnimationChannelFromItem_Execute()
 
 FRigUnit_SetBoolAnimationChannelFromItem_Execute()
 {
+	if (!IsValid(ExecuteContext.Hierarchy))
+	{
+		return;
+	}
+	
 	if(const FRigControlElement* ChannelElement = ExecuteContext.Hierarchy->Find<FRigControlElement>(Item))
 	{
 		if(ChannelElement->Settings.ControlType == ERigControlType::Bool)
@@ -154,6 +194,11 @@ FRigUnit_SetBoolAnimationChannelFromItem_Execute()
 
 FRigUnit_SetFloatAnimationChannelFromItem_Execute()
 {
+	if (!IsValid(ExecuteContext.Hierarchy))
+	{
+		return;
+	}
+	
 	if(const FRigControlElement* ChannelElement = ExecuteContext.Hierarchy->Find<FRigControlElement>(Item))
 	{
 		if(ChannelElement->Settings.ControlType == ERigControlType::Float || ChannelElement->Settings.ControlType == ERigControlType::ScaleFloat)
@@ -170,6 +215,11 @@ FRigUnit_SetFloatAnimationChannelFromItem_Execute()
 
 FRigUnit_SetIntAnimationChannelFromItem_Execute()
 {
+	if (!IsValid(ExecuteContext.Hierarchy))
+	{
+		return;
+	}
+	
 	if(const FRigControlElement* ChannelElement = ExecuteContext.Hierarchy->Find<FRigControlElement>(Item))
 	{
 		if(ChannelElement->Settings.ControlType == ERigControlType::Integer)
@@ -186,6 +236,11 @@ FRigUnit_SetIntAnimationChannelFromItem_Execute()
 
 FRigUnit_SetVector2DAnimationChannelFromItem_Execute()
 {
+	if (!IsValid(ExecuteContext.Hierarchy))
+	{
+		return;
+	}
+	
 	if(const FRigControlElement* ChannelElement = ExecuteContext.Hierarchy->Find<FRigControlElement>(Item))
 	{
 		if(ChannelElement->Settings.ControlType == ERigControlType::Vector2D)
@@ -202,6 +257,11 @@ FRigUnit_SetVector2DAnimationChannelFromItem_Execute()
 
 FRigUnit_SetVectorAnimationChannelFromItem_Execute()
 {
+	if (!IsValid(ExecuteContext.Hierarchy))
+	{
+		return;
+	}
+	
 	if(const FRigControlElement* ChannelElement = ExecuteContext.Hierarchy->Find<FRigControlElement>(Item))
 	{
 		if(ChannelElement->Settings.ControlType == ERigControlType::Position || ChannelElement->Settings.ControlType == ERigControlType::Scale)
@@ -218,6 +278,11 @@ FRigUnit_SetVectorAnimationChannelFromItem_Execute()
 
 FRigUnit_SetRotatorAnimationChannelFromItem_Execute()
 {
+	if (!IsValid(ExecuteContext.Hierarchy))
+	{
+		return;
+	}
+	
 	if(const FRigControlElement* ChannelElement = ExecuteContext.Hierarchy->Find<FRigControlElement>(Item))
 	{
 		if(ChannelElement->Settings.ControlType == ERigControlType::Rotator)
@@ -234,6 +299,11 @@ FRigUnit_SetRotatorAnimationChannelFromItem_Execute()
 
 FRigUnit_SetTransformAnimationChannelFromItem_Execute()
 {
+	if (!IsValid(ExecuteContext.Hierarchy))
+	{
+		return;
+	}
+	
 	if(const FRigControlElement* ChannelElement = ExecuteContext.Hierarchy->Find<FRigControlElement>(Item))
 	{
 		if(ChannelElement->Settings.ControlType == ERigControlType::Transform)
