Anim in Engine: Backout changes that can cause issues when trying to edit values via anim details

diff --git a/Engine/Plugins/Animation/ControlRig/Source/ControlRigEditor/Private/AnimDetails/Widgets/SAnimDetailsValueNumeric.cpp b/Engine/Plugins/Animation/ControlRig/Source/ControlRigEditor/Private/AnimDetails/Widgets/SAnimDetailsValueNumeric.cpp
--- a/Engine/Plugins/Animation/ControlRig/Source/ControlRigEditor/Private/AnimDetails/Widgets/SAnimDetailsValueNumeric.cpp
+++ b/Engine/Plugins/Animation/ControlRig/Source/ControlRigEditor/Private/AnimDetails/Widgets/SAnimDetailsValueNumeric.cpp
@@ -29,7 +29,7 @@ namespace UE::ControlRigEditor
 			TOptional<NumericType> SliderMinValue;
 			TOptional<NumericType> SliderMaxValue;
 			double SliderExponent = 1.0;
-			double Delta = 1.0;
+			double Delta = 0.0;
 			float ShiftMultiplier = 10.f;
 			float CtrlMultiplier = 0.1f;
 			bool SupportDynamicSliderMaxValue = false;
@@ -257,8 +257,8 @@ namespace UE::ControlRigEditor
 			.SupportDynamicSliderMaxValue(PropertyMetaData.SupportDynamicSliderMaxValue)
 			.MinValue(PropertyMetaData.MinValue)
 			.MaxValue(PropertyMetaData.MaxValue)
-			.MinSliderValue(PropertyMetaData.SliderMinValue)
-			.MaxSliderValue(PropertyMetaData.SliderMaxValue)
+			.MaxSliderValue(PropertyMetaData.SliderMinValue)
+			.MinSliderValue(PropertyMetaData.SliderMaxValue)
 			.ShiftMultiplier(PropertyMetaData.ShiftMultiplier)
 			.CtrlMultiplier(PropertyMetaData.CtrlMultiplier)
 			.SliderExponent(PropertyMetaData.SliderExponent)
