Control Rig: The Control Rig menu now shows directly in the viewport toolbar instead of being a submenu in the Show menu (eyeball icon)

diff --git a/Engine/Plugins/Animation/ControlRig/Source/ControlRigEditor/Private/Editor/ControlRigEditor.cpp b/Engine/Plugins/Animation/ControlRig/Source/ControlRigEditor/Private/Editor/ControlRigEditor.cpp
--- a/Engine/Plugins/Animation/ControlRig/Source/ControlRigEditor/Private/Editor/ControlRigEditor.cpp
+++ b/Engine/Plugins/Animation/ControlRig/Source/ControlRigEditor/Private/Editor/ControlRigEditor.cpp
@@ -4,6 +4,7 @@
 #include "Modules/ModuleManager.h"
 #include "ControlRigEditorModule.h"
 #include "ControlRigBlueprintLegacy.h"
+#include "Editor/ControlRigViewportToolbarExtensionControlRig.h"
 #include "SBlueprintEditorToolbar.h"
 #include "Editor/ControlRigEditorMode.h"
 #include "SEnumCombo.h"
@@ -136,118 +137,6 @@
 
 #define LOCTEXT_NAMESPACE "ControlRigEditor"
 
-namespace UE::ControlRigEditor::Private
-{
-TOptional<float> GetToolbarAxesScale()
-{
-	if (const UControlRigEditModeSettings* Settings = GetDefault<UControlRigEditModeSettings>())
-	{
-		return Settings->AxisScale;
-	}
-	return 0.0f;
-}
-
-void OnToolbarAxesScaleChanged(float InValue)
-{
-	if (UControlRigEditModeSettings* Settings = GetMutableDefault<UControlRigEditModeSettings>())
-	{
-		Settings->AxisScale = InValue;
-		Settings->SaveConfig();
-	}
-}
-
-FControlRigAssetInterfacePtr GetControlRigBlueprintFromContext(const UUnrealEdViewportToolbarContext* InContext)
-{
-	if (InContext)
-	{
-		if (TSharedPtr<SEditorViewport> Viewport = InContext->Viewport.Pin())
-		{
-			if (const TSharedPtr<FAnimationViewportClient>& AnimViewportClient =
-					StaticCastSharedPtr<FAnimationViewportClient>(Viewport->GetViewportClient()))
-			{
-#if WITH_RIGVMLEGACYEDITOR
-				// Check if we are editing UControlRigBlueprints, otherwise this is not a ControlRigEditor
-				if (const TArray<UObject*>* EditingObjects = AnimViewportClient->GetAssetEditorToolkit()->GetObjectsCurrentlyBeingEdited())
-				{
-					for (UObject* Object : *EditingObjects)
-					{
-						if (!Object->Implements<UControlRigAssetInterface>())
-						{
-							return nullptr;
-						}
-					}
-				}
-
-				if (AnimViewportClient->GetAssetEditorToolkit()->IsBlueprintEditor())
-				{
-					const TSharedRef<FBlueprintEditor>& BlueprintEditor =
-						StaticCastSharedRef<FBlueprintEditor>(AnimViewportClient->GetAssetEditorToolkit());
-
-					return BlueprintEditor->GetBlueprintObj();
-				}
-#endif
-				const TSharedRef<FControlRigEditor>& Editor =
-					StaticCastSharedRef<FControlRigEditor>(AnimViewportClient->GetAssetEditorToolkit());
-
-				return Editor->GetRigVMAssetInterface().GetObject();
-			}
-		}
-	}
-
-	return nullptr;
-}
-
-void FillControlRigMenu(UToolMenu* InMenu)
-{
-	if (const UUnrealEdViewportToolbarContext* const Context = InMenu->FindContext<UUnrealEdViewportToolbarContext>())
-	{
-		if (const FControlRigAssetInterfacePtr ControlRigBlueprint = GetControlRigBlueprintFromContext(Context))
-		{
-			FToolMenuSection& ControlRigSection = InMenu->FindOrAddSection("ControlRig", LOCTEXT("ControlRig_Label", "Control Rig"));
-			ControlRigSection.AddMenuEntry(FControlRigEditorCommands::Get().ToggleControlVisibility);
-			ControlRigSection.AddMenuEntry(FControlRigEditorCommands::Get().ToggleControlsAsOverlay);
-			ControlRigSection.AddMenuEntry(FControlRigEditorCommands::Get().ToggleDrawNulls);
-			ControlRigSection.AddMenuEntry(FControlRigEditorCommands::Get().ToggleDrawSockets);
-			ControlRigSection.AddMenuEntry(FControlRigEditorCommands::Get().ToggleDrawAxesOnSelection);
-
-			// Asset Axes Custom Widget
-			// clang-format off
-			TSharedRef<SWidget> AxesScaleEntryWidget = SNew(SBox)
-				.HAlign(HAlign_Right)
-				[
-					SNew(SBox)
-					.Padding(FMargin(4.0f, 0.0f, 0.0f, 0.0f))
-					.WidthOverride(100.0f)
-					[
-						SNew(SNumericEntryBox<float>)
-						.Font(FAppStyle::GetFontStyle(TEXT("MenuItem.Font")))
-						.AllowSpin(true)
-						.MinSliderValue(0.0f)
-						.MaxSliderValue(100.0f)
-						.Value_Static(&UE::ControlRigEditor::Private::GetToolbarAxesScale)
-						.OnValueChanged_Static(&UE::ControlRigEditor::Private::OnToolbarAxesScaleChanged)
-						.ToolTipText(LOCTEXT("ControlRigAxesScaleToolTip", "Scale of axes drawn for selected rig elements"))
-					]
-				];
-			// clang-format on
-
-			ControlRigSection.AddEntry(
-				FToolMenuEntry::InitWidget("AxesScale", AxesScaleEntryWidget, LOCTEXT("ControlRigAxesScale", "Axes Scale"))
-			);
-
-			// This entry only for Modular Rig
-			if (ControlRigBlueprint->IsModularRig())
-			{
-				FToolMenuSection& ModularRigSection =
-					InMenu->FindOrAddSection("ModularRig", LOCTEXT("ModularRig_Label", "Modular Rig"));
-				ModularRigSection.AddMenuEntry(FControlRigEditorCommands::Get().ToggleSchematicViewportVisibility);
-			}
-		}
-	}
-}
-
-} // namespace UE::ControlRigEditor::Private
-
 TAutoConsoleVariable<bool> CVarControlRigShowTestingToolbar(TEXT("ControlRig.Test.EnableTestingToolbar"), false, TEXT("When true we'll show the testing toolbar in Control Rig Editor."));
 TAutoConsoleVariable<bool> CVarShowSchematicPanelOverlay(TEXT("ControlRig.Preview.ShowSchematicPanelOverlay"), true, TEXT("When true we'll add an overlay to the persona viewport to show modular rig information."));
 
@@ -2921,38 +2810,7 @@ void FControlRigBaseEditor::HandleViewportCreated(const TSharedRef<class IPerson
 		ChangeShapeTransformNotificationOptions
 	);
 
-	if (UToolMenu* const Submenu = UToolMenus::Get()->ExtendMenu("AnimationEditor.ViewportToolbar.Show"))
-	{
-		FToolMenuSection& ControlRigSection = Submenu->FindOrAddSection(
-			"AnimViewportSceneElements", LOCTEXT("CharacterMenu_SceneElements", "Scene Elements")
-		);
-
-		ControlRigSection.AddDynamicEntry(
-			"ControlRigSubmenu",
-			FNewToolMenuSectionDelegate::CreateLambda(
-				[](FToolMenuSection& InSection)
-				{
-					if (UUnrealEdViewportToolbarContext* Context = InSection.FindContext<UUnrealEdViewportToolbarContext>())
-					{
-						const bool bIsControlRigEditor = !!UE::ControlRigEditor::Private::GetControlRigBlueprintFromContext(Context);
-
-						// This Submenu needs to be visible only for Control Rig Viewport Toolbars
-						if (bIsControlRigEditor)
-						{
-							InSection.AddSubMenu(
-								"ControlRigSubmenu",
-								LOCTEXT("ControlRigSubmenuLabel", "Control Rig"),
-								LOCTEXT("ControlRigSubmenuLabelToolTip", "Control Rig Options"),
-								FNewToolMenuChoice(
-									FNewToolMenuDelegate::CreateStatic(&UE::ControlRigEditor::Private::FillControlRigMenu)
-								)
-							);
-						}
-					}
-				}
-			)
-		);
-	}
+	UE::ControlRigEditor::PopulateControlRigViewportToolbarControlRigSubmenu("AnimationEditor.ViewportToolbar");
 
 	auto GetBorderColorAndOpacity = [this]()
 	{
diff --git a/Engine/Plugins/Animation/ControlRig/Source/ControlRigEditor/Private/Editor/ControlRigLegacyEditor.cpp b/Engine/Plugins/Animation/ControlRig/Source/ControlRigEditor/Private/Editor/ControlRigLegacyEditor.cpp
--- a/Engine/Plugins/Animation/ControlRig/Source/ControlRigEditor/Private/Editor/ControlRigLegacyEditor.cpp
+++ b/Engine/Plugins/Animation/ControlRig/Source/ControlRigEditor/Private/Editor/ControlRigLegacyEditor.cpp
@@ -3,10 +3,12 @@
 #include "Editor/ControlRigLegacyEditor.h"
 
 #if WITH_RIGVMLEGACYEDITOR
+#include "ControlRigContextMenuContext.h"
 #include "ControlRigEditorMode.h"
 #include "ControlRigEditorStyle.h"
 #include "IPersonaPreviewScene.h"
 #include "IPersonaToolkit.h"
+#include "ToolMenuContext.h"
 
 
 FControlRigLegacyEditor::FControlRigLegacyEditor()
@@ -124,4 +126,17 @@ bool FControlRigLegacyEditor::NewDocument_IsVisibleForType(FBlueprintEditor::ECr
 	return true;
 }
 
+void FControlRigLegacyEditor::InitToolMenuContext(FToolMenuContext& MenuContext)
+{
+	FRigVMLegacyEditor::InitToolMenuContext(MenuContext);
+
+	if (FRigVMAssetInterfacePtr RigVMAsset = GetRigVMAssetInterface())
+	{
+		UControlRigContextMenuContext* ContextObject = NewObject<UControlRigContextMenuContext>();
+		ContextObject->Init(SharedThis(this));
+
+		MenuContext.AddObject(ContextObject);
+	}
+}
+
 #endif
diff --git a/Engine/Plugins/Animation/ControlRig/Source/ControlRigEditor/Private/Editor/ControlRigLegacyEditor.h b/Engine/Plugins/Animation/ControlRig/Source/ControlRigEditor/Private/Editor/ControlRigLegacyEditor.h
--- a/Engine/Plugins/Animation/ControlRig/Source/ControlRigEditor/Private/Editor/ControlRigLegacyEditor.h
+++ b/Engine/Plugins/Animation/ControlRig/Source/ControlRigEditor/Private/Editor/ControlRigLegacyEditor.h
@@ -4,8 +4,11 @@
 
 #if WITH_RIGVMLEGACYEDITOR
 #include "ControlRigEditor.h"
+#include "Editor/RigVMEditor.h"
 #include "Types/SlateVector2.h"
 
+struct FToolMenuContext;
+
 class FControlRigLegacyEditor : public IControlRigLegacyEditor, public FControlRigBaseEditor
 {
 public:
@@ -51,7 +54,6 @@ public:
 	virtual URigVMGraph* GetFocusedModel() const { return IControlRigLegacyEditor::GetFocusedModel(); }
 	virtual TArray<FName> GetLastEventQueue() const override { return LastEventQueue; }
 
-
 	// FRigVMEditorBase interface
 	virtual UObject* GetOuterForHost() const override { return GetOuterForHostImpl();}
 	virtual UObject* GetOuterForHostSuper() const override { return IControlRigLegacyEditor::GetOuterForHost(); }
@@ -115,6 +117,9 @@ public:
 
 	
 protected:
+	//~ Begin IRigVMEditor interface
+	virtual void InitToolMenuContext(FToolMenuContext& MenuContext) override;
+	//~ End IRigVMEditor interface
 
 	virtual void BindCommands() override { return FControlRigBaseEditor::BindCommandsImpl(); }
 	virtual void BindCommandsSuper() override { return IControlRigLegacyEditor::BindCommands(); }
diff --git a/Engine/Plugins/Animation/ControlRig/Source/ControlRigEditor/Private/Editor/ControlRigNewEditor.cpp b/Engine/Plugins/Animation/ControlRig/Source/ControlRigEditor/Private/Editor/ControlRigNewEditor.cpp
--- a/Engine/Plugins/Animation/ControlRig/Source/ControlRigEditor/Private/Editor/ControlRigNewEditor.cpp
+++ b/Engine/Plugins/Animation/ControlRig/Source/ControlRigEditor/Private/Editor/ControlRigNewEditor.cpp
@@ -2,10 +2,12 @@
 
 #include "Editor/ControlRigNewEditor.h"
 
+#include "ControlRigContextMenuContext.h"
 #include "ControlRigEditorMode.h"
 #include "ControlRigEditorStyle.h"
 #include "IPersonaPreviewScene.h"
 #include "IPersonaToolkit.h"
+#include "ToolMenuContext.h"
 
 
 FControlRigEditor::FControlRigEditor()
@@ -137,3 +139,16 @@ void FControlRigEditor::Tick(float DeltaTime)
 	PreviewScene.UpdateCaptureContents();
 	FControlRigBaseEditor::TickImpl(DeltaTime);
 }
+
+void FControlRigEditor::InitToolMenuContext(FToolMenuContext& MenuContext)
+{
+	FRigVMNewEditor::InitToolMenuContext(MenuContext);
+
+	if (FRigVMAssetInterfacePtr RigVMAsset = GetRigVMAssetInterface())
+	{
+		UControlRigContextMenuContext* ContextObject = NewObject<UControlRigContextMenuContext>();
+		ContextObject->Init(SharedThis(this));
+
+		MenuContext.AddObject(ContextObject);
+	}
+}
diff --git a/Engine/Plugins/Animation/ControlRig/Source/ControlRigEditor/Private/Editor/ControlRigNewEditor.h b/Engine/Plugins/Animation/ControlRig/Source/ControlRigEditor/Private/Editor/ControlRigNewEditor.h
--- a/Engine/Plugins/Animation/ControlRig/Source/ControlRigEditor/Private/Editor/ControlRigNewEditor.h
+++ b/Engine/Plugins/Animation/ControlRig/Source/ControlRigEditor/Private/Editor/ControlRigNewEditor.h
@@ -3,9 +3,12 @@
 #pragma once
 
 #include "ControlRigEditor.h"
+#include "Editor/RigVMEditor.h"
 #include "Editor/SRigVMDetailsInspector.h"
 #include "PreviewScene.h"
 
+struct FToolMenuContext;
+
 class FControlRigEditor : public IControlRigNewEditor, public FControlRigBaseEditor
 {
 public:
@@ -51,7 +54,6 @@ public:
 	virtual URigVMGraph* GetFocusedModel() const { return IControlRigNewEditor::GetFocusedModel(); }
 	virtual TArray<FName> GetLastEventQueue() const override { return LastEventQueue; }
 
-
 	// FRigVMEditorBase interface
 	virtual UObject* GetOuterForHost() const override { return GetOuterForHostImpl();}
 	virtual UObject* GetOuterForHostSuper() const override { return IControlRigNewEditor::GetOuterForHost(); }
@@ -132,6 +134,10 @@ protected:
 	// FBaseToolKit overrides
 	virtual void CreateEditorModeManager() override { return FControlRigBaseEditor::CreateEditorModeManagerImpl(); }
 
+	//~ Begin IRigVMEditor interface
+	virtual void InitToolMenuContext(FToolMenuContext& MenuContext) override;
+	//~ End IRigVMEditor interface
+
 private:
 	/** Fill the toolbar with content */
 	virtual void FillToolbar(FToolBarBuilder& ToolbarBuilder, bool bEndSection = true) override { return FControlRigBaseEditor::FillToolbarImpl(ToolbarBuilder, bEndSection); }
diff --git a/Engine/Plugins/Animation/ControlRig/Source/ControlRigEditor/Private/Editor/ControlRigViewportToolbarExtensionControlRig.cpp b/Engine/Plugins/Animation/ControlRig/Source/ControlRigEditor/Private/Editor/ControlRigViewportToolbarExtensionControlRig.cpp
new file mode 100644
--- /dev/null
+++ b/Engine/Plugins/Animation/ControlRig/Source/ControlRigEditor/Private/Editor/ControlRigViewportToolbarExtensionControlRig.cpp
@@ -0,0 +1,148 @@
+// Copyright Epic Games, Inc. All Rights Reserved.
+
+#include "ControlRigViewportToolbarExtensionControlRig.h"
+
+#include "AnimationEditorViewportClient.h"
+#include "BlueprintEditor.h"
+#include "ControlRigContextMenuContext.h"
+#include "ControlRigEditorStyle.h"
+#include "EditMode/ControlRigEditModeSettings.h"
+#include "Editor/ControlRigEditor.h"
+#include "Editor/ControlRigEditorCommands.h"
+#include "Editor/ControlRigNewEditor.h"
+#include "SEditorViewport.h"
+#include "Templates/SharedPointer.h"
+#include "ToolMenu.h"
+#include "ToolMenuEntry.h"
+#include "ToolMenus.h"
+#include "ToolMenuSection.h"
+#include "UObject/ScriptInterface.h"
+#include "UObject/WeakInterfacePtr.h"
+#include "ViewportToolbar/UnrealEdViewportToolbarContext.h"
+#include "Widgets/Input/SNumericEntryBox.h"
+
+#define LOCTEXT_NAMESPACE "ControlRigViewportToolbarExtensionControlRig"
+
+namespace UE::ControlRigEditor
+{
+	/** Creates the conrol rig entry */
+	namespace Private
+	{
+		/** Creates a widget to edit the axis scale in editor */
+		TSharedRef<SWidget> MakeAxisScaleWidget()
+		{
+			const FText Tooltip = LOCTEXT("ControlRigAxesScaleToolTip", "Scale of axes drawn for selected rig elements");
+
+			return 
+				SNew(SBox)
+				.HAlign(HAlign_Right)
+				[
+					SNew(SBox)
+					.Padding(FMargin(4.0f, 0.0f, 0.0f, 0.0f))
+					.WidthOverride(100.0f)
+					[
+						SNew(SNumericEntryBox<float>)
+						.Font(FAppStyle::GetFontStyle(TEXT("MenuItem.Font")))
+						.AllowSpin(true)
+						.MinSliderValue(0.0f)
+						.MaxSliderValue(100.0f)
+						.Value_Lambda([]()
+							{										
+								return GetDefault<UControlRigEditModeSettings>()->AxisScale;
+							})
+						.OnValueChanged_Lambda([](float Value)
+						{
+							UControlRigEditModeSettings* Settings = GetMutableDefault<UControlRigEditModeSettings>();
+							Settings->AxisScale = FMath::Max(0.f, Value);
+							Settings->SaveConfig();
+						})
+						.ToolTipText(Tooltip)
+					]
+				];
+		}
+
+		/** Adds the control rig sub menu to the toolbar */
+		FToolMenuEntry CreateControlRigEntry(const IControlRigAssetInterface& Asset)
+		{
+			const TWeakInterfacePtr<const IControlRigAssetInterface> WeakAsset = &Asset;
+
+			FToolMenuEntry Entry = FToolMenuEntry::InitSubMenu(
+				"ControlRig",
+				LOCTEXT("ControlRigSubMenuLabel", "Control Rig"),
+				LOCTEXT("ControlRigSubMenuTooltip", "Control Rig related settings"),
+				FNewToolMenuDelegate::CreateLambda(
+					[WeakAsset](UToolMenu* SubMenu) -> void
+					{
+						FToolMenuSection& ControlRigSection = 
+							SubMenu->FindOrAddSection(
+								"ControlRig", 
+								LOCTEXT("ControlRigLabel", "Control Rig"));
+
+						const FControlRigEditorCommands& Commands = FControlRigEditorCommands::Get();
+
+						ControlRigSection.AddMenuEntry(Commands.ToggleControlVisibility);
+						ControlRigSection.AddMenuEntry(Commands.ToggleControlsAsOverlay);
+						ControlRigSection.AddMenuEntry(Commands.ToggleDrawNulls);
+						ControlRigSection.AddMenuEntry(Commands.ToggleDrawSockets);
+						ControlRigSection.AddMenuEntry(Commands.ToggleDrawAxesOnSelection);
+
+						ControlRigSection.AddSeparator("SliderSection");
+
+						ControlRigSection.AddEntry(FToolMenuEntry::InitWidget(
+							"AxesScale",
+							MakeAxisScaleWidget(),
+							LOCTEXT("ControlRigAxesScale", "Axes Scale")));
+
+						// Entries for Modular Rigs
+						if (WeakAsset.IsValid() && WeakAsset.Get()->IsModularRig())
+						{
+							FToolMenuSection& ModularRigSection = SubMenu->FindOrAddSection(
+								"ModularRig",
+								LOCTEXT("ModularRig_Label", "Modular Rig"));
+
+							FToolMenuSection& MRSection = SubMenu->FindOrAddSection(
+								"ControlRig",
+								LOCTEXT("ControlRigLabel", "Control Rig"));
+
+							ModularRigSection.AddMenuEntry(Commands.ToggleSchematicViewportVisibility);
+						}
+					})
+				);
+
+			Entry.InsertPosition = FToolMenuInsert("Show", EToolMenuInsertType::After);
+			Entry.Icon = FSlateIcon(FControlRigEditorStyle::Get().GetStyleSetName(), "ControlRig.Editor.TabIcon");
+			Entry.ToolBarData.LabelOverride = FText::GetEmpty();
+
+			return Entry;
+		}
+	}
+	//~ namespace Private
+
+	void PopulateControlRigViewportToolbarControlRigSubmenu(const FName& MenuName)
+	{
+		TWeakObjectPtr<UToolMenu> Toolbar = UToolMenus::Get()->ExtendMenu(MenuName);
+		FToolMenuSection* RightSection = Toolbar.IsValid() ? Toolbar->FindSection("Right") : nullptr;
+		if (!Toolbar.IsValid() ||
+			!RightSection)
+		{
+			return;
+		}
+
+		RightSection->AddDynamicEntry(
+			"ControlRig",
+			FNewToolMenuSectionDelegate::CreateLambda([](FToolMenuSection& Section)
+				{		
+					// Only for control rig assets
+					const UControlRigContextMenuContext* Context = Section.FindContext<UControlRigContextMenuContext>();
+					const TScriptInterface<IControlRigAssetInterface> ControlRigAsset = Context ? Context->GetControlRigAssetInterface() : nullptr;
+					if (Context &&
+						ControlRigAsset)
+					{
+						Section.AddEntry(Private::CreateControlRigEntry(*ControlRigAsset));
+					}
+				})
+		);	
+	}
+}
+
+#undef LOCTEXT_NAMESPACE
diff --git a/Engine/Plugins/Animation/ControlRig/Source/ControlRigEditor/Private/Editor/ControlRigViewportToolbarExtensionControlRig.h b/Engine/Plugins/Animation/ControlRig/Source/ControlRigEditor/Private/Editor/ControlRigViewportToolbarExtensionControlRig.h
new file mode 100644
--- /dev/null
+++ b/Engine/Plugins/Animation/ControlRig/Source/ControlRigEditor/Private/Editor/ControlRigViewportToolbarExtensionControlRig.h
@@ -0,0 +1,11 @@
+// Copyright Epic Games, Inc. All Rights Reserved.
+
+#pragma once
+
+class FName;
+
+namespace UE::ControlRigEditor
+{
+	/** Populates the Control Rig submenu in the viewport toolbar */
+	void PopulateControlRigViewportToolbarControlRigSubmenu(const FName& FName);
+}
