Fixed ControlRig spawn control node details panel not generating the shape name list as a combo box

diff --git a/Engine/Plugins/Animation/ControlRig/Source/ControlRigEditor/Private/AnimDetails/Customizations/AnimDetailsOverrideDetails.cpp b/Engine/Plugins/Animation/ControlRig/Source/ControlRigEditor/Private/AnimDetails/Customizations/AnimDetailsOverrideDetails.cpp
--- a/Engine/Plugins/Animation/ControlRig/Source/ControlRigEditor/Private/AnimDetails/Customizations/AnimDetailsOverrideDetails.cpp
+++ b/Engine/Plugins/Animation/ControlRig/Source/ControlRigEditor/Private/AnimDetails/Customizations/AnimDetailsOverrideDetails.cpp
@@ -14,11 +14,17 @@
 #include "IDetailGroup.h"
 #include "IPropertyUtilities.h"
 #include "IDetailChildrenBuilder.h"
+#include "ControlRigShapeNameList.h"
 
 #define LOCTEXT_NAMESPACE "AnimDetailsProxyDetails"
 
 namespace UE::ControlRigEditor
 {
+	FAnimDetailsOverrideDetails::FAnimDetailsOverrideDetails()
+		: ControlRigShapeNameList(MakeShared<FControlRigShapeNameList>())
+	{
+	}
+
 	void FAnimDetailsOverrideDetails::CustomizeHeader(TSharedRef<IPropertyHandle> StructPropertyHandle, FDetailWidgetRow& HeaderRow,
 		IPropertyTypeCustomizationUtils& StructCustomizationUtils)
 	{
@@ -49,76 +55,24 @@ namespace UE::ControlRigEditor
 		{
 			return;
 		}
-		const TMap<FString, FString>& LibraryNameMap = ControlRig->GetShapeLibraryNameMap();
-		TArray<TSoftObjectPtr<UControlRigShapeLibrary>> ShapeLibraries = ControlRig->GetShapeLibraries();
-		for(const TSoftObjectPtr<UControlRigShapeLibrary>& ShapeLibrary : ShapeLibraries)
-		{
-			if(ShapeLibrary.IsNull() || !ShapeLibrary.IsValid())
-			{
-				(void)ShapeLibrary.LoadSynchronous();
-			}
-			if(ShapeLibrary.IsNull() || !ShapeLibrary.IsValid())
-			{
-				continue;
-			}
-			const bool bUseNameSpace = ShapeLibraries.Num() > 1;
-			FString LibraryName = ShapeLibrary->GetName();
-			if(const FString* RemappedName = LibraryNameMap.Find(LibraryName))
-			{
-				LibraryName = *RemappedName;
-			}
-			
-			const FString NameSpace = bUseNameSpace ? LibraryName + TEXT(".") : FString();
-			ShapeNameList.Add(MakeShared<FRigVMStringWithTag>(UControlRigShapeLibrary::GetShapeName(ShapeLibrary.Get(), bUseNameSpace, LibraryNameMap, ShapeLibrary->DefaultShape)));
-			for (const FControlRigShapeDefinition& Shape : ShapeLibrary->Shapes)
-			{
-				ShapeNameList.Add(MakeShared<FRigVMStringWithTag>(UControlRigShapeLibrary::GetShapeName(ShapeLibrary.Get(), bUseNameSpace, LibraryNameMap, Shape)));
-			}
-		}
-		if(ShapeNameList.IsEmpty())
+
+		ControlRigShapeNameList->GenerateShapeLibraryList(ControlRig);
+
+		if(ControlRigShapeNameList->ShapeNameList.IsEmpty())
 		{
 			return;
 		}
-		StructBuilder.AddProperty(
-			StructPropertyHandle->GetChildHandle(GET_MEMBER_NAME_CHECKED(FRigUnit_HierarchyAddControl_ShapeSettings, Color)).ToSharedRef());
-		Property = StructPropertyHandle->GetChildHandle(GET_MEMBER_NAME_CHECKED(FRigUnit_HierarchyAddControl_ShapeSettings, Name));
+
+		StructBuilder.AddProperty(StructPropertyHandle->GetChildHandle(GET_MEMBER_NAME_CHECKED(FRigUnit_HierarchyAddControl_ShapeSettings, Color)).ToSharedRef());
+
+		TSharedPtr<IPropertyHandle> Property = StructPropertyHandle->GetChildHandle(GET_MEMBER_NAME_CHECKED(FRigUnit_HierarchyAddControl_ShapeSettings, Name));
 		if(Property.IsValid())
 		{
-			TSharedPtr<FRigVMStringWithTag> InitialSelected;
-			const FString CurrentShapeName = GetShapeNameListText().ToString();
-			for (TSharedPtr<FRigVMStringWithTag> Item : ShapeNameList)
-			{
-				if (Item->Equals(CurrentShapeName))
-				{
-					InitialSelected = Item;
-				}
-			}
-			IDetailPropertyRow& Row = StructBuilder.AddProperty(Property.ToSharedRef());
-			constexpr bool bShowChildren = true;
-			Row.CustomWidget(bShowChildren)
-			.NameContent()
-			[
-				Property->CreatePropertyNameWidget()
-			]
-			.ValueContent()
-			[
-				SAssignNew(ShapeNameListWidget, SRigVMGraphPinNameListValueWidget)
-				.OptionsSource(&ShapeNameList)
-				.OnGenerateWidget(this, &FAnimDetailsOverrideDetails::MakeShapeNameListItemWidget)
-				.OnSelectionChanged(this, &FAnimDetailsOverrideDetails::OnShapeNameListChanged)
-				.OnComboBoxOpening(this, &FAnimDetailsOverrideDetails::OnShapeNameListComboBox)
-				.InitiallySelectedItem(InitialSelected)
-				.Content()
-				[
-					SNew(STextBlock)
-					.Text(this, &FAnimDetailsOverrideDetails::GetShapeNameListText)
-					.Font(IDetailLayoutBuilder::GetDetailFont())
-				]
-			];
+			ControlRigShapeNameList->CreateShapeLibraryListWidget(StructBuilder, Property);
 		}
-		
-		StructBuilder.AddProperty(
-			StructPropertyHandle->GetChildHandle(GET_MEMBER_NAME_CHECKED(FRigUnit_HierarchyAddControl_ShapeSettings, bVisible)).ToSharedRef());
+
+		StructBuilder.AddProperty(StructPropertyHandle->GetChildHandle(GET_MEMBER_NAME_CHECKED(FRigUnit_HierarchyAddControl_ShapeSettings, bVisible)).ToSharedRef());
+
 		if(TSharedPtr<IPropertyHandle> TransformHandle = StructPropertyHandle->GetChildHandle(GET_MEMBER_NAME_CHECKED(FRigUnit_HierarchyAddControl_ShapeSettings, Transform)))
 		{
 			StructBuilder.AddProperty(TransformHandle->GetChildHandle(TEXT("Translation")).ToSharedRef());
@@ -127,76 +81,6 @@ namespace UE::ControlRigEditor
 		}
 	}
 	
-	TSharedRef<SWidget> FAnimDetailsOverrideDetails::MakeShapeNameListItemWidget(TSharedPtr<FRigVMStringWithTag> InItem)
-	{
-		return SNew(STextBlock)
-			.Text(FText::FromString(InItem->GetStringWithTag()))
-			.Font(IDetailLayoutBuilder::GetDetailFont());
-	}
-	
-	void FAnimDetailsOverrideDetails::OnShapeNameListChanged(TSharedPtr<FRigVMStringWithTag> NewSelection, ESelectInfo::Type SelectInfo)
-	{
-		if (SelectInfo == ESelectInfo::Direct)
-		{
-			return;
-		}
-		
-		const FString& NewShapeNameString = NewSelection->GetString();
-		if(Property)
-		{
-			const FName ShapeName = NewShapeNameString.IsEmpty() ? FName(NAME_None) : FName(*NewShapeNameString);
-			Property->SetValue(ShapeName);
-		}
-	}
-	
-	void FAnimDetailsOverrideDetails::OnShapeNameListComboBox()
-	{
-		const FString ShapeNameListText = GetShapeNameListText().ToString();
-		const TSharedPtr<FRigVMStringWithTag>* CurrentlySelectedItem =
-			ShapeNameList.FindByPredicate([ShapeNameListText](const TSharedPtr<FRigVMStringWithTag>& InItem)
-			{
-				return ShapeNameListText == InItem->GetString();
-			});
-		
-		if(CurrentlySelectedItem)
-		{
-			ShapeNameListWidget->SetSelectedItem(*CurrentlySelectedItem);
-		}
-	}
-	
-	FText FAnimDetailsOverrideDetails::GetShapeNameListText() const
-	{
-		if(!Property)
-		{
-			return FText();
-		}
-		
-		TOptional<FString> SharedValue;
-		for(int32 Index = 0; Index < Property->GetNumPerObjectValues(); Index++)
-		{
-			FString SingleValue;
-			if(!Property->GetPerObjectValue(Index, SingleValue))
-			{
-				SharedValue.Reset();
-				break;
-			}
-			if(!SharedValue.IsSet())
-			{
-				SharedValue = SingleValue;
-			}
-			else if(SharedValue.GetValue() != SingleValue)
-			{
-				SharedValue.Reset();
-				break;
-			}
-		}
-		
-		if(SharedValue.IsSet())
-		{
-			return FText::FromString(SharedValue.GetValue());
-		}
-		return LOCTEXT("MultipleValues", "Multiple Values");
-	}
-}
+} // end namespace UE::ControlRigEditor
 
 #undef LOCTEXT_NAMESPACE
diff --git a/Engine/Plugins/Animation/ControlRig/Source/ControlRigEditor/Private/AnimDetails/Customizations/AnimDetailsOverrideDetails.h b/Engine/Plugins/Animation/ControlRig/Source/ControlRigEditor/Private/AnimDetails/Customizations/AnimDetailsOverrideDetails.h
--- a/Engine/Plugins/Animation/ControlRig/Source/ControlRigEditor/Private/AnimDetails/Customizations/AnimDetailsOverrideDetails.h
+++ b/Engine/Plugins/Animation/ControlRig/Source/ControlRigEditor/Private/AnimDetails/Customizations/AnimDetailsOverrideDetails.h
@@ -11,6 +11,8 @@
 
 namespace UE::ControlRigEditor
 {
+	struct FControlRigShapeNameList;
+
 	/** Details customization for Control Shape settings */
 	class FAnimDetailsOverrideDetails : public IPropertyTypeCustomization
 	{
@@ -19,15 +21,12 @@ namespace UE::ControlRigEditor
 		{
 			return MakeShareable(new FAnimDetailsOverrideDetails());
 		}
+
+		FAnimDetailsOverrideDetails();
+
 		virtual void CustomizeHeader(TSharedRef<IPropertyHandle> StructPropertyHandle, FDetailWidgetRow& HeaderRow, IPropertyTypeCustomizationUtils& StructCustomizationUtils) override;
 		virtual void CustomizeChildren(TSharedRef<IPropertyHandle> StructPropertyHandle, IDetailChildrenBuilder& StructBuilder, IPropertyTypeCustomizationUtils& StructCustomizationUtils) override;
 	private:
-		TSharedRef<SWidget> MakeShapeNameListItemWidget(TSharedPtr<FRigVMStringWithTag> InItem);
-		void OnShapeNameListChanged(TSharedPtr<FRigVMStringWithTag> NewSelection, ESelectInfo::Type SelectInfo);
-		void OnShapeNameListComboBox();
-		FText GetShapeNameListText() const;
-		TSharedPtr<IPropertyHandle> Property;
-		TArray<TSharedPtr<FRigVMStringWithTag>> ShapeNameList;
-		TSharedPtr<SRigVMGraphPinNameListValueWidget> ShapeNameListWidget;
+		TSharedPtr<FControlRigShapeNameList> ControlRigShapeNameList;
 	};
 }
diff --git a/Engine/Plugins/Animation/ControlRig/Source/ControlRigEditor/Private/ControlRigEditorModule.cpp b/Engine/Plugins/Animation/ControlRig/Source/ControlRigEditor/Private/ControlRigEditorModule.cpp
--- a/Engine/Plugins/Animation/ControlRig/Source/ControlRigEditor/Private/ControlRigEditorModule.cpp
+++ b/Engine/Plugins/Animation/ControlRig/Source/ControlRigEditor/Private/ControlRigEditorModule.cpp
@@ -52,6 +52,7 @@
 #include "ControlRigAnimGraphDetails.h"
 #include "ControlRigBlueprintDetails.h"
 #include "ControlRigOverrideDetails.h"
+#include "ControlRigShapeSettingsDetails.h"
 #include "Subsystems/AssetEditorSubsystem.h"
 #include "Sequencer/ControlRigParameterTrackEditor.h"
 #include "ActorFactories/ActorFactorySkeletalMesh.h"
@@ -76,6 +77,7 @@
 #include "Rigs/FKControlRig.h"
 #include "SBakeToControlRigDialog.h"
 #include "Units/Execution/RigUnit_InverseExecution.h"
+#include "Units/Execution/RigUnit_DynamicHierarchy.h"
 #include "Widgets/SRigVMGraphPinVariableBinding.h"
 #include "ControlRigPythonLogDetails.h"
 #include "Dialog/SCustomDialog.h"
@@ -186,6 +188,10 @@ void FControlRigEditorModule::StartupModule()
 	PropertiesToUnregisterOnShutdown.Add(FControlRigOverrideContainer::StaticStruct()->GetFName());
 	PropertyEditorModule.RegisterCustomPropertyTypeLayout(PropertiesToUnregisterOnShutdown.Last(), FOnGetPropertyTypeCustomizationInstance::CreateStatic(&FControlRigOverrideDetails::MakeInstance));
 
+	PropertiesToUnregisterOnShutdown.Add(FRigUnit_HierarchyAddControl_ShapeSettings::StaticStruct()->GetFName());
+	PropertyEditorModule.RegisterCustomPropertyTypeLayout(PropertiesToUnregisterOnShutdown.Last(), FOnGetPropertyTypeCustomizationInstance::CreateStatic(&FControlRigShapeSettingsDetails::MakeInstance));
+	
+
 	// Register asset tools
 	auto RegisterAssetTypeAction = [this](const TSharedRef<IAssetTypeActions>& InAssetTypeAction)
 	{
diff --git a/Engine/Plugins/Animation/ControlRig/Source/ControlRigEditor/Private/ControlRigShapeNameList.cpp b/Engine/Plugins/Animation/ControlRig/Source/ControlRigEditor/Private/ControlRigShapeNameList.cpp
new file mode 100644
--- /dev/null
+++ b/Engine/Plugins/Animation/ControlRig/Source/ControlRigEditor/Private/ControlRigShapeNameList.cpp
@@ -0,0 +1,170 @@
+// Copyright Epic Games, Inc. All Rights Reserved.
+
+#include "ControlRigShapeNameList.h"
+#include "Widgets/SWidget.h"
+#include "DetailLayoutBuilder.h"
+#include "IDetailChildrenBuilder.h"
+#include "Widgets/SRigVMGraphPinNameListValueWidget.h"
+#include "ControlRig.h"
+#include "ControlRigAsset.h"
+
+#define LOCTEXT_NAMESPACE "ControlRigShapeNameList"
+
+namespace UE::ControlRigEditor
+{
+
+void FControlRigShapeNameList::CreateShapeLibraryListWidget(IDetailChildrenBuilder& InStructBuilder, TSharedPtr<IPropertyHandle>& InShapeSettingsNameProperty)
+{
+	if (!InShapeSettingsNameProperty.IsValid())
+	{
+		return;
+	}
+
+	ShapeSettingsNameProperty = InShapeSettingsNameProperty;
+
+	TSharedPtr<FRigVMStringWithTag> InitialSelected;
+	const FString CurrentShapeName = GetShapeNameListText().ToString();
+	for (TSharedPtr<FRigVMStringWithTag> Item : ShapeNameList)
+	{
+		if (Item->Equals(CurrentShapeName))
+		{
+			InitialSelected = Item;
+		}
+	}
+
+	IDetailPropertyRow& Row = InStructBuilder.AddProperty(ShapeSettingsNameProperty.ToSharedRef());
+
+	constexpr bool bShowChildren = true;
+	Row.CustomWidget(bShowChildren)
+		.NameContent()
+		[
+			ShapeSettingsNameProperty->CreatePropertyNameWidget()
+		]
+		.ValueContent()
+		[
+			SAssignNew(ShapeNameListWidget, SRigVMGraphPinNameListValueWidget)
+				.OptionsSource(&ShapeNameList)
+				.OnGenerateWidget(this, &FControlRigShapeNameList::MakeShapeNameListItemWidget)
+				.OnSelectionChanged(this, &FControlRigShapeNameList::OnShapeNameListChanged)
+				.OnComboBoxOpening(this, &FControlRigShapeNameList::OnShapeNameListComboBox)
+				.InitiallySelectedItem(InitialSelected)
+				.Content()
+				[
+					SNew(STextBlock)
+						.Text(this, &FControlRigShapeNameList::GetShapeNameListText)
+						.Font(IDetailLayoutBuilder::GetDetailFont())
+				]
+		];
+}
+
+void FControlRigShapeNameList::GenerateShapeLibraryList(const UControlRig* ControlRig)
+{
+	ShapeNameList.Reset();
+
+	if (ControlRig == nullptr)
+	{
+		return;
+	}
+
+	const TMap<FString, FString>& LibraryNameMap = ControlRig->GetShapeLibraryNameMap();
+	TArray<TSoftObjectPtr<UControlRigShapeLibrary>> ShapeLibraries = ControlRig->GetShapeLibraries();
+	for (const TSoftObjectPtr<UControlRigShapeLibrary>& ShapeLibrary : ShapeLibraries)
+	{
+		if (ShapeLibrary.IsNull() || !ShapeLibrary.IsValid())
+		{
+			(void)ShapeLibrary.LoadSynchronous();
+		}
+		if (ShapeLibrary.IsNull() || !ShapeLibrary.IsValid())
+		{
+			continue;
+		}
+		const bool bUseNameSpace = ShapeLibraries.Num() > 1;
+		FString LibraryName = ShapeLibrary->GetName();
+		if (const FString* RemappedName = LibraryNameMap.Find(LibraryName))
+		{
+			LibraryName = *RemappedName;
+		}
+
+		const FString NameSpace = bUseNameSpace ? LibraryName + TEXT(".") : FString();
+		ShapeNameList.Add(MakeShared<FRigVMStringWithTag>(UControlRigShapeLibrary::GetShapeName(ShapeLibrary.Get(), bUseNameSpace, LibraryNameMap, ShapeLibrary->DefaultShape)));
+		for (const FControlRigShapeDefinition& Shape : ShapeLibrary->Shapes)
+		{
+			ShapeNameList.Add(MakeShared<FRigVMStringWithTag>(UControlRigShapeLibrary::GetShapeName(ShapeLibrary.Get(), bUseNameSpace, LibraryNameMap, Shape)));
+		}
+	}
+}
+
+TSharedRef<SWidget> FControlRigShapeNameList::MakeShapeNameListItemWidget(TSharedPtr<FRigVMStringWithTag> InItem)
+{
+	return SNew(STextBlock)
+		.Text(FText::FromString(InItem->GetStringWithTag()))
+		.Font(IDetailLayoutBuilder::GetDetailFont());
+}
+	
+void FControlRigShapeNameList::OnShapeNameListChanged(TSharedPtr<FRigVMStringWithTag> NewSelection, ESelectInfo::Type SelectInfo)
+{
+	if (SelectInfo == ESelectInfo::Direct)
+	{
+		return;
+	}
+		
+	const FString& NewShapeNameString = NewSelection->GetString();
+	if(ShapeSettingsNameProperty)
+	{
+		const FName ShapeName = NewShapeNameString.IsEmpty() ? FName(NAME_None) : FName(*NewShapeNameString);
+		ShapeSettingsNameProperty->SetValue(ShapeName);
+	}
+}
+	
+void FControlRigShapeNameList::OnShapeNameListComboBox()
+{
+	const FString ShapeNameListText = GetShapeNameListText().ToString();
+	const TSharedPtr<FRigVMStringWithTag>* CurrentlySelectedItem =
+		ShapeNameList.FindByPredicate([ShapeNameListText](const TSharedPtr<FRigVMStringWithTag>& InItem)
+		{
+			return ShapeNameListText == InItem->GetString();
+		});
+		
+	if(CurrentlySelectedItem)
+	{
+		ShapeNameListWidget->SetSelectedItem(*CurrentlySelectedItem);
+	}
+}
+	
+FText FControlRigShapeNameList::GetShapeNameListText() const
+{
+	if(!ShapeSettingsNameProperty)
+	{
+		return FText();
+	}
+		
+	TOptional<FString> SharedValue;
+	for(int32 Index = 0; Index < ShapeSettingsNameProperty->GetNumPerObjectValues(); Index++)
+	{
+		FString SingleValue;
+		if(!ShapeSettingsNameProperty->GetPerObjectValue(Index, SingleValue))
+		{
+			SharedValue.Reset();
+			break;
+		}
+		if(!SharedValue.IsSet())
+		{
+			SharedValue = SingleValue;
+		}
+		else if(SharedValue.GetValue() != SingleValue)
+		{
+			SharedValue.Reset();
+			break;
+		}
+	}
+		
+	if(SharedValue.IsSet())
+	{
+		return FText::FromString(SharedValue.GetValue());
+	}
+	return LOCTEXT("MultipleValues", "Multiple Values");
+}
+
+} // end namespace UE::ControlRigEditor
+
+#undef LOCTEXT_NAMESPACE
diff --git a/Engine/Plugins/Animation/ControlRig/Source/ControlRigEditor/Private/ControlRigShapeNameList.h b/Engine/Plugins/Animation/ControlRig/Source/ControlRigEditor/Private/ControlRigShapeNameList.h
new file mode 100644
--- /dev/null
+++ b/Engine/Plugins/Animation/ControlRig/Source/ControlRigEditor/Private/ControlRigShapeNameList.h
@@ -0,0 +1,35 @@
+// Copyright Epic Games, Inc. All Rights Reserved.
+
+#pragma once
+
+#include "IDetailCustomization.h"
+#include "IPropertyTypeCustomization.h"
+#include "DetailWidgetRow.h"
+#include "Templates/SharedPointer.h"
+#include "Widgets/SRigVMGraphPinNameList.h"
+
+class IPropertyHandle;
+class UControlRig;
+
+namespace UE::ControlRigEditor
+{
+
+/*
+* Helper to create a ShapeLibraryListWidget for a ControlRig details
+*/
+struct FControlRigShapeNameList : public TSharedFromThis<FControlRigShapeNameList>
+{
+	void GenerateShapeLibraryList(const UControlRig* ControlRig);
+	TSharedRef<SWidget> MakeShapeNameListItemWidget(TSharedPtr<FRigVMStringWithTag> InItem);
+
+	void CreateShapeLibraryListWidget(IDetailChildrenBuilder& InStructBuilder, TSharedPtr<IPropertyHandle>& ShapeSettingsNameProperty);
+	void OnShapeNameListChanged(TSharedPtr<FRigVMStringWithTag> NewSelection, ESelectInfo::Type SelectInfo);
+	void OnShapeNameListComboBox();
+	FText GetShapeNameListText() const;
+
+	TSharedPtr<IPropertyHandle> ShapeSettingsNameProperty;
+	TArray<TSharedPtr<FRigVMStringWithTag>> ShapeNameList;
+	TSharedPtr<SRigVMGraphPinNameListValueWidget> ShapeNameListWidget;
+};
+
+} // end namespace UE::ControlRigEditor
diff --git a/Engine/Plugins/Animation/ControlRig/Source/ControlRigEditor/Private/ControlRigShapeSettingsDetails.cpp b/Engine/Plugins/Animation/ControlRig/Source/ControlRigEditor/Private/ControlRigShapeSettingsDetails.cpp
new file mode 100644
--- /dev/null
+++ b/Engine/Plugins/Animation/ControlRig/Source/ControlRigEditor/Private/ControlRigShapeSettingsDetails.cpp
@@ -0,0 +1,107 @@
+﻿// Copyright Epic Games, Inc. All Rights Reserved.
+
+#include "ControlRigShapeSettingsDetails.h"
+#include "Widgets/SWidget.h"
+#include "DetailLayoutBuilder.h"
+#include "DetailCategoryBuilder.h"
+#include "DetailWidgetRow.h"
+#include "IDetailChildrenBuilder.h"
+#include "IPropertyUtilities.h"
+#include "IPythonScriptPlugin.h"
+#include "RigVMPythonUtils.h"
+#include "Editor/ControlRigWrapperObject.h"
+#include "Units/Execution/RigUnit_DynamicHierarchy.h"
+#include "Widgets/SRigVMGraphPinNameListValueWidget.h"
+#include "ControlRig.h"
+#include "ControlRigAsset.h"
+#include "ControlRigShapeNameList.h"
+
+#define LOCTEXT_NAMESPACE "ControlRigCompilerDetails"
+
+namespace UE::ControlRigEditor
+{
+
+FControlRigShapeSettingsDetails::FControlRigShapeSettingsDetails()
+	: ControlRigShapeNameList(MakeShared<FControlRigShapeNameList>())
+{
+}
+
+void FControlRigShapeSettingsDetails::CustomizeHeader(TSharedRef<IPropertyHandle> InStructPropertyHandle, FDetailWidgetRow& InHeaderRow, IPropertyTypeCustomizationUtils& InStructCustomizationUtils)
+{
+	TArray<UObject*> Objects;
+	InStructPropertyHandle->GetOuterObjects(Objects);
+	ensure(Objects.Num() == 1); // This is in here to ensure we are only showing the modifier details in the blueprint editor
+
+	for (UObject* Object : Objects)
+	{
+		if (Object->IsA<UControlRigWrapperObject>())
+		{
+			ControlRigWrapperObject = Cast<UControlRigWrapperObject>(Object);
+			break;
+		}
+	}
+
+	InHeaderRow.NameContent()
+		[
+			InStructPropertyHandle->CreatePropertyNameWidget()
+		];	
+}
+
+void FControlRigShapeSettingsDetails::CustomizeChildren(TSharedRef<IPropertyHandle> InStructPropertyHandle, IDetailChildrenBuilder& InStructBuilder, IPropertyTypeCustomizationUtils& InStructCustomizationUtils)
+{
+	if (InStructPropertyHandle->IsValidHandle() && ControlRigWrapperObject != nullptr)
+	{
+		if (UControlRig* ControlRig = GetControlRig(ControlRigWrapperObject))
+		{
+			ControlRigShapeNameList->GenerateShapeLibraryList(ControlRig);
+
+			uint32 NumChildren = 0;
+			if (InStructPropertyHandle->GetNumChildren(NumChildren) == FPropertyAccess::Success)
+			{
+				TSharedPtr<IPropertyHandle> ShapeSettingsNameProperty = InStructPropertyHandle->GetChildHandle(GET_MEMBER_NAME_CHECKED(FRigUnit_HierarchyAddControl_ShapeSettings, Name));
+				if (!ShapeSettingsNameProperty.IsValid())
+				{
+					return;
+				}
+
+				for (uint32 Index = 0; Index < NumChildren; ++Index)
+				{
+					bool bDefaultContent = true;
+
+					TSharedRef<IPropertyHandle> ChildPropertyHandle = InStructPropertyHandle->GetChildHandle(Index).ToSharedRef();
+					if (ChildPropertyHandle->IsSamePropertyNode(ShapeSettingsNameProperty))
+					{
+						ControlRigShapeNameList->CreateShapeLibraryListWidget(InStructBuilder, ShapeSettingsNameProperty);
+						bDefaultContent = false;
+					}
+
+					if (bDefaultContent)
+					{
+						InStructBuilder.AddProperty(ChildPropertyHandle);
+					}
+				}
+			}
+		}
+	}
+}
+
+UControlRig* FControlRigShapeSettingsDetails::GetControlRig(UControlRigWrapperObject* InControlRigWrapperObject)
+{
+	if (URigVMUnitNode* UnitNode = Cast<URigVMUnitNode>(InControlRigWrapperObject->GetSubject()))
+	{
+		FControlRigAssetInterfacePtr AssetInterface = IControlRigAssetInterface::GetInterfaceOuter(UnitNode);
+		if (AssetInterface)
+		{
+			if (UControlRig* ControlRig = AssetInterface->GetDebuggedControlRig())
+			{
+				return ControlRig;
+			}
+		}
+	}
+
+	return nullptr;
+}
+
+} // end namespace UE::ControlRigEditor
+
+#undef LOCTEXT_NAMESPACE
diff --git a/Engine/Plugins/Animation/ControlRig/Source/ControlRigEditor/Private/ControlRigShapeSettingsDetails.h b/Engine/Plugins/Animation/ControlRig/Source/ControlRigEditor/Private/ControlRigShapeSettingsDetails.h
new file mode 100644
--- /dev/null
+++ b/Engine/Plugins/Animation/ControlRig/Source/ControlRigEditor/Private/ControlRigShapeSettingsDetails.h
@@ -0,0 +1,41 @@
+﻿// Copyright Epic Games, Inc. All Rights Reserved.
+
+#pragma once
+
+#include "IDetailCustomization.h"
+#include "IPropertyTypeCustomization.h"
+#include "DetailWidgetRow.h"
+#include "Templates/SharedPointer.h"
+#include "Widgets/SRigVMGraphPinNameList.h"
+
+class IPropertyHandle;
+class UControlRigWrapperObject;
+class UControlRig;
+
+namespace UE::ControlRigEditor
+{
+
+struct FControlRigShapeNameList;
+
+class FControlRigShapeSettingsDetails : public IPropertyTypeCustomization
+{
+public:
+	static TSharedRef<IPropertyTypeCustomization> MakeInstance()
+	{
+		return MakeShareable(new FControlRigShapeSettingsDetails);
+	}
+
+	FControlRigShapeSettingsDetails();
+
+	/** IPropertyTypeCustomization interface */
+	virtual void CustomizeHeader(TSharedRef<class IPropertyHandle> InStructPropertyHandle, class FDetailWidgetRow& InHeaderRow, IPropertyTypeCustomizationUtils& InStructCustomizationUtils) override;
+	virtual void CustomizeChildren(TSharedRef<class IPropertyHandle> InStructPropertyHandle, class IDetailChildrenBuilder& InStructBuilder, IPropertyTypeCustomizationUtils& InStructCustomizationUtils) override;
+
+protected:
+	UControlRigWrapperObject* ControlRigWrapperObject = nullptr;
+	TSharedPtr<FControlRigShapeNameList> ControlRigShapeNameList;
+
+	static UControlRig* GetControlRig(UControlRigWrapperObject* ControlRigWrapperObject);
+};
+
+} // end namespace UE::ControlRigEditor
