Fix for static analysis warnings

diff --git a/Engine/Plugins/Animation/ControlRig/Source/ControlRigEditor/Private/ControlRigShapeNameList.cpp b/Engine/Plugins/Animation/ControlRig/Source/ControlRigEditor/Private/ControlRigShapeNameList.cpp
--- a/Engine/Plugins/Animation/ControlRig/Source/ControlRigEditor/Private/ControlRigShapeNameList.cpp
+++ b/Engine/Plugins/Animation/ControlRig/Source/ControlRigEditor/Private/ControlRigShapeNameList.cpp
@@ -96,8 +96,10 @@ void FControlRigShapeNameList::GenerateShapeLibraryList(const UControlRig* Contr
 
 TSharedRef<SWidget> FControlRigShapeNameList::MakeShapeNameListItemWidget(TSharedPtr<FRigVMStringWithTag> InItem)
 {
+	const FText ItemText = InItem.IsValid() ? FText::FromString(InItem->GetStringWithTag()) : FText();
+
 	return SNew(STextBlock)
-		.Text(FText::FromString(InItem->GetStringWithTag()))
+		.Text(ItemText)
 		.Font(IDetailLayoutBuilder::GetDetailFont());
 }
 	
@@ -107,8 +109,8 @@ void FControlRigShapeNameList::OnShapeNameListChanged(TSharedPtr<FRigVMStringWit
 	{
 		return;
 	}
-		
-	const FString& NewShapeNameString = NewSelection->GetString();
+
+	const FString& NewShapeNameString = NewSelection.IsValid() ? NewSelection->GetString() : FString();
 	if(ShapeSettingsNameProperty)
 	{
 		const FName ShapeName = NewShapeNameString.IsEmpty() ? FName(NAME_None) : FName(*NewShapeNameString);
