Anim in Engine: Fix integer controls cannot be edited in Anim Details when their value is zero, fix currently unused min and max property values are inverted

diff --git a/Engine/Plugins/Animation/ControlRig/Source/ControlRigEditor/Private/AnimDetails/Proxies/AnimDetailsProxyInteger.h b/Engine/Plugins/Animation/ControlRig/Source/ControlRigEditor/Private/AnimDetails/Proxies/AnimDetailsProxyInteger.h
--- a/Engine/Plugins/Animation/ControlRig/Source/ControlRigEditor/Private/AnimDetails/Proxies/AnimDetailsProxyInteger.h
+++ b/Engine/Plugins/Animation/ControlRig/Source/ControlRigEditor/Private/AnimDetails/Proxies/AnimDetailsProxyInteger.h
@@ -12,7 +12,7 @@ struct FAnimDetailsInteger
 {
 	GENERATED_BODY()
 
-	UPROPERTY(EditAnywhere, Interp, Category = "Integer")
+	UPROPERTY(EditAnywhere, Interp, meta = (Delta = 1.0), Category = "Integer")
 	int64 Integer = 0;
 };
 
diff --git a/Engine/Plugins/Animation/ControlRig/Source/ControlRigEditor/Private/AnimDetails/Widgets/SAnimDetailsValueNumeric.cpp b/Engine/Plugins/Animation/ControlRig/Source/ControlRigEditor/Private/AnimDetails/Widgets/SAnimDetailsValueNumeric.cpp
--- a/Engine/Plugins/Animation/ControlRig/Source/ControlRigEditor/Private/AnimDetails/Widgets/SAnimDetailsValueNumeric.cpp
+++ b/Engine/Plugins/Animation/ControlRig/Source/ControlRigEditor/Private/AnimDetails/Widgets/SAnimDetailsValueNumeric.cpp
@@ -257,8 +257,8 @@ namespace UE::ControlRigEditor
 			.SupportDynamicSliderMaxValue(PropertyMetaData.SupportDynamicSliderMaxValue)
 			.MinValue(PropertyMetaData.MinValue)
 			.MaxValue(PropertyMetaData.MaxValue)
-			.MaxSliderValue(PropertyMetaData.SliderMinValue)
-			.MinSliderValue(PropertyMetaData.SliderMaxValue)
+			.MinSliderValue(PropertyMetaData.SliderMinValue)
+			.MaxSliderValue(PropertyMetaData.SliderMaxValue)
 			.ShiftMultiplier(PropertyMetaData.ShiftMultiplier)
 			.CtrlMultiplier(PropertyMetaData.CtrlMultiplier)
 			.SliderExponent(PropertyMetaData.SliderExponent)
