Anim in Engine: Safely access WeakProxyManager in SAnimDetailsValueNumeric::OnValueCommitted

#rnx

diff --git a/Engine/Plugins/Animation/ControlRig/Source/ControlRigEditor/Private/AnimDetails/Widgets/SAnimDetailsValueNumeric.cpp b/Engine/Plugins/Animation/ControlRig/Source/ControlRigEditor/Private/AnimDetails/Widgets/SAnimDetailsValueNumeric.cpp
--- a/Engine/Plugins/Animation/ControlRig/Source/ControlRigEditor/Private/AnimDetails/Widgets/SAnimDetailsValueNumeric.cpp
+++ b/Engine/Plugins/Animation/ControlRig/Source/ControlRigEditor/Private/AnimDetails/Widgets/SAnimDetailsValueNumeric.cpp
@@ -431,7 +431,8 @@ namespace UE::ControlRigEditor
 			constexpr bool bInteractive = false;
 			MultiEditChangePropertyValue(Value, bInteractive);
 		}
-		else if (WeakPropertyHandle.IsValid())
+		else if (WeakProxyManager.IsValid() &&
+			WeakPropertyHandle.IsValid())
 		{
 			const TSharedRef<IPropertyHandle> PropertyHandle = WeakPropertyHandle.Pin().ToSharedRef();
 
