Anim in Engine: Fix anim details options no longer show when reentering animation mode

diff --git a/Engine/Plugins/Animation/ControlRig/Source/ControlRigEditor/Private/AnimDetails/AnimDetailsOptionsMenuContext.h b/Engine/Plugins/Animation/ControlRig/Source/ControlRigEditor/Private/AnimDetails/AnimDetailsOptionsMenuContext.h
new file mode 100644
--- /dev/null
+++ b/Engine/Plugins/Animation/ControlRig/Source/ControlRigEditor/Private/AnimDetails/AnimDetailsOptionsMenuContext.h
@@ -0,0 +1,19 @@
+// Copyright Epic Games, Inc. All Rights Reserved.
+
+#pragma once
+
+#include "AnimDetailsOptionsMenuContext.generated.h"
+
+namespace UE::ControlRigEditor { class SAnimDetailsOptions; }
+
+/** Menu context for the anim details options menu */
+UCLASS()
+class UAnimDetailsOptionsMenuContext
+	: public UObject
+{
+	GENERATED_BODY()
+
+public:
+	/** The anim details options widget that uses this menu context */
+	TWeakPtr<UE::ControlRigEditor::SAnimDetailsOptions> WeakAnimDetailsOptions;
+};
diff --git a/Engine/Plugins/Animation/ControlRig/Source/ControlRigEditor/Private/AnimDetails/Widgets/SAnimDetailsOptions.cpp b/Engine/Plugins/Animation/ControlRig/Source/ControlRigEditor/Private/AnimDetails/Widgets/SAnimDetailsOptions.cpp
--- a/Engine/Plugins/Animation/ControlRig/Source/ControlRigEditor/Private/AnimDetails/Widgets/SAnimDetailsOptions.cpp
+++ b/Engine/Plugins/Animation/ControlRig/Source/ControlRigEditor/Private/AnimDetails/Widgets/SAnimDetailsOptions.cpp
@@ -2,7 +2,9 @@
 
 #include "SAnimDetailsOptions.h"
 
+#include "AnimDetails/AnimDetailsOptionsMenuContext.h"
 #include "AnimDetails/AnimDetailsSettings.h"
+#include "LevelEditor.h"
 #include "ToolMenuEntry.h"
 #include "ToolMenus.h"
 #include "Widgets/Images/SImage.h"
@@ -47,64 +49,91 @@ namespace UE::ControlRigEditor
 		{
 			UToolMenu* Menu = UToolMenus::Get()->RegisterMenu(AnimDetailsOptionsMenuName);
 
-			FToolMenuSection& OptionsSection = Menu->AddSection("Options");
+			Menu->AddDynamicSection("PopulateAnimDetailsMenu", FNewToolMenuDelegate::CreateStatic(&SAnimDetailsOptions::PopulateMenu));
+		}
+
+		FLevelEditorModule& LevelEditorModule = FModuleManager::Get().LoadModuleChecked<FLevelEditorModule>("LevelEditor");
+		const TSharedRef<FUICommandList> CommandList = LevelEditorModule.GetGlobalLevelEditorActions();
+
+		UAnimDetailsOptionsMenuContext* ContextObject = NewObject<UAnimDetailsOptionsMenuContext>();
+		ContextObject->WeakAnimDetailsOptions = SharedThis(this);
+
+		const FToolMenuContext MenuContext(CommandList, TSharedPtr<FExtender>(), ContextObject);
+		UToolMenu* ToolMenu = UToolMenus::Get()->GenerateMenu(AnimDetailsOptionsMenuName, MenuContext);
+
+		return UToolMenus::Get()->GenerateWidget(AnimDetailsOptionsMenuName, MenuContext);
+	}
+
+	void SAnimDetailsOptions::PopulateMenu(UToolMenu* InMenu)
+	{
+		constexpr const TCHAR* AnimDetailsOptionsMenuName = TEXT("AnimDetails.Options");
+
+		UAnimDetailsOptionsMenuContext* ContextObject = InMenu ? InMenu->FindContext<UAnimDetailsOptionsMenuContext>() : nullptr;
+		SAnimDetailsOptions* This = ContextObject && ContextObject->WeakAnimDetailsOptions.IsValid() ?
+			ContextObject->WeakAnimDetailsOptions.Pin().Get() :
+			nullptr;
+
+		if (!InMenu ||
+			!This)
+		{
+			return;
+		}
+
+		FToolMenuSection& OptionsSection = InMenu->AddSection("Options");
+		{
+			constexpr bool bIndent = false;
+			constexpr bool bSearchable = false;
+			constexpr bool bNoPadding = true;
+
+			// Num fractional digits option
 			{
-				constexpr bool bIndent = false;
-				constexpr bool bSearchable = false;
-				constexpr bool bNoPadding = true;
-
-				// Num fractional digits option
-				{
-					const TSharedRef<SWidget> NumFractionalDigitsOptionWidget =
-						SNew(SEditableTextBox)
-						.Font(FAppStyle::GetFontStyle("PropertyWindow.NormalFont"))
-						.Justification(ETextJustify::Right)
-						.MinDesiredWidth(40.f)
-						.SelectAllTextWhenFocused(true)
-						.Text(this, &SAnimDetailsOptions::GetNumFractionalDigitsText)
-						.OnTextCommitted(this, &SAnimDetailsOptions::OnNumFractionalDigitsCommitted);
-
-					OptionsSection.AddEntry(
-						FToolMenuEntry::InitWidget(
-							"NumFractionalDigitsOptionWidget", 
-							NumFractionalDigitsOptionWidget,
-							LOCTEXT("NumFractionalDigitsLabel", "Num Fractional Digits"),
-							bIndent,
-							bSearchable,
-							bNoPadding,
-							LOCTEXT("NumFractionalDigitsTooltip", "Sets the displayed num fractional digits"))
-					);
-				}
+				const TSharedRef<SWidget> NumFractionalDigitsOptionWidget =
+					SNew(SEditableTextBox)
+					.Font(FAppStyle::GetFontStyle("PropertyWindow.NormalFont"))
+					.Justification(ETextJustify::Right)
+					.MinDesiredWidth(40.f)
+					.SelectAllTextWhenFocused(true)
+					.Text(This, &SAnimDetailsOptions::GetNumFractionalDigitsText)
+					.OnTextCommitted(This, &SAnimDetailsOptions::OnNumFractionalDigitsCommitted);
+
+				OptionsSection.AddEntry(
+					FToolMenuEntry::InitWidget(
+						"NumFractionalDigitsOptionWidget", 
+						NumFractionalDigitsOptionWidget,
+						LOCTEXT("NumFractionalDigitsLabel", "Num Fractional Digits"),
+						bIndent,
+						bSearchable,
+						bNoPadding,
+						LOCTEXT("NumFractionalDigitsTooltip", "Sets the displayed num fractional digits"))
+				);
+			}
 
-				// LMB selects range option
-				{
-					const TSharedRef<SWidget> LMBSelectsRangeOptionWidget =
-						SNew(SBox)
-						.MinDesiredWidth(40.f)
+			// LMB selects range option
+			{
+				const TSharedRef<SWidget> LMBSelectsRangeOptionWidget =
+					SNew(SBox)
+					.MinDesiredWidth(40.f)
+					.HAlign(HAlign_Right)
+					.VAlign(VAlign_Center)
+					[
+						SNew(SCheckBox)
 						.HAlign(HAlign_Right)
-						.VAlign(VAlign_Center)
-						[
-							SNew(SCheckBox)
-							.HAlign(HAlign_Right)
-							.IsChecked(this, &SAnimDetailsOptions::GetLMBSelectsRangeCheckState)
-							.OnCheckStateChanged(this, &SAnimDetailsOptions::OnLMBSelectsRangeCheckStateChanged)
-						];
-
-					OptionsSection.AddEntry(
-						FToolMenuEntry::InitWidget(
-							"LMBSelectsRangeOptionWidget",
-							LMBSelectsRangeOptionWidget,
-							LOCTEXT("LMBSelectsRangeLabel", "LMB selects range"),
-							bIndent,
-							bSearchable,
-							bNoPadding,
-							LOCTEXT("LMBSelectsRangeTooltip", "When checked, selects a range when the left mouse button is down"))
-						);
-				}
+						.IsChecked(This, &SAnimDetailsOptions::GetLMBSelectsRangeCheckState)
+						.OnCheckStateChanged(This, &SAnimDetailsOptions::OnLMBSelectsRangeCheckStateChanged)
+					];
+
+				OptionsSection.AddEntry(
+					FToolMenuEntry::InitWidget(
+						"LMBSelectsRangeOptionWidget",
+						LMBSelectsRangeOptionWidget,
+						LOCTEXT("LMBSelectsRangeLabel", "LMB selects range"),
+						bIndent,
+						bSearchable,
+						bNoPadding,
+						LOCTEXT("LMBSelectsRangeTooltip", "When checked, selects a range when the left mouse button is down"))
+					);
 			}
 		}
-
-		return UToolMenus::Get()->GenerateWidget(AnimDetailsOptionsMenuName, FToolMenuContext());
 	}
 
 	FText SAnimDetailsOptions::GetNumFractionalDigitsText() const
diff --git a/Engine/Plugins/Animation/ControlRig/Source/ControlRigEditor/Private/AnimDetails/Widgets/SAnimDetailsOptions.h b/Engine/Plugins/Animation/ControlRig/Source/ControlRigEditor/Private/AnimDetails/Widgets/SAnimDetailsOptions.h
--- a/Engine/Plugins/Animation/ControlRig/Source/ControlRigEditor/Private/AnimDetails/Widgets/SAnimDetailsOptions.h
+++ b/Engine/Plugins/Animation/ControlRig/Source/ControlRigEditor/Private/AnimDetails/Widgets/SAnimDetailsOptions.h
@@ -4,6 +4,8 @@
 
 #include "Widgets/SCompoundWidget.h"
 
+class UToolMenu;
+
 namespace UE::ControlRigEditor
 {
 	/** The options menu for anim details */
@@ -26,6 +28,9 @@ namespace UE::ControlRigEditor
 		/** Creates the options menu */
 		TSharedRef<SWidget> MakeOptionsMenu();
 
+		/** Populates the menu */
+		static void PopulateMenu(UToolMenu* InMenu);
+
 		/** Returns text for the num fractional digits option */
 		FText GetNumFractionalDigitsText() const;
 
