Fixed ensure at Control Rig shap settings details customization, when selecting elements in the Sequencer

#rnx

diff --git a/Engine/Plugins/Animation/ControlRig/Source/ControlRigEditor/Private/ControlRigShapeSettingsDetails.cpp b/Engine/Plugins/Animation/ControlRig/Source/ControlRigEditor/Private/ControlRigShapeSettingsDetails.cpp
--- a/Engine/Plugins/Animation/ControlRig/Source/ControlRigEditor/Private/ControlRigShapeSettingsDetails.cpp
+++ b/Engine/Plugins/Animation/ControlRig/Source/ControlRigEditor/Private/ControlRigShapeSettingsDetails.cpp
@@ -30,14 +30,18 @@ void FControlRigShapeSettingsDetails::CustomizeHeader(TSharedRef<IPropertyHandle
 {
 	TArray<UObject*> Objects;
 	InStructPropertyHandle->GetOuterObjects(Objects);
-	ensure(Objects.Num() == 1); // This is in here to ensure we are only showing the modifier details in the blueprint editor
-
-	for (UObject* Object : Objects)
+	
+	// This block is in here to ensure we are only showing the modifier details in the blueprint editor
+	if (Objects.Num() == 1)
 	{
-		if (Object->IsA<UControlRigWrapperObject>())
+		for (UObject* Object : Objects)
 		{
-			ControlRigWrapperObject = Cast<UControlRigWrapperObject>(Object);
-			break;
+			if (Object->IsA<UControlRigWrapperObject>())
+			{
+				// Note that multi selection of rigs is not currently supported (each rig has to show a different name list)
+				ControlRigWrapperObject = Cast<UControlRigWrapperObject>(Object);
+				break;
+			}
 		}
 	}
 
