Anim in Engine: Fix Anim Details property value widgets are losing focus while scrubbing if a blueprint actor is controlled, fix in some cases editing Anim Details properties creates more than one undo step

diff --git a/Engine/Plugins/Animation/ControlRig/Source/ControlRigEditor/Private/AnimDetails/AnimDetailsMultiEditUtil.cpp b/Engine/Plugins/Animation/ControlRig/Source/ControlRigEditor/Private/AnimDetails/AnimDetailsMultiEditUtil.cpp
--- a/Engine/Plugins/Animation/ControlRig/Source/ControlRigEditor/Private/AnimDetails/AnimDetailsMultiEditUtil.cpp
+++ b/Engine/Plugins/Animation/ControlRig/Source/ControlRigEditor/Private/AnimDetails/AnimDetailsMultiEditUtil.cpp
@@ -49,7 +49,7 @@ namespace UE::ControlRigEditor
 					EPropertyChangeType::Interactive : 
 					EPropertyChangeType::ValueSet;
 
-				 PropertyHandle->NotifyPostChange(EPropertyChangeType::ValueSet);
+				 PropertyHandle->NotifyPostChange(ChangeType);
 			}
 		}
 	}
@@ -234,6 +234,7 @@ namespace UE::ControlRigEditor
 		{
 			GEditor->BeginTransaction(LOCTEXT("MultiEditSetPropertyValues", "Set Property Value"));
 		}
+		bIsInteractiveChangeOngoing = bInteractive;
 
 		PropertiesBeingEditedInteractively = GetPropertiesBeingEdited<ValueType>(ProxyManager, InstigatorProperty);
 
@@ -265,8 +266,6 @@ namespace UE::ControlRigEditor
 			}
 		}
 
-		// Only remember the interactive state now, the first change is not considered interactive. 
-		bIsInteractiveChangeOngoing = bInteractive;
 		if (IsInteractive())
 		{
 			AccumulatedDelta = AccumulatedDelta.Get<ValueType>() + Delta;
diff --git a/Engine/Plugins/Animation/ControlRig/Source/ControlRigEditor/Private/AnimDetails/AnimDetailsProxyManager.cpp b/Engine/Plugins/Animation/ControlRig/Source/ControlRigEditor/Private/AnimDetails/AnimDetailsProxyManager.cpp
--- a/Engine/Plugins/Animation/ControlRig/Source/ControlRigEditor/Private/AnimDetails/AnimDetailsProxyManager.cpp
+++ b/Engine/Plugins/Animation/ControlRig/Source/ControlRigEditor/Private/AnimDetails/AnimDetailsProxyManager.cpp
@@ -4,6 +4,7 @@
 
 #include "Algo/Find.h"
 #include "Algo/RemoveIf.h"
+#include "AnimDetails/AnimDetailsMultiEditUtil.h"
 #include "AnimDetails/AnimDetailsSelection.h"
 #include "AnimDetails/Proxies/AnimDetailsProxyBool.h"
 #include "AnimDetails/Proxies/AnimDetailsProxyEnum.h"
@@ -247,7 +248,11 @@ void UAnimDetailsProxyManager::OnSequencerSelectionChanged(TArray<FGuid> ObjectG
 
 void UAnimDetailsProxyManager::OnObjectsReplaced(const TMap<UObject*, UObject*>& OldToNewInstanceMap)
 {
-	RequestUpdateProxies();
+	// Objects may be replaced while we are editing interactively, don't update proxies in these cases.
+	if (!UE::ControlRigEditor::FAnimDetailsMultiEditUtil::Get().IsInteractive())
+	{
+		RequestUpdateProxies();
+	}
 }
 
 void UAnimDetailsProxyManager::ForceUpdateProxies()
