Control Rig: Fix crash when validating rig module
#rnx

diff --git a/Engine/Plugins/Animation/ControlRig/Source/ControlRigDeveloper/Private/ControlRigAsset.cpp b/Engine/Plugins/Animation/ControlRig/Source/ControlRigDeveloper/Private/ControlRigAsset.cpp
--- a/Engine/Plugins/Animation/ControlRig/Source/ControlRigDeveloper/Private/ControlRigAsset.cpp
+++ b/Engine/Plugins/Animation/ControlRig/Source/ControlRigDeveloper/Private/ControlRigAsset.cpp
@@ -2549,15 +2549,21 @@ void IControlRigAssetInterface::UpdateModularDependencyDelegates()
 	{
 		if(const UClass* Class = Element->Class.Get())
 		{
-			if(FControlRigAssetInterfacePtr Blueprint = Class->ClassGeneratedBy)
+			if(FControlRigAssetInterfacePtr ModuleBlueprint = Class->ClassGeneratedBy)
 			{
-				if(!VisitList.Contains(Blueprint))
+				if(!VisitList.Contains(ModuleBlueprint))
 				{
-					Blueprint->GetRigVMAssetInterface()->OnVMCompiled().RemoveAll(this);
-					Blueprint->OnModularRigCompiled().RemoveAll(this);
-					Blueprint->GetRigVMAssetInterface()->OnVMCompiled().AddRaw(this, &IControlRigAssetInterface::OnModularDependencyVMCompiled);
-					Blueprint->OnModularRigCompiled().AddRaw(this, &IControlRigAssetInterface::OnModularDependencyChanged);
-					VisitList.Add(Blueprint);
+					ModuleBlueprint->GetRigVMAssetInterface()->OnVMCompiled().RemoveAll(GetObject());
+					ModuleBlueprint->OnModularRigCompiled().RemoveAll(GetObject());
+					ModuleBlueprint->GetRigVMAssetInterface()->OnVMCompiled().AddWeakLambda(GetObject(), [this](UObject* InModuleBlueprint, URigVM* InVM, FRigVMExtendedExecuteContext& InExecuteContext)
+						{
+							OnModularDependencyVMCompiled(InModuleBlueprint, InVM, InExecuteContext);
+						});
+					ModuleBlueprint->OnModularRigCompiled().AddWeakLambda(GetObject(), [this](FRigVMAssetInterfacePtr InModuleBlueprint)
+						{
+							OnModularDependencyChanged(InModuleBlueprint);
+						});
+					VisitList.Add(ModuleBlueprint);
 				}
 			}
 		}
@@ -2565,15 +2571,15 @@ void IControlRigAssetInterface::UpdateModularDependencyDelegates()
 	});
 }
 
-void IControlRigAssetInterface::OnModularDependencyVMCompiled(UObject* InBlueprint, URigVM* InVM, FRigVMExtendedExecuteContext& InExecuteContext)
+void IControlRigAssetInterface::OnModularDependencyVMCompiled(UObject* InModuleBlueprint, URigVM* InVM, FRigVMExtendedExecuteContext& InExecuteContext)
 {
-	if(URigVMBlueprint* RigVMBlueprint = Cast<URigVMBlueprint>(InBlueprint))
+	if(URigVMBlueprint* RigVMBlueprint = Cast<URigVMBlueprint>(InModuleBlueprint))
 	{
 		OnModularDependencyChanged(FRigVMAssetInterfacePtr(RigVMBlueprint));
 	}
 }
 
-void IControlRigAssetInterface::OnModularDependencyChanged(FRigVMAssetInterfacePtr InBlueprint)
+void IControlRigAssetInterface::OnModularDependencyChanged(FRigVMAssetInterfacePtr InModuleBlueprint)
 {
 	RefreshModuleVariables();
 	RefreshModuleConnectors();
