Control Rig: Constraints window now remembers which inline tab you had last opened, and opens that.

diff --git a/Engine/Plugins/Animation/ControlRig/Source/ControlRigEditor/Private/EditMode/Constrain/ConstrainToolsManager.cpp b/Engine/Plugins/Animation/ControlRig/Source/ControlRigEditor/Private/EditMode/Constrain/ConstrainToolsManager.cpp
--- a/Engine/Plugins/Animation/ControlRig/Source/ControlRigEditor/Private/EditMode/Constrain/ConstrainToolsManager.cpp
+++ b/Engine/Plugins/Animation/ControlRig/Source/ControlRigEditor/Private/EditMode/Constrain/ConstrainToolsManager.cpp
@@ -6,6 +6,7 @@
 #include "Framework/Docking/TabManager.h"
 #include "SConstrainToolsRoot.h"
 #include "EditMode/ControlRigEditModeCommands.h"
+#include "EditMode/ControlRigEditModeSettings.h"
 #include "Widgets/Docking/SDockTab.h"
 
 #define LOCTEXT_NAMESPACE "FConstrainToolsManager"
@@ -50,6 +51,12 @@ void FConstrainToolsManager::ToggleVisibility() const
 void FConstrainToolsManager::ShowWidget() const
 {
 	TabManager->TryInvokeTab(FControlRigEditModeToolkit::ConstrainingTabName);
+	
+	if (const TSharedPtr<SConstrainToolsRoot> WidgetContentPin = ConstrainWindowContent.Pin())
+	{
+		EControlRigConstrainTab TabState = UControlRigEditModeSettings::Get()->LastUIStates.ConstraintsTabState.LastOpenInlineTab;
+		WidgetContentPin->OpenTab(TabState);
+	}
 }
 
 void FConstrainToolsManager::HideWidget() const
@@ -87,6 +94,12 @@ TSharedRef<SDockTab> FConstrainToolsManager::SpawnSnapperTab(const FSpawnTabArgs
 	return SAssignNew(SpawnedTab, SDockTab)
 	[
 		SAssignNew(ConstrainWindowContent, SConstrainToolsRoot, OwningEditMode, SelectionViewModel)
+		.OnTabSelected_Lambda([](EControlRigConstrainTab NewTab)
+		{
+			UControlRigEditModeSettings* Settings = UControlRigEditModeSettings::Get();
+			Settings->LastUIStates.ConstraintsTabState.LastOpenInlineTab = NewTab;
+			Settings->SaveConfig();
+		})
 	];
 }
 
@@ -100,9 +113,9 @@ void FConstrainToolsManager::BindCommands()
 	
 	FControlRigEditModeCommands& Commands = FControlRigEditModeCommands::Get();
 	CommandListPin->MapAction(Commands.ToggleConstrainTab, FExecuteAction::CreateRaw(this, &FConstrainToolsManager::ToggleVisibility));
-	CommandListPin->MapAction(Commands.OpenSpacesTab, FExecuteAction::CreateRaw(this, &FConstrainToolsManager::HandleOpenTabCommand, EConstrainTab::Spaces));
-	CommandListPin->MapAction(Commands.OpenConstraintsTab, FExecuteAction::CreateRaw(this, &FConstrainToolsManager::HandleOpenTabCommand, EConstrainTab::Constraints));
-	CommandListPin->MapAction(Commands.OpenSnapperTab, FExecuteAction::CreateRaw(this, &FConstrainToolsManager::HandleOpenTabCommand, EConstrainTab::Snapper));
+	CommandListPin->MapAction(Commands.OpenSpacesTab, FExecuteAction::CreateRaw(this, &FConstrainToolsManager::HandleOpenTabCommand, EControlRigConstrainTab::Spaces));
+	CommandListPin->MapAction(Commands.OpenConstraintsTab, FExecuteAction::CreateRaw(this, &FConstrainToolsManager::HandleOpenTabCommand, EControlRigConstrainTab::Constraints));
+	CommandListPin->MapAction(Commands.OpenSnapperTab, FExecuteAction::CreateRaw(this, &FConstrainToolsManager::HandleOpenTabCommand, EControlRigConstrainTab::Snapper));
 }
 
 void FConstrainToolsManager::UnbindCommands()
@@ -120,7 +133,7 @@ void FConstrainToolsManager::UnbindCommands()
 	CommandListPin->UnmapAction(Commands.OpenSnapperTab);
 }
 
-void FConstrainToolsManager::HandleOpenTabCommand(EConstrainTab InTab)
+void FConstrainToolsManager::HandleOpenTabCommand(EControlRigConstrainTab InTab)
 {
 	// This will flash the tab if already visible, which is desirable.
 	ShowWidget();
diff --git a/Engine/Plugins/Animation/ControlRig/Source/ControlRigEditor/Private/EditMode/Constrain/ConstrainToolsManager.h b/Engine/Plugins/Animation/ControlRig/Source/ControlRigEditor/Private/EditMode/Constrain/ConstrainToolsManager.h
--- a/Engine/Plugins/Animation/ControlRig/Source/ControlRigEditor/Private/EditMode/Constrain/ConstrainToolsManager.h
+++ b/Engine/Plugins/Animation/ControlRig/Source/ControlRigEditor/Private/EditMode/Constrain/ConstrainToolsManager.h
@@ -12,12 +12,12 @@ class FToolBarBuilder;
 class FUICommandList;
 class FWorkspaceItem;
 class SDockTab;
+enum class EControlRigConstrainTab : uint8;
 
 namespace UE::ControlRigEditor
 {
 class FRigSelectionViewModel;
 class SConstrainToolsRoot;
-enum class EConstrainTab : uint8;
 	
 /**
  * Manages UI that is used for constraining, namely:
@@ -70,7 +70,7 @@ private:
 	// Commands
 	void BindCommands();
 	void UnbindCommands();
-	void HandleOpenTabCommand(EConstrainTab InTab);
+	void HandleOpenTabCommand(EControlRigConstrainTab InTab);
 };
 }
 
diff --git a/Engine/Plugins/Animation/ControlRig/Source/ControlRigEditor/Private/EditMode/Constrain/SConstrainToolsRoot.cpp b/Engine/Plugins/Animation/ControlRig/Source/ControlRigEditor/Private/EditMode/Constrain/SConstrainToolsRoot.cpp
--- a/Engine/Plugins/Animation/ControlRig/Source/ControlRigEditor/Private/EditMode/Constrain/SConstrainToolsRoot.cpp
+++ b/Engine/Plugins/Animation/ControlRig/Source/ControlRigEditor/Private/EditMode/Constrain/SConstrainToolsRoot.cpp
@@ -3,13 +3,14 @@
 #include "SConstrainToolsRoot.h"
 
 #include "ControlRigEditorStyle.h"
-#include "SConstraintsEditionWidget_ConstrainPanel.h"
 #include "EditMode/ControlRigEditMode.h"
 #include "EditMode/SControlRigSnapper.h"
 #include "EditMode/SControlRigSpacePicker.h"
+#include "EditMode/Settings/ConstraintsTabRestoreState.h"
 #include "Editor/Constraints/SConstraintsWidget.h"
 #include "Misc/RigEditModeUtils.h"
 #include "Misc/Tab/SInlineTabPanel.h"
+#include "SConstraintsEditionWidget_ConstrainPanel.h"
 
 #define LOCTEXT_NAMESPACE "SConstrainToolsRoot"
 
@@ -20,9 +21,13 @@ void SConstrainToolsRoot::Construct(
 	)
 {
 	static_assert(
-		static_cast<int32>(EConstrainTab::Spaces) == 0 && static_cast<int32>(EConstrainTab::Constraints) == 1 && static_cast<int32>(EConstrainTab::Snapper) == 2,
+		static_cast<int32>(EControlRigConstrainTab::Spaces) == 0
+		&& static_cast<int32>(EControlRigConstrainTab::Constraints) == 1
+		&& static_cast<int32>(EControlRigConstrainTab::Snapper) == 2,
 		"Open tab relies on this order. Update the code accordingly."
 		);
+
+	OnTabSelectedDelegate = InArgs._OnTabSelected;
 	
 	ChildSlot
 	[
@@ -31,6 +36,7 @@ void SConstrainToolsRoot::Construct(
 			FInlineTabArgs()
 				.SetLabel(LOCTEXT("Spaces", "Spaces"))
 				.SetContent(CreateSpacesContent(InOwningEditMode, InSelectionViewModel))
+				.SetOnTabSelected(FSimpleDelegate::CreateSP(this, &SConstrainToolsRoot::BroadcastOnTabSelected, EControlRigConstrainTab::Spaces))
 				.SetIcon(FSlateIcon(FControlRigEditorStyle::Get().GetStyleSetName(), "ControlRig.Constrain.Spaces")),
 			
 			FInlineTabArgs()
@@ -42,12 +48,13 @@ void SConstrainToolsRoot::Construct(
 			FInlineTabArgs()
 				.SetLabel(LOCTEXT("Snapper", "Snapper"))
 				.SetContent(CreateSnapperContent())
+				.SetOnTabSelected(FSimpleDelegate::CreateSP(this, &SConstrainToolsRoot::BroadcastOnTabSelected, EControlRigConstrainTab::Snapper))
 				.SetIcon(FSlateIcon(FControlRigEditorStyle::Get().GetStyleSetName(), "ControlRig.Constrain.Snapper")),
 		})
 	];
 }
 
-void SConstrainToolsRoot::OpenTab(EConstrainTab InTab) const
+void SConstrainToolsRoot::OpenTab(EControlRigConstrainTab InTab) const
 {
 	// The enumeration is declared in display order.
 	InlineTabs->SwitchToTab(static_cast<int32>(InTab));
@@ -84,6 +91,7 @@ TSharedRef<SWidget> SConstrainToolsRoot::CreateSnapperContent()
 void SConstrainToolsRoot::OnTabSelected_ConstraintsEditionWidget() const
 {
 	ConstrainsEditionWidget->RefreshConstraintList();
+	BroadcastOnTabSelected(EControlRigConstrainTab::Constraints);
 }
 }
 
diff --git a/Engine/Plugins/Animation/ControlRig/Source/ControlRigEditor/Private/EditMode/Constrain/SConstrainToolsRoot.h b/Engine/Plugins/Animation/ControlRig/Source/ControlRigEditor/Private/EditMode/Constrain/SConstrainToolsRoot.h
--- a/Engine/Plugins/Animation/ControlRig/Source/ControlRigEditor/Private/EditMode/Constrain/SConstrainToolsRoot.h
+++ b/Engine/Plugins/Animation/ControlRig/Source/ControlRigEditor/Private/EditMode/Constrain/SConstrainToolsRoot.h
@@ -2,10 +2,12 @@
 
 #pragma once
 
+#include "Delegates/Delegate.h"
 #include "Widgets/DeclarativeSyntaxSupport.h"
 #include "Widgets/SCompoundWidget.h"
 
 class FControlRigEditMode;
+enum class EControlRigConstrainTab : uint8;
 
 namespace UE::ControlRigEditor
 {
@@ -13,12 +15,7 @@ class FRigSelectionViewModel;
 class SConstraintsEditionWidget_ConstrainPanel;
 class SInlineTabPanel;
 
-enum class EConstrainTab : uint8
-{
-	Spaces,
-	Constraints,
-	Snapper
-};
+DECLARE_DELEGATE_OneParam(FOnTabSelected, EControlRigConstrainTab);
 	
 /**
  * Coordinator widget that contains
@@ -32,12 +29,14 @@ class SConstrainToolsRoot : public SCompoundWidget
 public:
 	
 	SLATE_BEGIN_ARGS(SConstrainToolsRoot){}
+		/** Invoked when one of the named tabs is selected. */
+		SLATE_EVENT(FOnTabSelected, OnTabSelected)
 	SLATE_END_ARGS()
 
 	void Construct(const FArguments& InArgs, FControlRigEditMode& InOwningEditMode, const TSharedRef<FRigSelectionViewModel>& InSelectionViewModel);
 
 	/** Opens the specified tab. */
-	void OpenTab(EConstrainTab InTab) const;
+	void OpenTab(EControlRigConstrainTab InTab) const;
 	
 private:
 
@@ -46,6 +45,9 @@ private:
 	
 	/** The widget for editing constraints. */
 	TSharedPtr<SConstraintsEditionWidget_ConstrainPanel> ConstrainsEditionWidget;
+	
+	/** Invoked when one of the named tabs is selected. */
+	FOnTabSelected OnTabSelectedDelegate;
 
 	TSharedRef<SWidget> CreateSpacesContent(FControlRigEditMode& InOwningEditMode, const TSharedRef<FRigSelectionViewModel>& InSelectionViewModel);
 	TSharedRef<SWidget> CreateConstraintsContent(FControlRigEditMode& InOwningMode, const TSharedRef<FRigSelectionViewModel>& InSelectionViewModel);
@@ -53,5 +55,7 @@ private:
 
 	/** Refreshes the list when we switch to the tab. */
 	void OnTabSelected_ConstraintsEditionWidget() const;
+	/** Invoked when one of the named tabs is selected. */
+	void BroadcastOnTabSelected(EControlRigConstrainTab ControlRigConstrainTab) const { OnTabSelectedDelegate.ExecuteIfBound(ControlRigConstrainTab); }
 };
 }
diff --git a/Engine/Plugins/Animation/ControlRig/Source/ControlRigEditor/Private/EditMode/ControlRigEditModeToolkit.cpp b/Engine/Plugins/Animation/ControlRig/Source/ControlRigEditor/Private/EditMode/ControlRigEditModeToolkit.cpp
--- a/Engine/Plugins/Animation/ControlRig/Source/ControlRigEditor/Private/EditMode/ControlRigEditModeToolkit.cpp
+++ b/Engine/Plugins/Animation/ControlRig/Source/ControlRigEditor/Private/EditMode/ControlRigEditModeToolkit.cpp
@@ -496,7 +496,7 @@ void FControlRigEditModeToolkit::InvokeUI()
 	}
 #endif	
 
-	if (RestoreStates.bConstraintsOpen)
+	if (RestoreStates.ConstraintsTabState.bWasOpen)
 	{
 		ConstrainToolsManager->ShowWidget();
 	}
@@ -623,7 +623,7 @@ void FControlRigEditModeToolkit::SaveLayoutState() const
 	RestoreStates.bAnimLayerTabOpen = IsToolkitUIActive(AnimLayerTabName);
 	RestoreStates.bPoseTabOpen = IsToolkitUIActive(PoseTabName);
 	RestoreStates.bSnapperTabOpen = IsToolkitUIActive(SnapperTabName);
-	RestoreStates.bConstraintsOpen = IsToolkitUIActive(ConstrainingTabName);
+	RestoreStates.ConstraintsTabState.bWasOpen = IsToolkitUIActive(ConstrainingTabName);
 #if SELECTION_SETS_AS_TAB
 	RestoreStates.bSelectionSetsOpen = IsToolkitUIActive(SelectionSetsTabName);
 #else
diff --git a/Engine/Plugins/Animation/ControlRig/Source/ControlRigEditor/Private/EditMode/Settings/ConstraintsTabRestoreState.h b/Engine/Plugins/Animation/ControlRig/Source/ControlRigEditor/Private/EditMode/Settings/ConstraintsTabRestoreState.h
new file mode 100644
--- /dev/null
+++ b/Engine/Plugins/Animation/ControlRig/Source/ControlRigEditor/Private/EditMode/Settings/ConstraintsTabRestoreState.h
@@ -0,0 +1,26 @@
+﻿// Copyright Epic Games, Inc. All Rights Reserved.
+
+
+#pragma once
+
+#include "ConstraintsTabRestoreState.generated.h"
+
+UENUM()
+enum class EControlRigConstrainTab : uint8
+{
+	Spaces,
+	Constraints,
+	Snapper
+};
+
+USTRUCT()
+struct FControlRigConstraintsTabRestoreState
+{
+	GENERATED_BODY()
+
+	UPROPERTY()
+	bool bWasOpen = false;
+
+	UPROPERTY() 
+	EControlRigConstrainTab LastOpenInlineTab = EControlRigConstrainTab::Spaces;
+};
\ No newline at end of file
diff --git a/Engine/Plugins/Animation/ControlRig/Source/ControlRigEditor/Private/EditMode/Settings/ControlRigUIRestoreStates.h b/Engine/Plugins/Animation/ControlRig/Source/ControlRigEditor/Private/EditMode/Settings/ControlRigUIRestoreStates.h
--- a/Engine/Plugins/Animation/ControlRig/Source/ControlRigEditor/Private/EditMode/Settings/ControlRigUIRestoreStates.h
+++ b/Engine/Plugins/Animation/ControlRig/Source/ControlRigEditor/Private/EditMode/Settings/ControlRigUIRestoreStates.h
@@ -2,6 +2,7 @@
 
 #pragma once
 
+#include "ConstraintsTabRestoreState.h"
 #include "Misc/Flyout/FlyoutSavedState.h"
 #include "Overlay/DragBoxPosition.h"
 #include "ControlRigUIRestoreStates.generated.h"
@@ -25,7 +26,7 @@ struct FControlRigUIRestoreStates
 	bool bSnapperTabOpen = false;
 	
 	UPROPERTY()
-	bool bConstraintsOpen = false;
+	FControlRigConstraintsTabRestoreState ConstraintsTabState;
 
 	/** Whether the selection set was open when SELECTION_SETS_AS_TAB == 1 */
 	UPROPERTY()
diff --git a/Engine/Plugins/Animation/ControlRig/Source/ControlRigEditor/Private/Misc/Tab/SInlineTabPanel.cpp b/Engine/Plugins/Animation/ControlRig/Source/ControlRigEditor/Private/Misc/Tab/SInlineTabPanel.cpp
--- a/Engine/Plugins/Animation/ControlRig/Source/ControlRigEditor/Private/Misc/Tab/SInlineTabPanel.cpp
+++ b/Engine/Plugins/Animation/ControlRig/Source/ControlRigEditor/Private/Misc/Tab/SInlineTabPanel.cpp
@@ -52,7 +52,7 @@ static TArray<FTabEntry> Transform(TConstArrayView<FInlineTabArgs> InArgs, const
 				.OverflowPolicy(ETextOverflowPolicy::Ellipsis) // ... and cut off the rest of the text when there's not enough space
 			];
 		
-		Entry.OnTabSelected.BindLambda([Index, WeakOwnerPanel = This.ToWeakPtr(), Delegate = Entry.OnTabSelected]()
+		Entry.OnTabSelected.BindLambda([Index, WeakOwnerPanel = This.ToWeakPtr(), Delegate = CurrentArgs.OnTabSelected]()
 		{
 			if (const TSharedPtr<SInlineTabPanel> OwnerPin = WeakOwnerPanel.Pin())
 			{
