Control Rig: Avoid opening graph in modular rig editors
#rnx

diff --git a/Engine/Plugins/Animation/ControlRig/Source/ControlRigEditor/Private/Editor/ControlRigEditorMode.cpp b/Engine/Plugins/Animation/ControlRig/Source/ControlRigEditor/Private/Editor/ControlRigEditorMode.cpp
--- a/Engine/Plugins/Animation/ControlRig/Source/ControlRigEditor/Private/Editor/ControlRigEditorMode.cpp
+++ b/Engine/Plugins/Animation/ControlRig/Source/ControlRigEditor/Private/Editor/ControlRigEditorMode.cpp
@@ -165,7 +165,7 @@ FModularRigLegacyEditorMode::FModularRigLegacyEditorMode(const TSharedRef<FContr
 	TabFactories.RegisterFactory(MakeShared<FRigModuleAssetBrowserTabSummoner>(InControlRigEditor));
 	TabFactories.RegisterFactory(MakeShared<FModularRigEventQueueTabSummoner>(InControlRigEditor));
 
-	TabLayout = FTabManager::NewLayout("Standalone_ModularRigEditMode_Layout_v1.4")
+	TabLayout = FTabManager::NewLayout("Standalone_ModularRigEditMode_Layout_v1.5")
 		->AddArea
 		(
 			// Main application area
@@ -218,7 +218,6 @@ FModularRigLegacyEditorMode::FModularRigLegacyEditorMode(const TSharedRef<FContr
 					->AddTab(FBlueprintEditorTabs::DetailsID, ETabState::OpenedTab)
 					->AddTab(FPersonaTabs::AdvancedPreviewSceneSettingsID, ETabState::OpenedTab)
 					->AddTab(FRigHierarchyTabSummoner::TabID, ETabState::OpenedTab)
-					->AddTab("GraphEditor", ETabState::ClosedTab) // Adding the graph editor to avoid a warning about blueprint assumption there will be a graph editor
 				)
 			)
 		);
diff --git a/Engine/Plugins/Runtime/RigVM/Source/RigVMEditor/Public/Editor/RigVMEditor.h b/Engine/Plugins/Runtime/RigVM/Source/RigVMEditor/Public/Editor/RigVMEditor.h
--- a/Engine/Plugins/Runtime/RigVM/Source/RigVMEditor/Public/Editor/RigVMEditor.h
+++ b/Engine/Plugins/Runtime/RigVM/Source/RigVMEditor/Public/Editor/RigVMEditor.h
@@ -364,7 +364,7 @@ protected:
 	UE_API FReply OnSpawnGraphNodeByShortcut(FInputChord InChord, const FVector2f& InPosition, UEdGraph* InGraph);
 	bool ShouldLoadBPLibrariesFromAssetRegistry() { return false; }
 	UE_API bool JumpToHyperlinkImpl(const UObject* ObjectReference, bool bRequestRename = false);
-	bool ShouldOpenGraphByDefault() const { return true; }
+	virtual bool ShouldOpenGraphByDefault() const { return true; }
 	UE_API virtual void AddNewFunctionVariant(const UEdGraph* InOriginalFunction) override;
 
 	// FEditorUndoClient Interface
