Control Rig: Max clipping priority for the control rig menu in the viewport toolbar

#rnx

diff --git a/Engine/Plugins/Animation/ControlRig/Source/ControlRigEditor/Private/Editor/ControlRigViewportToolbarExtensionControlRig.cpp b/Engine/Plugins/Animation/ControlRig/Source/ControlRigEditor/Private/Editor/ControlRigViewportToolbarExtensionControlRig.cpp
--- a/Engine/Plugins/Animation/ControlRig/Source/ControlRigEditor/Private/Editor/ControlRigViewportToolbarExtensionControlRig.cpp
+++ b/Engine/Plugins/Animation/ControlRig/Source/ControlRigEditor/Private/Editor/ControlRigViewportToolbarExtensionControlRig.cpp
@@ -20,6 +20,7 @@
 #include "UObject/WeakInterfacePtr.h"
 #include "ViewportToolbar/UnrealEdViewportToolbarContext.h"
 #include "Widgets/Input/SNumericEntryBox.h"
+#include <limits>
 
 #define LOCTEXT_NAMESPACE "ControlRigViewportToolbarExtensionControlRig"
 
@@ -113,6 +114,9 @@ namespace UE::ControlRigEditor
 			Entry.Icon = FSlateIcon(FControlRigEditorStyle::Get().GetStyleSetName(), "ControlRig.Editor.TabIcon");
 			Entry.ToolBarData.LabelOverride = FText::GetEmpty();
 
+			// Set resize prio to highest as other entries in other sections dominate otherwise
+			Entry.ToolBarData.ResizeParams.ClippingPriority = std::numeric_limits<int32>::max(); 
+
 			return Entry;
 		}
 	}
