Fixed shape details customization not working on Modular Control Rigs and added multiselection support

#rnx

diff --git a/Engine/Plugins/Animation/ControlRig/Source/ControlRigEditor/Private/ControlRigShapeSettingsDetails.cpp b/Engine/Plugins/Animation/ControlRig/Source/ControlRigEditor/Private/ControlRigShapeSettingsDetails.cpp
--- a/Engine/Plugins/Animation/ControlRig/Source/ControlRigEditor/Private/ControlRigShapeSettingsDetails.cpp
+++ b/Engine/Plugins/Animation/ControlRig/Source/ControlRigEditor/Private/ControlRigShapeSettingsDetails.cpp
@@ -30,17 +30,23 @@ void FControlRigShapeSettingsDetails::CustomizeHeader(TSharedRef<IPropertyHandle
 {
 	TArray<UObject*> Objects;
 	InStructPropertyHandle->GetOuterObjects(Objects);
-	
-	// This block is in here to ensure we are only showing the modifier details in the blueprint editor
-	if (Objects.Num() == 1)
+
+	ControlRigObjects.Reset();
+
+	// If one of the objects is either a Control Rig or a Control Rig Wrapper, we can generate the combo box for selection
+	// Else, a FName Text Box will be shown (the default property details)
+	for (UObject* Object : Objects)
 	{
-		for (UObject* Object : Objects)
+		if (UControlRig* ControlRig = Cast<UControlRig>(Object))	// Needed by Modular Rig module details
 		{
-			if (Object->IsA<UControlRigWrapperObject>())
+			ControlRigObjects.AddUnique(ControlRig);
+		}
+		else if (UControlRigWrapperObject* ControlRigWrapperObject = Cast<UControlRigWrapperObject>(Object)) // Needed by Control Rig editor node editor details
+		{
+			if (UControlRig* WrappedControlRig = GetControlRig(ControlRigWrapperObject))
 			{
-				// Note that multi selection of rigs is not currently supported (each rig has to show a different name list)
-				ControlRigWrapperObject = Cast<UControlRigWrapperObject>(Object);
-				break;
+				// This customization works for multiselection of nodes with the same rig
+				ControlRigObjects.AddUnique(WrappedControlRig);
 			}
 		}
 	}
@@ -48,42 +54,67 @@ void FControlRigShapeSettingsDetails::CustomizeHeader(TSharedRef<IPropertyHandle
 	InHeaderRow.NameContent()
 		[
 			InStructPropertyHandle->CreatePropertyNameWidget()
-		];	
+		];
 }
 
 void FControlRigShapeSettingsDetails::CustomizeChildren(TSharedRef<IPropertyHandle> InStructPropertyHandle, IDetailChildrenBuilder& InStructBuilder, IPropertyTypeCustomizationUtils& InStructCustomizationUtils)
 {
-	if (InStructPropertyHandle->IsValidHandle() && ControlRigWrapperObject != nullptr)
+	if (InStructPropertyHandle->IsValidHandle())
 	{
-		if (UControlRig* ControlRig = GetControlRig(ControlRigWrapperObject))
+		uint32 NumChildren = 0;
+		if (InStructPropertyHandle->GetNumChildren(NumChildren) == FPropertyAccess::Success)
 		{
-			ControlRigShapeNameList->GenerateShapeLibraryList(ControlRig);
+			TSharedPtr<IPropertyHandle> ShapeSettingsNameProperty = InStructPropertyHandle->GetChildHandle(GET_MEMBER_NAME_CHECKED(FRigUnit_HierarchyAddControl_ShapeSettings, Name));
+			if (!ShapeSettingsNameProperty.IsValid())
+			{
+				return;
+			}
 
-			uint32 NumChildren = 0;
-			if (InStructPropertyHandle->GetNumChildren(NumChildren) == FPropertyAccess::Success)
+			for (uint32 Index = 0; Index < NumChildren; ++Index)
 			{
-				TSharedPtr<IPropertyHandle> ShapeSettingsNameProperty = InStructPropertyHandle->GetChildHandle(GET_MEMBER_NAME_CHECKED(FRigUnit_HierarchyAddControl_ShapeSettings, Name));
-				if (!ShapeSettingsNameProperty.IsValid())
-				{
-					return;
-				}
+				bool bDefaultContent = true;
 
-				for (uint32 Index = 0; Index < NumChildren; ++Index)
+				TSharedRef<IPropertyHandle> ChildPropertyHandle = InStructPropertyHandle->GetChildHandle(Index).ToSharedRef();
+				if (ChildPropertyHandle->IsSamePropertyNode(ShapeSettingsNameProperty))
 				{
-					bool bDefaultContent = true;
-
-					TSharedRef<IPropertyHandle> ChildPropertyHandle = InStructPropertyHandle->GetChildHandle(Index).ToSharedRef();
-					if (ChildPropertyHandle->IsSamePropertyNode(ShapeSettingsNameProperty))
+					const int32 NumUniqueControlRigs = ControlRigObjects.Num();
+					if (NumUniqueControlRigs > 0)
 					{
-						ControlRigShapeNameList->CreateShapeLibraryListWidget(InStructBuilder, ShapeSettingsNameProperty);
-						bDefaultContent = false;
-					}
+						if (NumUniqueControlRigs == 1)
+						{
+							if (UControlRig* ControlRig = ControlRigObjects[0].Get())
+							{
+								ControlRigShapeNameList->GenerateShapeLibraryList(ControlRig);
+								ControlRigShapeNameList->CreateShapeLibraryListWidget(InStructBuilder, ShapeSettingsNameProperty);
+								bDefaultContent = false;
+							}
+						}
+						else
+						{
+							IDetailPropertyRow& Row = InStructBuilder.AddProperty(ShapeSettingsNameProperty.ToSharedRef());
 
-					if (bDefaultContent)
-					{
-						InStructBuilder.AddProperty(ChildPropertyHandle);
+							constexpr bool bShowChildren = true;
+							Row.CustomWidget(bShowChildren)
+								.NameContent()
+								[
+									ShapeSettingsNameProperty->CreatePropertyNameWidget()
+								]
+								.ValueContent()
+								[
+									SNew(STextBlock)
+										.Text(LOCTEXT("ShapeSettings_MultipleControlRigs", "Multiple Control Rigs in selection"))
+										.Font(IDetailLayoutBuilder::GetDetailFont())
+								];
+
+							bDefaultContent = false;
+						}
 					}
 				}
+
+				if (bDefaultContent)
+				{
+					InStructBuilder.AddProperty(ChildPropertyHandle);
+				}
 			}
 		}
 	}
@@ -91,14 +122,17 @@ void FControlRigShapeSettingsDetails::CustomizeChildren(TSharedRef<IPropertyHand
 
 UControlRig* FControlRigShapeSettingsDetails::GetControlRig(UControlRigWrapperObject* InControlRigWrapperObject)
 {
-	if (URigVMUnitNode* UnitNode = Cast<URigVMUnitNode>(InControlRigWrapperObject->GetSubject()))
+	if (InControlRigWrapperObject)
 	{
-		FControlRigAssetInterfacePtr AssetInterface = IControlRigAssetInterface::GetInterfaceOuter(UnitNode);
-		if (AssetInterface)
+		if (URigVMUnitNode* UnitNode = Cast<URigVMUnitNode>(InControlRigWrapperObject->GetSubject()))
 		{
-			if (UControlRig* ControlRig = AssetInterface->GetDebuggedControlRig())
+			FControlRigAssetInterfacePtr AssetInterface = IControlRigAssetInterface::GetInterfaceOuter(UnitNode);
+			if (AssetInterface)
 			{
-				return ControlRig;
+				if (UControlRig* ControlRig = AssetInterface->GetDebuggedControlRig())
+				{
+					return ControlRig;
+				}
 			}
 		}
 	}
diff --git a/Engine/Plugins/Animation/ControlRig/Source/ControlRigEditor/Private/ControlRigShapeSettingsDetails.h b/Engine/Plugins/Animation/ControlRig/Source/ControlRigEditor/Private/ControlRigShapeSettingsDetails.h
--- a/Engine/Plugins/Animation/ControlRig/Source/ControlRigEditor/Private/ControlRigShapeSettingsDetails.h
+++ b/Engine/Plugins/Animation/ControlRig/Source/ControlRigEditor/Private/ControlRigShapeSettingsDetails.h
@@ -32,7 +32,7 @@ public:
 	virtual void CustomizeChildren(TSharedRef<class IPropertyHandle> InStructPropertyHandle, class IDetailChildrenBuilder& InStructBuilder, IPropertyTypeCustomizationUtils& InStructCustomizationUtils) override;
 
 protected:
-	UControlRigWrapperObject* ControlRigWrapperObject = nullptr;
+	TArray<TWeakObjectPtr<UControlRig>> ControlRigObjects;
 	TSharedPtr<FControlRigShapeNameList> ControlRigShapeNameList;
 
 	static UControlRig* GetControlRig(UControlRigWrapperObject* ControlRigWrapperObject);
