Control Rig: Fix change parent in sequencer not behaving correctly after save
#rnx

diff --git a/Engine/Plugins/Animation/ControlRig/Source/ControlRig/Private/Rigs/RigHierarchy.cpp b/Engine/Plugins/Animation/ControlRig/Source/ControlRig/Private/Rigs/RigHierarchy.cpp
--- a/Engine/Plugins/Animation/ControlRig/Source/ControlRig/Private/Rigs/RigHierarchy.cpp
+++ b/Engine/Plugins/Animation/ControlRig/Source/ControlRig/Private/Rigs/RigHierarchy.cpp
@@ -3807,6 +3807,9 @@ bool URigHierarchy::SetParentWeightArray(FRigBaseElement* InChild,  const TArray
 				MultiParentElement->GetDirtyState().MarkDirty(LocalType);
 			}
 
+			// Propagate before switching parent to make sure all children have a clean state
+			PropagateDirtyFlags(MultiParentElement, bInitial, bAffectChildren, true, false);
+
 			// first set all the zero weights
 			for(int32 WeightIndex=0; WeightIndex < InWeights.Num(); WeightIndex++)
 			{
diff --git a/Engine/Plugins/Animation/ControlRig/Source/ControlRig/Private/Sequencer/MovieSceneControlRigParameterTemplate.cpp b/Engine/Plugins/Animation/ControlRig/Source/ControlRig/Private/Sequencer/MovieSceneControlRigParameterTemplate.cpp
--- a/Engine/Plugins/Animation/ControlRig/Source/ControlRig/Private/Sequencer/MovieSceneControlRigParameterTemplate.cpp
+++ b/Engine/Plugins/Animation/ControlRig/Source/ControlRig/Private/Sequencer/MovieSceneControlRigParameterTemplate.cpp
@@ -1050,6 +1050,13 @@ struct FControlRigParameterPreAnimatedTokenProducer : IMovieScenePreAnimatedToke
 					}
 					case ERigControlType::EulerTransform:
 					{
+						FMovieSceneControlRigSpaceBaseKey SpaceValue;
+						//@helge how can we get the correct current space here? this is for restoring it.
+						//for now just using parent space
+						//FRigElementKey DefaultParent = RigHierarchy->GetFirstParent(ControlElement->GetKey());
+						SpaceValue.ControlRigElement = ControlElement->GetKey();
+						SpaceValue.SpaceType = EMovieSceneControlRigSpaceType::Parent;
+						Token.SpaceValues.Add(FControlSpaceAndValue(ControlElement->GetFName(), SpaceValue));
 						FEulerTransform EulerTransform = 
 							ControlRig
 							->GetControlValue(ControlElement, ERigControlValueType::Current).Get<FRigControlValue::FEulerTransform_Float>().ToTransform();
diff --git a/Engine/Plugins/Animation/ControlRig/Source/ControlRig/Public/Rigs/RigHierarchy.h b/Engine/Plugins/Animation/ControlRig/Source/ControlRig/Public/Rigs/RigHierarchy.h
--- a/Engine/Plugins/Animation/ControlRig/Source/ControlRig/Public/Rigs/RigHierarchy.h
+++ b/Engine/Plugins/Animation/ControlRig/Source/ControlRig/Public/Rigs/RigHierarchy.h
@@ -977,15 +977,19 @@ public:
 
 			Traverse([&ElementVisited, &Results](FRigBaseElement* InElement, bool& bContinue)
 			{
-				bContinue = !ElementVisited[InElement->GetIndex()];
-
-				if(bContinue)
+				bContinue = false;
+				if (ElementVisited.IsValidIndex(InElement->GetIndex()))
 				{
-					if(T* CastElement = Cast<T>(InElement))
-					{
-						Results.Add(CastElement);
-					}
-					ElementVisited[InElement->GetIndex()] = true;
+					bContinue = !ElementVisited[InElement->GetIndex()];
+
+				   if(bContinue)
+				   {
+					   if(T* CastElement = Cast<T>(InElement))
+					   {
+						   Results.Add(CastElement);
+					   }
+					   ElementVisited[InElement->GetIndex()] = true;
+				   }
 				}
 			});
 		}
diff --git a/Engine/Source/Editor/AnimationEditorWidgets/Private/SchematicGraphPanel/SSchematicGraphPanel.cpp b/Engine/Source/Editor/AnimationEditorWidgets/Private/SchematicGraphPanel/SSchematicGraphPanel.cpp
--- a/Engine/Source/Editor/AnimationEditorWidgets/Private/SchematicGraphPanel/SSchematicGraphPanel.cpp
+++ b/Engine/Source/Editor/AnimationEditorWidgets/Private/SchematicGraphPanel/SSchematicGraphPanel.cpp
@@ -710,7 +710,7 @@ void SSchematicGraphPanel::RemoveNode(const FSchematicGraphNode* InNodeToRemove)
 	const FGuid GuidToRemove = InNodeToRemove->GetGuid();
 	NodeByGuid.Remove(GuidToRemove);
 	
-	for (int32 Iter = 0; Iter != Children.Num(); ++Iter)
+	for (int32 Iter = Children.Num()-1; Iter >= 0; --Iter)
 	{
 		TSharedRef<SSchematicGraphNode> Widget = GetChild(Iter);
 		if (Widget->GetGuid() == GuidToRemove)
