Infinite constraints: be sure to add the current time if the child is not keyframed at the time of creation

#rnx

diff --git a/Engine/Plugins/Animation/ControlRig/Source/ControlRigEditor/Private/Sequencer/ControlConstraintChannelInterface.cpp b/Engine/Plugins/Animation/ControlRig/Source/ControlRigEditor/Private/Sequencer/ControlConstraintChannelInterface.cpp
--- a/Engine/Plugins/Animation/ControlRig/Source/ControlRigEditor/Private/Sequencer/ControlConstraintChannelInterface.cpp
+++ b/Engine/Plugins/Animation/ControlRig/Source/ControlRigEditor/Private/Sequencer/ControlConstraintChannelInterface.cpp
@@ -151,14 +151,27 @@ bool FControlConstraintChannelInterface::SmartConstraintKey(
 			
 			const FFrameRate TickResolution = InSequencer->GetFocusedTickResolution();
 			const FFrameTime CurrentTime = InSequencer->GetLocalTime().ConvertTo(TickResolution);
+			const FFrameNumber CurrentFrame = CurrentTime.GetFrame();
 			
-			// check if the first active key will be added before the current frame
-			// if that's the case then the current frame must be added and the compensation evaluator notified
 			TOptional<FFrameNumber> OptCurrentTime;
 			bool bPropagateRelative = true;
-			if (!bInfiniteInput && !bWasInfinite && !bHadConstraintChannel && ActiveValueToBeSet)
+			if (!bHadConstraintChannel && ActiveValueToBeSet)
 			{
-				Tie(OptCurrentTime, bPropagateRelative) = IsActiveKeyAddedBeforeCurrentTime(CurrentTime.GetFrame(), *Time, Channels, FramesToCompensate);
+				if (!bInfiniteInput && !bWasInfinite)
+				{
+					// check if the first active key will be added before the current frame
+                    // if that's the case then the current frame must be added and the compensation evaluator notified
+					Tie(OptCurrentTime, bPropagateRelative) =
+						IsActiveKeyAddedBeforeCurrentTime(CurrentFrame, *Time, Channels, FramesToCompensate);
+				}
+				else if (bInfiniteInput && !FramesToCompensate.IsEmpty())
+				{
+					// make sure to add the current frame first if the child has transforms to compensate
+					FramesToCompensate.Remove(CurrentFrame);
+					FramesToCompensate.Insert(CurrentFrame, 0);
+					OptCurrentTime = CurrentFrame;
+					bPropagateRelative = false;
+				}
 			}
 			
 			FCompensationEvaluator::FEvalParameters EvalParams(InSequencer, FramesToCompensate);
diff --git a/Engine/Source/Editor/MovieSceneTools/Private/Constraints/ComponentConstraintChannelInterface.cpp b/Engine/Source/Editor/MovieSceneTools/Private/Constraints/ComponentConstraintChannelInterface.cpp
--- a/Engine/Source/Editor/MovieSceneTools/Private/Constraints/ComponentConstraintChannelInterface.cpp
+++ b/Engine/Source/Editor/MovieSceneTools/Private/Constraints/ComponentConstraintChannelInterface.cpp
@@ -135,14 +135,27 @@ bool FComponentConstraintChannelInterface::SmartConstraintKey(
 			
 			const FFrameRate TickResolution = InSequencer->GetFocusedTickResolution();
 			const FFrameTime CurrentTime = InSequencer->GetLocalTime().ConvertTo(TickResolution);
+			const FFrameNumber CurrentFrame = CurrentTime.GetFrame();
 			
-			// check if the first active key will be added before the current frame
-			// if that's the case then the current frame must be added and the compensation evaluator notified
 			TOptional<FFrameNumber> OptCurrentTime;
 			bool bPropagateRelative = true;
-			if (!bInfiniteInput && !bWasInfinite && !bHadConstraintChannel && ActiveValueToBeSet)
+			if (!bHadConstraintChannel && ActiveValueToBeSet)
 			{
-				Tie(OptCurrentTime, bPropagateRelative) = IsActiveKeyAddedBeforeCurrentTime(CurrentTime.GetFrame(), *Time, Channels, FramesToCompensate);
+				if (!bInfiniteInput && !bWasInfinite)
+				{
+					// check if the first active key will be added before the current frame
+					// if that's the case then the current frame must be added and the compensation evaluator notified
+					Tie(OptCurrentTime, bPropagateRelative) =
+						IsActiveKeyAddedBeforeCurrentTime(CurrentFrame, *Time, Channels, FramesToCompensate);
+				}
+				else if (bInfiniteInput && !FramesToCompensate.IsEmpty())
+				{
+					// make sure to add the current frame first if the child has transforms to compensate
+					FramesToCompensate.Remove(CurrentFrame);
+					FramesToCompensate.Insert(CurrentFrame, 0);
+					OptCurrentTime = CurrentFrame;
+					bPropagateRelative = false;
+				}
 			}
 
 			FCompensationEvaluator::FEvalParameters EvalParams(InSequencer, FramesToCompensate);
