Control Rig: ECS version of 46252663

#rnx

diff --git a/Engine/Plugins/Animation/ControlRig/Source/ControlRig/Private/Sequencer/MovieSceneControlRigParameterBuffer.cpp b/Engine/Plugins/Animation/ControlRig/Source/ControlRig/Private/Sequencer/MovieSceneControlRigParameterBuffer.cpp
--- a/Engine/Plugins/Animation/ControlRig/Source/ControlRig/Private/Sequencer/MovieSceneControlRigParameterBuffer.cpp
+++ b/Engine/Plugins/Animation/ControlRig/Source/ControlRig/Private/Sequencer/MovieSceneControlRigParameterBuffer.cpp
@@ -584,6 +584,18 @@ void FControlRigParameterValues::AddCurrentValue(UControlRig* Rig, FRigControlEl
 {
 	const FName ControlName = ControlElement->GetFName();
 
+	auto GetSpaceValue = [ControlElement]()
+	{
+		//@helge how can we get the correct current space here? this is for restoring it.
+		//for now just using parent space
+		//FRigElementKey DefaultParent = RigHierarchy->GetFirstParent(ControlElement->GetKey());
+		FMovieSceneControlRigSpaceBaseKey SpaceValue;
+		SpaceValue.ControlRigElement = ControlElement->GetKey();
+		SpaceValue.SpaceType = EMovieSceneControlRigSpaceType::Parent;
+
+		return SpaceValue;
+	};
+	
 	switch (ControlElement->Settings.ControlType)
 	{
 	case ERigControlType::Bool:
@@ -627,13 +639,7 @@ void FControlRigParameterValues::AddCurrentValue(UControlRig* Rig, FRigControlEl
 	case ERigControlType::Scale:
 	case ERigControlType::Rotator:
 	{
-		FMovieSceneControlRigSpaceBaseKey SpaceValue;
-
-		//@helge how can we get the correct current space here? this is for restoring it.
-		//for now just using parent space
-		//FRigElementKey DefaultParent = RigHierarchy->GetFirstParent(ControlElement->GetKey());
-		SpaceValue.ControlRigElement = ControlElement->GetKey();
-		SpaceValue.SpaceType = EMovieSceneControlRigSpaceType::Parent;
+		const FMovieSceneControlRigSpaceBaseKey SpaceValue = GetSpaceValue();
 		Add(ControlName, SpaceValue);
 
 		FVector3f Val = Rig->GetControlValue(ControlElement, ERigControlValueType::Current).Get<FVector3f>();
@@ -649,12 +655,7 @@ void FControlRigParameterValues::AddCurrentValue(UControlRig* Rig, FRigControlEl
 
 	case ERigControlType::Transform:
 	{
-		FMovieSceneControlRigSpaceBaseKey SpaceValue;
-		//@helge how can we get the correct current space here? this is for restoring it.
-		//for now just using parent space
-		//FRigElementKey DefaultParent = RigHierarchy->GetFirstParent(ControlElement->GetKey());
-		SpaceValue.ControlRigElement = ControlElement->GetKey();
-		SpaceValue.SpaceType = EMovieSceneControlRigSpaceType::Parent;
+		const FMovieSceneControlRigSpaceBaseKey SpaceValue = GetSpaceValue();
 		Add(ControlName, SpaceValue);
 
 		const FTransform Val = Rig->GetControlValue(ControlElement, ERigControlValueType::Current).Get<FRigControlValue::FTransform_Float>().ToTransform();
@@ -676,7 +677,10 @@ void FControlRigParameterValues::AddCurrentValue(UControlRig* Rig, FRigControlEl
 		break;
 	}
 	case ERigControlType::EulerTransform:
-	{
+	{	
+		const FMovieSceneControlRigSpaceBaseKey SpaceValue = GetSpaceValue();
+		Add(ControlName, SpaceValue);
+			
 		FEulerTransform EulerTransform =
 			Rig
 			->GetControlValue(ControlElement, ERigControlValueType::Current).Get<FRigControlValue::FEulerTransform_Float>().ToTransform();
