Anim in Engine: Fix an issue where Anim Details in some cases shows channels as keyed on another frame when they aren't keyed at all

diff --git a/Engine/Plugins/Animation/ControlRig/Source/ControlRigEditor/Private/AnimDetails/Proxies/AnimDetailsProxyBase.cpp b/Engine/Plugins/Animation/ControlRig/Source/ControlRigEditor/Private/AnimDetails/Proxies/AnimDetailsProxyBase.cpp
--- a/Engine/Plugins/Animation/ControlRig/Source/ControlRigEditor/Private/AnimDetails/Proxies/AnimDetailsProxyBase.cpp
+++ b/Engine/Plugins/Animation/ControlRig/Source/ControlRigEditor/Private/AnimDetails/Proxies/AnimDetailsProxyBase.cpp
@@ -498,7 +498,7 @@ namespace UE::ControlRigEditor
 				break;
 			}
 
-			int32 EmptyChannelCount = 0;
+			bool bKeyedOnAnotherFrame = false;
 			for (const FMovieSceneChannelEntry& ChannelEntry : ChannelProxy.GetAllEntries())
 			{
 				if (ChannelEntry.GetChannelTypeName() != FMovieSceneDoubleChannel::StaticStruct()->GetFName() &&
@@ -527,7 +527,6 @@ namespace UE::ControlRigEditor
 					const int32 NumKeys = Channel->GetNumKeys();
 					if (NumKeys == 0)
 					{
-						++EmptyChannelCount;
 						continue;
 					}
 
@@ -537,7 +536,7 @@ namespace UE::ControlRigEditor
 					Channel->GetKeys(Range, &KeyTimesInRange, nullptr);
 					if (KeyTimesInRange.IsEmpty() && NumKeys == 0)
 					{
-						++EmptyChannelCount;
+						bKeyedOnAnotherFrame |= true;
 					}
 					else if (KeyTimesInRange.IsEmpty() && NumKeys > 0)
 					{
@@ -552,10 +551,9 @@ namespace UE::ControlRigEditor
 				break; // just do it for one type
 			}
 
-			const bool bIsPartiallySelectedGroup = EmptyChannelCount > 0;
-			if (bIsPartiallySelectedGroup)
+			if (bKeyedOnAnotherFrame)
 			{
-				SectionKeyedStatus = EPropertyKeyedStatus::PartiallyKeyed;
+				SectionKeyedStatus = FMath::Max(SectionKeyedStatus, EPropertyKeyedStatus::PartiallyKeyed);
 			}
 
 			return SectionKeyedStatus;
