Control Rig: Update Set Initial Transform from Current and related Commands

* Set Initial Transform from Current is now only applicable for Bones, Nulls and Connectors
* For Controls, Set Initital Transform from Current is now Zero Controls
* For Controls, Set Initital Transfrom from Closest Bone is now Zero Controls from Closest Bone, and no longer applicable for Bones, Nulls and Connectors
* Set Shape Transfrom from Current  is now Zero Control Shape
* Update tooltips
* Rename related methods accordingly

diff --git a/Engine/Plugins/Animation/ControlRig/Source/ControlRigEditor/Private/ControlRigHierarchyCommands.cpp b/Engine/Plugins/Animation/ControlRig/Source/ControlRigEditor/Private/ControlRigHierarchyCommands.cpp
--- a/Engine/Plugins/Animation/ControlRig/Source/ControlRigEditor/Private/ControlRigHierarchyCommands.cpp
+++ b/Engine/Plugins/Animation/ControlRig/Source/ControlRigEditor/Private/ControlRigHierarchyCommands.cpp
@@ -21,9 +21,10 @@ void FControlRigHierarchyCommands::RegisterCommands()
 	UI_COMMAND(PasteItems, "Paste", "Paste the selected items.", EUserInterfaceActionType::Button, FInputChord(EKeys::V, EModifierKey::Control));
 	UI_COMMAND(PasteLocalTransforms, "Paste Local Transform", "Paste the local transforms.", EUserInterfaceActionType::Button, FInputChord());
 	UI_COMMAND(PasteGlobalTransforms, "Paste Global Transform", "Paste the global transforms.", EUserInterfaceActionType::Button, FInputChord(EKeys::V, EModifierKey::Control | EModifierKey::Shift));
-	UI_COMMAND(SetInitialTransformFromClosestBone, "Set Initial Transform from Closest Bone", "Find the Closest Bone to Initial Transform of the Selected Bones/Nulls or To Offset Transform of Selected Controls", EUserInterfaceActionType::Button, FInputChord());
-	UI_COMMAND(SetInitialTransformFromCurrentTransform, "Set Initial Transform from Current", "Save the Current Transform To Initial Transform of the Selected Bones/Nulls or To Offset Transform of Selected Controls", EUserInterfaceActionType::Button, FInputChord());
-	UI_COMMAND(SetShapeTransformFromCurrent, "Set Shape Transform From Current", "Transfer the Current Local Transform of the Control to its Shape Transform, the Control's Transform Resets", EUserInterfaceActionType::Button, FInputChord());
+	UI_COMMAND(ZeroControls, "Zero Controls", "Sets the transform of the control as offset, zeroes the initial and current transform.", EUserInterfaceActionType::Button, FInputChord());
+	UI_COMMAND(ZeroControlsFromClosestBone, "Zero Controls from Closest Bone", "Sets the offset to the transform of the closest bone, zeroes the initial and current transform.", EUserInterfaceActionType::Button, FInputChord());
+	UI_COMMAND(ZeroControlShape, "Zero Control Shape", "Sets the current shape relative to the control's offset, zeroes the control's initial and current transform.", EUserInterfaceActionType::Button, FInputChord());
+	UI_COMMAND(SetInitialTransformFromCurrent, "Set Initial Transform from Current", "Set the current transform of selected bone, null or connector as its initial and current transform.", EUserInterfaceActionType::Button, FInputChord());
 	UI_COMMAND(ResetTransform, "Reset Transform", "Reset the Transform", EUserInterfaceActionType::Button, FInputChord(EKeys::G, EModifierKey::Control));
 	UI_COMMAND(ResetAllTransforms, "Reset All Transforms", "Resets all Transforms", EUserInterfaceActionType::Button, FInputChord(EKeys::G, EModifierKey::Control | EModifierKey::Shift));
 	UI_COMMAND(ResetNull, "Reset Null", "Resets or injects a Null below the Control", EUserInterfaceActionType::Button, FInputChord());
diff --git a/Engine/Plugins/Animation/ControlRig/Source/ControlRigEditor/Private/ControlRigHierarchyCommands.h b/Engine/Plugins/Animation/ControlRig/Source/ControlRigEditor/Private/ControlRigHierarchyCommands.h
--- a/Engine/Plugins/Animation/ControlRig/Source/ControlRigEditor/Private/ControlRigHierarchyCommands.h
+++ b/Engine/Plugins/Animation/ControlRig/Source/ControlRigEditor/Private/ControlRigHierarchyCommands.h
@@ -72,14 +72,17 @@ public:
 	/* Reset space */
 	TSharedPtr<FUICommandInfo> ResetNull;
 
-	/* Set initial transform from current */
-	TSharedPtr<FUICommandInfo> SetInitialTransformFromCurrentTransform;
+	/* Sets the transform of the control as offset, zeroes the initial and current transform */
+	TSharedPtr<FUICommandInfo> ZeroControls;
 
-	/* set initial transform from closest bone */
-	TSharedPtr<FUICommandInfo> SetInitialTransformFromClosestBone;
+	/* Sets the offset to the transform of the closest bone, zeroes the initial and current transform */
+	TSharedPtr<FUICommandInfo> ZeroControlsFromClosestBone;
 
-	/* set shape transform from the control's current local transform and reset the control's transform */
-	TSharedPtr<FUICommandInfo> SetShapeTransformFromCurrent;
+	/* Sets the current shape relative to the control's offset, zeroes the control's initial and current transform */
+	TSharedPtr<FUICommandInfo> ZeroControlShape;
+
+	/* Set the current transform of selected bone, null or connector as its initial and current transform */
+	TSharedPtr<FUICommandInfo> SetInitialTransformFromCurrent;
 
 	/* frames the selection in the tree */
 	TSharedPtr<FUICommandInfo> FrameSelection;
diff --git a/Engine/Plugins/Animation/ControlRig/Source/ControlRigEditor/Private/Editor/SRigHierarchy.cpp b/Engine/Plugins/Animation/ControlRig/Source/ControlRigEditor/Private/Editor/SRigHierarchy.cpp
--- a/Engine/Plugins/Animation/ControlRig/Source/ControlRigEditor/Private/Editor/SRigHierarchy.cpp
+++ b/Engine/Plugins/Animation/ControlRig/Source/ControlRigEditor/Private/Editor/SRigHierarchy.cpp
@@ -1,6 +1,7 @@
 // Copyright Epic Games, Inc. All Rights Reserved.
 
 #include "Editor/SRigHierarchy.h"
+
 #include "Widgets/Input/SComboButton.h"
 #include "Styling/AppStyle.h"
 #include "Widgets/Input/SButton.h"
@@ -358,19 +359,24 @@ void SRigHierarchy::BindCommands()
 		FExecuteAction::CreateSP(this, &SRigHierarchy::HandleResetTransform, false));
 
 	CommandList->MapAction(
-		Commands.SetInitialTransformFromClosestBone,
-		FExecuteAction::CreateSP(this, &SRigHierarchy::HandleSetInitialTransformFromClosestBone),
-		FCanExecuteAction::CreateSP(this, &SRigHierarchy::IsControlOrNullSelected, false));
+		Commands.ZeroControls,
+		FExecuteAction::CreateSP(this, &SRigHierarchy::HandleZeroControls),
+		FCanExecuteAction::CreateSP(this, &SRigHierarchy::CanHandleZeroControls));
 
 	CommandList->MapAction(
-		Commands.SetInitialTransformFromCurrentTransform,
-		FExecuteAction::CreateSP(this, &SRigHierarchy::HandleSetInitialTransformFromCurrentTransform),
-		FCanExecuteAction::CreateSP(this, &SRigHierarchy::IsMultiSelected, false));
+		Commands.ZeroControlsFromClosestBone,
+		FExecuteAction::CreateSP(this, &SRigHierarchy::HandleZeroControlsFromClosestBone),
+		FCanExecuteAction::CreateSP(this, &SRigHierarchy::CanHandleZeroControlsFromClosestBone));
+
+	CommandList->MapAction(
+		Commands.ZeroControlShape,
+		FExecuteAction::CreateSP(this, &SRigHierarchy::HandleZeroControlShape),
+		FCanExecuteAction::CreateSP(this, &SRigHierarchy::CanHandleZeroControlShape));
 
 	CommandList->MapAction(
-		Commands.SetShapeTransformFromCurrent,
-		FExecuteAction::CreateSP(this, &SRigHierarchy::HandleSetShapeTransformFromCurrent),
-		FCanExecuteAction::CreateSP(this, &SRigHierarchy::IsControlSelected, false));
+		Commands.SetInitialTransformFromCurrent,
+		FExecuteAction::CreateSP(this, &SRigHierarchy::HandleSetInitialTransformsAsCurrent),
+		FCanExecuteAction::CreateSP(this, &SRigHierarchy::CanHandleSetInitialTransformsAsCurrent));
 
 	CommandList->MapAction(
 		Commands.FrameSelection,
@@ -1710,13 +1716,15 @@ void SRigHierarchy::CreateContextMenu()
 							}
 						}
 
-						const FText FromCurrentLabel = FText::Format(LOCTEXT("SetTransformFromCurrentTransform", "Set {0} Transform from Current"), FText::FromString(*Keyword));
-						const FText FromClosestBoneLabel = FText::Format(LOCTEXT("SetTransformFromClosestBone", "Set {0} Transform from Closest Bone"), FText::FromString(*Keyword));
-						TransformsSection.AddMenuEntry(Commands.SetInitialTransformFromCurrentTransform, FromCurrentLabel);
-						TransformsSection.AddMenuEntry(Commands.SetInitialTransformFromClosestBone, FromClosestBoneLabel);
+						const FText FromCurrentLabel = FText::Format(LOCTEXT("SetTransformFromCurrentTransform", "Zero {0}"), FText::FromString(*Keyword));
+						const FText FromClosestBoneLabel = FText::Format(LOCTEXT("SetTransformFromClosestBone", "Zero {0} from Closest Bone"), FText::FromString(*Keyword));
+						TransformsSection.AddMenuEntry(Commands.ZeroControls, FromCurrentLabel);
+						TransformsSection.AddMenuEntry(Commands.ZeroControlsFromClosestBone, FromClosestBoneLabel);
 					}
 					
-					TransformsSection.AddMenuEntry(Commands.SetShapeTransformFromCurrent);
+					TransformsSection.AddMenuEntry(Commands.ZeroControlShape);
+					TransformsSection.AddMenuEntry(Commands.SetInitialTransformFromCurrent);
+
 					TransformsSection.AddMenuEntry(Commands.Unparent);
 
 					FToolMenuSection& AssetsSection = InMenu->AddSection(TEXT("Assets"), LOCTEXT("AssetsHeader", "Assets"));
@@ -3736,90 +3744,326 @@ void SRigHierarchy::HandleResetTransform(bool bSelectionOnly)
 	}
 }
 
-void SRigHierarchy::HandleSetInitialTransformFromCurrentTransform()
+void SRigHierarchy::HandleZeroControls()
 {
-	if (IsMultiSelected(false))
+	const TSharedPtr<IControlRigBaseEditor> PinnedControlRigEditor = ControlRigEditor.IsValid() ? ControlRigEditor.Pin() : nullptr;
+	const TScriptInterface<IControlRigAssetInterface> Blueprint = PinnedControlRigEditor.IsValid() ? PinnedControlRigEditor->GetControlRigAssetInterface() : nullptr;
+	URigHierarchy* DefaultHierarchy = GetDefaultHierarchy();
+	URigHierarchy* DebuggedHierarchy = GetHierarchy();
+
+	if (CanHandleZeroControls() &&
+		PinnedControlRigEditor.IsValid() &&
+		Blueprint &&
+		DefaultHierarchy &&
+		DebuggedHierarchy)
 	{
-		if (FControlRigAssetInterfacePtr Blueprint = ControlRigEditor.Pin()->GetControlRigAssetInterface())
+		const FScopedTransaction Transaction(LOCTEXT("HandleZeroControlsTransaction", "Zero Controls"));
+
+		TArray<FRigElementKey> SelectedKeys = GetSelectedElementKeys();
+
+		TMap<FRigElementKey, FTransform> GlobalTransforms;
+		TMap<FRigElementKey, FTransform> ParentGlobalTransforms;
+		for (const FRigElementKey& SelectedKey : SelectedKeys)
 		{
-			if (URigHierarchy* DebuggedHierarchy = GetHierarchy())
+			if (SelectedKey.Type != ERigElementType::Control)
 			{
-				FScopedTransaction Transaction(LOCTEXT("HierarchySetInitialTransforms", "Set Initial Transforms"));
+				continue;
+			}
 
-				TArray<FRigElementKey> SelectedKeys = GetSelectedElementKeys();
-				TMap<FRigElementKey, FTransform> GlobalTransforms;
-				TMap<FRigElementKey, FTransform> ParentGlobalTransforms;
+			GlobalTransforms.FindOrAdd(SelectedKey) = DebuggedHierarchy->GetGlobalTransform(SelectedKey);
+			ParentGlobalTransforms.FindOrAdd(SelectedKey) = DebuggedHierarchy->GetParentTransform(SelectedKey);
+		}
 
-				for (const FRigElementKey& SelectedKey : SelectedKeys)
+		for (const FRigElementKey& SelectedKey : SelectedKeys)
+		{
+			if (SelectedKey.Type != ERigElementType::Control)
+			{
+				continue;
+			}
+
+			const FTransform GlobalTransform = GlobalTransforms[SelectedKey];
+			const FTransform LocalTransform = GlobalTransform.GetRelativeTransform(ParentGlobalTransforms[SelectedKey]);
+
+			constexpr bool bAffectChildren = true;
+			constexpr bool bSetupUndo = true;
+			constexpr bool bForce = false;
+			constexpr bool bPrintPythonCommand = true;
+
+			if (FRigControlElement* DebuggedControlElement = DebuggedHierarchy->Find<FRigControlElement>(SelectedKey))
+			{
+				DebuggedHierarchy->SetControlOffsetTransform(DebuggedControlElement, LocalTransform, ERigTransformType::InitialLocal, bAffectChildren, bSetupUndo, bForce, bPrintPythonCommand);
+				DebuggedHierarchy->SetControlOffsetTransform(DebuggedControlElement, LocalTransform, ERigTransformType::CurrentLocal, bAffectChildren, bSetupUndo, bForce, bPrintPythonCommand);
+				DebuggedHierarchy->SetTransform(DebuggedControlElement, FTransform::Identity, ERigTransformType::InitialLocal, bAffectChildren, bSetupUndo, bForce, bPrintPythonCommand);
+				DebuggedHierarchy->SetTransform(DebuggedControlElement, FTransform::Identity, ERigTransformType::CurrentLocal, bAffectChildren, bSetupUndo, bForce, bPrintPythonCommand);
+			}
+
+			if (FRigControlElement* DefaultControlElement = DefaultHierarchy->Find<FRigControlElement>(SelectedKey))
+			{
+				DefaultHierarchy->SetControlOffsetTransform(DefaultControlElement, LocalTransform, ERigTransformType::InitialLocal, bAffectChildren, bSetupUndo);
+				DefaultHierarchy->SetControlOffsetTransform(DefaultControlElement, LocalTransform, ERigTransformType::CurrentLocal, bAffectChildren, bSetupUndo);
+				DefaultHierarchy->SetTransform(DefaultControlElement, FTransform::Identity, ERigTransformType::InitialLocal, bAffectChildren, bSetupUndo);
+				DefaultHierarchy->SetTransform(DefaultControlElement, FTransform::Identity, ERigTransformType::CurrentLocal, bAffectChildren, bSetupUndo);
+			}
+		}
+	}
+}
+
+bool SRigHierarchy::CanHandleZeroControls() const
+{
+	if (IsProceduralSelected())
+	{
+		return false;
+	}
+
+	return GetSelectedElementKeys().ContainsByPredicate(
+		[](const FRigElementKey& SelectedKey)
+		{
+			return SelectedKey.Type == ERigElementType::Control;
+		});
+}
+
+void SRigHierarchy::HandleZeroControlsFromClosestBone()
+{
+	const TSharedPtr<IControlRigBaseEditor> PinnedControlRigEditor = ControlRigEditor.IsValid() ? ControlRigEditor.Pin() : nullptr;
+	const TScriptInterface<IControlRigAssetInterface> Blueprint = PinnedControlRigEditor.IsValid() ? PinnedControlRigEditor->GetControlRigAssetInterface() : nullptr;
+	URigHierarchy* DefaultHierarchy = GetDefaultHierarchy();
+	URigHierarchy* DebuggedHierarchy = GetHierarchy();
+
+	if (CanHandleZeroControlsFromClosestBone() &&
+		PinnedControlRigEditor.IsValid() &&
+		Blueprint &&
+		DefaultHierarchy &&
+		DebuggedHierarchy)
+	{
+		const FScopedTransaction Transaction(LOCTEXT("ZeroControlsFromClosestBoneTransaction", "Zero Controls from Closest Bone"));
+
+		const TArray<FRigElementKey> SelectedKeys = GetSelectedElementKeys();
+
+		TMap<FRigElementKey, FTransform> ClosestTransforms;
+		TMap<FRigElementKey, FTransform> ParentGlobalTransforms;
+		for (const FRigElementKey& SelectedKey : GetSelectedElementKeys())
+		{
+			if (SelectedKey.Type == ERigElementType::Control || SelectedKey.Type == ERigElementType::Null)
+			{
+				const FTransform GlobalTransform = DebuggedHierarchy->GetGlobalTransform(SelectedKey);
+				FTransform ClosestTransform;
+				FName ClosestRigElement;
+
+				if (!FindClosestBone(GlobalTransform.GetLocation(), ClosestRigElement, ClosestTransform))
 				{
-					GlobalTransforms.FindOrAdd(SelectedKey) = DebuggedHierarchy->GetGlobalTransform(SelectedKey);
-					ParentGlobalTransforms.FindOrAdd(SelectedKey) = DebuggedHierarchy->GetParentTransform(SelectedKey);
+					continue;
 				}
 
-				URigHierarchy* DefaultHierarchy = GetDefaultHierarchy();
-				
-				for (const FRigElementKey& SelectedKey : SelectedKeys)
+				ClosestTransforms.FindOrAdd(SelectedKey) = ClosestTransform;
+				ParentGlobalTransforms.FindOrAdd(SelectedKey) = DebuggedHierarchy->GetParentTransform(SelectedKey);
+			}
+		}
+
+		for (const FRigElementKey& SelectedKey : SelectedKeys)
+		{
+			if (!ClosestTransforms.Contains(SelectedKey))
+			{
+				continue;
+			}
+			const FTransform GlobalTransform = ClosestTransforms[SelectedKey];
+			const FTransform LocalTransform = GlobalTransform.GetRelativeTransform(ParentGlobalTransforms[SelectedKey]);
+
+			if (SelectedKey.Type == ERigElementType::Control)
+			{
+				constexpr bool bAffectChildren = true;
+				constexpr bool bSetupUndo = true;
+				constexpr bool bForce = false;
+				constexpr bool bPrintPythonCommand = true;
+
+				if (FRigControlElement* DebuggedControlElement = DebuggedHierarchy->Find<FRigControlElement>(SelectedKey))
 				{
-					FTransform GlobalTransform = GlobalTransforms[SelectedKey];
-					FTransform LocalTransform = GlobalTransform.GetRelativeTransform(ParentGlobalTransforms[SelectedKey]);
+					DebuggedHierarchy->SetControlOffsetTransform(DebuggedControlElement, LocalTransform, ERigTransformType::InitialLocal, bAffectChildren, bSetupUndo, bForce, bPrintPythonCommand);
+					DebuggedHierarchy->SetControlOffsetTransform(DebuggedControlElement, LocalTransform, ERigTransformType::CurrentLocal, bAffectChildren, bSetupUndo, bForce, bPrintPythonCommand);
+					DebuggedHierarchy->SetTransform(DebuggedControlElement, FTransform::Identity, ERigTransformType::InitialLocal, bAffectChildren, bSetupUndo, bForce, bPrintPythonCommand);
+					DebuggedHierarchy->SetTransform(DebuggedControlElement, FTransform::Identity, ERigTransformType::CurrentLocal, bAffectChildren, bSetupUndo, bForce, bPrintPythonCommand);
+				}
 
-					if (SelectedKey.Type == ERigElementType::Control)
-					{
-						if(FRigControlElement* ControlElement = DebuggedHierarchy->Find<FRigControlElement>(SelectedKey))
-						{
-							DebuggedHierarchy->SetControlOffsetTransform(ControlElement, LocalTransform, ERigTransformType::InitialLocal, true, true, false, true);
-							DebuggedHierarchy->SetControlOffsetTransform(ControlElement, LocalTransform, ERigTransformType::CurrentLocal, true, true, false, true);
-							DebuggedHierarchy->SetTransform(ControlElement, FTransform::Identity, ERigTransformType::InitialLocal, true, true, false, true);
-							DebuggedHierarchy->SetTransform(ControlElement, FTransform::Identity, ERigTransformType::CurrentLocal, true, true, false, true);
-						}
+				if (FRigControlElement* DefaultControlElement = DefaultHierarchy->Find<FRigControlElement>(SelectedKey))
+				{
+					DefaultHierarchy->SetControlOffsetTransform(DefaultControlElement, LocalTransform, ERigTransformType::InitialLocal, bAffectChildren, bSetupUndo);
+					DefaultHierarchy->SetControlOffsetTransform(DefaultControlElement, LocalTransform, ERigTransformType::CurrentLocal, bAffectChildren, bSetupUndo);
+					DefaultHierarchy->SetTransform(DefaultControlElement, FTransform::Identity, ERigTransformType::InitialLocal, bAffectChildren, bSetupUndo);
+					DefaultHierarchy->SetTransform(DefaultControlElement, FTransform::Identity, ERigTransformType::CurrentLocal, bAffectChildren, bSetupUndo);
+				}
+			}
+		}
+	}
+}
 
-						if(DefaultHierarchy)
-						{
-							if(FRigControlElement* ControlElement = DefaultHierarchy->Find<FRigControlElement>(SelectedKey))
-							{
-								DefaultHierarchy->SetControlOffsetTransform(ControlElement, LocalTransform, ERigTransformType::InitialLocal, true, true);
-								DefaultHierarchy->SetControlOffsetTransform(ControlElement, LocalTransform, ERigTransformType::CurrentLocal, true, true);
-								DefaultHierarchy->SetTransform(ControlElement, FTransform::Identity, ERigTransformType::InitialLocal, true, true);
-								DefaultHierarchy->SetTransform(ControlElement, FTransform::Identity, ERigTransformType::CurrentLocal, true, true);
-							}
-						}
-					}
-					else if (SelectedKey.Type == ERigElementType::Null ||
-						SelectedKey.Type == ERigElementType::Bone ||
-						SelectedKey.Type == ERigElementType::Connector)
-					{
-						FTransform InitialTransform = LocalTransform;
-						if (ControlRigEditor.Pin()->GetPreviewInstance())
-						{
-							if (FAnimNode_ModifyBone* ModifyBone = ControlRigEditor.Pin()->GetPreviewInstance()->FindModifiedBone(SelectedKey.Name))
-							{
-								InitialTransform.SetTranslation(ModifyBone->Translation);
-								InitialTransform.SetRotation(FQuat(ModifyBone->Rotation));
-								InitialTransform.SetScale3D(ModifyBone->Scale);
-							}
-						}
+bool SRigHierarchy::CanHandleZeroControlsFromClosestBone() const
+{
+	if (IsProceduralSelected())
+	{
+		return false;
+	}
 
-						if(FRigTransformElement* TransformElement = DebuggedHierarchy->Find<FRigTransformElement>(SelectedKey))
-						{
-							DebuggedHierarchy->SetTransform(TransformElement, LocalTransform, ERigTransformType::InitialLocal, true, true, false, true);
-							DebuggedHierarchy->SetTransform(TransformElement, LocalTransform, ERigTransformType::CurrentLocal, true, true, false, true);
-						}
+	return GetSelectedElementKeys().ContainsByPredicate(
+		[](const FRigElementKey& SelectedKey)
+		{
+			return SelectedKey.Type == ERigElementType::Control;
+		});
+}
 
-						if(DefaultHierarchy)
-						{
-							if(FRigTransformElement* TransformElement = DefaultHierarchy->Find<FRigTransformElement>(SelectedKey))
-							{
-								DefaultHierarchy->SetTransform(TransformElement, LocalTransform, ERigTransformType::InitialLocal, true, true);
-								DefaultHierarchy->SetTransform(TransformElement, LocalTransform, ERigTransformType::CurrentLocal, true, true);
-							}
-						}
+void SRigHierarchy::HandleZeroControlShape()
+{
+	const TSharedPtr<IControlRigBaseEditor> PinnedControlRigEditor = ControlRigEditor.IsValid() ? ControlRigEditor.Pin() : nullptr;
+	const TScriptInterface<IControlRigAssetInterface> Blueprint = PinnedControlRigEditor.IsValid() ? PinnedControlRigEditor->GetControlRigAssetInterface() : nullptr;
+	URigHierarchy* DebuggedHierarchy = GetHierarchy();
+
+	if (CanHandleZeroControlShape() &&
+		PinnedControlRigEditor.IsValid() &&
+		Blueprint &&
+		DebuggedHierarchy)
+	{
+		const FScopedTransaction Transaction(LOCTEXT("ZeroControlShapeTransaction", "Zero Control Shape"));
+
+		const FRigHierarchyInteractionBracket DebuggedInteractionBracket(DebuggedHierarchy);
+
+		TArray<TSharedPtr<FRigTreeElement>> SelectedItems = TreeView->GetSelectedItems();
+		for (const TSharedPtr<FRigTreeElement>& SelectedItem : SelectedItems)
+		{
+			if (!SelectedItem->Key.IsElement())
+			{
+				continue;
+			}
+			if (FRigControlElement* ControlElement = DebuggedHierarchy->Find<FRigControlElement>(SelectedItem->Key.GetElement()))
+			{
+				if (!ControlElement->Settings.SupportsShape())
+				{
+					continue;
+				}
+
+				const FRigElementKey Key = ControlElement->GetKey();
+
+				const FTransform OffsetGlobalTransform = DebuggedHierarchy->GetGlobalControlOffsetTransform(Key);
+				const FTransform ShapeGlobalTransform = DebuggedHierarchy->GetGlobalControlShapeTransform(Key);
+				const FTransform ShapeLocalTransform = ShapeGlobalTransform.GetRelativeTransform(OffsetGlobalTransform);
+
+				constexpr bool bInitial = true;
+				constexpr bool bNotInitial = false;
+				constexpr bool bSetupUndo = true;
+				constexpr bool bAffectChildren = true;
+				constexpr bool bPrintPythonCommand = true;
+
+				DebuggedHierarchy->SetControlShapeTransform(Key, ShapeLocalTransform, bInitial, bSetupUndo);
+				DebuggedHierarchy->SetControlShapeTransform(Key, ShapeLocalTransform, bNotInitial, bSetupUndo);
+				GetHierarchy()->SetControlShapeTransform(Key, ShapeLocalTransform, bInitial, bSetupUndo);
+				GetHierarchy()->SetControlShapeTransform(Key, ShapeLocalTransform, bNotInitial, bSetupUndo);
+
+				DebuggedHierarchy->SetLocalTransform(Key, FTransform::Identity, bNotInitial, bAffectChildren, bSetupUndo);
+				DebuggedHierarchy->SetLocalTransform(Key, FTransform::Identity, bInitial, bAffectChildren, bSetupUndo);
+				GetHierarchy()->SetLocalTransform(Key, FTransform::Identity, bNotInitial, bAffectChildren, bSetupUndo, bPrintPythonCommand);
+				GetHierarchy()->SetLocalTransform(Key, FTransform::Identity, bInitial, bAffectChildren, bSetupUndo, bPrintPythonCommand);
+
+				if (FControlRigEditorEditMode* EditMode = PinnedControlRigEditor->GetEditMode())
+				{
+					EditMode->RequestToRecreateControlShapeActors();
+				}
+			}
+		}
+	}
+}
+
+bool SRigHierarchy::CanHandleZeroControlShape() const
+{
+	if (IsProceduralSelected())
+	{
+		return false;
+	}
+
+	return GetSelectedElementKeys().ContainsByPredicate(
+		[](const FRigElementKey& SelectedKey)
+		{
+			return SelectedKey.Type == ERigElementType::Control;
+		});
+}
+
+void SRigHierarchy::HandleSetInitialTransformsAsCurrent()
+{
+	const TScriptInterface<IControlRigAssetInterface> Blueprint = ControlRigEditor.Pin()->GetControlRigAssetInterface();
+	URigHierarchy* DefaultHierarchy = GetDefaultHierarchy();
+	URigHierarchy* DebuggedHierarchy = GetHierarchy();
+
+	if (CanHandleSetInitialTransformsAsCurrent() &&
+		Blueprint &&
+		DefaultHierarchy &&
+		DebuggedHierarchy)
+	{
+		const FScopedTransaction Transaction(LOCTEXT("SetInitialTransformsAsCurrentTransaction", "Set Initial Transforms as Current"));
+
+		const TArray<FRigElementKey> SelectedKeys = GetSelectedElementKeys();
+
+		TMap<FRigElementKey, FTransform> GlobalTransforms;
+		TMap<FRigElementKey, FTransform> ParentGlobalTransforms;
+		for (const FRigElementKey& SelectedKey : SelectedKeys)
+		{
+			GlobalTransforms.FindOrAdd(SelectedKey) = DebuggedHierarchy->GetGlobalTransform(SelectedKey);
+			ParentGlobalTransforms.FindOrAdd(SelectedKey) = DebuggedHierarchy->GetParentTransform(SelectedKey);
+		}
+
+		for (const FRigElementKey& SelectedKey : SelectedKeys)
+		{
+			if (SelectedKey.Type == ERigElementType::Null ||
+				SelectedKey.Type == ERigElementType::Bone ||
+				SelectedKey.Type == ERigElementType::Connector)
+			{
+				const FTransform GlobalTransform = GlobalTransforms[SelectedKey];
+				const FTransform LocalTransform = GlobalTransform.GetRelativeTransform(ParentGlobalTransforms[SelectedKey]);
+
+				FTransform InitialTransform = LocalTransform;
+				if (ControlRigEditor.Pin()->GetPreviewInstance())
+				{
+					if (FAnimNode_ModifyBone* ModifyBone = ControlRigEditor.Pin()->GetPreviewInstance()->FindModifiedBone(SelectedKey.Name))
+					{
+						InitialTransform.SetTranslation(ModifyBone->Translation);
+						InitialTransform.SetRotation(FQuat(ModifyBone->Rotation));
+						InitialTransform.SetScale3D(ModifyBone->Scale);
 					}
 				}
+
+				constexpr bool bAffectChildren = true;
+				constexpr bool bSetupUndo = true;
+				constexpr bool bForce = false;
+				constexpr bool bPrintPythonCommand = true;
+
+				if (FRigTransformElement* DebuggedTransformElement = DebuggedHierarchy->Find<FRigTransformElement>(SelectedKey))
+				{
+					DebuggedHierarchy->SetTransform(DebuggedTransformElement, LocalTransform, ERigTransformType::InitialLocal, bAffectChildren, bSetupUndo, bForce, bPrintPythonCommand);
+					DebuggedHierarchy->SetTransform(DebuggedTransformElement, LocalTransform, ERigTransformType::CurrentLocal, bAffectChildren, bSetupUndo, bForce, bPrintPythonCommand);
+				}
+				
+				if (FRigTransformElement* DefaultTransformElement = DefaultHierarchy->Find<FRigTransformElement>(SelectedKey))
+				{
+					DefaultHierarchy->SetTransform(DefaultTransformElement, LocalTransform, ERigTransformType::InitialLocal, bAffectChildren, bSetupUndo);
+					DefaultHierarchy->SetTransform(DefaultTransformElement, LocalTransform, ERigTransformType::CurrentLocal, bAffectChildren, bSetupUndo);
+				}
 			}
 		}
 	}
 }
 
+bool SRigHierarchy::CanHandleSetInitialTransformsAsCurrent() const
+{
+	if (IsProceduralSelected())
+	{
+		return false;
+	}
+
+	return GetSelectedElementKeys().ContainsByPredicate(
+		[](const FRigElementKey& SelectedKey)
+		{
+			return
+				SelectedKey.Type == ERigElementType::Null ||
+				SelectedKey.Type == ERigElementType::Bone ||
+				SelectedKey.Type == ERigElementType::Connector;
+		});
+}
+
 void SRigHierarchy::HandleFrameSelection()
 {
 	TArray<TSharedPtr<FRigTreeElement>> SelectedItems = TreeView->GetSelectedItems();
@@ -4246,140 +4490,5 @@ FReply SRigHierarchy::ResolveConnectorToArray(const FRigElementKey& DraggedKey,
 	return FReply::Unhandled();
 }
 
-void SRigHierarchy::HandleSetInitialTransformFromClosestBone()
-{
-	if (IsControlOrNullSelected(false))
-	{
-		if (FControlRigAssetInterfacePtr Blueprint = ControlRigEditor.Pin()->GetControlRigAssetInterface())
-		{
-			if (URigHierarchy* DebuggedHierarchy = GetHierarchy())
-			{
-				URigHierarchy* Hierarchy = GetDefaultHierarchy();
-  				check(Hierarchy);
-
-				FScopedTransaction Transaction(LOCTEXT("HierarchySetInitialTransforms", "Set Initial Transforms"));
-
-				TArray<FRigElementKey> SelectedKeys = GetSelectedElementKeys();
-				TMap<FRigElementKey, FTransform> ClosestTransforms;
-				TMap<FRigElementKey, FTransform> ParentGlobalTransforms;
-
-				for (const FRigElementKey& SelectedKey : SelectedKeys)
-				{
-					if (SelectedKey.Type == ERigElementType::Control || SelectedKey.Type == ERigElementType::Null)
-					{
-						const FTransform GlobalTransform = DebuggedHierarchy->GetGlobalTransform(SelectedKey);
-						FTransform ClosestTransform;
-						FName ClosestRigElement;
-
-						if (!FindClosestBone(GlobalTransform.GetLocation(), ClosestRigElement, ClosestTransform))
-						{
-							continue;
-						}
-
-						ClosestTransforms.FindOrAdd(SelectedKey) = ClosestTransform;
-						ParentGlobalTransforms.FindOrAdd(SelectedKey) = DebuggedHierarchy->GetParentTransform(SelectedKey);
-					}
-				}
-
-				for (const FRigElementKey& SelectedKey : SelectedKeys)
-				{
-					if (!ClosestTransforms.Contains(SelectedKey))
-					{
-						continue;
-					}
-					FTransform GlobalTransform = ClosestTransforms[SelectedKey];
-					FTransform LocalTransform = GlobalTransform.GetRelativeTransform(ParentGlobalTransforms[SelectedKey]);
-
-					if (SelectedKey.Type == ERigElementType::Control)
-					{
-						if(FRigControlElement* ControlElement = Hierarchy->Find<FRigControlElement>(SelectedKey))
-						{
-							Hierarchy->SetControlOffsetTransform(ControlElement, LocalTransform, ERigTransformType::InitialLocal, true, true, false, true);
-							Hierarchy->SetControlOffsetTransform(ControlElement, LocalTransform, ERigTransformType::CurrentLocal, true, true, false, true);
-							Hierarchy->SetTransform(ControlElement, FTransform::Identity, ERigTransformType::InitialLocal, true, true, false, true);
-							Hierarchy->SetTransform(ControlElement, FTransform::Identity, ERigTransformType::CurrentLocal, true, true, false, true);
-						}
-						if(FRigControlElement* ControlElement = DebuggedHierarchy->Find<FRigControlElement>(SelectedKey))
-						{
-							DebuggedHierarchy->SetControlOffsetTransform(ControlElement, LocalTransform, ERigTransformType::InitialLocal, true, true);
-							DebuggedHierarchy->SetControlOffsetTransform(ControlElement, LocalTransform, ERigTransformType::CurrentLocal, true, true);
-							DebuggedHierarchy->SetTransform(ControlElement, FTransform::Identity, ERigTransformType::InitialLocal, true, true);
-							DebuggedHierarchy->SetTransform(ControlElement, FTransform::Identity, ERigTransformType::CurrentLocal, true, true);
-						}
-					}
-					else if (SelectedKey.Type == ERigElementType::Null ||
-                        SelectedKey.Type == ERigElementType::Bone)
-					{
-						FTransform InitialTransform = LocalTransform;
-
-						if(FRigTransformElement* TransformElement = Hierarchy->Find<FRigTransformElement>(SelectedKey))
-						{
-							Hierarchy->SetTransform(TransformElement, LocalTransform, ERigTransformType::InitialLocal, true, true, false, true);
-							Hierarchy->SetTransform(TransformElement, LocalTransform, ERigTransformType::CurrentLocal, true, true, false, true);
-						}
-						if(FRigTransformElement* TransformElement = DebuggedHierarchy->Find<FRigTransformElement>(SelectedKey))
-						{
-							DebuggedHierarchy->SetTransform(TransformElement, LocalTransform, ERigTransformType::InitialLocal, true, true);
-							DebuggedHierarchy->SetTransform(TransformElement, LocalTransform, ERigTransformType::CurrentLocal, true, true);
-						}
-					}
-				}
-			}
-		}
-	}
-}
-
-void SRigHierarchy::HandleSetShapeTransformFromCurrent()
-{
-	if (IsControlSelected(false))
-	{
-		if (FControlRigAssetInterfacePtr Blueprint = ControlRigEditor.Pin()->GetControlRigAssetInterface())
-		{
-			if (URigHierarchy* DebuggedHierarchy = GetHierarchy())
-			{
-				FScopedTransaction Transaction(LOCTEXT("HierarchySetShapeTransforms", "Set Shape Transforms"));
-
-				FRigHierarchyInteractionBracket InteractionBracket(GetHierarchy());
-				FRigHierarchyInteractionBracket DebuggedInteractionBracket(DebuggedHierarchy);
-
-				TArray<TSharedPtr<FRigTreeElement>> SelectedItems = TreeView->GetSelectedItems();
-				for (const TSharedPtr<FRigTreeElement>& SelectedItem : SelectedItems)
-				{
-					if(!SelectedItem->Key.IsElement())
-					{
-						continue;
-					}
-					if(FRigControlElement* ControlElement = DebuggedHierarchy->Find<FRigControlElement>(SelectedItem->Key.GetElement()))
-					{
-						const FRigElementKey Key = ControlElement->GetKey();
-						
-						if (ControlElement->Settings.SupportsShape())
-						{
-							const FTransform OffsetGlobalTransform = DebuggedHierarchy->GetGlobalControlOffsetTransform(Key); 
-							const FTransform ShapeGlobalTransform = DebuggedHierarchy->GetGlobalControlShapeTransform(Key);
-							const FTransform ShapeLocalTransform = ShapeGlobalTransform.GetRelativeTransform(OffsetGlobalTransform);
-							
-							DebuggedHierarchy->SetControlShapeTransform(Key, ShapeLocalTransform, true, true);
-							DebuggedHierarchy->SetControlShapeTransform(Key, ShapeLocalTransform, false, true);
-							GetHierarchy()->SetControlShapeTransform(Key, ShapeLocalTransform, true, true);
-							GetHierarchy()->SetControlShapeTransform(Key, ShapeLocalTransform, false, true);
-
-							DebuggedHierarchy->SetLocalTransform(Key, FTransform::Identity, false, true, true);
-							DebuggedHierarchy->SetLocalTransform(Key, FTransform::Identity, true, true, true);
-							GetHierarchy()->SetLocalTransform(Key, FTransform::Identity, false, true, true, true);
-							GetHierarchy()->SetLocalTransform(Key, FTransform::Identity, true, true, true, true);
-						}
-
-						if (FControlRigEditorEditMode* EditMode = ControlRigEditor.Pin()->GetEditMode())
-						{
-							EditMode->RequestToRecreateControlShapeActors();
-						}
-					}
-				}
-			}
-		}
-	}
-}
-
 #undef LOCTEXT_NAMESPACE
 
diff --git a/Engine/Plugins/Animation/ControlRig/Source/ControlRigEditor/Private/Editor/SRigHierarchy.h b/Engine/Plugins/Animation/ControlRig/Source/ControlRigEditor/Private/Editor/SRigHierarchy.h
--- a/Engine/Plugins/Animation/ControlRig/Source/ControlRigEditor/Private/Editor/SRigHierarchy.h
+++ b/Engine/Plugins/Animation/ControlRig/Source/ControlRigEditor/Private/Editor/SRigHierarchy.h
@@ -32,7 +32,7 @@ struct FRigHierarchyImportSettings
 	: Mesh(nullptr)
 	{}
 
-	UPROPERTY(EditAnywhere, Category = "Hierachy Import")
+	UPROPERTY(EditAnywhere, Category = "Hierarchy Import")
 	TObjectPtr<USkeletalMesh> Mesh;
 };
 
@@ -198,10 +198,23 @@ private:
 	void UpdateMesh(USkeletalMesh* InMesh, const bool bImport) const;
 
 	void HandleResetTransform(bool bSelectionOnly);
-	void HandleResetInitialTransform();
-	void HandleSetInitialTransformFromCurrentTransform();
-	void HandleSetInitialTransformFromClosestBone();
-	void HandleSetShapeTransformFromCurrent();
+
+	/* Sets the transform of the control as offset, zeroes the initial and current transform */
+	void HandleZeroControls();
+	bool CanHandleZeroControls() const;
+
+	/* Sets the offset to the transform of the closest bone, zeroes the initial and current transform */
+	void HandleZeroControlsFromClosestBone();
+	bool CanHandleZeroControlsFromClosestBone() const;
+
+	/* Sets the current shape relative to the control's offset, zeroes the control's initial and current transform */
+	void HandleZeroControlShape();
+	bool CanHandleZeroControlShape() const;
+
+	/* Set the current transform of selected bone, null or connector as its initial and current transform */
+	void HandleSetInitialTransformsAsCurrent();
+	bool CanHandleSetInitialTransformsAsCurrent() const;
+
 	void HandleFrameSelection();
 	void HandleControlBoneOrSpaceTransform();
 	void HandleUnparent();
