Control Rig: Remove unused code in SRigHierarchy::HandleSetInitialTransformsAsCurrent

diff --git a/Engine/Plugins/Animation/ControlRig/Source/ControlRigEditor/Private/Editor/SRigHierarchy.cpp b/Engine/Plugins/Animation/ControlRig/Source/ControlRigEditor/Private/Editor/SRigHierarchy.cpp
--- a/Engine/Plugins/Animation/ControlRig/Source/ControlRigEditor/Private/Editor/SRigHierarchy.cpp
+++ b/Engine/Plugins/Animation/ControlRig/Source/ControlRigEditor/Private/Editor/SRigHierarchy.cpp
@@ -4015,17 +4015,6 @@ void SRigHierarchy::HandleSetInitialTransformsAsCurrent()
 				const FTransform GlobalTransform = GlobalTransforms[SelectedKey];
 				const FTransform LocalTransform = GlobalTransform.GetRelativeTransform(ParentGlobalTransforms[SelectedKey]);
 
-				FTransform InitialTransform = LocalTransform;
-				if (ControlRigEditor.Pin()->GetPreviewInstance())
-				{
-					if (FAnimNode_ModifyBone* ModifyBone = ControlRigEditor.Pin()->GetPreviewInstance()->FindModifiedBone(SelectedKey.Name))
-					{
-						InitialTransform.SetTranslation(ModifyBone->Translation);
-						InitialTransform.SetRotation(FQuat(ModifyBone->Rotation));
-						InitialTransform.SetScale3D(ModifyBone->Scale);
-					}
-				}
-
 				constexpr bool bAffectChildren = true;
 				constexpr bool bSetupUndo = true;
 				constexpr bool bForce = false;
