Anim in Engine: Remove unused code in Private::GetKeyedStatusInTrack in AnimDetailsProxyBase.cpp

diff --git a/Engine/Plugins/Animation/ControlRig/Source/ControlRigEditor/Private/AnimDetails/Proxies/AnimDetailsProxyBase.cpp b/Engine/Plugins/Animation/ControlRig/Source/ControlRigEditor/Private/AnimDetails/Proxies/AnimDetailsProxyBase.cpp
--- a/Engine/Plugins/Animation/ControlRig/Source/ControlRigEditor/Private/AnimDetails/Proxies/AnimDetailsProxyBase.cpp
+++ b/Engine/Plugins/Animation/ControlRig/Source/ControlRigEditor/Private/AnimDetails/Proxies/AnimDetailsProxyBase.cpp
@@ -525,7 +525,7 @@ namespace UE::ControlRigEditor
 					}
 				
 					const int32 NumKeys = Channel->GetNumKeys();
-					if (NumKeys == 0)
+					if (NumKeys < 1)
 					{
 						continue;
 					}
@@ -534,17 +534,15 @@ namespace UE::ControlRigEditor
 
 					TArray<FFrameNumber> KeyTimesInRange;
 					Channel->GetKeys(Range, &KeyTimesInRange, nullptr);
-					if (KeyTimesInRange.IsEmpty() && NumKeys == 0)
+
+					const bool bKeyedInFrame = KeyTimesInRange.IsEmpty();
+					if (bKeyedInFrame)
 					{
-						bKeyedOnAnotherFrame |= true;
-					}
-					else if (KeyTimesInRange.IsEmpty() && NumKeys > 0)
-					{
-						SectionKeyedStatus = FMath::Max(SectionKeyedStatus, EPropertyKeyedStatus::KeyedInOtherFrame);
+						SectionKeyedStatus = EPropertyKeyedStatus::KeyedInFrame;
 					}
 					else
 					{
-						SectionKeyedStatus = FMath::Max(SectionKeyedStatus, EPropertyKeyedStatus::KeyedInFrame);
+						SectionKeyedStatus = FMath::Max(SectionKeyedStatus, EPropertyKeyedStatus::KeyedInOtherFrame);
 					}
 				}
 
