Anim in Engine: Overcome regression from previous change, fix keys keyed in frame are considered as not keyed in frame and vice versa

#rnx

diff --git a/Engine/Plugins/Animation/ControlRig/Source/ControlRigEditor/Private/AnimDetails/Proxies/AnimDetailsProxyBase.cpp b/Engine/Plugins/Animation/ControlRig/Source/ControlRigEditor/Private/AnimDetails/Proxies/AnimDetailsProxyBase.cpp
--- a/Engine/Plugins/Animation/ControlRig/Source/ControlRigEditor/Private/AnimDetails/Proxies/AnimDetailsProxyBase.cpp
+++ b/Engine/Plugins/Animation/ControlRig/Source/ControlRigEditor/Private/AnimDetails/Proxies/AnimDetailsProxyBase.cpp
@@ -535,7 +535,7 @@ namespace UE::ControlRigEditor
 					TArray<FFrameNumber> KeyTimesInRange;
 					Channel->GetKeys(Range, &KeyTimesInRange, nullptr);
 
-					const bool bKeyedInFrame = KeyTimesInRange.IsEmpty();
+					const bool bKeyedInFrame = !KeyTimesInRange.IsEmpty();
 					if (bKeyedInFrame)
 					{
 						SectionKeyedStatus = EPropertyKeyedStatus::KeyedInFrame;
