Control Rig: Attempt to fix crash when OnControlRigAddedOrRemovedDelegate is broadcast
#rnx

diff --git a/Engine/Plugins/Animation/ControlRig/Source/ControlRigEditor/Private/EditMode/ControlRigBaseDockableView.cpp b/Engine/Plugins/Animation/ControlRig/Source/ControlRigEditor/Private/EditMode/ControlRigBaseDockableView.cpp
--- a/Engine/Plugins/Animation/ControlRig/Source/ControlRigEditor/Private/EditMode/ControlRigBaseDockableView.cpp
+++ b/Engine/Plugins/Animation/ControlRig/Source/ControlRigEditor/Private/EditMode/ControlRigBaseDockableView.cpp
@@ -42,10 +42,11 @@ void FControlRigBaseDockableView::SetEditMode(FControlRigEditMode& InEditMode)
 	if (FControlRigEditMode* EditMode = GetEditMode())
 	{
 		EditMode->OnControlRigAddedOrRemoved().RemoveAll(this);
-		EditMode->OnControlRigAddedOrRemoved().AddRaw(this, &FControlRigBaseDockableView::HandleControlAdded);
-		EditMode->OnControlRigVisibilityChanged().AddRaw(this, &FControlRigBaseDockableView::HandleRigVisibilityChanged);
+		EditMode->OnControlRigAddedOrRemoved().AddSP(AsSharedWidget(), &FControlRigBaseDockableView::HandleControlAdded);
+		EditMode->OnControlRigVisibilityChanged().RemoveAll(this);
+		EditMode->OnControlRigVisibilityChanged().AddSP(AsSharedWidget(), &FControlRigBaseDockableView::HandleRigVisibilityChanged);
 		EditMode->OnControlRigSelected().RemoveAll(this);
-		EditMode->OnControlRigSelected().AddRaw(this, &FControlRigBaseDockableView::HandlElementSelected);
+		EditMode->OnControlRigSelected().AddSP(AsSharedWidget(), &FControlRigBaseDockableView::HandlElementSelected);
 	}
 }
 
diff --git a/Engine/Plugins/Animation/ControlRig/Source/ControlRigEditor/Private/EditMode/ControlRigBaseDockableView.h b/Engine/Plugins/Animation/ControlRig/Source/ControlRigEditor/Private/EditMode/ControlRigBaseDockableView.h
--- a/Engine/Plugins/Animation/ControlRig/Source/ControlRigEditor/Private/EditMode/ControlRigBaseDockableView.h
+++ b/Engine/Plugins/Animation/ControlRig/Source/ControlRigEditor/Private/EditMode/ControlRigBaseDockableView.h
@@ -21,6 +21,7 @@ public:
 	virtual ~FControlRigBaseDockableView();
 	TArray<UControlRig*> GetControlRigs() const;
 
+	virtual TSharedRef<FControlRigBaseDockableView> AsSharedWidget() = 0;
 	virtual void SetEditMode(FControlRigEditMode& InEditMode);
 	FControlRigEditMode* GetEditMode() const;
 
diff --git a/Engine/Plugins/Animation/ControlRig/Source/ControlRigEditor/Private/EditMode/SControlRigControlViews.cpp b/Engine/Plugins/Animation/ControlRig/Source/ControlRigEditor/Private/EditMode/SControlRigControlViews.cpp
--- a/Engine/Plugins/Animation/ControlRig/Source/ControlRigEditor/Private/EditMode/SControlRigControlViews.cpp
+++ b/Engine/Plugins/Animation/ControlRig/Source/ControlRigEditor/Private/EditMode/SControlRigControlViews.cpp
@@ -496,7 +496,7 @@ void SControlRigPoseView::Construct(const FArguments& InArgs)
 	{
 		if (FControlRigEditMode* EditMode = OwningWidget.Pin()->GetEditMode())
 		{
-			EditMode->OnControlRigAddedOrRemoved().AddRaw(this, &SControlRigPoseView::HandleControlAdded);
+			EditMode->OnControlRigAddedOrRemoved().AddSP(SharedThis(this), &SControlRigPoseView::HandleControlAdded);
 			TArray<UControlRig*> ControlRigs = GetControlRigs();
 			for (UControlRig* ControlRig : ControlRigs)
 			{
diff --git a/Engine/Plugins/Animation/ControlRig/Source/ControlRigEditor/Private/EditMode/SControlRigOutliner.h b/Engine/Plugins/Animation/ControlRig/Source/ControlRigEditor/Private/EditMode/SControlRigOutliner.h
--- a/Engine/Plugins/Animation/ControlRig/Source/ControlRigEditor/Private/EditMode/SControlRigOutliner.h
+++ b/Engine/Plugins/Animation/ControlRig/Source/ControlRigEditor/Private/EditMode/SControlRigOutliner.h
@@ -353,6 +353,7 @@ class SControlRigOutliner : public FControlRigBaseDockableView, public SCompound
 	void Construct(const FArguments& InArgs, FControlRigEditMode& InEditMode);
 
 	//FControlRigBaseDockableView overrides
+	virtual TSharedRef<FControlRigBaseDockableView> AsSharedWidget() override { return SharedThis(this); }
 	virtual void SetEditMode(FControlRigEditMode& InEditMode) override;
 private:
 	virtual void HandleControlAdded(UControlRig* ControlRig, bool bIsAdded) override;
diff --git a/Engine/Plugins/Animation/ControlRig/Source/ControlRigEditor/Private/EditMode/SControlRigSpacePicker.h b/Engine/Plugins/Animation/ControlRig/Source/ControlRigEditor/Private/EditMode/SControlRigSpacePicker.h
--- a/Engine/Plugins/Animation/ControlRig/Source/ControlRigEditor/Private/EditMode/SControlRigSpacePicker.h
+++ b/Engine/Plugins/Animation/ControlRig/Source/ControlRigEditor/Private/EditMode/SControlRigSpacePicker.h
@@ -58,6 +58,7 @@ private:
 	FReply OnCompensateAllClicked();
 	void Compensate(TOptional<FFrameNumber> OptionalKeyTime, bool bSetPreviousTick);
 	EVisibility GetAddSpaceButtonVisibility() const;
+	virtual TSharedRef<FControlRigBaseDockableView> AsSharedWidget() override { return SharedThis(this); }
 
 	bool ReadyForBakeOrCompensation() const;
 	//for now picker works off of one ControlRig, this function gets the first control rig with a selection
diff --git a/Engine/Plugins/Animation/ControlRig/Source/ControlRigEditor/Private/Sequencer/AnimLayers/SAnimLayers.h b/Engine/Plugins/Animation/ControlRig/Source/ControlRigEditor/Private/Sequencer/AnimLayers/SAnimLayers.h
--- a/Engine/Plugins/Animation/ControlRig/Source/ControlRigEditor/Private/Sequencer/AnimLayers/SAnimLayers.h
+++ b/Engine/Plugins/Animation/ControlRig/Source/ControlRigEditor/Private/Sequencer/AnimLayers/SAnimLayers.h
@@ -30,6 +30,7 @@ class SAnimLayers : public FControlRigBaseDockableView, public SCompoundWidget
 	void Tick(const FGeometry& AllottedGeometry, const double InCurrentTime, const float InDeltaTime) override;
 
 	//FControlRigBaseDockableView overrides
+	virtual TSharedRef<FControlRigBaseDockableView> AsSharedWidget() override { return SharedThis(this); }
 	virtual void SetEditMode(FControlRigEditMode& InEditMode) override;
 	virtual void HandleControlAdded(UControlRig* ControlRig, bool bIsAdded) override;
 
diff --git a/Engine/Plugins/Animation/ControlRig/Source/ControlRigEditor/Private/Sequencer/SelectionSets/SSelectionSets.h b/Engine/Plugins/Animation/ControlRig/Source/ControlRigEditor/Private/Sequencer/SelectionSets/SSelectionSets.h
--- a/Engine/Plugins/Animation/ControlRig/Source/ControlRigEditor/Private/Sequencer/SelectionSets/SSelectionSets.h
+++ b/Engine/Plugins/Animation/ControlRig/Source/ControlRigEditor/Private/Sequencer/SelectionSets/SSelectionSets.h
@@ -66,6 +66,7 @@ class SSelectionSets : public FControlRigBaseDockableView, public SCompoundWidge
 	void Tick(const FGeometry& AllottedGeometry, const double InCurrentTime, const float InDeltaTime) override;
 
 	//FControlRigBaseDockableView overrides
+	virtual TSharedRef<FControlRigBaseDockableView> AsSharedWidget() override { return SharedThis(this); }
 	virtual void SetEditMode(FControlRigEditMode& InEditMode) override;
 	virtual void HandleControlAdded(UControlRig* ControlRig, bool bIsAdded) override;
 
