Anim in Engine: Fix Anim Details shows stucts as fully keyed when they're partially keyed

#rnx

diff --git a/Engine/Plugins/Animation/ControlRig/Source/ControlRigEditor/Private/AnimDetails/Proxies/AnimDetailsProxyBase.cpp b/Engine/Plugins/Animation/ControlRig/Source/ControlRigEditor/Private/AnimDetails/Proxies/AnimDetailsProxyBase.cpp
--- a/Engine/Plugins/Animation/ControlRig/Source/ControlRigEditor/Private/AnimDetails/Proxies/AnimDetailsProxyBase.cpp
+++ b/Engine/Plugins/Animation/ControlRig/Source/ControlRigEditor/Private/AnimDetails/Proxies/AnimDetailsProxyBase.cpp
@@ -498,7 +498,7 @@ namespace UE::ControlRigEditor
 				break;
 			}
 
-			bool bKeyedOnAnotherFrame = false;
+			TSet<EPropertyKeyedStatus> PerChannelStatuses;
 			for (const FMovieSceneChannelEntry& ChannelEntry : ChannelProxy.GetAllEntries())
 			{
 				if (ChannelEntry.GetChannelTypeName() != FMovieSceneDoubleChannel::StaticStruct()->GetFName() &&
@@ -523,35 +523,46 @@ namespace UE::ControlRigEditor
 					{
 						continue;
 					}
-				
+
 					const int32 NumKeys = Channel->GetNumKeys();
-					if (NumKeys < 1)
+					if (NumKeys == 0)
 					{
+						PerChannelStatuses.Add(EPropertyKeyedStatus::NotKeyed);
 						continue;
 					}
 
-					SectionKeyedStatus = FMath::Max(SectionKeyedStatus, EPropertyKeyedStatus::KeyedInOtherFrame);
-
 					TArray<FFrameNumber> KeyTimesInRange;
 					Channel->GetKeys(Range, &KeyTimesInRange, nullptr);
 
-					const bool bKeyedInFrame = !KeyTimesInRange.IsEmpty();
-					if (bKeyedInFrame)
+					if (KeyTimesInRange.IsEmpty() && NumKeys > 0)
 					{
-						SectionKeyedStatus = EPropertyKeyedStatus::KeyedInFrame;
+						PerChannelStatuses.Add(EPropertyKeyedStatus::KeyedInOtherFrame);
+
+						SectionKeyedStatus = FMath::Max(SectionKeyedStatus, EPropertyKeyedStatus::KeyedInOtherFrame);
 					}
 					else
 					{
-						SectionKeyedStatus = FMath::Max(SectionKeyedStatus, EPropertyKeyedStatus::KeyedInOtherFrame);
+						PerChannelStatuses.Add(EPropertyKeyedStatus::KeyedInFrame);
+
+						SectionKeyedStatus = EPropertyKeyedStatus::KeyedInFrame;
 					}
 				}
 
 				break; // just do it for one type
 			}
 
-			if (bKeyedOnAnotherFrame)
+			// For structs, show partially keyed when some but not all channels are keyed this frame
+			const bool bStructWithPerChannelStatuses = PerChannelStatuses.Num() > 1;
+			if (bStructWithPerChannelStatuses)
 			{
-				SectionKeyedStatus = FMath::Max(SectionKeyedStatus, EPropertyKeyedStatus::PartiallyKeyed);
+				const bool bKeyedInFrame = PerChannelStatuses.Contains(EPropertyKeyedStatus::KeyedInFrame);
+				const bool bKeyedInOtherFrame = PerChannelStatuses.Contains(EPropertyKeyedStatus::KeyedInOtherFrame);
+				const bool bNotKeyed = PerChannelStatuses.Contains(EPropertyKeyedStatus::NotKeyed);
+
+				if (bKeyedInFrame && (bKeyedInOtherFrame || bNotKeyed))
+				{
+					SectionKeyedStatus = EPropertyKeyedStatus::PartiallyKeyed;
+				}
 			}
 
 			return SectionKeyedStatus;
