Anim in Engine: Fix an issue where Anim Details displayed the actor or component name instead of the track name for attributes

diff --git a/Engine/Plugins/Animation/ControlRig/Source/ControlRigEditor/Private/AnimDetails/Customizations/AnimDetailsProxyManagerDetails.cpp b/Engine/Plugins/Animation/ControlRig/Source/ControlRigEditor/Private/AnimDetails/Customizations/AnimDetailsProxyManagerDetails.cpp
--- a/Engine/Plugins/Animation/ControlRig/Source/ControlRigEditor/Private/AnimDetails/Customizations/AnimDetailsProxyManagerDetails.cpp
+++ b/Engine/Plugins/Animation/ControlRig/Source/ControlRigEditor/Private/AnimDetails/Customizations/AnimDetailsProxyManagerDetails.cpp
@@ -78,6 +78,17 @@ namespace UE::ControlRigEditor
 							}
 						}
 					}
+					else if (const UObject* BoundObject = Proxy->GetSequencerItem().GetBoundObject())
+					{
+						if (const AActor* Actor = Cast<AActor>(BoundObject))
+						{
+							OptionalParentName = Actor->GetFName();
+						}
+						else if (const UActorComponent* Component = Cast<UActorComponent>(BoundObject))
+						{
+							OptionalParentName = Component->GetFName();
+						}
+					}
 
 					FAnimDetailsGroupedAttributeProxies* AttributeGroupPtr = [&AttributeGroups, &OptionalParentName]() -> FAnimDetailsGroupedAttributeProxies*
 						{
@@ -136,9 +147,9 @@ namespace UE::ControlRigEditor
 			}
 
 			// Attribute proxies are displayed individually, but are still multi-edited across multiple control rig
-			if (AttributeGroups.Num() == 1)
+			const bool bOmitAttributeGroupName = !TopLevelDetailRowIDToProxiesMap.IsEmpty() && AttributeGroups.Num() == 1;
+			if (bOmitAttributeGroupName)
 			{
-				// Add without subgroup if theres only one attribute group
 				const FText CategoryName = LOCTEXT("AttributeCategoryName", "Attributes");
 				IDetailCategoryBuilder& CategoryBuilder_Attributes = DetailLayout.EditCategory("Attributes", CategoryName);
 				const FAnimDetailsGroupedAttributeProxies& AttributeGroup = AttributeGroups[0];
diff --git a/Engine/Plugins/Animation/ControlRig/Source/ControlRigEditor/Private/AnimDetails/Customizations/AnimDetailsValueCustomization.cpp b/Engine/Plugins/Animation/ControlRig/Source/ControlRigEditor/Private/AnimDetails/Customizations/AnimDetailsValueCustomization.cpp
--- a/Engine/Plugins/Animation/ControlRig/Source/ControlRigEditor/Private/AnimDetails/Customizations/AnimDetailsValueCustomization.cpp
+++ b/Engine/Plugins/Animation/ControlRig/Source/ControlRigEditor/Private/AnimDetails/Customizations/AnimDetailsValueCustomization.cpp
@@ -115,27 +115,43 @@ namespace UE::ControlRigEditor
 				continue;
 			}
 
-			const FName PropertyName = ChildHandle->GetProperty()->GetFName();
-			const FText PropertyDisplayText = [&ChildHandle, this]()
+			TArray<UObject*> OuterObjects;
+			ChildHandle->GetOuterObjects(OuterObjects);
+
+			UAnimDetailsProxyBase* FirstProxy = OuterObjects.IsEmpty() ? nullptr : Cast<UAnimDetailsProxyBase>(OuterObjects[0]);
+			if (!FirstProxy)
+			{
+				continue;
+			}
+
+			const FText PropertyDisplayText = [FirstProxy, &ChildHandle, this]()
 				{
 					// Draw the proxy display name if the property has only one child.
-					// This aviods issues where multi-editing would show the non-customized property name when multi editing.
-					if (SortedChildHandles.Num() == 1)
+					const bool bShowHeaderRow = SortedChildHandles.Num() > 1;
+					if (!bShowHeaderRow)
 					{
-						TArray<UObject*> OuterObjects;
-						ChildHandle->GetOuterObjects(OuterObjects);
-
-						UAnimDetailsProxyBase* FirstProxy = OuterObjects.IsEmpty() ? nullptr : Cast<UAnimDetailsProxyBase>(OuterObjects[0]);
 						if (FirstProxy)
 						{
-							return FirstProxy->GetDisplayNameText(EElementNameDisplayMode::ForceShort);
+							const UControlRig* ControlRig = FirstProxy->GetControlRig();
+							const URigHierarchy* Hierarchy = ControlRig ? ControlRig->GetHierarchy() : nullptr;
+							const FRigControlElement* ControlElement = FirstProxy->GetControlElement();
+							if (ControlRig && Hierarchy && ControlElement)
+							{
+								// For single control rig controls, show the show the short element name
+								constexpr EElementNameDisplayMode NameDisplayMode = EElementNameDisplayMode::ForceShort;
+								return Hierarchy->GetDisplayNameForUI(ControlElement, NameDisplayMode);
+							}
+							else if (const UMovieSceneTrack* Track = FirstProxy->GetSequencerItem().GetMovieSceneTrack())
+							{
+								// For single sequencer bindings, show the track name
+								return Track->GetDisplayName();
+							}
 						}
 					}
 					
 					return ChildHandle->GetPropertyDisplayName();
 				}();
 
-
 			// Handle visibilty for child content widgets. This is required so that  
 			// FAnimDetailsNavigableWidgetRegistry can find the right widget to navigate to.
 			const TSharedRef<SWidget> ValueWidget = MakeChildWidget(InStructPropertyHandle, ChildHandle);
@@ -332,7 +348,7 @@ namespace UE::ControlRigEditor
 			{
 				return
 					SNew(STextBlock)
-					.Text(InStructPropertyHandle->GetPropertyDisplayName())
+					.Text(OuterProxy->GetDisplayNameText())
 					.Font(FAppStyle::Get().GetFontStyle("PropertyWindow.NormalFont"));
 			}
 		}
@@ -445,13 +461,15 @@ namespace UE::ControlRigEditor
 
 	bool FAnimDetailsValueCustomization::IsResetToDefaultVisible(const TWeakPtr<IPropertyHandle> WeakPropertyHandle) const
 	{
-		if (!WeakPropertyHandle.IsValid())
+		const TSharedPtr<IPropertyHandle> PropertyHandle = WeakPropertyHandle.IsValid() ? WeakPropertyHandle.Pin() : nullptr;
+		const FProperty* Property = PropertyHandle.IsValid() ? PropertyHandle->GetProperty() : nullptr;
+		if (!PropertyHandle.IsValid() ||
+			!Property)
 		{
 			return false;
 		}
-		const TSharedRef<IPropertyHandle> PropertyHandle = WeakPropertyHandle.Pin().ToSharedRef();
 
-		const FName PropertyName = PropertyHandle->GetProperty()->GetFName();
+		const FName PropertyName = Property->GetFName();
 
 		TArray<UObject*> OuterObjects;
 		PropertyHandle->GetOuterObjects(OuterObjects);
@@ -470,13 +488,15 @@ namespace UE::ControlRigEditor
 
 	void FAnimDetailsValueCustomization::OnResetToDefaultClicked(const TWeakPtr<IPropertyHandle> WeakPropertyHandle)
 	{
-		if (!WeakPropertyHandle.IsValid())
+		const TSharedPtr<IPropertyHandle> PropertyHandle = WeakPropertyHandle.IsValid() ? WeakPropertyHandle.Pin() : nullptr;
+		const FProperty* Property = PropertyHandle.IsValid() ? PropertyHandle->GetProperty() : nullptr;
+		if (!PropertyHandle.IsValid() ||
+			!Property)
 		{
 			return;
 		}
-		const TSharedRef<IPropertyHandle> PropertyHandle = WeakPropertyHandle.Pin().ToSharedRef();
 
-		const FName PropertyName = PropertyHandle->GetProperty()->GetFName();
+		const FName PropertyName = Property->GetFName();
 
 		TArray<UObject*> OuterObjects;
 		PropertyHandle->GetOuterObjects(OuterObjects);
diff --git a/Engine/Plugins/Animation/ControlRig/Source/ControlRigEditor/Private/AnimDetails/Proxies/AnimDetailsProxyBase.cpp b/Engine/Plugins/Animation/ControlRig/Source/ControlRigEditor/Private/AnimDetails/Proxies/AnimDetailsProxyBase.cpp
--- a/Engine/Plugins/Animation/ControlRig/Source/ControlRigEditor/Private/AnimDetails/Proxies/AnimDetailsProxyBase.cpp
+++ b/Engine/Plugins/Animation/ControlRig/Source/ControlRigEditor/Private/AnimDetails/Proxies/AnimDetailsProxyBase.cpp
@@ -813,9 +813,9 @@ FName UAnimDetailsProxyBase::GetDetailRowID() const
 		{
 			return ControlElement->GetKey().Name;
 		}
-		else if (UMovieSceneTrack* MovieSceneTrack = SequencerItem.GetMovieSceneTrack())
+		else if (const FProperty* Property = SequencerItem.GetProperty())
 		{
-			return *MovieSceneTrack->GetPathName();
+			return *Property->GetPathName();
 		}
 		else
 		{
