Control Rig: Suspend warnings when looking up elements using name spaced paths during load / patching

#rnx

#FailedValidations Preflight

diff --git a/Engine/Plugins/Animation/ControlRig/Source/ControlRig/Private/Rigs/RigHierarchy.cpp b/Engine/Plugins/Animation/ControlRig/Source/ControlRig/Private/Rigs/RigHierarchy.cpp
--- a/Engine/Plugins/Animation/ControlRig/Source/ControlRig/Private/Rigs/RigHierarchy.cpp
+++ b/Engine/Plugins/Animation/ControlRig/Source/ControlRig/Private/Rigs/RigHierarchy.cpp
@@ -126,6 +126,7 @@ URigHierarchy::URigHierarchy()
 , LastInteractedKey()
 , bSuspendNotifications(false)
 , bSuspendMetadataNotifications(false)
+, bSuspendNameSpaceKeyWarnings(false)
 , HierarchyController(nullptr)
 , bIsControllerAvailable(true)
 , ResetPoseHash(INDEX_NONE)
@@ -1816,7 +1817,7 @@ int32 URigHierarchy::GetIndex(const FRigElementKey& InKey, bool bFollowRedirecto
 	if(PatchedKey != InKey)
 	{
 #if WITH_EDITOR
-		if(!IsLoading())
+		if(!IsLoading() && !bSuspendNameSpaceKeyWarnings)
 		{
 			const int32 NumReceivedNameSpaceBasedKeys = ReceivedNameSpaceBasedKeys.Num(); 
 			if(ReceivedNameSpaceBasedKeys.AddUnique(InKey) == NumReceivedNameSpaceBasedKeys)
diff --git a/Engine/Plugins/Animation/ControlRig/Source/ControlRig/Private/Sequencer/MovieSceneControlRigParameterSection.cpp b/Engine/Plugins/Animation/ControlRig/Source/ControlRig/Private/Sequencer/MovieSceneControlRigParameterSection.cpp
--- a/Engine/Plugins/Animation/ControlRig/Source/ControlRig/Private/Sequencer/MovieSceneControlRigParameterSection.cpp
+++ b/Engine/Plugins/Animation/ControlRig/Source/ControlRig/Private/Sequencer/MovieSceneControlRigParameterSection.cpp
@@ -3147,25 +3147,28 @@ void UMovieSceneControlRigParameterSection::RecreateWithThisControlRig(UControlR
 		}
 	}
 	
-	
-	//HELGE MAKE SURE THIS WORKS FOR THE MODULAR STUFF
 	// rename all existing parameters based on short name vs long name.
 	// this also recovers from controls being stored with an original long path
 	// which may now be outdated. UControlRig::FindControl has backwards compat for that.
 	if(InControlRig->IsModularRig())
 	{
-		ForEachParameter([this, InControlRig](FBaseParameterNameAndValue* InOutParameter)
+		URigHierarchy* RigHierarchy = InControlRig->GetHierarchy();
+		if (RigHierarchy)
 		{
-			if(const FRigControlElement* Control = InControlRig->FindControl(InOutParameter->ParameterName))
+			const FRigHierarchyNameSpaceWarningBracket SuspendNameSpaceWarnings(RigHierarchy); 
+			
+			ForEachParameter([this, InControlRig](FBaseParameterNameAndValue* InOutParameter)
 			{
-				const FName& ControlName = Control->GetFName();
-				if(ControlName != InOutParameter->ParameterName)
+				if(const FRigControlElement* Control = InControlRig->FindControl(InOutParameter->ParameterName))
 				{
-					InOutParameter->ParameterName = ControlName;
+					const FName& ControlName = Control->GetFName();
+					if(ControlName != InOutParameter->ParameterName)
+					{
+						InOutParameter->ParameterName = ControlName;
+					}
 				}
-			}
-		});
-
+			});
+		}
 	}
 
 	//also need update Modular Rig FNames on the transformable handles on the constraint channels for those that are constraints. Update both child and parent
diff --git a/Engine/Plugins/Animation/ControlRig/Source/ControlRig/Public/Rigs/RigHierarchy.h b/Engine/Plugins/Animation/ControlRig/Source/ControlRig/Public/Rigs/RigHierarchy.h
--- a/Engine/Plugins/Animation/ControlRig/Source/ControlRig/Public/Rigs/RigHierarchy.h
+++ b/Engine/Plugins/Animation/ControlRig/Source/ControlRig/Public/Rigs/RigHierarchy.h
@@ -4963,6 +4963,11 @@ private:
 	 */
 	bool bSuspendMetadataNotifications = false;
 
+	/**
+	 * If set to true all name space key related warnings will be suspended
+	 */
+	bool bSuspendNameSpaceKeyWarnings;
+
 	/**
 	 * The event fired during undo / redo
 	 */
@@ -5520,6 +5525,7 @@ private:
 	friend class FControlRigHierarchyRelinkElementStorage;
 	friend class FRigHierarchyPoseAdapter;
 	friend class UControlRigGraph;
+	friend struct FRigHierarchyNameSpaceWarningBracket;
 };
 
 struct FRigHierarchyInteractionBracket
@@ -5807,3 +5813,11 @@ public:
 protected:
 	TArray<FName>& Names;
 };
+
+struct FRigHierarchyNameSpaceWarningBracket : TGuardValue<bool>
+{
+	FRigHierarchyNameSpaceWarningBracket(URigHierarchy* InHierarchy)
+		: TGuardValue<bool>(InHierarchy->bSuspendNameSpaceKeyWarnings, true)
+	{
+	}
+};
diff --git a/Engine/Plugins/Animation/ControlRig/Source/ControlRigDeveloper/Private/ControlRigAsset.cpp b/Engine/Plugins/Animation/ControlRig/Source/ControlRigDeveloper/Private/ControlRigAsset.cpp
--- a/Engine/Plugins/Animation/ControlRig/Source/ControlRigDeveloper/Private/ControlRigAsset.cpp
+++ b/Engine/Plugins/Animation/ControlRig/Source/ControlRigDeveloper/Private/ControlRigAsset.cpp
@@ -2350,7 +2350,11 @@ void IControlRigAssetInterface::UpdateElementKeyRedirector(UControlRig* InContro
 {
 	InControlRig->HierarchySettings = GetHierarchySettings();
 	InControlRig->RigModuleSettings = GetRigModuleSettings();
-	InControlRig->ElementKeyRedirector = FRigElementKeyRedirector(GetArrayConnectionMap(), InControlRig->GetHierarchy());
+	if (URigHierarchy* RigHierarchy = InControlRig->GetHierarchy())
+	{
+		const FRigHierarchyNameSpaceWarningBracket SuspendNameSpaceWarnings(RigHierarchy); 
+		InControlRig->ElementKeyRedirector = FRigElementKeyRedirector(GetArrayConnectionMap(), RigHierarchy);
+	}
 }
 
 void IControlRigAssetInterface::PropagatePoseFromInstanceToBP(UControlRig* InControlRig) const
