Control Rig: Follow up: Suspend warnings when looking up elements using name spaced paths during load / patching

#rnx

diff --git a/Engine/Plugins/Animation/ControlRig/Source/ControlRig/Private/Sequencer/MovieSceneControlRigParameterSection.cpp b/Engine/Plugins/Animation/ControlRig/Source/ControlRig/Private/Sequencer/MovieSceneControlRigParameterSection.cpp
--- a/Engine/Plugins/Animation/ControlRig/Source/ControlRig/Private/Sequencer/MovieSceneControlRigParameterSection.cpp
+++ b/Engine/Plugins/Animation/ControlRig/Source/ControlRig/Private/Sequencer/MovieSceneControlRigParameterSection.cpp
@@ -3152,8 +3152,7 @@ void UMovieSceneControlRigParameterSection::RecreateWithThisControlRig(UControlR
 	// which may now be outdated. UControlRig::FindControl has backwards compat for that.
 	if(InControlRig->IsModularRig())
 	{
-		URigHierarchy* RigHierarchy = InControlRig->GetHierarchy();
-		if (RigHierarchy)
+		if (URigHierarchy* RigHierarchy = InControlRig->GetHierarchy())
 		{
 			const FRigHierarchyNameSpaceWarningBracket SuspendNameSpaceWarnings(RigHierarchy); 
 			
@@ -3180,12 +3179,16 @@ void UMovieSceneControlRigParameterSection::RecreateWithThisControlRig(UControlR
 			{
 				if (HandleControlRig->IsModularRig())
 				{
-					if (const FRigControlElement* Control = HandleControlRig->FindControl(Handle->ControlName))
+					if (URigHierarchy* RigHierarchy = HandleControlRig->GetHierarchy())
 					{
-						const FName& ControlName = Control->GetFName();
-						if (ControlName != Handle->ControlName)
+						const FRigHierarchyNameSpaceWarningBracket SuspendNameSpaceWarnings(RigHierarchy); 
+						if (const FRigControlElement* Control = HandleControlRig->FindControl(Handle->ControlName))
 						{
-							Handle->ControlName = ControlName;
+							const FName& ControlName = Control->GetFName();
+							if (ControlName != Handle->ControlName)
+							{
+								Handle->ControlName = ControlName;
+							}
 						}
 					}
 				}
