Anim in Engine: Anim Details now uses meta data of Control Rig elements and Object Properties that seqeuncer controls, value ranges are now displayed and respected

diff --git a/Engine/Plugins/Animation/ControlRig/Source/ControlRigEditor/Private/AnimDetails/Customizations/AnimDetailsProxyDetails.cpp b/Engine/Plugins/Animation/ControlRig/Source/ControlRigEditor/Private/AnimDetails/Customizations/AnimDetailsProxyDetails.cpp
--- a/Engine/Plugins/Animation/ControlRig/Source/ControlRigEditor/Private/AnimDetails/Customizations/AnimDetailsProxyDetails.cpp
+++ b/Engine/Plugins/Animation/ControlRig/Source/ControlRigEditor/Private/AnimDetails/Customizations/AnimDetailsProxyDetails.cpp
@@ -55,7 +55,6 @@ namespace UE::ControlRigEditor
 		{
 			// Always use the long name for the category
 			DisplayNameText = AllProxies[0]->GetDisplayNameText(EElementNameDisplayMode::ForceLong);
-			AllProxies[0]->UpdatePropertyDisplayNames(DetailBuilder);
 		}
 		else
 		{
diff --git a/Engine/Plugins/Animation/ControlRig/Source/ControlRigEditor/Private/AnimDetails/Customizations/AnimDetailsSetInstancedPropertyMetadata.cpp b/Engine/Plugins/Animation/ControlRig/Source/ControlRigEditor/Private/AnimDetails/Customizations/AnimDetailsSetInstancedPropertyMetadata.cpp
new file mode 100644
--- /dev/null
+++ b/Engine/Plugins/Animation/ControlRig/Source/ControlRigEditor/Private/AnimDetails/Customizations/AnimDetailsSetInstancedPropertyMetadata.cpp
@@ -0,0 +1,267 @@
+// Copyright Epic Games, Inc. All Rights Reserved.
+
+#include "AnimDetailsSetInstancedPropertyMetadata.h"
+
+#include "AnimDetails/Proxies/AnimDetailsProxyBase.h"
+#include "ControlRig.h"
+#include "Math/Transform.h"
+#include "Math/Vector.h"
+#include "PropertyHandle.h"
+
+namespace UE::ControlRigEditor::AnimDetailsMetaDataUtil
+{
+	TAutoConsoleVariable<bool> CVarAnimDetailsUseExternalPropertyMetadata(
+		TEXT("AnimDetails.UseExternalPropertyMetadata"),
+		true,
+		TEXT("When enabled, Anim Details uses metadata from the controls and properties, instead of its own")
+	);
+
+	namespace Private
+	{
+		/** Returns the value at the component index of the type */
+		template <typename TValueType>
+		void GetValueFromComponentIndex(const int32 InComponentIndex, const TValueType& InValue, TOptional<float>& OutOptionalValue)
+		{
+			if constexpr (std::is_same_v<TValueType, float> || std::is_same_v<TValueType, int32>)
+			{
+				OutOptionalValue = InValue;
+			}
+			else if constexpr (std::is_same_v<TValueType, FVector3f>)
+			{
+				OutOptionalValue = InValue.Component(InComponentIndex);
+			}
+			else if constexpr (
+				std::is_same_v<TValueType, FRigControlValue::FTransform_Float> ||
+				std::is_same_v<TValueType, FRigControlValue::FTransformNoScale_Float> ||
+				std::is_same_v<TValueType, FRigControlValue::FEulerTransform_Float>)
+			{
+				const FTransform Transform = [&InValue]()
+					{
+						if constexpr (
+							std::is_same_v<TValueType, FRigControlValue::FTransformNoScale_Float> ||
+							std::is_same_v<TValueType, FRigControlValue::FEulerTransform_Float>)
+						{
+							return InValue.ToTransform().ToFTransform();
+						}
+						else
+						{
+							return InValue.ToTransform();
+						}
+					}();
+
+				if (InComponentIndex == 0)
+				{
+					OutOptionalValue = Transform.GetTranslation().X;
+				}
+				else if (InComponentIndex == 1)
+				{
+					OutOptionalValue = Transform.GetTranslation().Y;
+				}
+				else if (InComponentIndex == 2)
+				{
+					OutOptionalValue = Transform.GetTranslation().Z;
+				}
+				else if (InComponentIndex == 3)
+				{
+					OutOptionalValue = Transform.GetRotation().Rotator().Roll;
+				}
+				else if (InComponentIndex == 4)
+				{
+					OutOptionalValue = Transform.GetRotation().Rotator().Pitch;
+				}
+				else if (InComponentIndex == 5)
+				{
+					OutOptionalValue = Transform.GetRotation().Rotator().Yaw;
+				}
+
+				// Scale only for transforms with scale
+				if constexpr (
+					std::is_same_v<TValueType, FRigControlValue::FTransform_Float> ||
+					std::is_same_v<TValueType, FRigControlValue::FEulerTransform_Float>)
+				{
+					if (InComponentIndex == 6)
+					{
+						OutOptionalValue = Transform.GetScale3D().X;
+					}
+					else if (InComponentIndex == 7)
+					{
+						OutOptionalValue = Transform.GetScale3D().Y;
+					}
+					else if (InComponentIndex == 8)
+					{
+						OutOptionalValue = Transform.GetScale3D().Z;
+					}
+				}
+			}
+			else
+			{
+				[] <bool SupportedType = false>()
+				{
+					static_assert(SupportedType, "Unsupported type in GetValueFromComponentIndex");
+				}();
+			}
+		}
+
+		/** Gets the min and max value from the component index of the type */
+		template <typename TValueType>
+		void GetMinMaxValueFromComponentIndex(
+			const FRigControlElement& InControlElement,
+			const int32 InComponentIndex,
+			TOptional<float>& OutOptionalMin,
+			TOptional<float>& OutOptionalMax)
+		{
+			const TArray<FRigControlLimitEnabled>& LimitEnabled = InControlElement.Settings.LimitEnabled;
+			if (!ensureMsgf(LimitEnabled.IsValidIndex(InComponentIndex), TEXT("Invalid value index when trying to get min max from element")))
+			{
+				return;
+			}
+
+			const bool bWithMinimum = !LimitEnabled.IsEmpty() && LimitEnabled[InComponentIndex].bMinimum;
+			const TOptional<TValueType> OptionalMinValue = bWithMinimum ? InControlElement.Settings.MinimumValue.Get<TValueType>() : TOptional<TValueType>();
+
+			const bool bWithMaxium = !LimitEnabled.IsEmpty() && LimitEnabled[InComponentIndex].bMaximum;
+			const TOptional<TValueType> OptionalMaxValue = bWithMaxium ? InControlElement.Settings.MaximumValue.Get<TValueType>() : TOptional<TValueType>();
+
+			if (OptionalMinValue.IsSet())
+			{
+				GetValueFromComponentIndex(InComponentIndex, OptionalMinValue.GetValue(), OutOptionalMin);
+			}
+
+			if (OptionalMaxValue.IsSet())
+			{
+				GetValueFromComponentIndex(InComponentIndex, OptionalMaxValue.GetValue(), OutOptionalMax);
+			}
+		}
+
+		/** Returns the min and max value at given value index in the control elmeent value */
+		void GetMinMaxFromControlElement(const FRigControlElement& InControlElement, const int32 InComponentIndex, TOptional<float>& OutOptionalMin, TOptional<float>& OutOptionalMax)
+		{
+			switch (InControlElement.Settings.ControlType)
+			{
+			case ERigControlType::Bool:
+			{
+				// Bool has no min max
+				break;
+			}
+			case ERigControlType::Float:
+			case ERigControlType::ScaleFloat:
+			{
+				GetMinMaxValueFromComponentIndex<float>(InControlElement, InComponentIndex, OutOptionalMin, OutOptionalMax);
+				break;
+			}
+			case ERigControlType::Integer:
+			{
+				GetMinMaxValueFromComponentIndex<int32>(InControlElement, InComponentIndex, OutOptionalMin, OutOptionalMax);
+				break;
+			}
+			case ERigControlType::Vector2D:
+			case ERigControlType::Rotator:
+			case ERigControlType::Position:
+			case ERigControlType::Scale:
+			{
+				GetMinMaxValueFromComponentIndex<FVector3f>(InControlElement, InComponentIndex, OutOptionalMin, OutOptionalMax);
+				break;
+			}
+			case ERigControlType::Transform:
+			{
+				GetMinMaxValueFromComponentIndex<FRigControlValue::FTransform_Float>(InControlElement, InComponentIndex, OutOptionalMin, OutOptionalMax);
+				break;
+			}
+			case ERigControlType::TransformNoScale:
+			{
+				GetMinMaxValueFromComponentIndex<FRigControlValue::FTransformNoScale_Float>(InControlElement, InComponentIndex, OutOptionalMin, OutOptionalMax);
+				break;
+			}
+			case ERigControlType::EulerTransform:
+			{
+				GetMinMaxValueFromComponentIndex<FRigControlValue::FEulerTransform_Float>(InControlElement, InComponentIndex, OutOptionalMin, OutOptionalMax);
+				break;
+			}
+			default:
+			{
+				break;
+			}
+			}
+		}
+	}
+
+	void SetInstancedPropertyMetaData(const UAnimDetailsProxyBase& TemplateProxy, const TSharedRef<IPropertyHandle>& ValuePropertyHandle)
+	{
+		if (!CVarAnimDetailsUseExternalPropertyMetadata.GetValueOnGameThread())
+		{
+			return;
+		}
+
+		const FProperty* Property = ValuePropertyHandle->GetProperty();
+		if (!Property)
+		{
+			return;
+		}
+		const FName PropertyName = Property->GetFName();
+
+		const UControlRig* ControlRig = TemplateProxy.GetControlRig();
+		const FRigControlElement* ControlElement = TemplateProxy.GetControlElement();
+		if (ControlRig && ControlElement)
+		{
+			// Always use a linear delta sensitivity and a slider exponent of 1 for control rig controls
+			ValuePropertyHandle->SetInstanceMetaData("LinearDeltaSensitivity", FString::Printf(TEXT("%f"), 1.0f));
+			ValuePropertyHandle->SetInstanceMetaData("SliderExponent", FString::Printf(TEXT("%f"), 1.0f));
+
+			// Adopt limits from the rig element
+			const int32 ComponentIndex = TemplateProxy.GetPropertyNames().IndexOfByKey(PropertyName);
+
+			TOptional<float> OptionalMin;
+			TOptional<float> OptionalMax;
+			Private::GetMinMaxFromControlElement(*ControlElement, ComponentIndex, OptionalMin, OptionalMax);
+
+			// If min and max don't define a finite range, and the property has no delta, use a delta of 0.5
+			const bool bWithRange = OptionalMin.IsSet() && OptionalMax.IsSet();
+			const bool bWithDeltaMetadata = !ValuePropertyHandle->GetMetaData("Delta").IsEmpty();
+			if (!bWithRange && !bWithDeltaMetadata)
+			{
+				ValuePropertyHandle->SetInstanceMetaData("Delta", FString::Printf(TEXT("%f"), .5f));
+			}
+			
+			if (OptionalMin.IsSet())
+			{
+				ValuePropertyHandle->SetInstanceMetaData("ClampMin", FString::Printf(TEXT("%f"), OptionalMin.GetValue()));
+				ValuePropertyHandle->SetInstanceMetaData("UIMin", FString::Printf(TEXT("%f"), OptionalMin.GetValue()));
+			}
+
+			if (OptionalMax.IsSet())
+			{
+				ValuePropertyHandle->SetInstanceMetaData("ClampMax", FString::Printf(TEXT("%f"), OptionalMax.GetValue()));
+				ValuePropertyHandle->SetInstanceMetaData("UIMax", FString::Printf(TEXT("%f"), OptionalMax.GetValue()));
+			}
+		}
+		else if (const FProperty* SequencerProperty = TemplateProxy.GetSequencerItem().GetProperty())
+		{
+			// Adopt supported metadata from the controlled property
+			if (const TMap<FName, FString>* PropertyMetadatas = SequencerProperty->GetMetaDataMap())
+			{
+				static const TSet<FName> SupprtedMetadataNames =
+				{
+					 "DisplayName",
+					 "Tooltip",
+					 "SliderExponent",
+					 "LinearDeltaSensitivity",
+					 "Delta",
+					 "ClampMin",
+					 "ClampMax",
+					 "UIMin",
+					 "UIMax",
+					 "Units"
+				};
+
+				for (const TTuple<FName, FString>& PropertyMetadata : *PropertyMetadatas)
+				{
+					if (SupprtedMetadataNames.Contains(PropertyMetadata.Key) &&
+						!PropertyMetadata.Value.IsEmpty())
+					{
+						ValuePropertyHandle->SetInstanceMetaData(PropertyMetadata.Key, PropertyMetadata.Value);
+					}
+				}
+			}
+		}
+	}
+}
diff --git a/Engine/Plugins/Animation/ControlRig/Source/ControlRigEditor/Private/AnimDetails/Customizations/AnimDetailsSetInstancedPropertyMetadata.h b/Engine/Plugins/Animation/ControlRig/Source/ControlRigEditor/Private/AnimDetails/Customizations/AnimDetailsSetInstancedPropertyMetadata.h
new file mode 100644
--- /dev/null
+++ b/Engine/Plugins/Animation/ControlRig/Source/ControlRigEditor/Private/AnimDetails/Customizations/AnimDetailsSetInstancedPropertyMetadata.h
@@ -0,0 +1,14 @@
+// Copyright Epic Games, Inc. All Rights Reserved.
+
+#pragma once
+
+#include "Templates/SharedPointer.h"
+
+class IPropertyHandle;
+class UAnimDetailsProxyBase;
+
+namespace UE::ControlRigEditor::AnimDetailsMetaDataUtil
+{
+	/** Sets property meta data for the property instance from the data in the template proxy, for example min and max values. */
+	void SetInstancedPropertyMetaData(const UAnimDetailsProxyBase& TemplateProxy, const TSharedRef<IPropertyHandle>& ValuePropertyHandle);
+}
diff --git a/Engine/Plugins/Animation/ControlRig/Source/ControlRigEditor/Private/AnimDetails/Customizations/AnimDetailsValueCustomization.cpp b/Engine/Plugins/Animation/ControlRig/Source/ControlRigEditor/Private/AnimDetails/Customizations/AnimDetailsValueCustomization.cpp
--- a/Engine/Plugins/Animation/ControlRig/Source/ControlRigEditor/Private/AnimDetails/Customizations/AnimDetailsValueCustomization.cpp
+++ b/Engine/Plugins/Animation/ControlRig/Source/ControlRigEditor/Private/AnimDetails/Customizations/AnimDetailsValueCustomization.cpp
@@ -4,6 +4,7 @@
 
 #include "Algo/AnyOf.h"
 #include "AnimDetails/AnimDetailsProxyManager.h"
+#include "AnimDetails/Customizations/AnimDetailsSetInstancedPropertyMetadata.h"
 #include "AnimDetails/Proxies/AnimDetailsProxyLocation.h"
 #include "AnimDetails/Proxies/AnimDetailsProxyRotation.h"
 #include "AnimDetails/Proxies/AnimDetailsProxyScale.h"
@@ -124,6 +125,8 @@ namespace UE::ControlRigEditor
 				continue;
 			}
 
+			AnimDetailsMetaDataUtil::SetInstancedPropertyMetaData(*FirstProxy, ChildHandle);
+
 			const FText PropertyDisplayText = [FirstProxy, &ChildHandle, this]()
 				{
 					// Draw the proxy display name if the property has only one child.
diff --git a/Engine/Plugins/Animation/ControlRig/Source/ControlRigEditor/Private/AnimDetails/Proxies/AnimDetailsProxyBase.h b/Engine/Plugins/Animation/ControlRig/Source/ControlRigEditor/Private/AnimDetails/Proxies/AnimDetailsProxyBase.h
--- a/Engine/Plugins/Animation/ControlRig/Source/ControlRigEditor/Private/AnimDetails/Proxies/AnimDetailsProxyBase.h
+++ b/Engine/Plugins/Animation/ControlRig/Source/ControlRigEditor/Private/AnimDetails/Proxies/AnimDetailsProxyBase.h
@@ -87,9 +87,6 @@ public:
 	 */
 	virtual FName GetPropertyID(const FName& PropertyName) const;
 
-	/** Convenience function to allow the proxy to update its related property display name */
-	virtual void UpdatePropertyDisplayNames(IDetailLayoutBuilder& DetailBuilder) {};
-
 	/** Returns the property names in this proxy controls. Should be in order they're declared */
 	virtual TArray<FName> GetPropertyNames() const PURE_VIRTUAL(UAnimDetailsProxyBase::GetPropertyNames, return {};)
 
diff --git a/Engine/Plugins/Animation/ControlRig/Source/ControlRigEditor/Private/AnimDetails/Proxies/AnimDetailsProxyBool.cpp b/Engine/Plugins/Animation/ControlRig/Source/ControlRigEditor/Private/AnimDetails/Proxies/AnimDetailsProxyBool.cpp
--- a/Engine/Plugins/Animation/ControlRig/Source/ControlRigEditor/Private/AnimDetails/Proxies/AnimDetailsProxyBool.cpp
+++ b/Engine/Plugins/Animation/ControlRig/Source/ControlRigEditor/Private/AnimDetails/Proxies/AnimDetailsProxyBool.cpp
@@ -14,37 +14,6 @@ FName UAnimDetailsProxyBool::GetCategoryName() const
 	return "Bool";
 }
 
-void UAnimDetailsProxyBool::UpdatePropertyDisplayNames(IDetailLayoutBuilder& DetailBuilder)
-{
-	// Use the short name as the property display name when displaying control rig elements
-	constexpr EElementNameDisplayMode NameDisplayMode = EElementNameDisplayMode::ForceShort;
-
-	TSharedPtr<IPropertyHandle> ValuePropertyHandle = DetailBuilder.GetProperty(GET_MEMBER_NAME_CHECKED(UAnimDetailsProxyBool, Bool), GetClass());
-	if (ValuePropertyHandle.IsValid())
-	{
-		ValuePropertyHandle->SetPropertyDisplayName(GetDisplayNameText(NameDisplayMode));
-	}
-
-	if (bIsIndividual)
-	{
-		const FString PropertyPath = FString::Printf(TEXT("%s.%s"), GET_MEMBER_NAME_STRING_CHECKED(UAnimDetailsProxyBool, Bool), GET_MEMBER_NAME_STRING_CHECKED(FAnimDetailsBool, Bool));
-		ValuePropertyHandle = DetailBuilder.GetProperty(FName(*PropertyPath), GetClass());
-		if (ValuePropertyHandle.IsValid())
-		{
-			ValuePropertyHandle->SetPropertyDisplayName(GetDisplayNameText(NameDisplayMode));
-		}
-	}
-	else
-	{
-		const FString PropertyPath = FString::Printf(TEXT("%s.%s"), GET_MEMBER_NAME_STRING_CHECKED(UAnimDetailsProxyBool, Bool), GET_MEMBER_NAME_STRING_CHECKED(FAnimDetailsBool, Bool));
-		ValuePropertyHandle = DetailBuilder.GetProperty(FName(*PropertyPath), GetClass());
-		if (ValuePropertyHandle.IsValid() && SequencerItem.IsValid())
-		{
-			ValuePropertyHandle->SetPropertyDisplayName(FText::FromName(SequencerItem.GetBinding()->GetPropertyName()));
-		}
-	}
-}
-
 TArray<FName> UAnimDetailsProxyBool::GetPropertyNames() const
 {
 	return
diff --git a/Engine/Plugins/Animation/ControlRig/Source/ControlRigEditor/Private/AnimDetails/Proxies/AnimDetailsProxyBool.h b/Engine/Plugins/Animation/ControlRig/Source/ControlRigEditor/Private/AnimDetails/Proxies/AnimDetailsProxyBool.h
--- a/Engine/Plugins/Animation/ControlRig/Source/ControlRigEditor/Private/AnimDetails/Proxies/AnimDetailsProxyBool.h
+++ b/Engine/Plugins/Animation/ControlRig/Source/ControlRigEditor/Private/AnimDetails/Proxies/AnimDetailsProxyBool.h
@@ -26,7 +26,6 @@ class UAnimDetailsProxyBool
 public:
 	//~ Begin UAnimDetailsProxyBase interface
 	virtual FName GetCategoryName() const override;
-	virtual void UpdatePropertyDisplayNames(IDetailLayoutBuilder& DetailBuilder);
 	virtual TArray<FName> GetPropertyNames() const override;
 	virtual TSet<ERigControlType> GetSupportedControlTypes() const override;
 	virtual bool PropertyIsOnProxy(const FProperty* Property, const FProperty* MemberProperty) override;
diff --git a/Engine/Plugins/Animation/ControlRig/Source/ControlRigEditor/Private/AnimDetails/Proxies/AnimDetailsProxyEnum.cpp b/Engine/Plugins/Animation/ControlRig/Source/ControlRigEditor/Private/AnimDetails/Proxies/AnimDetailsProxyEnum.cpp
--- a/Engine/Plugins/Animation/ControlRig/Source/ControlRigEditor/Private/AnimDetails/Proxies/AnimDetailsProxyEnum.cpp
+++ b/Engine/Plugins/Animation/ControlRig/Source/ControlRigEditor/Private/AnimDetails/Proxies/AnimDetailsProxyEnum.cpp
@@ -24,37 +24,6 @@ FName UAnimDetailsProxyEnum::GetDetailRowID() const
 	return UAnimDetailsProxyBase::GetDetailRowID();
 }
 
-void UAnimDetailsProxyEnum::UpdatePropertyDisplayNames(IDetailLayoutBuilder& DetailBuilder)
-{
-	// Use the short name as the property display name when displaying control rig elements
-	constexpr EElementNameDisplayMode NameDisplayMode = EElementNameDisplayMode::ForceShort;
-
-	TSharedPtr<IPropertyHandle> ValuePropertyHandle = DetailBuilder.GetProperty(GET_MEMBER_NAME_CHECKED(UAnimDetailsProxyEnum, Enum), GetClass());
-	if (ValuePropertyHandle.IsValid())
-	{
-		ValuePropertyHandle->SetPropertyDisplayName(GetDisplayNameText(NameDisplayMode));
-	}
-
-	if (SequencerItem.IsValid())
-	{
-		const FString PropertyPath = FString::Printf(TEXT("%s.%s"), GET_MEMBER_NAME_STRING_CHECKED(UAnimDetailsProxyEnum, Enum), GET_MEMBER_NAME_STRING_CHECKED(FAnimDetailsEnum, EnumIndex));
-		ValuePropertyHandle = DetailBuilder.GetProperty(FName(*PropertyPath), GetClass());
-		if (ValuePropertyHandle.IsValid())
-		{
-			ValuePropertyHandle->SetPropertyDisplayName(FText::FromName(SequencerItem.GetBinding()->GetPropertyName()));
-		}
-	}
-	else if (bIsIndividual)
-	{
-		const FString PropertyPath = FString::Printf(TEXT("%s.%s"), GET_MEMBER_NAME_STRING_CHECKED(UAnimDetailsProxyEnum, Enum), GET_MEMBER_NAME_STRING_CHECKED(FAnimDetailsEnum, EnumIndex));
-		ValuePropertyHandle = DetailBuilder.GetProperty(FName(*PropertyPath), GetClass());
-		if (ValuePropertyHandle.IsValid())
-		{
-			ValuePropertyHandle->SetPropertyDisplayName(GetDisplayNameText(NameDisplayMode));
-		}
-	}
-}
-
 TArray<FName> UAnimDetailsProxyEnum::GetPropertyNames() const
 {
 	return
diff --git a/Engine/Plugins/Animation/ControlRig/Source/ControlRigEditor/Private/AnimDetails/Proxies/AnimDetailsProxyEnum.h b/Engine/Plugins/Animation/ControlRig/Source/ControlRigEditor/Private/AnimDetails/Proxies/AnimDetailsProxyEnum.h
--- a/Engine/Plugins/Animation/ControlRig/Source/ControlRigEditor/Private/AnimDetails/Proxies/AnimDetailsProxyEnum.h
+++ b/Engine/Plugins/Animation/ControlRig/Source/ControlRigEditor/Private/AnimDetails/Proxies/AnimDetailsProxyEnum.h
@@ -36,7 +36,6 @@ public:
 	//~ Begin UAnimDetailsProxyBase interface
 	virtual FName GetCategoryName() const override;
 	virtual FName GetDetailRowID() const override;
-	virtual void UpdatePropertyDisplayNames(IDetailLayoutBuilder& DetailBuilder);
 	virtual TArray<FName> GetPropertyNames() const override;
 	virtual TSet<ERigControlType> GetSupportedControlTypes() const override;
 	virtual bool PropertyIsOnProxy(const FProperty* Property, const FProperty* MemberProperty) override;
diff --git a/Engine/Plugins/Animation/ControlRig/Source/ControlRigEditor/Private/AnimDetails/Proxies/AnimDetailsProxyFloat.cpp b/Engine/Plugins/Animation/ControlRig/Source/ControlRigEditor/Private/AnimDetails/Proxies/AnimDetailsProxyFloat.cpp
--- a/Engine/Plugins/Animation/ControlRig/Source/ControlRigEditor/Private/AnimDetails/Proxies/AnimDetailsProxyFloat.cpp
+++ b/Engine/Plugins/Animation/ControlRig/Source/ControlRigEditor/Private/AnimDetails/Proxies/AnimDetailsProxyFloat.cpp
@@ -14,37 +14,6 @@ FName UAnimDetailsProxyFloat::GetCategoryName() const
 	return "Float";
 }
 
-void UAnimDetailsProxyFloat::UpdatePropertyDisplayNames(IDetailLayoutBuilder& DetailBuilder)
-{
-	// Use the short name as the property display name when displaying control rig elements
-	constexpr EElementNameDisplayMode NameDisplayMode = EElementNameDisplayMode::ForceShort;
-
-	TSharedPtr<IPropertyHandle> ValuePropertyHandle = DetailBuilder.GetProperty(GET_MEMBER_NAME_CHECKED(UAnimDetailsProxyFloat, Float), GetClass());
-	if (ValuePropertyHandle.IsValid())
-	{
-		ValuePropertyHandle->SetPropertyDisplayName(GetDisplayNameText(NameDisplayMode));
-	}
-
-	if (bIsIndividual)
-	{
-		const FString PropertyPath = FString::Printf(TEXT("%s.%s"), GET_MEMBER_NAME_STRING_CHECKED(UAnimDetailsProxyFloat, Float), GET_MEMBER_NAME_STRING_CHECKED(FAnimDetailsFloat, Float));
-		ValuePropertyHandle = DetailBuilder.GetProperty(FName(*PropertyPath), GetClass());
-		if (ValuePropertyHandle.IsValid())
-		{
-			ValuePropertyHandle->SetPropertyDisplayName(GetDisplayNameText(NameDisplayMode));
-		}
-	}
-	else if (SequencerItem.IsValid())
-	{
-		const FString PropertyPath = FString::Printf(TEXT("%s.%s"), GET_MEMBER_NAME_STRING_CHECKED(UAnimDetailsProxyFloat, Float), GET_MEMBER_NAME_STRING_CHECKED(FAnimDetailsFloat, Float));
-		ValuePropertyHandle = DetailBuilder.GetProperty(FName(*PropertyPath), GetClass());
-		if (ValuePropertyHandle.IsValid())
-		{
-			ValuePropertyHandle->SetPropertyDisplayName(FText::FromName(SequencerItem.GetBinding()->GetPropertyName()));
-		}
-	}
-}
-
 TArray<FName> UAnimDetailsProxyFloat::GetPropertyNames() const
 {
 	return
diff --git a/Engine/Plugins/Animation/ControlRig/Source/ControlRigEditor/Private/AnimDetails/Proxies/AnimDetailsProxyFloat.h b/Engine/Plugins/Animation/ControlRig/Source/ControlRigEditor/Private/AnimDetails/Proxies/AnimDetailsProxyFloat.h
--- a/Engine/Plugins/Animation/ControlRig/Source/ControlRigEditor/Private/AnimDetails/Proxies/AnimDetailsProxyFloat.h
+++ b/Engine/Plugins/Animation/ControlRig/Source/ControlRigEditor/Private/AnimDetails/Proxies/AnimDetailsProxyFloat.h
@@ -17,7 +17,7 @@ struct FAnimDetailsFloat
 {
 	GENERATED_BODY()
 
-	UPROPERTY(EditAnywhere, Interp, Category = "Float", meta = (SliderExponent = "1.0"))
+	UPROPERTY(EditAnywhere, Interp, Category = "Float")
 	double Float = 0.0;
 };
 
@@ -31,7 +31,6 @@ class UAnimDetailsProxyFloat
 public:
 	//~ Begin UAnimDetailsProxyBase interface
 	virtual FName GetCategoryName() const override;
-	virtual void UpdatePropertyDisplayNames(IDetailLayoutBuilder& DetailBuilder);
 	virtual TArray<FName> GetPropertyNames() const override;
 	virtual TSet<ERigControlType> GetSupportedControlTypes() const override;
 	virtual bool PropertyIsOnProxy(const FProperty* Property, const FProperty* MemberProperty) override;
diff --git a/Engine/Plugins/Animation/ControlRig/Source/ControlRigEditor/Private/AnimDetails/Proxies/AnimDetailsProxyInteger.cpp b/Engine/Plugins/Animation/ControlRig/Source/ControlRigEditor/Private/AnimDetails/Proxies/AnimDetailsProxyInteger.cpp
--- a/Engine/Plugins/Animation/ControlRig/Source/ControlRigEditor/Private/AnimDetails/Proxies/AnimDetailsProxyInteger.cpp
+++ b/Engine/Plugins/Animation/ControlRig/Source/ControlRigEditor/Private/AnimDetails/Proxies/AnimDetailsProxyInteger.cpp
@@ -14,37 +14,6 @@ FName UAnimDetailsProxyInteger::GetCategoryName() const
 	return "Integer";
 }
 
-void UAnimDetailsProxyInteger::UpdatePropertyDisplayNames(IDetailLayoutBuilder& DetailBuilder)
-{
-	// Use the short name as the property display name when displaying control rig elements
-	constexpr EElementNameDisplayMode NameDisplayMode = EElementNameDisplayMode::ForceShort;
-
-	TSharedPtr<IPropertyHandle> ValuePropertyHandle = DetailBuilder.GetProperty(GET_MEMBER_NAME_CHECKED(UAnimDetailsProxyInteger, Integer), GetClass());
-	if (ValuePropertyHandle.IsValid())
-	{
-		ValuePropertyHandle->SetPropertyDisplayName(GetDisplayNameText(NameDisplayMode));
-	}
-
-	if (SequencerItem.IsValid())
-	{
-		const FString PropertyPath = FString::Printf(TEXT("%s.%s"), GET_MEMBER_NAME_STRING_CHECKED(UAnimDetailsProxyInteger, Integer), GET_MEMBER_NAME_STRING_CHECKED(FAnimDetailsInteger, Integer));
-		ValuePropertyHandle = DetailBuilder.GetProperty(FName(*PropertyPath), GetClass());
-		if (ValuePropertyHandle.IsValid())
-		{
-			ValuePropertyHandle->SetPropertyDisplayName(FText::FromName(SequencerItem.GetBinding()->GetPropertyName()));
-		}
-	}
-	else if (bIsIndividual)
-	{
-		const FString PropertyPath = FString::Printf(TEXT("%s.%s"), GET_MEMBER_NAME_STRING_CHECKED(UAnimDetailsProxyInteger, Integer), GET_MEMBER_NAME_STRING_CHECKED(FAnimDetailsInteger, Integer));
-		ValuePropertyHandle = DetailBuilder.GetProperty(FName(*PropertyPath), GetClass());
-		if (ValuePropertyHandle.IsValid())
-		{
-			ValuePropertyHandle->SetPropertyDisplayName(GetDisplayNameText(NameDisplayMode));
-		}
-	}
-}
-
 TArray<FName> UAnimDetailsProxyInteger::GetPropertyNames() const
 {
 	return 
diff --git a/Engine/Plugins/Animation/ControlRig/Source/ControlRigEditor/Private/AnimDetails/Proxies/AnimDetailsProxyInteger.h b/Engine/Plugins/Animation/ControlRig/Source/ControlRigEditor/Private/AnimDetails/Proxies/AnimDetailsProxyInteger.h
--- a/Engine/Plugins/Animation/ControlRig/Source/ControlRigEditor/Private/AnimDetails/Proxies/AnimDetailsProxyInteger.h
+++ b/Engine/Plugins/Animation/ControlRig/Source/ControlRigEditor/Private/AnimDetails/Proxies/AnimDetailsProxyInteger.h
@@ -26,7 +26,6 @@ class UAnimDetailsProxyInteger
 public:
 	//~ Begin UAnimDetailsProxyBase interface
 	virtual FName GetCategoryName() const override;
-	virtual void UpdatePropertyDisplayNames(IDetailLayoutBuilder& DetailBuilder);
 	virtual TArray<FName> GetPropertyNames() const override;
 	virtual TSet<ERigControlType> GetSupportedControlTypes() const override;
 	virtual bool PropertyIsOnProxy(const FProperty* Property, const FProperty* MemberProperty) override;
diff --git a/Engine/Plugins/Animation/ControlRig/Source/ControlRigEditor/Private/AnimDetails/Proxies/AnimDetailsProxyLocation.h b/Engine/Plugins/Animation/ControlRig/Source/ControlRigEditor/Private/AnimDetails/Proxies/AnimDetailsProxyLocation.h
--- a/Engine/Plugins/Animation/ControlRig/Source/ControlRigEditor/Private/AnimDetails/Proxies/AnimDetailsProxyLocation.h
+++ b/Engine/Plugins/Animation/ControlRig/Source/ControlRigEditor/Private/AnimDetails/Proxies/AnimDetailsProxyLocation.h
@@ -12,13 +12,13 @@ struct FAnimDetailsLocation
 {
 	GENERATED_BODY()
 
-	UPROPERTY(EditAnywhere, Interp, Category = "Location", meta = (Delta = "0.5", SliderExponent = "1", LinearDeltaSensitivity = "1"))
+	UPROPERTY(EditAnywhere, Interp, Category = "Location")
 	double LX = 0.0;
 
-	UPROPERTY(EditAnywhere, Interp, Category = "Location", meta = (Delta = "0.5", SliderExponent = "1", LinearDeltaSensitivity = "1"))
+	UPROPERTY(EditAnywhere, Interp, Category = "Location")
 	double LY = 0.0;
 
-	UPROPERTY(EditAnywhere, Interp, Category = "Location", meta = (Delta = "0.5", SliderExponent = "1", LinearDeltaSensitivity = "1"))
+	UPROPERTY(EditAnywhere, Interp, Category = "Location")
 	double LZ = 0.0;
 
 	FAnimDetailsLocation() = default;
diff --git a/Engine/Plugins/Animation/ControlRig/Source/ControlRigEditor/Private/AnimDetails/Proxies/AnimDetailsProxyRotation.h b/Engine/Plugins/Animation/ControlRig/Source/ControlRigEditor/Private/AnimDetails/Proxies/AnimDetailsProxyRotation.h
--- a/Engine/Plugins/Animation/ControlRig/Source/ControlRigEditor/Private/AnimDetails/Proxies/AnimDetailsProxyRotation.h
+++ b/Engine/Plugins/Animation/ControlRig/Source/ControlRigEditor/Private/AnimDetails/Proxies/AnimDetailsProxyRotation.h
@@ -12,13 +12,13 @@ struct  FAnimDetailsRotation
 {
 	GENERATED_BODY()
 
-	UPROPERTY(EditAnywhere, Interp, Category = "Rotation", meta = (Delta = "0.5", SliderExponent = "1", LinearDeltaSensitivity = "1"))
+	UPROPERTY(EditAnywhere, Interp, Category = "Rotation")
 	double RX = 0.0;
 
-	UPROPERTY(EditAnywhere, Interp, Category = "Rotation", meta = (Delta = "0.5", SliderExponent = "1", LinearDeltaSensitivity = "1"))
+	UPROPERTY(EditAnywhere, Interp, Category = "Rotation")
 	double RY = 0.0;
 
-	UPROPERTY(EditAnywhere, Interp, Category = "Rotation", meta = (Delta = "0.5", SliderExponent = "1", LinearDeltaSensitivity = "1"))
+	UPROPERTY(EditAnywhere, Interp, Category = "Rotation")
 	double RZ = 0.0;
 
 	FAnimDetailsRotation() = default;
diff --git a/Engine/Plugins/Animation/ControlRig/Source/ControlRigEditor/Private/AnimDetails/Proxies/AnimDetailsProxyScale.h b/Engine/Plugins/Animation/ControlRig/Source/ControlRigEditor/Private/AnimDetails/Proxies/AnimDetailsProxyScale.h
--- a/Engine/Plugins/Animation/ControlRig/Source/ControlRigEditor/Private/AnimDetails/Proxies/AnimDetailsProxyScale.h
+++ b/Engine/Plugins/Animation/ControlRig/Source/ControlRigEditor/Private/AnimDetails/Proxies/AnimDetailsProxyScale.h
@@ -12,13 +12,13 @@ struct FAnimDetailsScale
 {
 	GENERATED_BODY()
 
-	UPROPERTY(EditAnywhere, Interp, Category = "Scale", meta = (Delta = "0.5", SliderExponent = "1", LinearDeltaSensitivity = "1"))
+	UPROPERTY(EditAnywhere, Interp, Category = "Scale")
 	double SX = 1.0;
 
-	UPROPERTY(EditAnywhere, Interp, Category = "Scale", meta = (Delta = "0.5", SliderExponent = "1", LinearDeltaSensitivity = "1"))
+	UPROPERTY(EditAnywhere, Interp, Category = "Scale")
 	double SY = 1.0;
 
-	UPROPERTY(EditAnywhere, Interp, Category = "Scale", meta = (Delta = "0.5", SliderExponent = "1", LinearDeltaSensitivity = "1"))
+	UPROPERTY(EditAnywhere, Interp, Category = "Scale")
 	double SZ = 1.0;
 
 	FAnimDetailsScale() = default;
diff --git a/Engine/Plugins/Animation/ControlRig/Source/ControlRigEditor/Private/AnimDetails/Widgets/SAnimDetailsValueNumeric.cpp b/Engine/Plugins/Animation/ControlRig/Source/ControlRigEditor/Private/AnimDetails/Widgets/SAnimDetailsValueNumeric.cpp
--- a/Engine/Plugins/Animation/ControlRig/Source/ControlRigEditor/Private/AnimDetails/Widgets/SAnimDetailsValueNumeric.cpp
+++ b/Engine/Plugins/Animation/ControlRig/Source/ControlRigEditor/Private/AnimDetails/Widgets/SAnimDetailsValueNumeric.cpp
@@ -20,9 +20,9 @@ namespace UE::ControlRigEditor
 	{
 		/** Struct to extract meta data from a property handle */
 		template <typename NumericType>
-		struct FFloatingPointMetaData
+		struct FAnimDetailsPropertyMetadata
 		{
-			FFloatingPointMetaData(const TSharedRef<IPropertyHandle>& PropertyHandle);
+			FAnimDetailsPropertyMetadata(const TSharedRef<IPropertyHandle>& PropertyHandle);
 
 			TOptional<NumericType> MinValue;
 			TOptional<NumericType> MaxValue;
@@ -37,24 +37,18 @@ namespace UE::ControlRigEditor
 		};
 
 		template <typename NumericType>
-		FFloatingPointMetaData<NumericType>::FFloatingPointMetaData(const TSharedRef<IPropertyHandle>& PropertyHandle)
+		FAnimDetailsPropertyMetadata<NumericType>::FAnimDetailsPropertyMetadata(const TSharedRef<IPropertyHandle>& PropertyHandle)
 		{
-			const FProperty* Property = PropertyHandle->GetProperty();
-			if (!Property)
-			{
-				return;
-			}
-
-			const FString& MetaUIMinString = Property->GetMetaData(TEXT("UIMin"));
-			const FString& MetaUIMaxString = Property->GetMetaData(TEXT("UIMax"));
-			const FString& SliderExponentString = Property->GetMetaData(TEXT("SliderExponent"));
-			const FString& DeltaString = Property->GetMetaData(TEXT("Delta"));
-			const FString& ShiftMultiplierString = Property->GetMetaData(TEXT("ShiftMultiplier"));
-			const FString& CtrlMultiplierString = Property->GetMetaData(TEXT("CtrlMultiplier"));
-			const FString& SupportDynamicSliderMaxValueString = Property->GetMetaData(TEXT("SupportDynamicSliderMaxValue"));
-			const FString& SupportDynamicSliderMinValueString = Property->GetMetaData(TEXT("SupportDynamicSliderMinValue"));
-			const FString& ClampMinString = Property->GetMetaData(TEXT("ClampMin"));
-			const FString& ClampMaxString = Property->GetMetaData(TEXT("ClampMax"));
+			const FString& MetaUIMinString = PropertyHandle->GetMetaData(TEXT("UIMin"));
+			const FString& MetaUIMaxString = PropertyHandle->GetMetaData(TEXT("UIMax"));
+			const FString& SliderExponentString = PropertyHandle->GetMetaData(TEXT("SliderExponent"));
+			const FString& DeltaString = PropertyHandle->GetMetaData(TEXT("Delta"));
+			const FString& ShiftMultiplierString = PropertyHandle->GetMetaData(TEXT("ShiftMultiplier"));
+			const FString& CtrlMultiplierString = PropertyHandle->GetMetaData(TEXT("CtrlMultiplier"));
+			const FString& SupportDynamicSliderMaxValueString = PropertyHandle->GetMetaData(TEXT("SupportDynamicSliderMaxValue"));
+			const FString& SupportDynamicSliderMinValueString = PropertyHandle->GetMetaData(TEXT("SupportDynamicSliderMinValue"));
+			const FString& ClampMinString = PropertyHandle->GetMetaData(TEXT("ClampMin"));
+			const FString& ClampMaxString = PropertyHandle->GetMetaData(TEXT("ClampMax"));
 
 			// If no UIMin/Max was specified then use the clamp string
 			const FString& UIMinString = MetaUIMinString.Len() ? MetaUIMinString : ClampMinString;
@@ -238,7 +232,7 @@ namespace UE::ControlRigEditor
 		UAnimDetailsProxyManager* ProxyManager = WeakProxyManager.Get();
 
 		const TSharedRef<IPropertyHandle> PropertyHandle = WeakPropertyHandle.Pin().ToSharedRef();
-		const PropertyUtils::FFloatingPointMetaData<NumericType>& PropertyMetaData = PropertyUtils::FFloatingPointMetaData<NumericType>(PropertyHandle);
+		const PropertyUtils::FAnimDetailsPropertyMetadata<NumericType>& PropertyMetaData = PropertyUtils::FAnimDetailsPropertyMetadata<NumericType>(PropertyHandle);
 
 		const uint8 NumFractionalDigits = GetDefault<UAnimDetailsSettings>()->NumFractionalDigits;
 
@@ -264,7 +258,7 @@ namespace UE::ControlRigEditor
 			.SliderExponent(PropertyMetaData.SliderExponent)
 			.Delta(PropertyMetaData.Delta)
 			.LinearDeltaSensitivity(1)
-			.ClearKeyboardFocusOnCommit(true)
+			.ClearKeyboardFocusOnCommit(false)
 			.TypeInterface(TypeInterface);
 
 		NavigableWidgetRegistrar.RegisterAsNavigable(*ProxyManager, AsShared(), SpinBox.ToSharedRef());
