Anim in Engine: Avoid refreshing the Anim Details view systematically whenever Anim Details self is editing interactively

diff --git a/Engine/Plugins/Animation/ControlRig/Source/ControlRigEditor/Private/AnimDetails/AnimDetailsProxyManager.cpp b/Engine/Plugins/Animation/ControlRig/Source/ControlRigEditor/Private/AnimDetails/AnimDetailsProxyManager.cpp
--- a/Engine/Plugins/Animation/ControlRig/Source/ControlRigEditor/Private/AnimDetails/AnimDetailsProxyManager.cpp
+++ b/Engine/Plugins/Animation/ControlRig/Source/ControlRigEditor/Private/AnimDetails/AnimDetailsProxyManager.cpp
@@ -76,9 +76,27 @@ void UAnimDetailsProxyManager::RequestUpdateProxies()
 	}
 
 	if (ensureMsgf(IsInGameThread(), TEXT("Anim Details proxies can only be updated in game thread. Ignoring call")) &&
-		!RequestUpdateProxiesTimerHandle.IsValid())
+		!RequestUpdateProxiesTickerHandle.IsValid())
 	{
-		RequestUpdateProxiesTimerHandle = GEditor->GetTimerManager()->SetTimerForNextTick(FTimerDelegate::CreateUObject(this, &UAnimDetailsProxyManager::ForceUpdateProxies));
+		const TWeakObjectPtr<UAnimDetailsProxyManager> WeakThis = this;
+		RequestUpdateProxiesTickerHandle = FTSTicker::GetCoreTicker().AddTicker(FTickerDelegate::CreateLambda([WeakThis](float DeltaTime)
+			{
+				if (!WeakThis.IsValid())
+				{
+					return false;
+				}
+
+				// Never update proxies while anim details self is editing interactively, but handle the request after
+				if (UE::ControlRigEditor::FAnimDetailsMultiEditUtil::Get().IsInteractive())
+				{
+					return true;
+				}
+
+				WeakThis->ForceUpdateProxies();
+				
+				return false;
+			})
+		);
 	}
 }
 
@@ -248,11 +266,7 @@ void UAnimDetailsProxyManager::OnSequencerSelectionChanged(TArray<FGuid> ObjectG
 
 void UAnimDetailsProxyManager::OnObjectsReplaced(const TMap<UObject*, UObject*>& OldToNewInstanceMap)
 {
-	// Objects may be replaced while we are editing interactively, don't update proxies in these cases.
-	if (!UE::ControlRigEditor::FAnimDetailsMultiEditUtil::Get().IsInteractive())
-	{
-		RequestUpdateProxies();
-	}
+	RequestUpdateProxies();
 }
 
 void UAnimDetailsProxyManager::ForceUpdateProxies()
@@ -263,7 +277,7 @@ void UAnimDetailsProxyManager::ForceUpdateProxies()
 	}
 
 	// FTimerHandle has to be invalidated on the game thread
-	RequestUpdateProxiesTimerHandle.Invalidate();
+	RequestUpdateProxiesTickerHandle.Reset();
 
 	const FControlRigEditMode* EditMode = static_cast<FControlRigEditMode*>(GLevelEditorModeTools().GetActiveMode(FControlRigEditMode::ModeName));
 	if (!EditMode || !UpdateSequencer())
diff --git a/Engine/Plugins/Animation/ControlRig/Source/ControlRigEditor/Private/AnimDetails/AnimDetailsProxyManager.h b/Engine/Plugins/Animation/ControlRig/Source/ControlRigEditor/Private/AnimDetails/AnimDetailsProxyManager.h
--- a/Engine/Plugins/Animation/ControlRig/Source/ControlRigEditor/Private/AnimDetails/AnimDetailsProxyManager.h
+++ b/Engine/Plugins/Animation/ControlRig/Source/ControlRigEditor/Private/AnimDetails/AnimDetailsProxyManager.h
@@ -4,6 +4,7 @@
 
 #include "AnimDetailsFilter.h"
 #include "AnimDetailsNavigableWidgetRegistry.h"
+#include "Containers/Ticker.h"
 #include "EditorUndoClient.h"
 #include "Engine/TimerHandle.h"
 #include "TimerManager.h"
@@ -169,8 +170,8 @@ private:
 	/** Delegate broadcase when the selection changed */
 	FSimpleMulticastDelegate OnProxiesChangedDelegate;
 
-	/** Timer handle for the RequestUpdateProxies member function */
-	FTimerHandle RequestUpdateProxiesTimerHandle;
+	/** Ticker handle for the RequestUpdateProxies member function */
+	FTSTicker::FDelegateHandle RequestUpdateProxiesTickerHandle;
 
 	/** Timer handle for the RequestUpdateProxyValues member function */
 	FTimerHandle RequestUpdateProxyValuesTimerHandle;
