Anim in Engine: AnimDetailsProxyManager now properly unregisters the RequestUpdateProxies Ticker when ForceUpdateProxies

#rnx

diff --git a/Engine/Plugins/Animation/ControlRig/Source/ControlRigEditor/Private/AnimDetails/AnimDetailsProxyManager.cpp b/Engine/Plugins/Animation/ControlRig/Source/ControlRigEditor/Private/AnimDetails/AnimDetailsProxyManager.cpp
--- a/Engine/Plugins/Animation/ControlRig/Source/ControlRigEditor/Private/AnimDetails/AnimDetailsProxyManager.cpp
+++ b/Engine/Plugins/Animation/ControlRig/Source/ControlRigEditor/Private/AnimDetails/AnimDetailsProxyManager.cpp
@@ -276,7 +276,7 @@ void UAnimDetailsProxyManager::ForceUpdateProxies()
 		return;
 	}
 
-	// FTimerHandle has to be invalidated on the game thread
+	FTSTicker::GetCoreTicker().RemoveTicker(RequestUpdateProxiesTickerHandle);
 	RequestUpdateProxiesTickerHandle.Reset();
 
 	const FControlRigEditMode* EditMode = static_cast<FControlRigEditMode*>(GLevelEditorModeTools().GetActiveMode(FControlRigEditMode::ModeName));
