Constraints: fixed new evaluation with ControlRigComponent

#rnx

#FailedValidations Preflight

diff --git a/Engine/Plugins/Animation/ControlRig/Source/ControlRig/Private/Constraints/ControlRigTransformableHandle.cpp b/Engine/Plugins/Animation/ControlRig/Source/ControlRig/Private/Constraints/ControlRigTransformableHandle.cpp
--- a/Engine/Plugins/Animation/ControlRig/Source/ControlRig/Private/Constraints/ControlRigTransformableHandle.cpp
+++ b/Engine/Plugins/Animation/ControlRig/Source/ControlRig/Private/Constraints/ControlRigTransformableHandle.cpp
@@ -251,7 +251,7 @@ FTransform UTransformableControlHandle::GetGlobalTransform() const
 		return FTransform::Identity;
 	}
 	
-	const USceneComponent* BoundComponent = GetBoundComponent();
+	USceneComponent* BoundComponent = GetBoundComponent();
 	if (!BoundComponent)
 	{
 		return FTransform::Identity;
@@ -260,11 +260,14 @@ FTransform UTransformableControlHandle::GetGlobalTransform() const
 	const FTransform& ComponentTransform = BoundComponent->GetComponentTransform();
 	if (TransformableHandleUtils::SkipTicking())
 	{
-		if (!ControlRig->IsEvaluating())
+		if (USkeletalMeshComponent* SkeletalMeshComponent = Cast<USkeletalMeshComponent>(BoundComponent))
 		{
-			TransformableHandleUtils::EvaluateComponent( GetSkeletalMesh(), GetEvaluationTask() );
+			if (!ControlRig->IsEvaluating())
+			{
+				TransformableHandleUtils::EvaluateComponent( SkeletalMeshComponent, GetEvaluationTask() );
+			}
+			return Cache * ComponentTransform;
 		}
-		return Cache * ComponentTransform;
 	}
 
 	const FRigElementKey& ControlKey = ControlElement->GetKey();
