Anim in Engine: Anim Details now follows the same order as Sequencer for non Control Rig properties

diff --git a/Engine/Plugins/Animation/ControlRig/Source/ControlRigEditor/Private/AnimDetails/AnimDetailsProxyManager.cpp b/Engine/Plugins/Animation/ControlRig/Source/ControlRigEditor/Private/AnimDetails/AnimDetailsProxyManager.cpp
--- a/Engine/Plugins/Animation/ControlRig/Source/ControlRigEditor/Private/AnimDetails/AnimDetailsProxyManager.cpp
+++ b/Engine/Plugins/Animation/ControlRig/Source/ControlRigEditor/Private/AnimDetails/AnimDetailsProxyManager.cpp
@@ -422,12 +422,26 @@ void UAnimDetailsProxyManager::ForceUpdateProxies()
 	ForceUpdateProxyValues();
 
 	// Grouped proxies are displayed first, reflect that in the external selection so that anim details and its selection use the same order
-	Algo::StableSortBy(ExternalSelection,
-		[](const UAnimDetailsProxyBase* Proxy)
+	ExternalSelection.StableSort(
+		[](const UAnimDetailsProxyBase& ProxyA, UAnimDetailsProxyBase& ProxyB)
 		{
-			constexpr int32 OrderInFront = 0;
-			constexpr int32 OrderToBack = 1;
-			return Proxy && !Proxy->bIsIndividual ? OrderInFront : OrderToBack;
+			const UMovieSceneTrack* TrackA = ProxyA.GetSequencerItem().GetMovieSceneTrack();
+			const UMovieSceneTrack* TrackB = ProxyB.GetSequencerItem().GetMovieSceneTrack();
+			const bool bIsIndividualA = ProxyA.bIsIndividual;
+			const bool bIsIndividualB = ProxyB.bIsIndividual;
+
+			if (TrackA &&
+				TrackB &&
+				bIsIndividualA == bIsIndividualB)
+			{
+				const bool bLexicalLess = TrackA->GetDisplayName().ToString() < TrackB->GetDisplayName().ToString();
+
+				return bLexicalLess;
+			}
+			else
+			{
+				return !bIsIndividualA || bIsIndividualA == bIsIndividualB;
+			}
 		});
 
 	OnProxiesChangedDelegate.Broadcast();
