Anim in Engine: Fix improper sorting of Anim Details Proxies led to a regression where drag selecting in Anim Details no longer worked as expected

#rnx

diff --git a/Engine/Plugins/Animation/ControlRig/Source/ControlRigEditor/Private/AnimDetails/AnimDetailsProxyManager.cpp b/Engine/Plugins/Animation/ControlRig/Source/ControlRigEditor/Private/AnimDetails/AnimDetailsProxyManager.cpp
--- a/Engine/Plugins/Animation/ControlRig/Source/ControlRigEditor/Private/AnimDetails/AnimDetailsProxyManager.cpp
+++ b/Engine/Plugins/Animation/ControlRig/Source/ControlRigEditor/Private/AnimDetails/AnimDetailsProxyManager.cpp
@@ -425,23 +425,25 @@ void UAnimDetailsProxyManager::ForceUpdateProxies()
 	ExternalSelection.StableSort(
 		[](const UAnimDetailsProxyBase& ProxyA, UAnimDetailsProxyBase& ProxyB)
 		{
-			const UMovieSceneTrack* TrackA = ProxyA.GetSequencerItem().GetMovieSceneTrack();
-			const UMovieSceneTrack* TrackB = ProxyB.GetSequencerItem().GetMovieSceneTrack();
 			const bool bIsIndividualA = ProxyA.bIsIndividual;
 			const bool bIsIndividualB = ProxyB.bIsIndividual;
+			const bool bOrderFront = !bIsIndividualA  && bIsIndividualB;
+			if (bIsIndividualB)
+			{
+				return true;
+			}
 
+			const UMovieSceneTrack* TrackA = ProxyA.GetSequencerItem().GetMovieSceneTrack();
+			const UMovieSceneTrack* TrackB = ProxyB.GetSequencerItem().GetMovieSceneTrack();
 			if (TrackA &&
-				TrackB &&
-				bIsIndividualA == bIsIndividualB)
+				TrackB)
 			{
 				const bool bLexicalLess = TrackA->GetDisplayName().ToString() < TrackB->GetDisplayName().ToString();
 
 				return bLexicalLess;
 			}
-			else
-			{
-				return !bIsIndividualA || bIsIndividualA == bIsIndividualB;
-			}
+
+			return false;
 		});
 
 	OnProxiesChangedDelegate.Broadcast();
