Anim in Engine: Fix Anim Details still does not sort non Control Rig proxies in the same order as Sequencer

#rnx

diff --git a/Engine/Plugins/Animation/ControlRig/Source/ControlRigEditor/Private/AnimDetails/AnimDetailsProxyManager.cpp b/Engine/Plugins/Animation/ControlRig/Source/ControlRigEditor/Private/AnimDetails/AnimDetailsProxyManager.cpp
--- a/Engine/Plugins/Animation/ControlRig/Source/ControlRigEditor/Private/AnimDetails/AnimDetailsProxyManager.cpp
+++ b/Engine/Plugins/Animation/ControlRig/Source/ControlRigEditor/Private/AnimDetails/AnimDetailsProxyManager.cpp
@@ -427,10 +427,9 @@ void UAnimDetailsProxyManager::ForceUpdateProxies()
 		{
 			const bool bIsIndividualA = ProxyA.bIsIndividual;
 			const bool bIsIndividualB = ProxyB.bIsIndividual;
-			const bool bOrderFront = !bIsIndividualA  && bIsIndividualB;
-			if (bIsIndividualB)
+			if (bIsIndividualA != bIsIndividualB)
 			{
-				return true;
+				return !bIsIndividualA;
 			}
 
 			const UMovieSceneTrack* TrackA = ProxyA.GetSequencerItem().GetMovieSceneTrack();
