Control Rig: Exchange icons for constraints window

diff --git a/Engine/Plugins/Animation/ControlRig/Source/ControlRigEditor/Private/ControlRigEditorStyle.h b/Engine/Plugins/Animation/ControlRig/Source/ControlRigEditor/Private/ControlRigEditorStyle.h
--- a/Engine/Plugins/Animation/ControlRig/Source/ControlRigEditor/Private/ControlRigEditorStyle.h
+++ b/Engine/Plugins/Animation/ControlRig/Source/ControlRigEditor/Private/ControlRigEditorStyle.h
@@ -249,9 +249,9 @@ public:
 				.SetPadding(FMargin(16, 6))
 			);
 			
-			Set("ControlRig.Constrain.Spaces", new IMAGE_BRUSH_SVG("Slate/Tab_Spaces_16x", Icon16x16));
-			Set("ControlRig.Constrain.Constraints", new IMAGE_BRUSH_SVG("Slate/Tab_Constraints_16x", Icon16x16));
-			Set("ControlRig.Constrain.Snapper", new IMAGE_BRUSH_SVG("Slate/Tab_RigSnapper_16x", Icon16x16));
+			Set("ControlRig.Constrain.Spaces", new IMAGE_BRUSH_SVG("Slate/AnimationSpaces_16", Icon16x16));
+			Set("ControlRig.Constrain.Constraints", new IMAGE_BRUSH_SVG("Slate/AnimationConstraint_16", Icon16x16));
+			Set("ControlRig.Constrain.Snapper", new IMAGE_BRUSH_SVG("Slate/AnimationControlRigSnapper_16", Icon16x16));
 		}
 		
 		FSlateStyleRegistry::RegisterSlateStyle(*this);
