Constraints & MRQ: make sure to reset the control transform handles when ending PIE
rendering with MRQ using nested sub-sequences could leave control handles pointing to gcd control rig instances when the MRQ world was destroyed.
this made constraints invalid, and this state could further be serialized when saving the level sequence.
handles are now updated with the default control rig when ending PIE.

#rnx

diff --git a/Engine/Plugins/Animation/ControlRig/Source/ControlRig/Private/Sequencer/MovieSceneControlRigParameterTrack.cpp b/Engine/Plugins/Animation/ControlRig/Source/ControlRig/Private/Sequencer/MovieSceneControlRigParameterTrack.cpp
--- a/Engine/Plugins/Animation/ControlRig/Source/ControlRig/Private/Sequencer/MovieSceneControlRigParameterTrack.cpp
+++ b/Engine/Plugins/Animation/ControlRig/Source/ControlRig/Private/Sequencer/MovieSceneControlRigParameterTrack.cpp
@@ -2,6 +2,7 @@
 
 #include "Sequencer/MovieSceneControlRigParameterTrack.h"
 #include "Compilation/MovieSceneCompilerRules.h"
+#include "Constraints/ControlRigTransformableHandle.h"
 #include "Evaluation/MovieSceneEvaluationTrack.h"
 #include "MovieSceneCommonHelpers.h"
 #include "EntitySystem/MovieSceneEntitySystem.h"
@@ -10,10 +11,10 @@
 #include "MovieSceneTimeHelpers.h"
 #include "Channels/MovieSceneChannelProxy.h"
 #include "Rigs/FKControlRig.h"
-#include "Rigs/RigHierarchyController.h"
 #include "UObject/Package.h"
 #include "Async/Async.h"
 #include "Sequencer/MovieSceneControlRigSystem.h"
+#include "Transform/TransformConstraint.h"
 
 #if WITH_EDITOR
 #include "ScopedTransaction.h"
@@ -80,7 +81,8 @@ UMovieSceneControlRigParameterTrack::UMovieSceneControlRigParameterTrack(const F
 	UE::MovieScene::ShouldUseLegacyControlRigTemplate = &UMovieSceneControlRigParameterTrack::ShouldUseLegacyTemplate;
 
 #if WITH_EDITOR
-	FEditorDelegates::ShutdownPIE.AddUObject(this, &UMovieSceneControlRigParameterTrack::OnShutdownPIE);
+	EndPIEHandle = FEditorDelegates::EndPIE.AddUObject(this, &UMovieSceneControlRigParameterTrack::OnEndPIE);
+	ShutdownPIEHandle = FEditorDelegates::ShutdownPIE.AddUObject(this, &UMovieSceneControlRigParameterTrack::OnShutdownPIE);
 #endif
 }
 
@@ -88,6 +90,11 @@ void UMovieSceneControlRigParameterTrack::BeginDestroy()
 {
 	Super::BeginDestroy();
 	UnbindControlRigDelegates();
+	
+#if WITH_EDITOR
+	FEditorDelegates::EndPIE.Remove(EndPIEHandle);
+	FEditorDelegates::ShutdownPIE.Remove(ShutdownPIEHandle);
+#endif
 }
 
 // UMovieSceneTrack interface
@@ -134,6 +141,45 @@ void UMovieSceneControlRigParameterTrack::CheckForNameChanges()
 
 #if WITH_EDITOR
 
+void UMovieSceneControlRigParameterTrack::OnEndPIE(bool bIsSimulating)
+{
+	UControlRig* DefaultRig = ControlRig.Get();
+	if (!DefaultRig || GameWorldControlRigs.IsEmpty())
+	{
+		return;
+	}
+
+	constexpr bool bEvenIfPendingKill = true;
+	for (const auto& [WeakGameWorldPtr, GameRigPtr]: GameWorldControlRigs)
+	{
+		UWorld* World = WeakGameWorldPtr.Get(bEvenIfPendingKill);
+		UControlRig* GameRig = GameRigPtr.Get();
+		if (World && GameRig)
+		{
+			auto ResetHandle = [GameRig, DefaultRig](const TObjectPtr<UTransformableHandle>& Handle)
+			{
+				UTransformableControlHandle* ControlHandle = Cast<UTransformableControlHandle>(Handle);
+				if (ControlHandle && ControlHandle->ControlRig == GameRig)
+				{
+					ControlHandle->UnregisterDelegates();
+					ControlHandle->ControlRig = DefaultRig;
+					ControlHandle->RegisterDelegates();
+				}						
+			};
+
+			const FConstraintsManagerController& Controller = FConstraintsManagerController::Get(World);
+			for (const TWeakObjectPtr<UTickableConstraint>& Constraint: Controller.GetAllConstraints())
+			{
+				if (UTickableTransformConstraint* TransformConstraint = Cast< UTickableTransformConstraint>(Constraint.Get()))
+				{
+					ResetHandle(TransformConstraint->ChildTRSHandle);
+					ResetHandle(TransformConstraint->ParentTRSHandle);
+				}
+			}
+		}
+	}
+}
+
 void UMovieSceneControlRigParameterTrack::OnShutdownPIE(bool bIsSimulating)
 {
 	TArray<TWeakObjectPtr<UWorld>> KeysToRemove;
diff --git a/Engine/Plugins/Animation/ControlRig/Source/ControlRig/Public/Sequencer/MovieSceneControlRigParameterTrack.h b/Engine/Plugins/Animation/ControlRig/Source/ControlRig/Public/Sequencer/MovieSceneControlRigParameterTrack.h
--- a/Engine/Plugins/Animation/ControlRig/Source/ControlRig/Public/Sequencer/MovieSceneControlRigParameterTrack.h
+++ b/Engine/Plugins/Animation/ControlRig/Source/ControlRig/Public/Sequencer/MovieSceneControlRigParameterTrack.h
@@ -215,7 +215,10 @@ private:
 	void CheckForNameChanges();
 
 #if WITH_EDITOR
+	void OnEndPIE(bool bIsSimulating);
 	void OnShutdownPIE(bool bIsSimulating);
+	FDelegateHandle ShutdownPIEHandle;
+	FDelegateHandle EndPIEHandle;
 #endif
 
 private:
