Control Rig: Change icon of the Constrain window

diff --git a/Engine/Plugins/Animation/ControlRig/Source/ControlRigEditor/Private/ControlRigEditorStyle.h b/Engine/Plugins/Animation/ControlRig/Source/ControlRigEditor/Private/ControlRigEditorStyle.h
--- a/Engine/Plugins/Animation/ControlRig/Source/ControlRigEditor/Private/ControlRigEditorStyle.h
+++ b/Engine/Plugins/Animation/ControlRig/Source/ControlRigEditor/Private/ControlRigEditorStyle.h
@@ -86,7 +86,7 @@ public:
 		//Tool Styles
 		{
 			Set("ControlRig.OnlySelectControls", new IMAGE_BRUSH_SVG("Slate/AnimationSelectOnlyControlRig", Icon16x16));
-			Set("ControlRig.SnapperTool", new IMAGE_BRUSH_SVG("Slate/AnimationControlRigSnapper", Icon16x16));
+			Set("ControlRig.ConstraintTools", new IMAGE_BRUSH_SVG("Slate/AnimationConstraintTools_16", Icon16x16));
 			Set("ControlRig.PoseTool", new IMAGE_BRUSH_SVG("Slate/AnimationPoses", Icon16x16));
 			Set("ControlRig.CreatePose", new IMAGE_BRUSH_SVG("Slate/AnimationCreatePose", Icon16x16));
 			Set("ControlRig.TweenTool", new IMAGE_BRUSH_SVG("Slate/AnimationTweens", Icon16x16));
diff --git a/Engine/Plugins/Animation/ControlRig/Source/ControlRigEditor/Private/EditMode/Constrain/ConstrainToolsManager.cpp b/Engine/Plugins/Animation/ControlRig/Source/ControlRigEditor/Private/EditMode/Constrain/ConstrainToolsManager.cpp
--- a/Engine/Plugins/Animation/ControlRig/Source/ControlRigEditor/Private/EditMode/Constrain/ConstrainToolsManager.cpp
+++ b/Engine/Plugins/Animation/ControlRig/Source/ControlRigEditor/Private/EditMode/Constrain/ConstrainToolsManager.cpp
@@ -79,7 +79,7 @@ void FConstrainToolsManager::RegisterTabSpawner(const TSharedRef<FTabManager>& I
 		)
 		.SetDisplayName(LOCTEXT("Tab.Title", "Constrain"))
 		.SetTooltipText(LOCTEXT("Tab.ToolTip", "Snap child objects to a parent object over a set of frames."))
-		.SetIcon(FSlateIcon(TEXT("ControlRigEditorStyle"), TEXT("ControlRig.SnapperTool")))
+		.SetIcon(FSlateIcon(TEXT("ControlRigEditorStyle"), TEXT("ControlRig.ConstraintTools")))
 		.SetGroup(InWorkspaceMenuGroup);
 	InTabManager->RegisterDefaultTabWindowSize(FControlRigEditModeToolkit::ConstrainingTabName, FVector2D(380, 500));
 }
diff --git a/Engine/Plugins/Animation/ControlRig/Source/ControlRigEditor/Private/EditMode/ControlRigEditModeToolkit.cpp b/Engine/Plugins/Animation/ControlRig/Source/ControlRigEditor/Private/EditMode/ControlRigEditModeToolkit.cpp
--- a/Engine/Plugins/Animation/ControlRig/Source/ControlRigEditor/Private/EditMode/ControlRigEditModeToolkit.cpp
+++ b/Engine/Plugins/Animation/ControlRig/Source/ControlRigEditor/Private/EditMode/ControlRigEditModeToolkit.cpp
@@ -399,7 +399,7 @@ void FControlRigEditModeToolkit::RequestModeUITabs()
 			.SetDisplayName(LOCTEXT("ControlRigSnapperTab", "Control Rig Snapper"))
 			.SetTooltipText(LOCTEXT("ControlRigSnapperTabTooltip", "Snap child objects to a parent object over a set of frames."))
 			.SetGroup(MenuGroup)
-			.SetIcon(FSlateIcon(TEXT("ControlRigEditorStyle"), TEXT("ControlRig.SnapperTool")));
+			.SetIcon(FSlateIcon(TEXT("ControlRigEditorStyle"), TEXT("ControlRig.ConstraintTools")));
 		ModeUILayerPtr->GetTabManager()->RegisterDefaultTabWindowSize(SnapperTabName, FVector2D(300, 325));
 
 		ModeUILayerPtr->GetTabManager()->UnregisterTabSpawner(PoseTabName);
@@ -686,7 +686,7 @@ void FControlRigEditModeToolkit::BuildToolBar(FToolBarBuilder& InToolBarBuilder)
 		NAME_None,
 		LOCTEXT("Constrain", "Constrain"),
 		LOCTEXT("SnapperTooltip", "Tools for constraining, such as snapping, space picker, and constraints edition."),
-		FSlateIcon(TEXT("ControlRigEditorStyle"), TEXT("ControlRig.SnapperTool")),
+		FSlateIcon(TEXT("ControlRigEditorStyle"), TEXT("ControlRig.ConstraintTools")),
 		EUserInterfaceActionType::ToggleButton
 	);
 	
