Constraints: fixed delta serialization missing properties when savig uasset

#rnx

diff --git a/Engine/Plugins/Animation/ControlRig/Source/ControlRigEditor/Private/Editor/Constraints/SConstraintsWidget.cpp b/Engine/Plugins/Animation/ControlRig/Source/ControlRigEditor/Private/Editor/Constraints/SConstraintsWidget.cpp
--- a/Engine/Plugins/Animation/ControlRig/Source/ControlRigEditor/Private/Editor/Constraints/SConstraintsWidget.cpp
+++ b/Engine/Plugins/Animation/ControlRig/Source/ControlRigEditor/Private/Editor/Constraints/SConstraintsWidget.cpp
@@ -63,25 +63,15 @@ const TMap< UClass*, ETransformConstraintType >& FConstraintInfo::GetConstraintT
 	return ConstraintToType;
 }
 
-const TArray< UTickableTransformConstraint* >& FConstraintInfo::GetMutableDefaults()
+const TArray< UTickableTransformConstraint* >& FConstraintInfo::GetConfigurableConstraints()
 {
-	static TArray< UTickableTransformConstraint* > MutableDefaults({
-		GetMutableDefault<UTickableTranslationConstraint>(),
-		GetMutableDefault<UTickableRotationConstraint>(),
-		GetMutableDefault<UTickableScaleConstraint>(),
-		GetMutableDefault<UTickableParentConstraint>(),
-		GetMutableDefault<UTickableLookAtConstraint>()});
-	return MutableDefaults;
-}
-
-const TArray< TFunction<UTickableTransformConstraint*()> >& FConstraintInfo::GetConfigurableConstraints()
-{
-	static TArray< TFunction<UTickableTransformConstraint*()> > ConfigurableConstraints({
-		[](){ return NewObject<UTickableTranslationConstraint>(GetTransientPackage(), NAME_None, RF_Transactional, GetMutableDefault<UTickableTranslationConstraint>(), true); },
-		[](){ return NewObject<UTickableRotationConstraint>(GetTransientPackage(), NAME_None, RF_Transactional, GetMutableDefault<UTickableRotationConstraint>(), true); },
-		[](){ return NewObject<UTickableScaleConstraint>(GetTransientPackage(), NAME_None, RF_Transactional, GetMutableDefault<UTickableScaleConstraint>(), true); },
-		[](){ return NewObject<UTickableParentConstraint>(GetTransientPackage(), NAME_None, RF_Transactional, GetMutableDefault<UTickableParentConstraint>(), true); },
-		[](){ return NewObject<UTickableLookAtConstraint>(GetTransientPackage(), NAME_None, RF_Transactional, GetMutableDefault<UTickableLookAtConstraint>(), true); },
+	using namespace UE::TransformConstraintUtil;
+	static TArray< UTickableTransformConstraint* > ConfigurableConstraints({
+		GetConfig(UTickableTranslationConstraint::StaticClass()),
+		GetConfig(UTickableRotationConstraint::StaticClass()),
+		GetConfig(UTickableScaleConstraint::StaticClass()),
+		GetConfig(UTickableParentConstraint::StaticClass()),
+		GetConfig(UTickableLookAtConstraint::StaticClass()),
 	});
 	return ConfigurableConstraints;
 }
@@ -106,24 +96,13 @@ int8 FConstraintInfo::GetType(UClass* InClass)
 	return -1;
 }
 
-UTickableTransformConstraint* FConstraintInfo::GetMutable(ETransformConstraintType InType)
-{
-	static const UEnum* ETransformConstraintTypeEnum = StaticEnum<ETransformConstraintType>();
-	const uint8 ConstraintType = static_cast<uint8>(InType);
-	if (ETransformConstraintTypeEnum->IsValidEnumValue(ConstraintType))
-	{
-		return GetMutableDefaults()[ConstraintType];
-	}
-	return nullptr;
-}
-
 UTickableTransformConstraint* FConstraintInfo::GetConfigurable(ETransformConstraintType InType)
 {
 	static const UEnum* ETransformConstraintTypeEnum = StaticEnum<ETransformConstraintType>();
 	const uint8 ConstraintType = static_cast<uint8>(InType);
 	if (ETransformConstraintTypeEnum->IsValidEnumValue(ConstraintType))
 	{
-		return GetConfigurableConstraints()[ConstraintType]();
+		return GetConfigurableConstraints()[ConstraintType];
 	}
 	return nullptr;
 }
@@ -448,13 +427,17 @@ TSharedRef<SWidget> SConstraintMenuEntry::GenerateConstraintDefaultWidget()
 		DetailsViewArgs.bShowOptions = false;
 		DetailsViewArgs.bShowModifiedPropertiesOption = false;
 		DetailsViewArgs.ColumnWidth = 0.45f;
-		DetailsViewArgs.NotifyHook = this;
 		DetailsViewArgs.bAllowMultipleTopLevelObjects = true;
 	}
 	
 	TSharedRef<IDetailsView> DetailsView = FModuleManager::GetModuleChecked<FPropertyEditorModule>("PropertyEditor").CreateDetailView(DetailsViewArgs);
 
 	ConfigurableConstraint = FConstraintInfo::GetConfigurable(ConstraintType);
+
+	if (!ensure(ConfigurableConstraint))
+	{
+		return SNullWidget::NullWidget;
+	}
 	
 	// manage properties visibility
 	const bool bIsLookAtConstraint = ConfigurableConstraint->GetClass() == UTickableLookAtConstraint::StaticClass();
@@ -462,7 +445,7 @@ TSharedRef<SWidget> SConstraintMenuEntry::GenerateConstraintDefaultWidget()
 		[bIsLookAtConstraint](const FPropertyAndParent& InPropertyAndParent)
 		{
 			const FName PropertyName = InPropertyAndParent.Property.GetFName();
-			
+		
 			// hide active property
 			static const FName ActivePropName = GET_MEMBER_NAME_CHECKED(UTickableTransformConstraint, Active);
 			if (PropertyName == ActivePropName)
@@ -487,11 +470,11 @@ TSharedRef<SWidget> SConstraintMenuEntry::GenerateConstraintDefaultWidget()
 				TWeakPtr<ISequencer> WeakSequencer = UE::Private::GetSequencerChecked();
 				return UE::Private::GetSequencerChecked().IsValid();
 			}
-				
+			
 			return true;
 		});
 	DetailsView->SetIsPropertyVisibleDelegate(PropertyVisibility);
-	
+
 	TArray<UObject*> Objects(
 	{
 		ConfigurableConstraint,
@@ -564,36 +547,6 @@ FReply SConstraintMenuEntry::OnDragDetected(const FGeometry& MyGeometry, const F
 	return FReply::Handled();
 }
 
-void SConstraintMenuEntry::NotifyPostChange(const FPropertyChangedEvent& InPropertyChangedEvent, FProperty* InPropertyThatChanged)
-{
-	if (!InPropertyThatChanged || InPropertyChangedEvent.ChangeType != EPropertyChangeType::ValueSet)
-	{
-		return;
-	}
-
-	UTickableTransformConstraint* MutableConstraint = FConstraintInfo::GetMutable(ConstraintType);
-	if (!IsValid(ConfigurableConstraint) || !IsValid(MutableConstraint))
-	{
-		return;
-	}
-
-	const UClass* ConstraintClass = MutableConstraint->GetClass();
-	
-	if (const FProperty* CDOProperty = ConstraintClass->FindPropertyByName(InPropertyChangedEvent.GetPropertyName()))
-	{
-		// copy the property
-		CDOProperty->CopyCompleteValue_InContainer(MutableConstraint, ConfigurableConstraint.Get());
-	}
-	else if (InPropertyChangedEvent.MemberProperty)
-	{
-		// copy the member property instead
-		if (const FProperty* CDOMemberProperty = ConstraintClass->FindPropertyByName(InPropertyChangedEvent.GetMemberPropertyName()))
-		{
-			CDOMemberProperty->CopyCompleteValue_InContainer(MutableConstraint, ConfigurableConstraint.Get());
-		}
-	}
-}
-
 FReply SConstraintMenuEntry::CreateFromSelectionOrPicker(const bool bUseDefault) const
 {
 	// We need to hide the creation menu. This destroys the menu that is keeping this SConstraintMenuEntry alive.
diff --git a/Engine/Plugins/Animation/ControlRig/Source/ControlRigEditor/Private/Editor/Constraints/SConstraintsWidget.h b/Engine/Plugins/Animation/ControlRig/Source/ControlRigEditor/Private/Editor/Constraints/SConstraintsWidget.h
--- a/Engine/Plugins/Animation/ControlRig/Source/ControlRigEditor/Private/Editor/Constraints/SConstraintsWidget.h
+++ b/Engine/Plugins/Animation/ControlRig/Source/ControlRigEditor/Private/Editor/Constraints/SConstraintsWidget.h
@@ -10,7 +10,6 @@
 #include "BakingAnimationKeySettings.h"
 #include "Misc/QualifiedFrameTime.h"
 #include "ConstraintsManager.h"
-#include "Misc/NotifyHook.h"
 
 #define UE_API CONTROLRIGEDITOR_API
 
@@ -43,14 +42,12 @@ class FConstraintInfo
 public:
 	static const FSlateBrush* GetBrush(uint8 InType);
 	static int8 GetType(UClass* InClass);
-	static UTickableTransformConstraint* GetMutable(ETransformConstraintType InType);
 	static UTickableTransformConstraint* GetConfigurable(ETransformConstraintType InType);
 	
 private:
 	static const TArray< const FSlateBrush* >& GetBrushes();
 	static const TMap< UClass*, ETransformConstraintType >& GetConstraintToType();
-	static const TArray< UTickableTransformConstraint* >& GetMutableDefaults();
-	static const TArray< TFunction<UTickableTransformConstraint*()> >& GetConfigurableConstraints();
+	static const TArray< UTickableTransformConstraint* >& GetConfigurableConstraints();
 };
 
 /**
@@ -86,7 +83,7 @@ private:
  * SConstraintMenuEntry
  */
 
-class SConstraintMenuEntry : public SCompoundWidget, public FNotifyHook
+class SConstraintMenuEntry : public SCompoundWidget
 {
 public:
 	SLATE_BEGIN_ARGS(SConstraintMenuEntry){}
@@ -108,10 +105,6 @@ public:
 	virtual FReply OnDragDetected(const FGeometry& MyGeometry, const FPointerEvent& MouseEvent) override;
 	// End of SWidget interface
 
-	// FNotifyHook interface
-	/** NotifyPostChange is used to copy the property value from the configurable constraint to the cdo. */
-	virtual void NotifyPostChange( const FPropertyChangedEvent& InPropertyChangedEvent, FProperty* InPropertyThatChanged ) override;
-	
 private:
 
 	/** Uses the current selected parent or pick a new one if none has been set. */
diff --git a/Engine/Source/Runtime/Experimental/Animation/Constraints/Private/ConstraintSubsystem.cpp b/Engine/Source/Runtime/Experimental/Animation/Constraints/Private/ConstraintSubsystem.cpp
--- a/Engine/Source/Runtime/Experimental/Animation/Constraints/Private/ConstraintSubsystem.cpp
+++ b/Engine/Source/Runtime/Experimental/Animation/Constraints/Private/ConstraintSubsystem.cpp
@@ -6,6 +6,7 @@
 #include "ConstraintsManager.h"
 #include "Algo/ForEach.h"
 #include "Misc/CoreDelegates.h"
+#include "Transform/TransformConstraint.h"
 
 //needs to be static to avoid system getting deleted with dangling handles.
 
@@ -24,18 +25,24 @@ void UConstraintSubsystem::Initialize(FSubsystemCollectionBase& Collection)
 
 	if (GEngine && GEngine->IsInitialized())
 	{
-		RegisterWorldDelegates();
+		InitializeInternal();
 	}
 	else
 	{
-		FCoreDelegates::OnPostEngineInit.AddUObject(this, &UConstraintSubsystem::RegisterWorldDelegates);
+		FCoreDelegates::OnPostEngineInit.AddUObject(this, &UConstraintSubsystem::InitializeInternal);
 	}
 	
 	SetFlags(RF_Transactional);
 }
 
-void UConstraintSubsystem::RegisterWorldDelegates()
+void UConstraintSubsystem::InitializeInternal()
 {
+	RegisterConfig<UTickableTranslationConstraint>();
+	RegisterConfig<UTickableRotationConstraint>();
+	RegisterConfig<UTickableScaleConstraint>();
+	RegisterConfig<UTickableParentConstraint>();
+	RegisterConfig<UTickableLookAtConstraint>();
+	
 	OnWorldInitHandle = FWorldDelegates::OnPreWorldInitialization.AddStatic(&UConstraintSubsystem::OnWorldInit);
 	OnWorldCleanupHandle = FWorldDelegates::OnWorldCleanup.AddStatic(&UConstraintSubsystem::OnWorldCleanup);
 	OnPostGarbageCollectHandle = FCoreUObjectDelegates::GetPostGarbageCollect().AddStatic(&UConstraintSubsystem::OnPostGarbageCollect);
@@ -43,6 +50,13 @@ void UConstraintSubsystem::RegisterWorldDelegates()
 	FCoreDelegates::OnPostEngineInit.RemoveAll(this);
 }
 
+UTickableConstraint* UConstraintSubsystem::GetConfig(const UClass* InConstraintClass) 
+{
+	const TObjectPtr<UTickableConstraint>* ObjectKey = ConstraintsConfig.Find(InConstraintClass);
+	ensure(ObjectKey && ObjectKey->Get());
+	return ObjectKey ? ObjectKey->Get() : nullptr;
+}
+
 void UConstraintSubsystem::Deinitialize()
 {
 	for (int32 Index = ConstraintsInWorld.Num() - 1; Index >= 0; --Index)
diff --git a/Engine/Source/Runtime/Experimental/Animation/Constraints/Private/ConstraintSubsystem.h b/Engine/Source/Runtime/Experimental/Animation/Constraints/Private/ConstraintSubsystem.h
--- a/Engine/Source/Runtime/Experimental/Animation/Constraints/Private/ConstraintSubsystem.h
+++ b/Engine/Source/Runtime/Experimental/Animation/Constraints/Private/ConstraintSubsystem.h
@@ -80,6 +80,31 @@ public:
 	UPROPERTY(BlueprintAssignable, Category = Constraints, meta = (DisplayName = "OnConstraintRemoved"))
 	FOnConstraintRemovedFromSystem OnConstraintRemovedFromSystem_BP;
 
+	/** Returns a UTickableConstraint instance that can be used as a configuration constraints. */
+	UTickableConstraint* GetConfig(const UClass* InConstraintClass);
+
+	/** Register a new constraint class type and allocate a new configuration constraints for that class. */
+	template<typename ConstraintType>
+	void RegisterConfig()
+	{
+		static_assert(TIsDerivedFrom<ConstraintType, UTickableConstraint>::Value,
+		"The template class ConstraintType must be a subclass of UTickableConstraint.");
+
+		if (UConstraintSubsystem* SubSystem = Get())
+		{
+			UClass* ConstraintClass = ConstraintType::StaticClass();
+			if (ensureAlways(!ConstraintClass->HasAnyClassFlags(CLASS_Abstract)))
+			{
+				if (!ConstraintsConfig.Contains(ConstraintClass))
+				{
+					ConstraintsConfig.Emplace(
+					ConstraintClass,
+					NewObject<ConstraintType>(SubSystem, NAME_None, RF_Transactional));
+				}
+			}
+		}
+	}
+
 private:
 	UPROPERTY(transient)
 	TArray<FConstraintsInWorld> ConstraintsInWorld;
@@ -95,7 +120,7 @@ private:
 	static FDelegateHandle OnPostGarbageCollectHandle;
 	static void OnPostGarbageCollect();
 	
-	void RegisterWorldDelegates();
+	void InitializeInternal();
 
 	mutable bool bNeedsCleanup = false;
 
@@ -106,4 +131,8 @@ private:
 	int32 GetConstraintsInWorldIndex(const UWorld* InWorld) const;
 
 	void CleanupInvalidConstraints() const;
+
+	/** Per UClass constraint configuration instance. */
+	UPROPERTY(transient)
+	TMap<TObjectPtr<UClass>, TObjectPtr<UTickableConstraint>> ConstraintsConfig;
 };
diff --git a/Engine/Source/Runtime/Experimental/Animation/Constraints/Private/Transform/TransformConstraintUtil.cpp b/Engine/Source/Runtime/Experimental/Animation/Constraints/Private/Transform/TransformConstraintUtil.cpp
--- a/Engine/Source/Runtime/Experimental/Animation/Constraints/Private/Transform/TransformConstraintUtil.cpp
+++ b/Engine/Source/Runtime/Experimental/Animation/Constraints/Private/Transform/TransformConstraintUtil.cpp
@@ -796,5 +796,11 @@ void GetChildrenConstraints( UWorld* World, const UTickableTransformConstraint*
 	const TArray<ConstraintPtr> FilteredConstraints = Controller.GetConstraintsByPredicate(Predicate);
 	OutConstraints.Append(FilteredConstraints);
 }
+
+UTickableTransformConstraint* GetConfig(const UClass* InConstraintClass)
+{
+	UConstraintSubsystem* Subsystem = UConstraintSubsystem::Get();
+	return Subsystem && InConstraintClass ? Cast<UTickableTransformConstraint>(Subsystem->GetConfig(InConstraintClass)) : nullptr;
+}
 	
 }
diff --git a/Engine/Source/Runtime/Experimental/Animation/Constraints/Public/ConstraintsManager.inl b/Engine/Source/Runtime/Experimental/Animation/Constraints/Public/ConstraintsManager.inl
--- a/Engine/Source/Runtime/Experimental/Animation/Constraints/Public/ConstraintsManager.inl
+++ b/Engine/Source/Runtime/Experimental/Animation/Constraints/Public/ConstraintsManager.inl
@@ -28,7 +28,8 @@ TConstraint* FConstraintsManagerController::AllocateConstraintT(const FName& InB
 		Name = MakeUniqueObjectName(Subsystem, TConstraint::StaticClass(), InBaseName, EUniqueObjectNameOptions::GloballyUnique);
 	}
 
-	TConstraint* NewConstraint = NewObject<TConstraint>(Subsystem, Name, RF_Transactional, bUseDefault ? GetMutableDefault<TConstraint>() : nullptr, bUseDefault);
+	UObject* Template = bUseDefault ? Subsystem->GetConfig(TConstraint::StaticClass()) : nullptr;
+	TConstraint* NewConstraint = NewObject<TConstraint>(Subsystem, Name, RF_Transactional, Template, Template != nullptr);
 	NewConstraint->Modify();
 	return NewConstraint;
 }
@@ -66,7 +67,7 @@ TArray< TWeakObjectPtr<UTickableConstraint> > FConstraintsManagerController::Get
 	
 	Algo::StableSort(Constraints, [TicksBefore](const ConstraintPtr& LHS, const ConstraintPtr& RHS)
 	{
-		return TicksBefore(*LHS, *RHS);
+		return (LHS.IsValid() && RHS.IsValid()) ? TicksBefore(*LHS, *RHS) : false;
 	});
 	
 	return Constraints;
diff --git a/Engine/Source/Runtime/Experimental/Animation/Constraints/Public/Transform/TransformConstraintUtil.h b/Engine/Source/Runtime/Experimental/Animation/Constraints/Public/Transform/TransformConstraintUtil.h
--- a/Engine/Source/Runtime/Experimental/Animation/Constraints/Public/Transform/TransformConstraintUtil.h
+++ b/Engine/Source/Runtime/Experimental/Animation/Constraints/Public/Transform/TransformConstraintUtil.h
@@ -124,6 +124,8 @@ namespace UE::TransformConstraintUtil
 	/** Adjust the transform on a scene component so it's effected by the constraint*/
 	CONSTRAINTS_API void UpdateTransformBasedOnConstraint(FTransform& InOutCurrentTransform, const USceneComponent* InSceneComponent);
 
+	/** Returns the configuration constraint instance for a specific constraint class if it exists */
+	CONSTRAINTS_API UTickableTransformConstraint* GetConfig(const UClass* InConstraintClass);
 }
 
 #undef UE_API
