Fixed RigVM legacy parameters getting public status lost when loading an old asset

#rnx

diff --git a/Engine/Plugins/Animation/ControlRig/Source/ControlRigDeveloper/Private/ControlRigAsset.cpp b/Engine/Plugins/Animation/ControlRig/Source/ControlRigDeveloper/Private/ControlRigAsset.cpp
--- a/Engine/Plugins/Animation/ControlRig/Source/ControlRigDeveloper/Private/ControlRigAsset.cpp
+++ b/Engine/Plugins/Animation/ControlRig/Source/ControlRigDeveloper/Private/ControlRigAsset.cpp
@@ -2232,6 +2232,7 @@ void IControlRigAssetInterface::CreateMemberVariablesOnLoad()
 			const FName VarName = FindHostMemberVariableUniqueName(NameValidator, Description.Name.ToString());
 			FRigVMExternalVariable ExternalVariable = Description.ToExternalVariable();
 			ExternalVariable.Name = VarName;
+			ExternalVariable.bIsPublic = true; // parameters get this set by default
 			const FName ResultName = GetRigVMAssetInterface()->AddHostMemberVariableFromExternal(ExternalVariable, FString());
 			if (!ResultName.IsNone())
 			{
diff --git a/Engine/Plugins/Runtime/RigVM/Source/RigVMDeveloper/Private/RigVMBlueprintLegacy.cpp b/Engine/Plugins/Runtime/RigVM/Source/RigVMDeveloper/Private/RigVMBlueprintLegacy.cpp
--- a/Engine/Plugins/Runtime/RigVM/Source/RigVMDeveloper/Private/RigVMBlueprintLegacy.cpp
+++ b/Engine/Plugins/Runtime/RigVM/Source/RigVMDeveloper/Private/RigVMBlueprintLegacy.cpp
@@ -684,7 +684,7 @@ TArray<FRigVMGraphVariableDescription> URigVMBlueprint::GetAssetVariables() cons
 		{
 			NewVariable.bExposeToCinematics = true;
 		}
-		if (BPVariable.PropertyFlags & CPF_DisableEditOnInstance)
+		if (!(BPVariable.PropertyFlags & CPF_DisableEditOnInstance))
 		{
 			NewVariable.bPublic = true;
 		}
