Control Rig: Fix crash when undoing change preview mesh
#rnx

diff --git a/Engine/Plugins/Animation/ControlRig/Source/ControlRigEditor/Private/Editor/ControlRigEditor.cpp b/Engine/Plugins/Animation/ControlRig/Source/ControlRigEditor/Private/Editor/ControlRigEditor.cpp
--- a/Engine/Plugins/Animation/ControlRig/Source/ControlRigEditor/Private/Editor/ControlRigEditor.cpp
+++ b/Engine/Plugins/Animation/ControlRig/Source/ControlRigEditor/Private/Editor/ControlRigEditor.cpp
@@ -2067,7 +2067,9 @@ void FControlRigBaseEditor::PostTransactionImpl(bool bSuccess, const FTransactio
 	//Compile();
 
 	UpdateRigVMHost();
-	
+
+	// SetPreviewMesh sometimes destroys the editor. Lets make sure we still have a valid editor.
+	TSharedPtr<FControlRigBaseEditor> SharedPointer = SharedRef().ToSharedPtr();
 	USkeletalMesh* PreviewMesh = GetPersonaToolkit()->GetPreviewScene()->GetPreviewMesh();
 	if (PreviewMesh != RigBlueprint->GetPreviewMesh())
 	{
@@ -2075,20 +2077,23 @@ void FControlRigBaseEditor::PostTransactionImpl(bool bSuccess, const FTransactio
 		GetPersonaToolkit()->SetPreviewMesh(PreviewMesh, true);
 	}
 
-	if (UControlRig* DebuggedControlRig = Cast<UControlRig>(RigBlueprint->GetObjectBeingDebugged()))
+	if (SharedPointer.IsValid())
 	{
-		if(URigHierarchy* Hierarchy = DebuggedControlRig->GetHierarchy())
+		if (UControlRig* DebuggedControlRig = Cast<UControlRig>(RigBlueprint->GetObjectBeingDebugged()))
 		{
-			if(Hierarchy->Num() == 0)
+			if(URigHierarchy* Hierarchy = DebuggedControlRig->GetHierarchy())
 			{
-				OnHierarchyChanged();
+				if(Hierarchy->Num() == 0)
+				{
+					OnHierarchyChanged();
+				}
 			}
 		}
-	}
 
-	if (FControlRigEditMode* EditMode = GetEditMode())
-	{
-		EditMode->RequestToRecreateControlShapeActors();
+		if (FControlRigEditMode* EditMode = GetEditMode())
+		{
+			EditMode->RequestToRecreateControlShapeActors();
+		}
 	}
 }
 
