Constraints: evaluate non-interacting rigs when they are related to an object/control being manipulated
- non-interactiing rigs are cached when a new gizmo manipulation is started so that they can also be updated on tick
- constraints related to these non-interacting rigs are also marked for evaluation so that manipulation is smooth and reliable

#rnx

diff --git a/Engine/Plugins/Animation/ControlRig/Source/ControlRigEditor/Private/EditMode/ControlRigEditMode.cpp b/Engine/Plugins/Animation/ControlRig/Source/ControlRigEditor/Private/EditMode/ControlRigEditMode.cpp
--- a/Engine/Plugins/Animation/ControlRig/Source/ControlRigEditor/Private/EditMode/ControlRigEditMode.cpp
+++ b/Engine/Plugins/Animation/ControlRig/Source/ControlRigEditor/Private/EditMode/ControlRigEditMode.cpp
@@ -509,6 +509,10 @@ void FControlRigEditMode::Enter()
 	{
 		UE::ControlRig::PopulateControlRigViewportToolbarTransformSubmenu("LevelEditor.ViewportToolbar.Transform");
 		UE::ControlRig::PopulateControlRigViewportToolbarShowSubmenu("LevelEditor.ViewportToolbar.Show");
+		if (GEngine)
+		{
+			GEngine->OnActorMoving().AddRaw(this, &FControlRigEditMode::HandleActorMoving);
+		}
 	}
 	
 	RegisterPendingFocusMode();
@@ -557,6 +561,11 @@ void FControlRigEditMode::Exit()
 		UE::ControlRig::RemoveControlRigViewportToolbarExtensions();
 	}
 
+	if (GEngine)
+	{
+		GEngine->OnActorMoving().RemoveAll(this);
+	}
+
 	UnregisterPendingFocusMode();
 	
 	ClearOutAnyActiveTools();
@@ -765,6 +774,28 @@ void FControlRigEditMode::Tick(FEditorViewportClient* ViewportClient, float Delt
 	DetailKeyFrameCache->UpdateIfDirty();
 	
 	UpdateRotationContext();
+
+	if (bUpdateNonInteractingRigs && InteractionScopes.IsEmpty())
+	{
+		// mark non-interacting constraints for evaluation
+		for (const TWeakObjectPtr<UControlRig>& NonInteractingRig: NonInteractingRigs)
+		{
+			UControlRig* Rig = NonInteractingRig.Get();
+			if (USceneComponent* BoundComponent = Rig ? GetHostingSceneComponent(Rig) : nullptr)
+			{
+				for (const TWeakObjectPtr<UTickableConstraint>& WeakConstraint: ConstraintsCache.Get(BoundComponent, WorldPtr))
+				{
+					if (UTickableConstraint* Constraint = WeakConstraint.Get(); Constraint->Active)
+					{
+						FConstraintsManagerController::Get(WorldPtr).MarkConstraintForEvaluation(Constraint);
+					}
+				}
+			}
+		}
+		FConstraintsManagerController::Get(WorldPtr).FlushEvaluationGraph();
+		TickManipulatableObjects(NonInteractingRigs);
+		bUpdateNonInteractingRigs = false;
+	}
 }
 
 //Hit proxy for FK Rigs and bones.
@@ -1285,7 +1316,10 @@ bool FControlRigEditMode::HandleBeginTransform(const FEditorViewportClient* InVi
 	
 	// do not stop if there are other selected objects that are not related to the control rig mode
 	// note that we are only interested in this value if controls are also being manipulated
-	bOtherSelectedObjects = bInteractingWithControls ? HasAnyOtherSelectedObjects() : false;
+	CacheAnyOtherSelectedObjects();
+	CacheNonInteractingRigsToUpdate();
+	
+	bOtherSelectedObjects = bInteractingWithControls ? !SelectedObjects.IsEmpty() : false;
 	
 	return bOtherSelectedObjects ? false : bInteractingWithControls;
 }
@@ -1305,6 +1339,9 @@ bool FControlRigEditMode::HandleEndTransform(FEditorViewportClient* InViewportCl
 	bIsTracking = false;
 	InteractionDependencies.Reset();
 	bOtherSelectedObjects = false;
+	SelectedObjects.Reset();
+	bUpdateNonInteractingRigs = false;
+	NonInteractingRigs.Reset();
 	
 	if (InteractionScopes.Num() > 0)
 	{		
@@ -2335,6 +2372,24 @@ bool FControlRigEditMode::InputDelta(FEditorViewportClient* InViewportClient, FV
 			return InteractionData.Key;
 		});
 		
+		RigsToTick.Append(NonInteractingRigs);
+
+		// mark non-interacting constraints for evaluation
+		for (const TWeakObjectPtr<UControlRig>& NonInteractingRig: NonInteractingRigs)
+		{
+			UControlRig* Rig = NonInteractingRig.Get();
+			if (USceneComponent* BoundComponent = Rig ? GetHostingSceneComponent(Rig) : nullptr)
+			{
+				for (const TWeakObjectPtr<UTickableConstraint>& WeakConstraint: ConstraintsCache.Get(BoundComponent, WorldPtr))
+				{
+					if (UTickableConstraint* Constraint = WeakConstraint.Get(); Constraint->Active)
+					{
+						FConstraintsManagerController::Get(WorldPtr).MarkConstraintForEvaluation(Constraint);
+					}
+				}
+			}
+		}
+		
 		TickManipulatableObjects(RigsToTick);
 	}
 
@@ -6183,42 +6238,44 @@ void FControlRigEditMode::DestroyShapesActors(UControlRig* InControlRig)
 
 USceneComponent* FControlRigEditMode::GetHostingSceneComponent(const UControlRig* ControlRig) const
 {
-	if (ControlRig == nullptr && GetControlRigs().Num() > 0)
+	const TArrayView<const TWeakObjectPtr<UControlRig>> Rigs = GetControlRigs();
+	if (!ControlRig && !Rigs.IsEmpty())
 	{
-		ControlRig = GetControlRigs()[0].Get();
+		ControlRig = Rigs[0].Get();
 	}
-	if (ControlRig)
+
+	const TSharedPtr<IControlRigObjectBinding>& ObjectBinding = ControlRig ? ControlRig->GetObjectBinding() : nullptr;
+	UObject* BoundObject = ObjectBinding ? ObjectBinding->GetBoundObject() : nullptr;
+	if (!BoundObject)
+	{
+		return nullptr;
+	}
+
+	if (USceneComponent* BoundSceneComponent = Cast<USceneComponent>(BoundObject))
+	{
+		return BoundSceneComponent;
+	}
+
+	if (const USkeleton* BoundSkeleton = Cast<USkeleton>(BoundObject))
 	{
-		TSharedPtr<IControlRigObjectBinding> ObjectBinding = ControlRig->GetObjectBinding();
-		if (ObjectBinding.IsValid())
+		// Bound to a Skeleton means we are previewing an Animation Sequence
+		if (ULevel* Level = WorldPtr ? WorldPtr->PersistentLevel : nullptr)
 		{
-			if (USceneComponent* BoundSceneComponent = Cast<USceneComponent>(ObjectBinding->GetBoundObject()))
+			const TObjectPtr<AActor>* PreviewActor = Level->Actors.FindByPredicate([](const TObjectPtr<AActor>& Actor)
 			{
-				return BoundSceneComponent;
-			}
-			else if (USkeleton* BoundSkeleton = Cast<USkeleton>(ObjectBinding->GetBoundObject()))
+				return Actor && Actor->GetClass() == AAnimationEditorPreviewActor::StaticClass();
+			});
+
+			if (PreviewActor)
 			{
-				// Bound to a Skeleton means we are previewing an Animation Sequence
-				if (WorldPtr)
+				if (UDebugSkelMeshComponent* DebugComponent = (*PreviewActor)->FindComponentByClass<UDebugSkelMeshComponent>())
 				{
-					TObjectPtr<AActor>* PreviewActor = WorldPtr->PersistentLevel->Actors.FindByPredicate([](TObjectPtr<AActor> Actor)
-					{
-						return Actor && Actor->GetClass() == AAnimationEditorPreviewActor::StaticClass();
-					});
-
-					if(PreviewActor)
-					{
-						if (UDebugSkelMeshComponent* DebugComponent = (*PreviewActor)->FindComponentByClass<UDebugSkelMeshComponent>())
-						{
-							return DebugComponent;
-						}
-					}
+					return DebugComponent;
 				}
-			}			
+			}
 		}
-		
-	}	
-
+	}
+	
 	return nullptr;
 }
 
@@ -7050,40 +7107,165 @@ TOptional<FTransform> FControlRigEditMode::GetConstraintParentTransform(const UC
 	return Dummy;
 }
 
-bool FControlRigEditMode::HasAnyOtherSelectedObjects() const
+void FControlRigEditMode::HandleActorMoving(AActor* InMovedActor)
 {
-	if (IsInLevelEditor())
+	if (!bUpdateNonInteractingRigs && InMovedActor)
 	{
-		// check actor selection (minus control shapes)
-		TArray< TWeakObjectPtr<UObject> > Selection;
-		GetModeManager()->GetSelectedActors()->GetSelectedObjects(Selection);
-		Selection.RemoveAll([](const TWeakObjectPtr<UObject>& WeakActor)
+		bUpdateNonInteractingRigs = SelectedObjects.Contains(InMovedActor->GetRootComponent());
+
+		if (!bUpdateNonInteractingRigs)
 		{
-			return !WeakActor.IsValid() || WeakActor->IsA<AControlRigShapeActor>();
-		});
+			constexpr bool bFromChildActors = true;
+			InMovedActor->ForEachComponent<USceneComponent>(
+				bFromChildActors,
+				[&Objects = SelectedObjects, &Found = bUpdateNonInteractingRigs](USceneComponent* Component) 
+				{
+					if (!Found)
+					{
+						Found = Objects.Contains(Component);
+					}
+				});
+		}
+	}
+}
 
-		if (!Selection.IsEmpty())
+void FControlRigEditMode::CacheAnyOtherSelectedObjects()
+{
+	if (!IsInLevelEditor())
+	{
+		return;
+	}
+
+	SelectedObjects.Reset();
+	
+	// check actor selection (minus control shapes)
+	TArray< TWeakObjectPtr<UObject> > Selection;
+	GetModeManager()->GetSelectedActors()->GetSelectedObjects(Selection);
+	Selection.RemoveAll([](const TWeakObjectPtr<UObject>& WeakActor)
+	{
+		return !WeakActor.IsValid() || WeakActor->IsA<AControlRigShapeActor>();
+	});
+
+	if (!Selection.IsEmpty())
+	{
+		SelectedObjects.Reserve(Selection.Num());
+		for (TWeakObjectPtr<UObject> WeakActor: Selection)
 		{
-			return true;
+			if (AActor* Actor = Cast<AActor>(WeakActor))
+			{
+				if (USceneComponent* SceneComponent = Actor->GetRootComponent())
+				{
+					SelectedObjects.Add(SceneComponent);
+				}
+			}
 		}
-		
-		// check component selection (minus bound components)
-		GetModeManager()->GetSelectedComponents()->GetSelectedObjects(Selection);
-		for (const TWeakObjectPtr<UControlRig>& ControlRigPtr : RuntimeControlRigs)
+		return;
+	}
+	
+	// check component selection (minus bound components)
+	GetModeManager()->GetSelectedComponents()->GetSelectedObjects(Selection);
+	for (const TWeakObjectPtr<UControlRig>& ControlRigPtr : RuntimeControlRigs)
+	{
+		if (const TSharedPtr<IControlRigObjectBinding>& Binding = ControlRigPtr.IsValid() ? ControlRigPtr->GetObjectBinding() : nullptr)
 		{
-			if (const TSharedPtr<IControlRigObjectBinding>& Binding = ControlRigPtr.IsValid() ? ControlRigPtr->GetObjectBinding() : nullptr)
+			if (UObject* BoundObject = Binding->GetBoundObject())
 			{
-				if (UObject* BoundObject = Binding->GetBoundObject())
+				Selection.Remove(BoundObject);
+			}
+		}
+	}
+	SelectedObjects.Append(Selection);
+}
+
+void FControlRigEditMode::CacheNonInteractingRigsToUpdate()
+{
+	if (!IsInLevelEditor())
+	{
+		return;
+	}
+
+	bUpdateNonInteractingRigs = false;
+	NonInteractingRigs.Reset();
+
+	// cache any other rig which is not currently interacting
+	if (!InteractionScopes.IsEmpty())
+	{
+		TSet<TWeakObjectPtr<UControlRig>> InteractingRigs;
+		TSet<USceneComponent*> InteractingComponents;
+
+		for (const auto& [InteractingRig, InteractionScope] : InteractionScopes)
+		{
+			if (InteractingRig)
+			{
+				InteractingRigs.Add(InteractingRig);
+				if (USceneComponent* BoundComponent = GetHostingSceneComponent(InteractingRig))
 				{
-					Selection.Remove(BoundObject);
+					InteractingComponents.Add(BoundComponent);
 				}
 			}
 		}
+
+		const TArrayView<const TWeakObjectPtr<UControlRig>> Rigs = GetControlRigs();
+		if (InteractingRigs.Num() == Rigs.Num())
+		{
+			return;
+		}
 		
-		return !Selection.IsEmpty();		
+		for (const TWeakObjectPtr<UControlRig>& Rig: Rigs)
+		{
+			if (UControlRig* NonInteractingRig = Rig.Get(); !InteractingRigs.Contains(NonInteractingRig))
+			{
+				if (USceneComponent* BoundComponent = GetHostingSceneComponent(NonInteractingRig))
+				{
+					for (USceneComponent* ComponentToTick: InteractingComponents)
+					{
+						if (ConstraintsCache.HasAnyDependency(BoundComponent, ComponentToTick, WorldPtr))
+						{
+							NonInteractingRigs.Add(NonInteractingRig);
+							break;
+						}
+					}
+				}
+			}
+		}
+		return;
+	}
+
+	// cache all the rigs that are related to the external objects that might be interacted with
+	if (!SelectedObjects.IsEmpty())
+	{
+		for (const TWeakObjectPtr<UControlRig>& Rig: GetControlRigs())
+		{
+			if (UControlRig* NonInteractingRig = Rig.Get())
+			{
+				if (USceneComponent* BoundComponent = GetHostingSceneComponent(NonInteractingRig))
+				{
+					for (const TWeakObjectPtr<UObject>& SelectedObject: SelectedObjects)
+					{
+						if (UObject* Object = SelectedObject.Get())
+						{
+							if (AActor* Actor = Cast<AActor>(Object))
+							{
+								Object = Actor->GetRootComponent();
+							}
+							
+							if (BoundComponent == Object)
+							{
+								NonInteractingRigs.Add(NonInteractingRig);
+								break;
+							}
+
+							if (ConstraintsCache.HasAnyDependency(BoundComponent, Object, WorldPtr))
+							{
+								NonInteractingRigs.Add(NonInteractingRig);
+								break;
+							}
+						}
+					}
+				}
+			}
+		}
 	}
-	
-	return false;
 }
 
 #undef LOCTEXT_NAMESPACE
diff --git a/Engine/Plugins/Animation/ControlRig/Source/ControlRigEditor/Public/EditMode/ControlRigEditMode.h b/Engine/Plugins/Animation/ControlRig/Source/ControlRigEditor/Public/EditMode/ControlRigEditMode.h
--- a/Engine/Plugins/Animation/ControlRig/Source/ControlRigEditor/Public/EditMode/ControlRigEditMode.h
+++ b/Engine/Plugins/Animation/ControlRig/Source/ControlRigEditor/Public/EditMode/ControlRigEditMode.h
@@ -785,8 +785,13 @@ private:
 
 	// whether there are other objects selected in the world that are not shape actors or components bound to the rigs
 	bool bOtherSelectedObjects = false;
-	
-	bool HasAnyOtherSelectedObjects() const;
+	TSet<TWeakObjectPtr<UObject>> SelectedObjects;
+	void CacheAnyOtherSelectedObjects();
+
+	void HandleActorMoving(AActor* InMovedActor);
+	bool bUpdateNonInteractingRigs = false;
+	void CacheNonInteractingRigsToUpdate();
+	TArray<TWeakObjectPtr<UControlRig>> NonInteractingRigs;
 
 	friend class FControlRigEditorModule;
 	friend class FControlRigBaseEditor;
diff --git a/Engine/Source/Runtime/Experimental/Animation/Constraints/Private/Transform/TransformConstraintUtil.cpp b/Engine/Source/Runtime/Experimental/Animation/Constraints/Private/Transform/TransformConstraintUtil.cpp
--- a/Engine/Source/Runtime/Experimental/Animation/Constraints/Private/Transform/TransformConstraintUtil.cpp
+++ b/Engine/Source/Runtime/Experimental/Animation/Constraints/Private/Transform/TransformConstraintUtil.cpp
@@ -44,6 +44,12 @@ UTransformableHandle* GetHandle(UObject* InObject, const FName& InSocketName)
 
 	return nullptr;
 }
+
+UObject* GetTarget(const TObjectPtr<UTransformableHandle>& InHandle)
+{
+	return IsValid(InHandle) ? InHandle->GetTarget().Get() : nullptr;
+}
+
 	
 }
 
@@ -157,7 +163,7 @@ const TArray< TWeakObjectPtr<UTickableConstraint> >& FConstraintsInteractionCach
 	{
 		if (const UTickableTransformConstraint* TransformConstraint = Cast<UTickableTransformConstraint>(WeakConstraint.Get()))
 		{
-			const UObject* Target = IsValid(TransformConstraint->ChildTRSHandle) ? TransformConstraint->ChildTRSHandle->GetTarget().Get() : nullptr;
+			const UObject* Target = Private::GetTarget(TransformConstraint->ChildTRSHandle);
 			return Target && Target == InTarget;
 		}
 		return false;
@@ -170,7 +176,69 @@ const TArray< TWeakObjectPtr<UTickableConstraint> >& FConstraintsInteractionCach
 bool FConstraintsInteractionCache::HasAnyActiveConstraint(UObject* InTarget, UWorld* InWorld)
 {
 	const TArray< TWeakObjectPtr<UTickableConstraint> >& TransformConstraints = Get(InTarget, InWorld);
-	return GetLastActiveConstraintIndex(TransformConstraints) != INDEX_NONE;
+	return TransformConstraints.ContainsByPredicate([](const TWeakObjectPtr<UTickableConstraint>& InConstraint)
+	{
+	   if (const UTickableTransformConstraint* TransformConstraint = Cast<UTickableTransformConstraint>(InConstraint.Get()))
+	   {
+		   return TransformConstraint->Active;
+	   }
+	   return false;
+	});
+}
+
+bool FConstraintsInteractionCache::HasAnyDependency(UObject* InChild, UObject* InParent, UWorld* InWorld)
+{
+	TSet< UObject* > Visited({InChild});
+	return HasAnyDependencyInternal(InChild, InParent, InWorld, Visited);
+}
+
+bool FConstraintsInteractionCache::HasAnyDependencyInternal(UObject* InChild, UObject* InParent, UWorld* InWorld, TSet< UObject* >& Visited)
+{
+	if (!InChild || !InParent || !InWorld || InChild == InParent || Visited.Contains(InParent))
+	{
+		return false;
+	}
+	
+	if (!HasAnyActiveConstraint(InChild, InWorld))
+	{
+		return false;
+	}
+
+	const TArray< TWeakObjectPtr<UTickableConstraint> >& Constraints = Get(InChild, InWorld);
+	TSet< UObject* > Parents;
+	Parents.Reserve(Constraints.Num());
+	const bool bIsInParentAnActualParent = Constraints.ContainsByPredicate([InParent, &Parents](const TWeakObjectPtr<UTickableConstraint>& InConstraint)
+	{
+	   if (const UTickableTransformConstraint* TransformConstraint = Cast<UTickableTransformConstraint>(InConstraint.Get()))
+	   {
+		   if (UObject* Parent = Private::GetTarget(TransformConstraint->ParentTRSHandle))
+		   {
+			   Parents.Add(Parent);
+			   if (Parent == InParent && InConstraint->Active)
+			   {
+				   return true;
+			   }
+		   }
+	   }
+	   return false;
+	});
+
+	if (bIsInParentAnActualParent)
+	{
+		return true;
+	}
+	
+	Visited.Add(InParent);
+	
+	for (UObject* Parent: Parents)
+	{
+		if (HasAnyDependencyInternal(InChild, Parent, InWorld, Visited))
+		{
+			return true;
+		}
+	}
+		
+	return false;
 }
 	
 void FConstraintsInteractionCache::Reset()
diff --git a/Engine/Source/Runtime/Experimental/Animation/Constraints/Public/Transform/TransformConstraintUtil.h b/Engine/Source/Runtime/Experimental/Animation/Constraints/Public/Transform/TransformConstraintUtil.h
--- a/Engine/Source/Runtime/Experimental/Animation/Constraints/Public/Transform/TransformConstraintUtil.h
+++ b/Engine/Source/Runtime/Experimental/Animation/Constraints/Public/Transform/TransformConstraintUtil.h
@@ -42,11 +42,15 @@ public:
 
 	UE_API const TArray< TWeakObjectPtr<UTickableConstraint> >& Get(UObject* InTarget, UWorld* InWorld = nullptr);
 	UE_API bool HasAnyActiveConstraint(UObject* InTarget, UWorld* InWorld = nullptr);
+
+	UE_API bool HasAnyDependency(UObject* InChild, UObject* InParent, UWorld* InWorld);
 	
 	UE_API void RegisterNotifications(UWorld* InWorld);
 	UE_API void UnregisterNotifications(UWorld* InWorld);
 
 private:
+
+	bool HasAnyDependencyInternal(UObject* InChild, UObject* InParent, UWorld* InWorld, TSet< UObject* >& Visited);
 	
 	TMap<uint32, TArray< TWeakObjectPtr<UTickableConstraint> >> PerHandleConstraints;
 	TMap<TObjectKey<UObject>, TArray< TWeakObjectPtr<UTickableConstraint> >> PerTargetConstraints;
