Fixed static analysis issue from 46732089

#rnx

diff --git a/Engine/Plugins/Animation/ControlRig/Source/ControlRigEditor/Private/EditMode/ControlRigEditMode.cpp b/Engine/Plugins/Animation/ControlRig/Source/ControlRigEditor/Private/EditMode/ControlRigEditMode.cpp
--- a/Engine/Plugins/Animation/ControlRig/Source/ControlRigEditor/Private/EditMode/ControlRigEditMode.cpp
+++ b/Engine/Plugins/Animation/ControlRig/Source/ControlRigEditor/Private/EditMode/ControlRigEditMode.cpp
@@ -6259,9 +6259,9 @@ USceneComponent* FControlRigEditMode::GetHostingSceneComponent(const UControlRig
 	if (const USkeleton* BoundSkeleton = Cast<USkeleton>(BoundObject))
 	{
 		// Bound to a Skeleton means we are previewing an Animation Sequence
-		if (ULevel* Level = WorldPtr ? WorldPtr->PersistentLevel : nullptr)
+		if (WorldPtr && WorldPtr->PersistentLevel)
 		{
-			const TObjectPtr<AActor>* PreviewActor = Level->Actors.FindByPredicate([](const TObjectPtr<AActor>& Actor)
+			const TObjectPtr<AActor>* PreviewActor = WorldPtr->PersistentLevel->Actors.FindByPredicate([](const TObjectPtr<AActor>& Actor)
 			{
 				return Actor && Actor->GetClass() == AAnimationEditorPreviewActor::StaticClass();
 			});
