Selection Sets: Fix for subsequence navigation with sequencer. Main issue is we don't spawn and create control rigs till evaluation.

diff --git a/Engine/Plugins/Animation/ControlRig/Source/ControlRigEditor/Private/Sequencer/SelectionSets/SSelectionSets.cpp b/Engine/Plugins/Animation/ControlRig/Source/ControlRigEditor/Private/Sequencer/SelectionSets/SSelectionSets.cpp
--- a/Engine/Plugins/Animation/ControlRig/Source/ControlRigEditor/Private/Sequencer/SelectionSets/SSelectionSets.cpp
+++ b/Engine/Plugins/Animation/ControlRig/Source/ControlRigEditor/Private/Sequencer/SelectionSets/SSelectionSets.cpp
@@ -909,15 +909,19 @@ void SSelectionSets::PopulateActiveActors()
 														SNew(SImage)
 															.Image_Lambda([GetSelectionSetInfo]() -> const FSlateBrush*
 															{
-																const FActorWithSelectionSet* ActorWithSeletionSet = GetSelectionSetInfo();
-																for (const TWeakObjectPtr<UObject>& Object : ActorWithSeletionSet->Actors)
+																if (const FActorWithSelectionSet* ActorWithSeletionSet = GetSelectionSetInfo())
 																{
-																	if (UControlRig* ControlRig = Cast<UControlRig>(Object.Get()))
+
+																	for (const TWeakObjectPtr<UObject>& Object : ActorWithSeletionSet->Actors)
 																	{
-																		return FSlateIconFinder::FindIconForClass(USkeleton::StaticClass()).GetIcon();
+																		if (UControlRig* ControlRig = Cast<UControlRig>(Object.Get()))
+																		{
+																			return FSlateIconFinder::FindIconForClass(USkeleton::StaticClass()).GetIcon();
+																		}
 																	}
+																	return FSlateIconFinder::FindIconForClass(AActor::StaticClass()).GetIcon();
 																}
-																return FSlateIconFinder::FindIconForClass(AActor::StaticClass()).GetIcon();
+																return nullptr;
 															})
 															.ColorAndOpacity(FSlateColor::UseForeground())
 													]
@@ -1232,18 +1236,24 @@ void  SSelectionSets::OnCloseEvent(TSharedRef<ISequencer> InSequencer)
 
 void SSelectionSets::UpdateBindings()
 {
-	if (TSharedPtr<ISequencer> SequencerPtr = UAIESelectionSets::GetSequencerFromAsset())
+	GEditor->GetTimerManager()->SetTimerForNextTick([this, WeakThis = this->AsWeak()]()
 	{
-		if (UAIESelectionSets* SelectionSets = UAIESelectionSets::GetSelectionSets(SequencerPtr, false /*Add if doesn't exist*/))
+		if (TSharedPtr<ISequencer> SequencerPtr = UAIESelectionSets::GetSequencerFromAsset())
 		{
-			SelectionSets->SelectionListChanged().RemoveAll(this);
-			SelectionSets->SelectionListChanged().AddSP(this, &SSelectionSets::HandleSelectionListChanged);
+			if (WeakThis.IsValid())
+			{
+				if (UAIESelectionSets* SelectionSets = UAIESelectionSets::GetSelectionSets(SequencerPtr, false /*Add if doesn't exist*/))
+				{
+					SelectionSets->SelectionListChanged().RemoveAll(this);
+					SelectionSets->SelectionListChanged().AddSP(this, &SSelectionSets::HandleSelectionListChanged);
 
-			SelectionSets->SequencerBindingsAdded(SequencerPtr);
-			PopulateUI();
+					SelectionSets->SequencerBindingsAdded(SequencerPtr);
+				}
+				PopulateUI();
+			}
 		}
-	}
-
+	});
+	
 }
 
 void SSelectionSets::OnMovieSceneBindingsChanged()
diff --git a/Engine/Plugins/Animation/ControlRig/Source/ControlRigEditor/Private/Sequencer/SelectionSets/SelectionSets.cpp b/Engine/Plugins/Animation/ControlRig/Source/ControlRigEditor/Private/Sequencer/SelectionSets/SelectionSets.cpp
--- a/Engine/Plugins/Animation/ControlRig/Source/ControlRigEditor/Private/Sequencer/SelectionSets/SelectionSets.cpp
+++ b/Engine/Plugins/Animation/ControlRig/Source/ControlRigEditor/Private/Sequencer/SelectionSets/SelectionSets.cpp
@@ -215,15 +215,18 @@ void FActorsWithSelectionSets::SetUpFromSequencer(TSharedPtr<ISequencer>& Sequen
 	TArray<TWeakObjectPtr<AActor>> CurrentSet;
 	ActorsWithSet.GenerateKeyArray(CurrentSet);
 
-	TArray<UControlRig*> ControlRigs;
-	FControlRigEditMode* ControlRigEditMode = static_cast<FControlRigEditMode*>(GLevelEditorModeTools().GetActiveMode(FControlRigEditMode::ModeName));
-	if (ControlRigEditMode)
-	{
-		ControlRigs = ControlRigEditMode->GetControlRigsArray(false /*bIsVisible*/);
-	}
-
+	
 	if (ULevelSequence* LevelSequence = Cast<ULevelSequence>(Sequencer->GetFocusedMovieSceneSequence()))
 	{
+		TArray<UControlRig*> ControlRigs;
+		TArray<FControlRigSequencerBindingProxy> Proxies = UControlRigSequencerEditorLibrary::GetControlRigs(LevelSequence);
+		for (FControlRigSequencerBindingProxy& Proxy : Proxies)
+		{
+			if (UControlRig* ControlRig = Proxy.ControlRig.Get())
+			{
+				ControlRigs.Add(ControlRig);
+			}
+		}
 		const UMovieScene* MovieScene = LevelSequence->GetMovieScene();
 		if (MovieScene)
 		{
@@ -260,41 +263,41 @@ void FActorsWithSelectionSets::SetUpFromSequencer(TSharedPtr<ISequencer>& Sequen
 				}
 			}
 		}
-	}
 
-	for (UControlRig* ControlRig : ControlRigs)
-	{
-		if (ControlRig)
+		for (UControlRig* ControlRig : ControlRigs)
 		{
-			const TArray<FRigControlElement*> Controls = ControlRig->AvailableControls();
-			for (const FRigControlElement* Control : Controls)
+			if (ControlRig)
 			{
-				if (Control)
+				const TArray<FRigControlElement*> Controls = ControlRig->AvailableControls();
+				for (const FRigControlElement* Control : Controls)
 				{
-					for (const TPair<FGuid, FAIESelectionSetItem>& Pair : SelectionSets->SelectionSets)
+					if (Control)
 					{
-						if (Pair.Value.ContainsItem(Control->GetName()))
+						for (const TPair<FGuid, FAIESelectionSetItem>& Pair : SelectionSets->SelectionSets)
 						{
-							if (AActor * Actor = FActorWithSelectionSet::GetControlRigActor(ControlRig))
+							if (Pair.Value.ContainsItem(Control->GetName()))
 							{
-								FString ActorName = Actor->GetActorLabel(false);
+								if (AActor* Actor = FActorWithSelectionSet::GetControlRigActor(ControlRig))
+								{
+									FString ActorName = Actor->GetActorLabel(false);
 
-								const bool bExists = ActorsWithSet.Contains(Actor);
-								FActorWithSelectionSet& Set = ActorsWithSet.FindOrAdd(Actor);
+									const bool bExists = ActorsWithSet.Contains(Actor);
+									FActorWithSelectionSet& Set = ActorsWithSet.FindOrAdd(Actor);
 
-								Set.Name = FText::FromString(ActorName);
-								if (bExists == false)
-								{
-									Set.Color = SelectionSets->GetNextActorColor();
-									Set.SetSelected(true);
-								}
-								Set.Actors.Add(ControlRig);
-								Set.SetsThatMatch.Add(Pair.Key);
-								if (MakeSelectedIfPresent.Contains(ControlRig))
-								{
-									Set.SetSelected(true);
+									Set.Name = FText::FromString(ActorName);
+									if (bExists == false)
+									{
+										Set.Color = SelectionSets->GetNextActorColor();
+										Set.SetSelected(true);
+									}
+									Set.Actors.Add(ControlRig);
+									Set.SetsThatMatch.Add(Pair.Key);
+									if (MakeSelectedIfPresent.Contains(ControlRig))
+									{
+										Set.SetSelected(true);
+									}
+									CurrentSet.Remove(Actor);
 								}
-								CurrentSet.Remove(Actor);
 							}
 						}
 					}
