ControlRig & Sequencer: ensure that spaces and constraints delegates are updated when replacing a rig

#rnx

diff --git a/Engine/Plugins/Animation/ControlRig/Source/ControlRig/Private/Sequencer/MovieSceneControlRigParameterSection.cpp b/Engine/Plugins/Animation/ControlRig/Source/ControlRig/Private/Sequencer/MovieSceneControlRigParameterSection.cpp
--- a/Engine/Plugins/Animation/ControlRig/Source/ControlRig/Private/Sequencer/MovieSceneControlRigParameterSection.cpp
+++ b/Engine/Plugins/Animation/ControlRig/Source/ControlRig/Private/Sequencer/MovieSceneControlRigParameterSection.cpp
@@ -2063,22 +2063,18 @@ EMovieSceneChannelProxyType UMovieSceneControlRigParameterSection::CacheChannelP
 			const FName SpaceName = FName(TEXT("Space"));
 
 			// begin constraints
-		// begin constraints
 			auto AddConstrainChannels = [this, &ConstraintsChannelIndex, &SortOrder, &Channels](
 				const FName& InControlName, const FText& InGroup, const bool bEnabled)
 				{
-
-					const FConstraintsManagerController& Controller = FConstraintsManagerController::Get(ControlRig->GetWorld());
-
-					static constexpr bool bSorted = true;
 					const uint32 ControlHash = UTransformableControlHandle::ComputeHash(ControlRig.Get(), InControlName);
-					TArray<TWeakObjectPtr<UTickableConstraint>> Constraints = Controller.GetParentConstraints(ControlHash, bSorted);
-					for (const TWeakObjectPtr<UTickableConstraint>& Constraint : Constraints)
+					for (FConstraintAndActiveChannel& ConstraintChannel : ConstraintsChannels)
 					{
-						if (Constraint.IsValid())
+						UTickableTransformConstraint* Constraint = Cast<UTickableTransformConstraint>(ConstraintChannel.GetConstraint());
+						UTransformableControlHandle* Handle = Constraint ? Cast<UTransformableControlHandle>(Constraint->ChildTRSHandle.Get()) : nullptr;
+						if (Handle && Handle->GetHash() == ControlHash)
 						{
-							const FGuid& ConstraintID = Constraint->ConstraintID;
-							if (FConstraintAndActiveChannel* ConstraintChannel = GetConstraintChannel(ConstraintID))
+							const TWeakObjectPtr<UTickableConstraint>& WeakConstraint = MakeWeakObjectPtr(ConstraintChannel.GetConstraint());
+							if (WeakConstraint.IsValid())
 							{
 								if (FChannelMapInfo* ChannelInfo = ControlChannelMap.Find(InControlName))
 								{
@@ -2086,29 +2082,29 @@ EMovieSceneChannelProxyType UMovieSceneControlRigParameterSection::CacheChannelP
 								}
 
 #if WITH_EDITOR
-								ConstraintChannel->ActiveChannel.ExtraLabel = [Constraint]
+								ConstraintChannel.ActiveChannel.ExtraLabel = [WeakConstraint]
+								{
+									if (WeakConstraint.IsValid())
 									{
-										if (Constraint.IsValid())
+										FString ParentStr; WeakConstraint->GetLabel().Split(TEXT("."), &ParentStr, nullptr);
+										if (!ParentStr.IsEmpty())
 										{
-											FString ParentStr; Constraint->GetLabel().Split(TEXT("."), &ParentStr, nullptr);
-											if (!ParentStr.IsEmpty())
-											{
-												return ParentStr;
-											}
+											return ParentStr;
 										}
-										static const FString DummyStr;
-										return DummyStr;
-									};
+									}
+									static const FString DummyStr;
+									return DummyStr;
+								};
 
-								const FText DisplayText = FText::FromString(Constraint->GetTypeLabel());
-								FMovieSceneChannelMetaData MetaData(Constraint->GetFName(), DisplayText, InGroup, bEnabled);
+								const FText DisplayText = FText::FromString(WeakConstraint->GetTypeLabel());
+								FMovieSceneChannelMetaData MetaData(WeakConstraint->GetFName(), DisplayText, InGroup, bEnabled);
 								ConstraintsChannelIndex += 1;
 								MetaData.SortOrder = SortOrder++;
 								MetaData.bCanCollapseToTrack = true;
 
-								Channels.Add(ConstraintChannel->ActiveChannel, MetaData, TMovieSceneExternalValue<bool>());
+								Channels.Add(ConstraintChannel.ActiveChannel, MetaData, TMovieSceneExternalValue<bool>());
 #else
-								Channels.Add(ConstraintChannel->ActiveChannel);
+								Channels.Add(ConstraintChannel.ActiveChannel);
 #endif
 							}
 						}
diff --git a/Engine/Plugins/Animation/ControlRig/Source/ControlRigEditor/Private/Sequencer/ControlRigParameterTrackEditor.cpp b/Engine/Plugins/Animation/ControlRig/Source/ControlRigEditor/Private/Sequencer/ControlRigParameterTrackEditor.cpp
--- a/Engine/Plugins/Animation/ControlRig/Source/ControlRigEditor/Private/Sequencer/ControlRigParameterTrackEditor.cpp
+++ b/Engine/Plugins/Animation/ControlRig/Source/ControlRigEditor/Private/Sequencer/ControlRigParameterTrackEditor.cpp
@@ -542,17 +542,20 @@ FControlRigParameterTrackEditor::~FControlRigParameterTrackEditor()
 
 void FControlRigParameterTrackEditor::BindControlRig(UControlRig* ControlRig)
 {
-	if (ControlRig && BoundControlRigs.Contains(ControlRig) == false)
+	if (ControlRig)
 	{
-		ControlRig->ControlModified().AddRaw(this, &FControlRigParameterTrackEditor::HandleControlModified);
-		ControlRig->OnPostConstruction_AnyThread().AddRaw(this, &FControlRigParameterTrackEditor::HandleOnPostConstructed);
-		ControlRig->ControlSelected().AddRaw(this, &FControlRigParameterTrackEditor::HandleControlSelected);
-		ControlRig->ControlUndoBracket().AddRaw(this, &FControlRigParameterTrackEditor::HandleControlUndoBracket);
-		ControlRig->ControlRigBound().AddRaw(this, &FControlRigParameterTrackEditor::HandleOnControlRigBound);
+		if (!BoundControlRigs.Contains(ControlRig))
+		{
+			ControlRig->ControlModified().AddRaw(this, &FControlRigParameterTrackEditor::HandleControlModified);
+			ControlRig->OnPostConstruction_AnyThread().AddRaw(this, &FControlRigParameterTrackEditor::HandleOnPostConstructed);
+			ControlRig->ControlSelected().AddRaw(this, &FControlRigParameterTrackEditor::HandleControlSelected);
+			ControlRig->ControlUndoBracket().AddRaw(this, &FControlRigParameterTrackEditor::HandleControlUndoBracket);
+			ControlRig->ControlRigBound().AddRaw(this, &FControlRigParameterTrackEditor::HandleOnControlRigBound);
 		
-		BoundControlRigs.Add(ControlRig);
-		UMovieSceneControlRigParameterTrack* Track = FindTrack(ControlRig);
-		if (Track)
+			BoundControlRigs.Add(ControlRig);
+		}
+
+		if (UMovieSceneControlRigParameterTrack* Track = FindTrack(ControlRig))
 		{
 			for (UMovieSceneSection* BaseSection : Track->GetAllSections())
 			{
@@ -574,8 +577,16 @@ void FControlRigParameterTrackEditor::BindControlRig(UControlRig* ControlRig)
 					}
 				}
 			}
-			Track->SpaceChannelAdded().AddRaw(this, &FControlRigParameterTrackEditor::HandleOnSpaceAdded);
-			Track->ConstraintChannelAdded().AddRaw(this, &FControlRigParameterTrackEditor::HandleOnConstraintAdded);
+
+			if (!Track->SpaceChannelAdded().IsBoundToObject(this))
+			{
+				Track->SpaceChannelAdded().AddRaw(this, &FControlRigParameterTrackEditor::HandleOnSpaceAdded);
+			}
+
+			if (!Track->ConstraintChannelAdded().IsBoundToObject(this))
+			{
+				Track->ConstraintChannelAdded().AddRaw(this, &FControlRigParameterTrackEditor::HandleOnConstraintAdded);
+			}
 		}
 	}
 }
