Fixed a crash at SRigSpacePickerWidget by making sure the weak pointer is cleared before calling request destroy window, as that creates an infinite recursion

diff --git a/Engine/Plugins/Animation/ControlRig/Source/ControlRigEditor/Private/Editor/SRigSpacePickerWidget.cpp b/Engine/Plugins/Animation/ControlRig/Source/ControlRigEditor/Private/Editor/SRigSpacePickerWidget.cpp
--- a/Engine/Plugins/Animation/ControlRig/Source/ControlRigEditor/Private/Editor/SRigSpacePickerWidget.cpp
+++ b/Engine/Plugins/Animation/ControlRig/Source/ControlRigEditor/Private/Editor/SRigSpacePickerWidget.cpp
@@ -356,12 +356,11 @@ TSharedPtr<SWindow> SRigSpacePickerWidget::OpenDialog(bool bModal)
 
 void SRigSpacePickerWidget::CloseDialog()
 {
-	
-	if ( DialogWindow.IsValid() )
+	if (TSharedPtr<SWindow> DialogWindowShared = DialogWindow.Pin())
 	{
-		DialogWindow.Pin()->GetOnWindowDeactivatedEvent().RemoveAll(this);
-		DialogWindow.Pin()->RequestDestroyWindow();
-		DialogWindow.Reset();
+		DialogWindow.Reset(); // we have to reset before calling request destroy window, or an infinite recursion will happen on Mac
+		DialogWindowShared->GetOnWindowDeactivatedEvent().RemoveAll(this);
+		DialogWindowShared->RequestDestroyWindow();
 	}
 }
 
