[ControlRig & RigVM] replace function graph task with ExecuteOnGameThread to avoid scheduling tasks during package save, which can lead random crashes

diff --git a/Engine/Plugins/Animation/ControlRig/Source/ControlRig/Private/Rigs/RigHierarchy.cpp b/Engine/Plugins/Animation/ControlRig/Source/ControlRig/Private/Rigs/RigHierarchy.cpp
--- a/Engine/Plugins/Animation/ControlRig/Source/ControlRig/Private/Rigs/RigHierarchy.cpp
+++ b/Engine/Plugins/Animation/ControlRig/Source/ControlRig/Private/Rigs/RigHierarchy.cpp
@@ -4538,10 +4538,11 @@ void URigHierarchy::SendEvent(const FRigEventContext& InEvent, bool bAsynchronou
 
 		if (bAsynchronous)
 		{
-			FFunctionGraphTask::CreateAndDispatchWhenReady([WeakThis, Delegate, InEvent]()
-            {
+
+			ExecuteOnGameThread(UE_SOURCE_LOCATION,[WeakThis, Delegate, InEvent]()
+			{
                 Delegate.Broadcast(WeakThis.Get(), InEvent);
-            }, TStatId(), NULL, ENamedThreads::GameThread);
+            });
 		}
 		else
 		{
diff --git a/Engine/Plugins/Animation/ControlRig/Source/ControlRigEditor/Private/ControlRigEditorModule.cpp b/Engine/Plugins/Animation/ControlRig/Source/ControlRigEditor/Private/ControlRigEditorModule.cpp
--- a/Engine/Plugins/Animation/ControlRig/Source/ControlRigEditor/Private/ControlRigEditorModule.cpp
+++ b/Engine/Plugins/Animation/ControlRig/Source/ControlRigEditor/Private/ControlRigEditorModule.cpp
@@ -910,10 +910,10 @@ void FControlRigEditorModule::GetDirectManipulationMenuActions(IRigVMClientHost*
 								}
 
 								// run the task after a bit so that the rig has the opportunity to run first
-								FFunctionGraphTask::CreateAndDispatchWhenReady([ControlRigBlueprint, UnitNode, Target]()
+								ExecuteOnGameThread(UE_SOURCE_LOCATION, [ControlRigBlueprint, UnitNode, Target]()
 								{
 									ControlRigBlueprint->AddTransientControl(UnitNode, Target);
-								}, TStatId(), NULL, ENamedThreads::GameThread);
+								});
 							}),
 							FCanExecuteAction::CreateLambda([HasNoUnconstrainedAffectedPin, IsSliced]() -> bool
 							{
diff --git a/Engine/Plugins/Animation/ControlRig/Source/ControlRigEditor/Private/ControlRigElementDetails.cpp b/Engine/Plugins/Animation/ControlRig/Source/ControlRigEditor/Private/ControlRigElementDetails.cpp
--- a/Engine/Plugins/Animation/ControlRig/Source/ControlRigEditor/Private/ControlRigElementDetails.cpp
+++ b/Engine/Plugins/Animation/ControlRig/Source/ControlRigEditor/Private/ControlRigElementDetails.cpp
@@ -1819,14 +1819,14 @@ void FRigBaseElementDetails::CustomizeMetadata(IDetailLayoutBuilder& DetailBuild
 							}
 						}
 						// run the refresh of the user interface on the next tick on the game thread 
-						FFunctionGraphTask::CreateAndDispatchWhenReady([WeakPropertyUtilities]()
+						ExecuteOnGameThread(UE_SOURCE_LOCATION, [WeakPropertyUtilities]()
 						{
 							const TSharedPtr<IPropertyUtilities> PropertyUtilities = WeakPropertyUtilities.IsValid() ? WeakPropertyUtilities.Pin() : nullptr;
 							if(PropertyUtilities.IsValid())
 							{
 								PropertyUtilities->ForceRefresh();
 							}
-						}, TStatId(), NULL, ENamedThreads::GameThread);
+						});
 					}
 				}
 			});
diff --git a/Engine/Plugins/Animation/ControlRig/Source/ControlRigEditor/Private/EditMode/ControlRigEditMode.cpp b/Engine/Plugins/Animation/ControlRig/Source/ControlRigEditor/Private/EditMode/ControlRigEditMode.cpp
--- a/Engine/Plugins/Animation/ControlRig/Source/ControlRigEditor/Private/EditMode/ControlRigEditMode.cpp
+++ b/Engine/Plugins/Animation/ControlRig/Source/ControlRigEditor/Private/EditMode/ControlRigEditMode.cpp
@@ -5085,8 +5085,8 @@ void FControlRigEditMode::OnHierarchyModified_AnyThread(ERigHierarchyNotificatio
 	}
 
 	TWeakObjectPtr<URigHierarchy> WeakHierarchy = InHierarchy;
-	
-	FFunctionGraphTask::CreateAndDispatchWhenReady([this, InNotif, WeakHierarchy, ElementKey, ComponentKey]()
+
+	ExecuteOnGameThread(UE_SOURCE_LOCATION, [this, InNotif, WeakHierarchy, ElementKey, ComponentKey]()
 	{
 		if(!WeakHierarchy.IsValid())
 		{
@@ -5101,7 +5101,7 @@ void FControlRigEditMode::OnHierarchyModified_AnyThread(ERigHierarchyNotificatio
 			OnHierarchyModified(InNotif, WeakHierarchy.Get(), Element);
 		}
 		
-	}, TStatId(), NULL, ENamedThreads::GameThread);
+	});
 }
 
 void FControlRigEditMode::OnControlModified(UControlRig* Subject, FRigControlElement* InControlElement, const FRigControlModifiedContext& Context)
@@ -5160,10 +5160,10 @@ void FControlRigEditMode::OnPostConstruction_AnyThread(UControlRig* InRig, const
 		}
 		else
 		{
-			FFunctionGraphTask::CreateAndDispatchWhenReady([Task]()
+			ExecuteOnGameThread(UE_SOURCE_LOCATION, [Task]()
 			{
 				Task();
-			}, TStatId(), NULL, ENamedThreads::GameThread);
+			});
 		}
 	}
 }
diff --git a/Engine/Plugins/Animation/ControlRig/Source/ControlRigEditor/Private/EditMode/SControlRigOutliner.cpp b/Engine/Plugins/Animation/ControlRig/Source/ControlRigEditor/Private/EditMode/SControlRigOutliner.cpp
--- a/Engine/Plugins/Animation/ControlRig/Source/ControlRigEditor/Private/EditMode/SControlRigOutliner.cpp
+++ b/Engine/Plugins/Animation/ControlRig/Source/ControlRigEditor/Private/EditMode/SControlRigOutliner.cpp
@@ -2489,7 +2489,7 @@ void SControlRigOutliner::HandleOnObjectBoundToControlRig(UObject* InObject)
 	//just refresh the views, but do so on next tick since with FK control rig's the controls aren't set up
 	//until AFTER we are bound.
 	TWeakPtr<SControlRigOutliner> WeakPtr = StaticCastSharedRef<SControlRigOutliner>(AsShared()).ToWeakPtr();
-	FFunctionGraphTask::CreateAndDispatchWhenReady([WeakPtr]()
+	ExecuteOnGameThread(UE_SOURCE_LOCATION, [WeakPtr]()
 	{
 		if (WeakPtr.IsValid())
 		{
@@ -2506,7 +2506,7 @@ void SControlRigOutliner::HandleOnObjectBoundToControlRig(UObject* InObject)
 				}
 			}
 		}
-	}, TStatId(), nullptr, ENamedThreads::GameThread);
+	});
 }
 
 void SControlRigOutliner::HandlePostConstruction(UControlRig* InControlRig, const FName& InEventName)
diff --git a/Engine/Plugins/Animation/ControlRig/Source/ControlRigEditor/Private/Editor/ControlRigEditor.cpp b/Engine/Plugins/Animation/ControlRig/Source/ControlRigEditor/Private/Editor/ControlRigEditor.cpp
--- a/Engine/Plugins/Animation/ControlRig/Source/ControlRigEditor/Private/Editor/ControlRigEditor.cpp
+++ b/Engine/Plugins/Animation/ControlRig/Source/ControlRigEditor/Private/Editor/ControlRigEditor.cpp
@@ -4117,10 +4117,10 @@ void FControlRigBaseEditor::OnDirectManipulationChanged(TSharedPtr<FString> NewV
 		if(Target.Name.Equals(DesiredTarget, ESearchCase::CaseSensitive))
 		{
 			// run the task after a bit so that the rig has the opportunity to run first
-			FFunctionGraphTask::CreateAndDispatchWhenReady([ControlRigBlueprint, UnitNode, Target]()
+			ExecuteOnGameThread(UE_SOURCE_LOCATION, [ControlRigBlueprint, UnitNode, Target]()
 			{
 				ControlRigBlueprint->AddTransientControl(UnitNode, Target);
-			}, TStatId(), NULL, ENamedThreads::GameThread);
+			});
 			break;
 		}
 	}
@@ -4176,10 +4176,10 @@ void FControlRigBaseEditor::RefreshDirectManipulationTextList()
 	(void)GetDirectManipulationTargetTextList();
 	if(DirectManipulationCombo.IsValid())
 	{
-		FFunctionGraphTask::CreateAndDispatchWhenReady([this]()
+		ExecuteOnGameThread(UE_SOURCE_LOCATION, [this]()
 		{
 			DirectManipulationCombo->RefreshOptions();
-		}, TStatId(), NULL, ENamedThreads::GameThread);
+		});
 	}
 }
 
@@ -4757,10 +4757,10 @@ void FControlRigBaseEditor::OnHierarchyModified_AnyThread(ERigHierarchyNotificat
 	}
 	else
 	{
-		FFunctionGraphTask::CreateAndDispatchWhenReady([Task]()
+		ExecuteOnGameThread(UE_SOURCE_LOCATION, [Task]()
 		{
 			Task();
-		}, TStatId(), NULL, ENamedThreads::GameThread);
+		});
 	}
 }
 
@@ -6595,10 +6595,10 @@ void FControlRigBaseEditor::OnPostConstruction_AnyThread(UControlRig* InRig, con
 		}
 		else
 		{
-			FFunctionGraphTask::CreateAndDispatchWhenReady([Task]()
+			ExecuteOnGameThread(UE_SOURCE_LOCATION, [Task]()
 			{
 				Task();
-			}, TStatId(), NULL, ENamedThreads::GameThread);
+			});
 		}
 	}
 	else if(bShouldExecute)
diff --git a/Engine/Plugins/Animation/ControlRig/Source/ControlRigEditor/Private/Editor/ControlRigSchematicModel.cpp b/Engine/Plugins/Animation/ControlRig/Source/ControlRigEditor/Private/Editor/ControlRigSchematicModel.cpp
--- a/Engine/Plugins/Animation/ControlRig/Source/ControlRigEditor/Private/Editor/ControlRigSchematicModel.cpp
+++ b/Engine/Plugins/Animation/ControlRig/Source/ControlRigEditor/Private/Editor/ControlRigSchematicModel.cpp
@@ -1431,7 +1431,7 @@ void FControlRigSchematicModel::HandleSchematicCancelDrag(SSchematicGraphPanel*
 			}
 		}
 
-		FFunctionGraphTask::CreateAndDispatchWhenReady([this, Keys]()
+		ExecuteOnGameThread(UE_SOURCE_LOCATION, [this, Keys]()
 		{
 			if (UModularRigController* Controller = ControlRigBlueprint->GetModularRigController())
 			{
@@ -1442,7 +1442,7 @@ void FControlRigSchematicModel::HandleSchematicCancelDrag(SSchematicGraphPanel*
 					Controller->DisconnectConnector(Key);
 				}
 			}
-		}, TStatId(), NULL, ENamedThreads::GameThread);
+		});
 	}
 }
 
@@ -1542,7 +1542,7 @@ void FControlRigSchematicModel::HandleSchematicDrop(SSchematicGraphPanel* InPane
 
 			if(ControlRigClass && ModuleSettings.IsValidModule())
 			{
-				FFunctionGraphTask::CreateAndDispatchWhenReady([this, ControlRigClass, ModuleSettings, TargetKeys]()
+				ExecuteOnGameThread(UE_SOURCE_LOCATION, [this, ControlRigClass, ModuleSettings, TargetKeys]()
 				{
 					if (UModularRigController* Controller = ControlRigBlueprint->GetModularRigController())
 					{
@@ -1587,14 +1587,14 @@ void FControlRigSchematicModel::HandleSchematicDrop(SSchematicGraphPanel* InPane
 
 						ClearSelection();
 					}
-				}, TStatId(), NULL, ENamedThreads::GameThread);
+				});
 			}
 		}
 	}
 	else if(SchematicDragDropOp.IsValid())
 	{
 		const TArray<FGuid> Sources = SchematicDragDropOp->GetElements();
-		FFunctionGraphTask::CreateAndDispatchWhenReady([this, Sources, TargetKeys]()
+		ExecuteOnGameThread(UE_SOURCE_LOCATION, [this, Sources, TargetKeys]()
 		{
 			UModularRig* ControlRig = Cast<UModularRig>(ControlRigBlueprint->GetDebuggedControlRig());
 			if (!ControlRig)
@@ -1630,12 +1630,12 @@ void FControlRigSchematicModel::HandleSchematicDrop(SSchematicGraphPanel* InPane
 
 				ClearSelection();
 			}
-		}, TStatId(), NULL, ENamedThreads::GameThread);
+		});
 	}
 	else if(ModuleDragDropOperation.IsValid())
 	{
 		const TArray<FName> ModuleNames = ModuleDragDropOperation->GetModules();
-		FFunctionGraphTask::CreateAndDispatchWhenReady([this, ModuleNames, TargetKeys]()
+		ExecuteOnGameThread(UE_SOURCE_LOCATION, [this, ModuleNames, TargetKeys]()
 		{
 			UModularRig* ControlRig = Cast<UModularRig>(ControlRigBlueprint->GetDebuggedControlRig());
 			if (!ControlRig)
@@ -1691,7 +1691,7 @@ void FControlRigSchematicModel::HandleSchematicDrop(SSchematicGraphPanel* InPane
 					}
 				}
 			}
-		}, TStatId(), NULL, ENamedThreads::GameThread);
+		});
 	}
 }
 
diff --git a/Engine/Plugins/Animation/ControlRig/Source/ControlRigEditor/Private/Editor/ControlRigSkeletalMeshComponent.cpp b/Engine/Plugins/Animation/ControlRig/Source/ControlRigEditor/Private/Editor/ControlRigSkeletalMeshComponent.cpp
--- a/Engine/Plugins/Animation/ControlRig/Source/ControlRigEditor/Private/Editor/ControlRigSkeletalMeshComponent.cpp
+++ b/Engine/Plugins/Animation/ControlRig/Source/ControlRigEditor/Private/Editor/ControlRigSkeletalMeshComponent.cpp
@@ -307,10 +307,10 @@ void UControlRigSkeletalMeshComponent::OnHierarchyModified_AnyThread(ERigHierarc
 	}
 	else
 	{
-		FFunctionGraphTask::CreateAndDispatchWhenReady([Task]()
+		ExecuteOnGameThread(UE_SOURCE_LOCATION, [Task]()
 		{
 			Task();		
-		}, TStatId(), NULL, ENamedThreads::GameThread);
+		});
 	}
 }
 
@@ -329,9 +329,9 @@ void UControlRigSkeletalMeshComponent::OnPostConstruction_AnyThread(UControlRig*
 	}
 	else
 	{
-		FFunctionGraphTask::CreateAndDispatchWhenReady([this]()
+		ExecuteOnGameThread(UE_SOURCE_LOCATION, [this]()
 		{
 			RebuildDebugDrawSkeleton();
-		}, TStatId(), NULL, ENamedThreads::GameThread);
+		});
 	}
 }
\ No newline at end of file
diff --git a/Engine/Plugins/Animation/ControlRig/Source/ControlRigEditor/Private/Editor/SRigCurveContainer.cpp b/Engine/Plugins/Animation/ControlRig/Source/ControlRigEditor/Private/Editor/SRigCurveContainer.cpp
--- a/Engine/Plugins/Animation/ControlRig/Source/ControlRigEditor/Private/Editor/SRigCurveContainer.cpp
+++ b/Engine/Plugins/Animation/ControlRig/Source/ControlRigEditor/Private/Editor/SRigCurveContainer.cpp
@@ -464,10 +464,10 @@ void SRigCurveContainer::RefreshCurveList_AnyThread()
 	}
 	else
 	{
-		FFunctionGraphTask::CreateAndDispatchWhenReady([this]()
+		ExecuteOnGameThread(UE_SOURCE_LOCATION, [this]()
 		{
 			RefreshCurveList();
-		}, TStatId(), nullptr, ENamedThreads::GameThread);
+		});
 	}
 }
 
@@ -645,7 +645,7 @@ void SRigCurveContainer::OnHierarchyModified_AnyThread(ERigHierarchyNotification
 		}
 		
 		TWeakObjectPtr<URigHierarchy> WeakHierarchy = InHierarchy;
-		FFunctionGraphTask::CreateAndDispatchWhenReady([this, InNotif, WeakHierarchy,Key]()
+		ExecuteOnGameThread(UE_SOURCE_LOCATION, [this, InNotif, WeakHierarchy,Key]()
 		{
 			if(!WeakHierarchy.IsValid())
 			{
@@ -656,7 +656,7 @@ void SRigCurveContainer::OnHierarchyModified_AnyThread(ERigHierarchyNotification
 			{
 				OnHierarchyModified(InNotif, WeakHierarchy.Get(), Element);
 			}
-		}, TStatId(), NULL, ENamedThreads::GameThread);
+		});
 	}
 }
 
diff --git a/Engine/Plugins/Animation/ControlRig/Source/ControlRigEditor/Private/Editor/SRigHierarchy.cpp b/Engine/Plugins/Animation/ControlRig/Source/ControlRigEditor/Private/Editor/SRigHierarchy.cpp
--- a/Engine/Plugins/Animation/ControlRig/Source/ControlRigEditor/Private/Editor/SRigHierarchy.cpp
+++ b/Engine/Plugins/Animation/ControlRig/Source/ControlRigEditor/Private/Editor/SRigHierarchy.cpp
@@ -608,10 +608,10 @@ void SRigHierarchy::RefreshTreeView(bool bRebuildContent)
 		// is the rig currently running
 		if(Hierarchy->HasExecuteContext())
 		{
-			FFunctionGraphTask::CreateAndDispatchWhenReady([this, bRebuildContent]()
+			ExecuteOnGameThread(UE_SOURCE_LOCATION, [this, bRebuildContent]()
 			{
 				RefreshTreeView(bRebuildContent);
-			}, TStatId(), NULL, ENamedThreads::GameThread);
+			});
 		}
 	}
 	
@@ -972,7 +972,7 @@ void SRigHierarchy::OnHierarchyModified_AnyThread(ERigHierarchyNotification InNo
 
 		TWeakObjectPtr<URigHierarchy> WeakHierarchy = InHierarchy;
 
-		FFunctionGraphTask::CreateAndDispatchWhenReady([this, InNotif, WeakHierarchy, Key]()
+		ExecuteOnGameThread(UE_SOURCE_LOCATION, [this, InNotif, WeakHierarchy, Key]()
         {
             if(!WeakHierarchy.IsValid())
             {
@@ -983,7 +983,7 @@ void SRigHierarchy::OnHierarchyModified_AnyThread(ERigHierarchyNotification InNo
 	            OnHierarchyModified(InNotif, WeakHierarchy.Get(), Element);
             }
 			
-        }, TStatId(), NULL, ENamedThreads::GameThread);
+        });
 	}
 }
 
@@ -1141,10 +1141,10 @@ void SRigHierarchy::OnPostConstruction_AnyThread(UControlRig* InRig, const FName
 		}
 		else
 		{
-			FFunctionGraphTask::CreateAndDispatchWhenReady([Task]()
+			ExecuteOnGameThread(UE_SOURCE_LOCATION, [Task]()
 			{
 				Task();
-			}, TStatId(), NULL, ENamedThreads::GameThread);
+			});
 		}
 	}
 }
diff --git a/Engine/Plugins/Animation/ControlRig/Source/ControlRigEditor/Private/Sequencer/ControlRigParameterTrackEditor.cpp b/Engine/Plugins/Animation/ControlRig/Source/ControlRigEditor/Private/Sequencer/ControlRigParameterTrackEditor.cpp
--- a/Engine/Plugins/Animation/ControlRig/Source/ControlRigEditor/Private/Sequencer/ControlRigParameterTrackEditor.cpp
+++ b/Engine/Plugins/Animation/ControlRig/Source/ControlRigEditor/Private/Sequencer/ControlRigParameterTrackEditor.cpp
@@ -1719,14 +1719,14 @@ void FControlRigParameterTrackEditor::FRecreateRigOperator::RefreshSettingsFromC
 
 void FControlRigParameterTrackEditor::FRecreateRigOperator::OnPostControlRigForwardSolve_AnyThread(UControlRig*, const FName&)
 {
-	FFunctionGraphTask::CreateAndDispatchWhenReady(
+	ExecuteOnGameThread(UE_SOURCE_LOCATION, 
 		[WeakThis = TWeakPtr<FControlRigParameterTrackEditor::FRecreateRigOperator>(SharedThis(this))]()
 	{
 		if (WeakThis.IsValid())
 		{
 			WeakThis.Pin()->RefreshSettingsFromControlRig();
 		}
-	}, TStatId(), nullptr, ENamedThreads::GameThread);
+	});
 }
 
 void FControlRigParameterTrackEditor::FRecreateRigOperator::OnRigSettingsChanged(const FPropertyChangedEvent& InEvent)
diff --git a/Engine/Plugins/Runtime/RigVM/Source/RigVMDeveloper/Private/EdGraph/RigVMEdGraphNode.cpp b/Engine/Plugins/Runtime/RigVM/Source/RigVMDeveloper/Private/EdGraph/RigVMEdGraphNode.cpp
--- a/Engine/Plugins/Runtime/RigVM/Source/RigVMDeveloper/Private/EdGraph/RigVMEdGraphNode.cpp
+++ b/Engine/Plugins/Runtime/RigVM/Source/RigVMDeveloper/Private/EdGraph/RigVMEdGraphNode.cpp
@@ -2440,7 +2440,7 @@ void URigVMEdGraphNode::RequestRename(float InDelay)
 				
 				if(InDelay < SMALL_NUMBER)
 				{
-					FFunctionGraphTask::CreateAndDispatchWhenReady(RequestRenameLambda, TStatId(), NULL, ENamedThreads::GameThread);
+					ExecuteOnGameThread(UE_SOURCE_LOCATION, RequestRenameLambda);
 				}
 				else
 				{
diff --git a/Engine/Plugins/Runtime/RigVM/Source/RigVMEditor/Private/Editor/RigVMEditor.cpp b/Engine/Plugins/Runtime/RigVM/Source/RigVMEditor/Private/Editor/RigVMEditor.cpp
--- a/Engine/Plugins/Runtime/RigVM/Source/RigVMEditor/Private/Editor/RigVMEditor.cpp
+++ b/Engine/Plugins/Runtime/RigVM/Source/RigVMEditor/Private/Editor/RigVMEditor.cpp
@@ -396,7 +396,7 @@ void FRigVMEditorBase::InitRigVMEditor(const EToolkitMode::Type Mode, const TSha
 	}
 
 	TWeakPtr<IRigVMEditor> ThisWeakPtr = SharedRef().ToWeakPtr();
-	FFunctionGraphTask::CreateAndDispatchWhenReady([ThisWeakPtr, InRigVMBlueprint]()
+	ExecuteOnGameThread(UE_SOURCE_LOCATION, [ThisWeakPtr, InRigVMBlueprint]()
 	{
 		if (TSharedPtr<IRigVMEditor> Editor = ThisWeakPtr.Pin())
 		{
@@ -420,7 +420,7 @@ void FRigVMEditorBase::InitRigVMEditor(const EToolkitMode::Type Mode, const TSha
 			   TabManager->TryInvokeTab(MyBlueprintTabId);
 		   }
 		}
-	}, TStatId(), nullptr, ENamedThreads::GameThread);
+	});
 
 	bRigVMEditorInitialized = true;
 	UpdateStaleWatchedPins();
diff --git a/Engine/Plugins/Runtime/RigVM/Source/RigVMEditor/Private/Widgets/SRigVMBulkEditWidget.cpp b/Engine/Plugins/Runtime/RigVM/Source/RigVMEditor/Private/Widgets/SRigVMBulkEditWidget.cpp
--- a/Engine/Plugins/Runtime/RigVM/Source/RigVMEditor/Private/Widgets/SRigVMBulkEditWidget.cpp
+++ b/Engine/Plugins/Runtime/RigVM/Source/RigVMEditor/Private/Widgets/SRigVMBulkEditWidget.cpp
@@ -226,7 +226,7 @@ void SRigVMBulkEditWidget::Tick(const FGeometry& AllottedGeometry, const double
 			Transaction = MakeShareable(new FScopedTransaction(LOCTEXT("BulkEdit", "Bulk Edit")));
 		}
 		
-		FFunctionGraphTask::CreateAndDispatchWhenReady([this]()
+		ExecuteOnGameThread(UE_SOURCE_LOCATION, [this]()
 		{
 			TSharedPtr<FRigVMTreeTask> Task;
 			{
@@ -335,7 +335,7 @@ void SRigVMBulkEditWidget::Tick(const FGeometry& AllottedGeometry, const double
 
 				ReportDelegateGuard.Reset();
 			}
-		}, TStatId(), nullptr, ENamedThreads::GameThread);		
+		});		
 	}
 	else if(!CompletedTasks.IsEmpty())
 	{
diff --git a/Engine/Plugins/Runtime/RigVM/Source/RigVMEditor/Private/Widgets/SRigVMExecutionStackView.cpp b/Engine/Plugins/Runtime/RigVM/Source/RigVMEditor/Private/Widgets/SRigVMExecutionStackView.cpp
--- a/Engine/Plugins/Runtime/RigVM/Source/RigVMEditor/Private/Widgets/SRigVMExecutionStackView.cpp
+++ b/Engine/Plugins/Runtime/RigVM/Source/RigVMEditor/Private/Widgets/SRigVMExecutionStackView.cpp
@@ -1183,7 +1183,7 @@ void SRigVMExecutionStackView::HandleHostInitializedEvent(URigVMHost* InHost, co
 	}
 	else
 	{
-		FFunctionGraphTask::CreateAndDispatchWhenReady(HandleHostInitializedEventLambda, TStatId(), NULL, ENamedThreads::GameThread);
+		ExecuteOnGameThread(UE_SOURCE_LOCATION, HandleHostInitializedEventLambda);
 	}
 }
 
@@ -1256,7 +1256,7 @@ void SRigVMExecutionStackView::HandleHostExecutedEvent(URigVMHost* InHost, const
 	}
 	else
 	{
-		FFunctionGraphTask::CreateAndDispatchWhenReady(HandleHostExecutedEventLambda, TStatId(), NULL, ENamedThreads::GameThread);
+		ExecuteOnGameThread(UE_SOURCE_LOCATION, HandleHostExecutedEventLambda);
 	}
 
 }
diff --git a/Engine/Plugins/Runtime/RigVM/Source/RigVMEditor/Private/Widgets/SRigVMGraphNode.cpp b/Engine/Plugins/Runtime/RigVM/Source/RigVMEditor/Private/Widgets/SRigVMGraphNode.cpp
--- a/Engine/Plugins/Runtime/RigVM/Source/RigVMEditor/Private/Widgets/SRigVMGraphNode.cpp
+++ b/Engine/Plugins/Runtime/RigVM/Source/RigVMEditor/Private/Widgets/SRigVMGraphNode.cpp
@@ -410,13 +410,13 @@ void SRigVMGraphNode::Construct( const FArguments& InArgs )
 	if(EdGraphNode->bRenameIsPending)
 	{
 		TWeakObjectPtr<URigVMEdGraphNode> WeakNode = EdGraphNode;
-		FFunctionGraphTask::CreateAndDispatchWhenReady([WeakNode]()
+		ExecuteOnGameThread(UE_SOURCE_LOCATION, [WeakNode]()
 		{
 			if(WeakNode.IsValid())
 			{
 				(void)WeakNode->RequestRename();
 			}
-		}, TStatId(), NULL, ENamedThreads::GameThread);
+		});
 	}
 }
 
