Control Rig Pose: Guard against nullptr

#rnx

diff --git a/Engine/Plugins/Animation/ControlRig/Source/ControlRig/Private/Tools/ControlRigPose.cpp b/Engine/Plugins/Animation/ControlRig/Source/ControlRig/Private/Tools/ControlRigPose.cpp
--- a/Engine/Plugins/Animation/ControlRig/Source/ControlRig/Private/Tools/ControlRigPose.cpp
+++ b/Engine/Plugins/Animation/ControlRig/Source/ControlRig/Private/Tools/ControlRigPose.cpp
@@ -14,6 +14,11 @@
 
 void FControlRigControlPose::SavePose(UControlRig* ControlRig, bool bUseAll)
 {
+	if (!ControlRig)
+	{
+		return;
+	}
+	
 	TArray<FRigControlElement*> CurrentControls;
 	ControlRig->GetControlsInOrder(CurrentControls);
 	CopyOfControls.SetNum(0);
@@ -52,6 +57,10 @@ void FControlRigControlPose::SavePose(UControlRig* ControlRig, bool bUseAll)
 
 void FControlRigControlPose::PastePose(UControlRig* ControlRig, bool bDoKey, bool bDoMirror)
 {
+	if (!ControlRig)
+	{
+		return;
+	}
 	PastePoseInternal(ControlRig, bDoKey, bDoMirror, CopyOfControls);
 	ControlRig->Evaluate_AnyThread();
 	PastePoseInternal(ControlRig, bDoKey, bDoMirror, CopyOfControls);
@@ -61,6 +70,10 @@ void FControlRigControlPose::PastePose(UControlRig* ControlRig, bool bDoKey, boo
 void FControlRigControlPose::SetControlMirrorTransform(bool bDoLocal, UControlRig* ControlRig, const FName& Name, bool bIsMatched, 
 	const FTransform& GlobalTransform, const FTransform& LocalTransform, bool bNotify, const  FRigControlModifiedContext&Context,bool bSetupUndo)
 {
+	if (!ControlRig)
+	{
+		return;
+	}
 	if (bDoLocal || bIsMatched)
 	{
 		ControlRig->SetControlLocalTransform(Name, LocalTransform, bNotify,Context,bSetupUndo, true/* bFixEulerFlips*/);
@@ -74,6 +87,10 @@ void FControlRigControlPose::SetControlMirrorTransform(bool bDoLocal, UControlRi
 
 void FControlRigControlPose::PastePoseInternal(UControlRig* ControlRig, bool bDoKey, bool bDoMirror, const TArray<FRigControlCopy>& ControlsToPaste)
 {
+	if (!ControlRig)
+	{
+		return;
+	}
 	FRigControlModifiedContext Context;
 	Context.SetKey = bDoKey ? EControlRigSetKey::Always : EControlRigSetKey::DoNotCare;
 	FControlRigPoseMirrorTable MirrorTable;
@@ -176,6 +193,10 @@ void FControlRigControlPose::PastePoseInternal(UControlRig* ControlRig, bool bDo
 
 void FControlRigControlPose::BlendWithInitialPoses(FControlRigControlPose& InitialPose, UControlRig* ControlRig, bool bDoKey, bool bDoMirror, float BlendValue, bool bDoAdditive)
 {
+	if (!ControlRig)
+	{
+		return;
+	}
 
 	auto BlendTransforms = [bDoAdditive, BlendValue](FTransform& InitialVal, FTransform& Val) 
 	{ 
@@ -365,6 +386,10 @@ void UControlRigPoseAsset::SavePose(UControlRig* InControlRig, bool bUseAll)
 
 void UControlRigPoseAsset::PastePose(UControlRig* InControlRig, bool bDoKey, bool bDoMirror, bool bDoAdditive)
 {
+	if (!InControlRig)
+	{
+		return;
+	}
 #if WITH_EDITOR
 	FScopedTransaction ScopedTransaction(LOCTEXT("PastePoseTransaction", "Paste Pose"));
 	InControlRig->Modify();
