Control Rig: Multiple Control Rig tests are being skipped due to an assertion error

#rnx

diff --git a/Engine/Plugins/Animation/ControlRig/Source/ControlRigEditor/Private/Editor/ControlRigEditor.cpp b/Engine/Plugins/Animation/ControlRig/Source/ControlRigEditor/Private/Editor/ControlRigEditor.cpp
--- a/Engine/Plugins/Animation/ControlRig/Source/ControlRigEditor/Private/Editor/ControlRigEditor.cpp
+++ b/Engine/Plugins/Animation/ControlRig/Source/ControlRigEditor/Private/Editor/ControlRigEditor.cpp
@@ -4172,15 +4172,23 @@ const TArray<TSharedPtr<FString>>& FControlRigBaseEditor::GetDirectManipulationT
 
 void FControlRigBaseEditor::RefreshDirectManipulationTextList()
 {
-	DirectManipulationTextList.Reset();
-	(void)GetDirectManipulationTargetTextList();
-	if(DirectManipulationCombo.IsValid())
-	{
-		ExecuteOnGameThread(UE_SOURCE_LOCATION, [this]()
+	const TWeakPtr<FControlRigBaseEditor> WeakThis = SharedRef();
+	const TWeakPtr<SComboBox<TSharedPtr<FString>>> WeakDirectManipulationCombo = DirectManipulationCombo;
+	ExecuteOnGameThread(UE_SOURCE_LOCATION,
+		[WeakThis, WeakDirectManipulationCombo]()
 		{
-			DirectManipulationCombo->RefreshOptions();
+			const TSharedPtr<FControlRigBaseEditor> SharedThis = WeakThis.IsValid() ? WeakThis.Pin() : nullptr;
+			const TSharedPtr<SComboBox<TSharedPtr<FString>>> SharedDirectManipulationCombo = WeakDirectManipulationCombo.IsValid() ? WeakDirectManipulationCombo.Pin() : nullptr;
+
+			if (SharedThis.IsValid() &&
+				SharedDirectManipulationCombo.IsValid())
+			{
+				SharedThis->DirectManipulationTextList.Reset();
+				SharedThis->GetDirectManipulationTargetTextList();
+
+				SharedDirectManipulationCombo->RefreshOptions();
+			}
 		});
-	}
 }
 
 EVisibility FControlRigBaseEditor::GetPreviewingNodeVisibility() const
