Control Rig: Patch for loading overrides when properties are missing

#rnx

diff --git a/Engine/Plugins/Animation/ControlRig/Source/ControlRig/Private/ControlRigOverride.cpp b/Engine/Plugins/Animation/ControlRig/Source/ControlRig/Private/ControlRigOverride.cpp
--- a/Engine/Plugins/Animation/ControlRig/Source/ControlRig/Private/ControlRigOverride.cpp
+++ b/Engine/Plugins/Animation/ControlRig/Source/ControlRig/Private/ControlRigOverride.cpp
@@ -128,43 +128,49 @@ bool FControlRigOverrideValue::Serialize(FArchive& Ar)
 			FPropertyInfo Info;
 			Ar << Info.Property;
 			Ar << Info.ArrayIndex;
-			Properties.Add(Info);
+			if (Info.Property.IsValid())
+			{
+				Properties.Add(Info);
+			}
 		}
 
-		// rebuild path from properties
-		FStringBuilderBase PathBuilder;
-		bool bFirstProperty = true;
-
-		for(int32 Index = 0; Index < NumProperties; Index++)
+		if (Properties.Num() == NumProperties)
 		{
-			if(const FProperty* Property = Properties[Index].Property.Get())
+			// rebuild path from properties
+			FStringBuilderBase PathBuilder;
+			bool bFirstProperty = true;
+
+			for(int32 Index = 0; Index < NumProperties; Index++)
 			{
-				if(!bFirstProperty)
+				if(const FProperty* Property = Properties[Index].Property.Get())
 				{
-					PathBuilder.Append(FControlRigOverrideValue::PathSeparator);
+					if(!bFirstProperty)
+					{
+						PathBuilder.Append(FControlRigOverrideValue::PathSeparator);
+					}
+					else
+					{
+						bFirstProperty = false;
+					}
+					PathBuilder.Append(Property->GetName());
+					Path = Property->GetName();
+					if(Properties[Index].ArrayIndex != INDEX_NONE)
+					{
+						Path->Appendf(TEXT("[%d]"),  Properties[Index].ArrayIndex);
+					}
 				}
 				else
 				{
-					bFirstProperty = false;
-				}
-				PathBuilder.Append(Property->GetName());
-				Path = Property->GetName();
-				if(Properties[Index].ArrayIndex != INDEX_NONE)
-				{
-					Path->Appendf(TEXT("[%d]"),  Properties[Index].ArrayIndex);
+					PathBuilder.Reset();
+					break;
 				}
 			}
-			else
+
+			if(PathBuilder.Len() > 0)
 			{
-				PathBuilder.Reset();
-				break;
+				Path = PathBuilder.ToString();
 			}
 		}
-
-		if(PathBuilder.Len() > 0)
-		{
-			Path = PathBuilder.ToString();
-		}
 	}
 	else
 	{
@@ -185,46 +191,61 @@ bool FControlRigOverrideValue::Serialize(FArchive& Ar)
 
 	if(Ar.IsLoading())
 	{
+		// during save the data skip offset was stored as an int32, we need to convert it over.
+		// this is fixed in 5.8 - but for 5.7 we'll need to add this dead-end workaround.
+		// for the hotfix we cannot introduce a new object version - which is what we did for 5.8.
+		{
+			const int64 ArchivePositionCurrent = Ar.Tell();
+			Ar.Seek(ArchivePositionBeforeData);
+			int32 OffsetForDataAsInt32 = 0;
+			Ar << OffsetForDataAsInt32;
+			OffsetForData = OffsetForDataAsInt32;
+			Ar.Seek(ArchivePositionCurrent);
+		}
+		
 		bool bSkipData = true;
 		
-		if(const FProperty* LeafProperty = GetLeafProperty())
+		if (Properties.Num() == NumProperties)
 		{
-			int32 PropertySize = 0;
-			uint32 PropertyHash = 0;
-
-			if(Ar.CustomVer(FControlRigObjectVersion::GUID) >= FControlRigObjectVersion::OverridesStoreTOCDataForProperties)
-			{
-				Ar << PropertySize;
-				Ar << PropertyHash;
-			}
-			else
+			if(const FProperty* LeafProperty = GetLeafProperty())
 			{
-				PropertySize = LeafProperty->GetSize();
-				PropertyHash = GetPropertyHash();
-			}
+				int32 PropertySize = 0;
+				uint32 PropertyHash = 0;
 
-			if (LeafProperty->GetSize() == PropertySize && GetPropertyHash() == PropertyHash)
-			{
-				void* DataPtr = AllocateDataIfRequired();
-				LeafProperty->SerializeItem(FStructuredArchiveFromArchive(Ar).GetSlot(), DataPtr);
+				if(Ar.CustomVer(FControlRigObjectVersion::GUID) >= FControlRigObjectVersion::OverridesStoreTOCDataForProperties)
+				{
+					Ar << PropertySize;
+					Ar << PropertyHash;
+				}
+				else
+				{
+					PropertySize = LeafProperty->GetSize();
+					PropertyHash = GetPropertyHash();
+				}
 
-				FString ValueAsString;
-				LeafProperty->ExportTextItem_Direct(ValueAsString, DataPtr, DataPtr, nullptr, PPF_None, nullptr);
-				CachedStringValue = ValueAsString;
+				if (LeafProperty->GetSize() == PropertySize && GetPropertyHash() == PropertyHash)
+				{
+					void* DataPtr = AllocateDataIfRequired();
+					LeafProperty->SerializeItem(FStructuredArchiveFromArchive(Ar).GetSlot(), DataPtr);
 
-				bSkipData = false;
-			}
-			else
-			{
-				DataPropertyHash = 0;
-				DataArray.Reset();
+					FString ValueAsString;
+					LeafProperty->ExportTextItem_Direct(ValueAsString, DataPtr, DataPtr, nullptr, PPF_None, nullptr);
+					CachedStringValue = ValueAsString;
+
+					bSkipData = false;
+				}
+				else
+				{
+					DataPropertyHash = 0;
+					DataArray.Reset();
+				}
 			}
 		}
-
 		if (bSkipData)
 		{
 			// we can't deserialize the data - skip it
 			Ar.Seek(ArchivePositionBeforeData + OffsetForData);
+			Reset();
 		}
 
 		UpdateHash();
@@ -1475,11 +1496,22 @@ void FControlRigOverrideContainer::SetFromUObject(const UObject* InSubjectPtr)
 
 bool FControlRigOverrideContainer::Serialize(FArchive& Ar)
 {
+	if (!Ar.IsLoading())
+	{
+		Values.RemoveAll([](const FControlRigOverrideValue& Value) -> bool {
+			return !Value.IsValid();
+		});
+	}
+	
 	Ar << bUsesKeyForSubject;
 	Ar << Values;
 	
 	if(Ar.IsLoading())
 	{
+		Values.RemoveAll([](const FControlRigOverrideValue& Value) -> bool {
+			return !Value.IsValid();
+		});
+
 		RebuildLookup();
 		InvalidateCache();
 	}
