Control Rig: Properly access references UObjects to avoid them from being GCed

#rnx

diff --git a/Engine/Plugins/Animation/ControlRig/Source/ControlRig/Private/ControlRigOverride.cpp b/Engine/Plugins/Animation/ControlRig/Source/ControlRig/Private/ControlRigOverride.cpp
--- a/Engine/Plugins/Animation/ControlRig/Source/ControlRig/Private/ControlRigOverride.cpp
+++ b/Engine/Plugins/Animation/ControlRig/Source/ControlRig/Private/ControlRigOverride.cpp
@@ -534,6 +534,15 @@ uint32 GetTypeHash(const FControlRigOverrideValue& InOverride)
 
 void FControlRigOverrideValue::AddStructReferencedObjects(FReferenceCollector& Collector) const
 {
+	if(const FProperty* RootProperty = GetRootProperty())
+	{
+		if (UStruct* OwnerStruct = RootProperty->GetOwnerStruct())
+		{
+			TObjectPtr<UObject> OwnerStructObj = OwnerStruct;
+			Collector.AddReferencedObject(OwnerStructObj);
+		}
+	}
+
 	if (const FProperty* LeafProperty = GetLeafProperty())
 	{
 		if (const FArrayProperty* ArrayProperty = CastField<FArrayProperty>(LeafProperty))
@@ -545,10 +554,12 @@ void FControlRigOverrideValue::AddStructReferencedObjects(FReferenceCollector& C
 					FScriptArrayHelper ArrayHelper(ArrayProperty, DataPtr);
 					for (int32 Index = 0; Index < ArrayHelper.Num(); Index++)
 					{
-						void* ElementData = ArrayHelper.GetRawPtr(Index);
-						if (TObjectPtr<UObject>* ObjPtr = (TObjectPtr<UObject>*)(ElementData))
+						if (const void* ElementData = ArrayHelper.GetRawPtr(Index))
 						{
-							Collector.AddReferencedObject(*ObjPtr);
+							if (TObjectPtr<UObject> ObjPtr = ObjProperty->GetObjectPtrPropertyValue(ElementData))
+							{
+								Collector.AddReferencedObject(ObjPtr);
+							}
 						}
 					}
 				}
@@ -573,9 +584,9 @@ void FControlRigOverrideValue::AddStructReferencedObjects(FReferenceCollector& C
 			{
 				if (const void* DataPtr = GetData())
 				{
-					if (TObjectPtr<const UObject>* ObjPtr = ObjProperty->ContainerPtrToValuePtr<TObjectPtr<const UObject>>(const_cast<void*>(DataPtr)))
+					if (TObjectPtr<UObject> ObjPtr = ObjProperty->GetObjectPtrPropertyValue(DataPtr))
 					{
-						Collector.AddReferencedObject(*ObjPtr);
+						Collector.AddReferencedObject(ObjPtr);
 					}
 				}
 			}
