Control Rig: Remove the need for a game thread callback to refresh the direct manipulator combo.

#rnx
#Synced-CL 48443308

diff --git a/Engine/Plugins/Animation/ControlRig/Source/ControlRigEditor/Private/Editor/ControlRigEditor.cpp b/Engine/Plugins/Animation/ControlRig/Source/ControlRigEditor/Private/Editor/ControlRigEditor.cpp
--- a/Engine/Plugins/Animation/ControlRig/Source/ControlRigEditor/Private/Editor/ControlRigEditor.cpp
+++ b/Engine/Plugins/Animation/ControlRig/Source/ControlRigEditor/Private/Editor/ControlRigEditor.cpp
@@ -4159,28 +4159,19 @@ const TArray<FRigDirectManipulationTarget> FControlRigBaseEditor::GetDirectManip
 
 const TArray<TSharedPtr<FString>>& FControlRigBaseEditor::GetDirectManipulationTargetTextList() const
 {
-	if(DirectManipulationTextList.IsEmpty())
-	{
-		const TArray<FRigDirectManipulationTarget> Targets = GetDirectManipulationTargets();
-		for(const FRigDirectManipulationTarget& Target : Targets)
-		{
-			DirectManipulationTextList.Emplace(new FString(Target.Name));
-		}
-	}
 	return DirectManipulationTextList;
 }
 
 void FControlRigBaseEditor::RefreshDirectManipulationTextList()
 {
 	DirectManipulationTextList.Reset();
-	(void)GetDirectManipulationTargetTextList();
-	if(DirectManipulationCombo.IsValid())
+	for(const FRigDirectManipulationTarget& Target : GetDirectManipulationTargets())
 	{
-		ExecuteOnGameThread(UE_SOURCE_LOCATION, [this]()
-		{
-			DirectManipulationCombo->RefreshOptions();
-		});
+		DirectManipulationTextList.Emplace(new FString(Target.Name));
 	}
+	
+	DirectManipulationCombo->ClearSelection();
+	DirectManipulationCombo->RefreshOptions();
 }
 
 EVisibility FControlRigBaseEditor::GetPreviewingNodeVisibility() const
diff --git a/Engine/Plugins/Animation/ControlRig/Source/ControlRigEditor/Private/Editor/ControlRigEditor.h b/Engine/Plugins/Animation/ControlRig/Source/ControlRigEditor/Private/Editor/ControlRigEditor.h
--- a/Engine/Plugins/Animation/ControlRig/Source/ControlRigEditor/Private/Editor/ControlRigEditor.h
+++ b/Engine/Plugins/Animation/ControlRig/Source/ControlRigEditor/Private/Editor/ControlRigEditor.h
@@ -506,7 +506,7 @@ protected:
 	TStrongObjectPtr<UControlRigReplay> ReplayStrongPtr;
 
 	TWeakObjectPtr<const URigVMUnitNode> DirectManipulationSubject;
-	mutable TArray<TSharedPtr<FString>> DirectManipulationTextList;
+	TArray<TSharedPtr<FString>> DirectManipulationTextList;
 	TSharedPtr<SComboBox<TSharedPtr<FString>>> DirectManipulationCombo;
 	bool bRefreshDirectionManipulationTargetsRequired;
 	FSimpleMulticastDelegate RequestNavigateToConnectorWarningDelegate;
