Anim in Engine: Selecting Location Controls in Anim Details now filters channel curves correctly

* Control Rig now names Sequencer Location Channel Proxies "Location" instead of "Position" in all cases
* Inherently overcome previous typo "Positon"

diff --git a/Engine/Plugins/Animation/ControlRig/Source/ControlRig/Private/Sequencer/MovieSceneControlRigParameterSection.cpp b/Engine/Plugins/Animation/ControlRig/Source/ControlRig/Private/Sequencer/MovieSceneControlRigParameterSection.cpp
--- a/Engine/Plugins/Animation/ControlRig/Source/ControlRig/Private/Sequencer/MovieSceneControlRigParameterSection.cpp
+++ b/Engine/Plugins/Animation/ControlRig/Source/ControlRig/Private/Sequencer/MovieSceneControlRigParameterSection.cpp
@@ -204,7 +204,7 @@ struct FParameterVectorChannelEditorData
 		static constexpr FStringView NameSuffixes[] = { TEXT(".X"), TEXT(".Y"), TEXT(".Z"), TEXT(".W") };
 		static_assert(UE_ARRAY_COUNT(NameSuffixes) == UE_ARRAY_COUNT(MetaData));
 		
-		static constexpr FStringView PositionNameSuffixes[] = { TEXT(".Positon.X"), TEXT(".Position.Y"), TEXT(".Position.Z"), nullptr };
+		static constexpr FStringView PositionNameSuffixes[] = { TEXT(".Location.X"), TEXT(".Location.Y"), TEXT(".Location.Z"), nullptr };
 		static_assert(UE_ARRAY_COUNT(PositionNameSuffixes) == UE_ARRAY_COUNT(MetaData));
 		
 		static constexpr FStringView RotationNameSuffixes[] = { TEXT(".Rotation.Roll"), TEXT(".Rotation.Pitch"), TEXT(".Rotation.Yaw"), nullptr };
