Control Rig: Avoid crash when patching modular rig model
#rnx

diff --git a/Engine/Plugins/Animation/ControlRig/Source/ControlRig/Private/ModularRigModel.cpp b/Engine/Plugins/Animation/ControlRig/Source/ControlRig/Private/ModularRigModel.cpp
--- a/Engine/Plugins/Animation/ControlRig/Source/ControlRig/Private/ModularRigModel.cpp
+++ b/Engine/Plugins/Animation/ControlRig/Source/ControlRig/Private/ModularRigModel.cpp
@@ -224,8 +224,11 @@ void FModularRigModel::PatchModelsOnLoad()
 			FRigModuleReference& Module = Modules[ModuleIndex];
 			if(!Module.ParentPath_DEPRECATED.IsEmpty())
 			{
-				Module.ParentModuleName = PreviousModulePaths.FindChecked(Module.ParentPath_DEPRECATED);
-				Module.ParentPath_DEPRECATED.Reset();
+				if (FName* ParentPath = PreviousModulePaths.Find(Module.ParentPath_DEPRECATED))
+				{
+					Module.ParentModuleName = *ParentPath;
+					Module.ParentPath_DEPRECATED.Reset();
+				}
 			}
 		}
 
