Control Rig: Fix a possible crash in FControlRigBaseEditor::RefreshDirectManipulationTextList

diff --git a/Engine/Plugins/Animation/ControlRig/Source/ControlRigEditor/Private/Editor/ControlRigEditor.cpp b/Engine/Plugins/Animation/ControlRig/Source/ControlRigEditor/Private/Editor/ControlRigEditor.cpp
--- a/Engine/Plugins/Animation/ControlRig/Source/ControlRigEditor/Private/Editor/ControlRigEditor.cpp
+++ b/Engine/Plugins/Animation/ControlRig/Source/ControlRigEditor/Private/Editor/ControlRigEditor.cpp
@@ -4164,14 +4164,17 @@ const TArray<TSharedPtr<FString>>& FControlRigBaseEditor::GetDirectManipulationT
 
 void FControlRigBaseEditor::RefreshDirectManipulationTextList()
 {
-	DirectManipulationTextList.Reset();
-	for(const FRigDirectManipulationTarget& Target : GetDirectManipulationTargets())
+	if (DirectManipulationCombo.IsValid())
 	{
-		DirectManipulationTextList.Emplace(new FString(Target.Name));
+		DirectManipulationTextList.Reset();
+		for (const FRigDirectManipulationTarget& Target : GetDirectManipulationTargets())
+		{
+			DirectManipulationTextList.Emplace(new FString(Target.Name));
+		}
+
+		DirectManipulationCombo->ClearSelection();
+		DirectManipulationCombo->RefreshOptions();
 	}
-	
-	DirectManipulationCombo->ClearSelection();
-	DirectManipulationCombo->RefreshOptions();
 }
 
 EVisibility FControlRigBaseEditor::GetPreviewingNodeVisibility() const
