Sequencer - BakeToControlRig: prevent baking if the playback context or the skeletal mesh actor is invalid

#rnx

diff --git a/Engine/Plugins/Animation/ControlRig/Source/ControlRigEditor/Private/ControlRigEditorModule.cpp b/Engine/Plugins/Animation/ControlRig/Source/ControlRigEditor/Private/ControlRigEditorModule.cpp
--- a/Engine/Plugins/Animation/ControlRig/Source/ControlRigEditor/Private/ControlRigEditorModule.cpp
+++ b/Engine/Plugins/Animation/ControlRig/Source/ControlRigEditor/Private/ControlRigEditorModule.cpp
@@ -108,6 +108,7 @@
 #endif
 #include "Editor/ControlRigNewEditor.h"
 #include "Graph/ControlRigGraphSchema.h"
+#include "Misc/LevelSequencePlaybackContext.h"
 
 #if WITH_RIGVMLEGACYEDITOR
 #include "SKismetInspector.h"
@@ -1091,13 +1092,54 @@ void FControlRigEditorModule::EditWithFKControlRig(UAnimSequence* AnimSequence,
 	BakeToControlRig(UFKControlRig::StaticClass(),AnimSequence, SkelMesh, InSkeleton);
 }
 
+namespace UE::ControlRigEditor::Private
+{
+
+void LogBakeToControlRigError(const FString& InError)
+{
+	const FText Notification = FText::Format(
+		NSLOCTEXT("ControlRigEditor", "BakeToControlRigError",
+		"BakeToControlRig: invalid baking context: {0}. Aborting baking."), FText::FromString(InError));
+	
+	// output log notification
+	UE_LOG(LogControlRigEditor, Error, TEXT("%s"), *Notification.ToString());
+
+	// editor notification
+	FNotificationInfo Info(Notification);
+	Info.Image = FAppStyle::GetBrush(TEXT("MessageLog.Error"));
+	Info.ExpireDuration = 5.0f;
+
+	FSlateNotificationManager::Get().AddNotification(Info);
+}
+	
+}
 
 void FControlRigEditorModule::BakeToControlRig(UClass* ControlRigClass, UAnimSequence* AnimSequence, USkeletalMesh* SkelMesh, USkeleton* InSkeleton)
 {
+	if (!FSlateApplication::IsInitialized())
+	{
+		return;
+	}
+	
 	FSlateApplication::Get().DismissAllMenus();
 
 	UWorld* World = GCurrentLevelEditingViewportClient ? GCurrentLevelEditingViewportClient->GetWorld() : nullptr;
 
+	{
+		FLevelSequencePlaybackContext LevelSequencePlaybackContext(nullptr);
+		const UObject* PlaybackContext = LevelSequencePlaybackContext.GetPlaybackContext();
+		UWorld* WorldPlaybackContext = PlaybackContext ? PlaybackContext->GetWorld() : nullptr;
+		if (!WorldPlaybackContext || WorldPlaybackContext->IsGameWorld() || World != WorldPlaybackContext)
+		{
+			const FString ContextError =
+				!WorldPlaybackContext ? TEXT("null playback world") :
+				WorldPlaybackContext->IsGameWorld() ? TEXT("cannot bake in a game world") :
+				World ? TEXT("the playback context is different than the current editor world") : TEXT("null editor world");
+			UE::ControlRigEditor::Private::LogBakeToControlRigError(ContextError);
+			return;
+		}
+	}
+
 	if (World)
 	{
 		FControlRigEditorModule::UnLinkLevelSequence(AnimSequence);
@@ -1149,8 +1191,14 @@ void FControlRigEditorModule::BakeToControlRig(UClass* ControlRigClass, UAnimSeq
 			MeshActor->RegisterAllComponents();
 			TArray<TWeakObjectPtr<AActor> > ActorsToAdd;
 			ActorsToAdd.Add(MeshActor);
-			TArray<FGuid> ActorTracks = SequencerPtr->AddActors(ActorsToAdd, false);
-			FGuid ActorTrackGuid = ActorTracks[0];
+			const TArray<FGuid> ActorTracks = SequencerPtr->AddActors(ActorsToAdd, false);
+
+			FGuid ActorTrackGuid = ActorTracks.IsEmpty() ? FGuid() : ActorTracks[0];
+			if (!ActorTrackGuid.IsValid())
+			{
+				World->EditorDestroyActor(MeshActor, true);
+				return;
+			}
 
 			// By default, convert this to a spawnable and delete the existing actor. If for some reason, 
 			// the spawnable couldn't be generated, use the existing actor as a possessable (this could 
@@ -1166,14 +1214,22 @@ void FControlRigEditorModule::BakeToControlRig(UClass* ControlRigClass, UAnimSeq
 				{
 					GCurrentLevelEditingViewportClient->GetWorld()->EditorDestroyActor(MeshActor, true);
 					MeshActor = Cast<ASkeletalMeshActor>(SpawnedMesh);
-					if (SkelMesh)
+					if (MeshActor)
 					{
-						MeshActor->GetSkeletalMeshComponent()->SetSkeletalMesh(SkelMesh);
+						if (SkelMesh)
+						{
+							MeshActor->GetSkeletalMeshComponent()->SetSkeletalMesh(SkelMesh);
+						}
+						MeshActor->RegisterAllComponents();
 					}
-					MeshActor->RegisterAllComponents();
 				}
 			}
 
+			if (!ensureMsgf(MeshActor, TEXT("BakeToControlRig: The skeletal mesh actor on which to bake the animations is null. Aborting baking.")))
+			{
+				return;
+			}
+
 			//Delete binding from default animating rig
 			//if we have skel mesh component binding we can just delete that
 			FGuid CompGuid = SequencerPtr->FindObjectId(*(MeshActor->GetSkeletalMeshComponent()), SequencerPtr->GetFocusedTemplateID());
