Constraints & Layered ControlRig: make sure to mark the skeletal mesh component for evaluation when updating the constraints
the nofication system that marks the bound object for evaluation is not triggered before the constraint evaluation as it's only fired when the rig is evaluated (which happens after)
marking the bound object for evaluation will force a rig evaluation via the AnimationEvaluation system

#rnx

diff --git a/Engine/Plugins/Animation/ControlRig/Source/ControlRig/Private/Sequencer/MovieSceneControlRigParameterTemplate.cpp b/Engine/Plugins/Animation/ControlRig/Source/ControlRig/Private/Sequencer/MovieSceneControlRigParameterTemplate.cpp
--- a/Engine/Plugins/Animation/ControlRig/Source/ControlRig/Private/Sequencer/MovieSceneControlRigParameterTemplate.cpp
+++ b/Engine/Plugins/Animation/ControlRig/Source/ControlRig/Private/Sequencer/MovieSceneControlRigParameterTemplate.cpp
@@ -1461,6 +1461,18 @@ struct FControlRigParameterExecutionToken : IMovieSceneExecutionToken
 				{
 					UWorld* BoundObjectWorld = BoundObject->GetWorld();
 					
+					USkeletalMeshComponent* SkeletalMeshComponent = Cast<USkeletalMeshComponent>(FControlRigObjectBinding::GetBindableObject(BoundObject));
+					auto MarkRigForEvaluationIfAdditive = [ControlRig, SkeletalMeshComponent]()
+					{
+						if (TransformableHandleUtils::SkipTicking())
+						{
+							if (SkeletalMeshComponent && ControlRig && ControlRig->IsAdditive())
+							{
+								TransformableHandleUtils::MarkComponentForEvaluation(SkeletalMeshComponent);
+							}
+						}
+					};
+										
 					TSharedRef<UE::MovieScene::FSharedPlaybackState> SharedPlaybackState = Player.GetSharedPlaybackState();
 					for (FConstraintAndActiveValue& ConstraintValue : ConstraintsValues)
 					{
@@ -1479,6 +1491,8 @@ struct FControlRigParameterExecutionToken : IMovieSceneExecutionToken
 
 							if (TransformConstraint && ConstraintValue.Value && TransformableHandleUtils::SkipTicking())
 							{
+								MarkRigForEvaluationIfAdditive();
+								
 								if (UTransformableComponentHandle* ParentHandle = Cast<UTransformableComponentHandle>(TransformConstraint->ParentTRSHandle))
 								{
 									constexpr bool bRecursive = true;
diff --git a/Engine/Plugins/Animation/ControlRig/Source/ControlRig/Private/Sequencer/MovieSceneControlRigSystem.cpp b/Engine/Plugins/Animation/ControlRig/Source/ControlRig/Private/Sequencer/MovieSceneControlRigSystem.cpp
--- a/Engine/Plugins/Animation/ControlRig/Source/ControlRig/Private/Sequencer/MovieSceneControlRigSystem.cpp
+++ b/Engine/Plugins/Animation/ControlRig/Source/ControlRig/Private/Sequencer/MovieSceneControlRigSystem.cpp
@@ -1148,6 +1148,18 @@ struct FEvaluateBaseControlRigs
 
 		FMovieSceneSequenceID SequenceID = OptSequenceID ? *OptSequenceID : MovieSceneSequenceID::Root;
 
+		USkeletalMeshComponent* SkeletalMeshComponent = Cast<USkeletalMeshComponent>(FControlRigObjectBinding::GetBindableObject(BoundObject));
+		auto MarkRigForEvaluationIfAdditive = [ControlRig, SkeletalMeshComponent]()
+		{
+			if (TransformableHandleUtils::SkipTicking())
+			{
+				if (SkeletalMeshComponent && ControlRig && ControlRig->IsAdditive())
+				{
+					TransformableHandleUtils::MarkComponentForEvaluation(SkeletalMeshComponent);
+				}
+			}
+		};
+		
 		for (const FConstraintAndActiveChannel& ConstraintAndActiveChannel : InOutBaseEvalData.Section->GetConstraintsChannels())
 		{
 			bool Value = false;
@@ -1167,6 +1179,8 @@ struct FEvaluateBaseControlRigs
 
 				if (TransformConstraint && Value && TransformableHandleUtils::SkipTicking())
 				{
+					MarkRigForEvaluationIfAdditive();
+					
 					if (UTransformableComponentHandle* ParentHandle = Cast<UTransformableComponentHandle>(TransformConstraint->ParentTRSHandle))
 					{
 						constexpr bool bRecursive = true;
