Control Rig: Fix pyton context menu extenders not working
#rnx

diff --git a/Engine/Plugins/Animation/ControlRig/Source/ControlRigDeveloper/Public/ControlRigAsset.h b/Engine/Plugins/Animation/ControlRig/Source/ControlRigDeveloper/Public/ControlRigAsset.h
--- a/Engine/Plugins/Animation/ControlRig/Source/ControlRigDeveloper/Public/ControlRigAsset.h
+++ b/Engine/Plugins/Animation/ControlRig/Source/ControlRigDeveloper/Public/ControlRigAsset.h
@@ -35,7 +35,7 @@ class UControlRigGraph;
 struct FEndLoadPackageContext;
 
 
-UINTERFACE(BlueprintType)
+UINTERFACE(NotBlueprintable)
 class UE_API UControlRigAssetInterface : public UInterface
 {
 	GENERATED_BODY()
@@ -51,6 +51,8 @@ public:
 	
 	virtual FRigVMAssetInterfacePtr GetRigVMAssetInterface() = 0;
 	virtual const FRigVMAssetInterfacePtr GetRigVMAssetInterface() const { return const_cast<IControlRigAssetInterface*>(this)->GetRigVMAssetInterface(); }
+
+	UFUNCTION(BlueprintCallable, Category = ControlRigAsset)
 	virtual URigHierarchy* GetHierarchy() = 0;
 	virtual URigHierarchy* GetHierarchy() const { return const_cast<IControlRigAssetInterface*>(this)->GetHierarchy(); }
 	UE_API virtual FRigVMClient* GetRigVMClient();
@@ -243,7 +245,8 @@ public:
 
 public:
 
-	UE_API virtual URigHierarchyController* GetHierarchyController() { return GetHierarchy()->GetController(true); }
+	UFUNCTION(BlueprintCallable, Category = ControlRigAsset)
+	virtual URigHierarchyController* GetHierarchyController() { return GetHierarchy()->GetController(true); }
 
 	virtual FModularRigModel& GetModularRigModel() = 0;
 	const FModularRigModel& GetModularRigModel() const { return const_cast<IControlRigAssetInterface*>(this)->GetModularRigModel(); }
diff --git a/Engine/Plugins/Animation/ControlRig/Source/ControlRigEditor/Private/Editor/ControlRigContextMenuContext.h b/Engine/Plugins/Animation/ControlRig/Source/ControlRigEditor/Private/Editor/ControlRigContextMenuContext.h
--- a/Engine/Plugins/Animation/ControlRig/Source/ControlRigEditor/Private/Editor/ControlRigContextMenuContext.h
+++ b/Engine/Plugins/Animation/ControlRig/Source/ControlRigEditor/Private/Editor/ControlRigContextMenuContext.h
@@ -4,6 +4,7 @@
 
 #include "CoreMinimal.h"
 #include "ControlRigAsset.h"
+#include "ControlRigBlueprintLegacy.h"
 #include "EdGraph/EdGraph.h"
 #include "Rigs/RigHierarchyDefines.h"
 
@@ -121,7 +122,7 @@ public:
 
 	/** Get the control rig blueprint that we are editing */
 	UFUNCTION(BlueprintCallable, Category = ControlRigEditorExtensions, meta=(DeprecatedFunction, DeprecationMessage="Function has been deprecated, please rely on GetControlRigAssetInterface instead."))
-    UControlRigBlueprint* GetControlRigBlueprint() const { return nullptr; }
+    UControlRigBlueprint* GetControlRigBlueprint() const { return Cast<UControlRigBlueprint>(GetControlRigAssetInterface().GetObject()); }
 
 	UFUNCTION(BlueprintCallable, Category = ControlRigEditorExtensions)
 	TScriptInterface<IControlRigAssetInterface> GetControlRigAssetInterface() const;
