Control Rig: Fix Selected Control Rig Control sequencer filter not working when Force Short name mode is active
#rnx

diff --git a/Engine/Plugins/Animation/ControlRig/Source/ControlRigEditor/Private/Sequencer/ControlRigSequencerFilter.cpp b/Engine/Plugins/Animation/ControlRig/Source/ControlRigEditor/Private/Sequencer/ControlRigSequencerFilter.cpp
--- a/Engine/Plugins/Animation/ControlRig/Source/ControlRigEditor/Private/Sequencer/ControlRigSequencerFilter.cpp
+++ b/Engine/Plugins/Animation/ControlRig/Source/ControlRigEditor/Private/Sequencer/ControlRigSequencerFilter.cpp
@@ -10,6 +10,7 @@
 #include "MVVM/Extensions/IOutlinerExtension.h"
 #include "MVVM/ViewModels/ChannelModel.h"
 #include "Sequencer/MovieSceneControlRigParameterTrack.h"
+#include "Settings/ControlRigSettings.h"
 #include "Styling/AppStyle.h"
 #include "Styling/SlateIconFinder.h"
 
@@ -178,11 +179,12 @@ public:
 		}
 
 		const FString ControlTrackLabel = OutlinerExtension->GetLabel().ToString();
+		const EElementNameDisplayMode ElementNameDisplayMode = UControlRigEditorSettings::Get()->ElementNameDisplayMode;
 
 		const TArray<const FRigBaseElement*> BaseRigElements = ControlRigHierarchy->GetSelectedElements(ERigElementType::Control);
 		for (const FRigBaseElement* const BaseRigElement : BaseRigElements)
 		{
-			const FString ElementDisplayName = ControlRigHierarchy->GetDisplayNameForUI(BaseRigElement).ToString();
+			const FString ElementDisplayName = ControlRigHierarchy->GetDisplayNameForUI(BaseRigElement, ElementNameDisplayMode).ToString();
 			if (ElementDisplayName.Equals(ControlTrackLabel))
 			{
 				return true;
@@ -197,7 +199,7 @@ public:
 					{
 						if (const FRigBaseElement* const DrivenKeyElement = ControlRigHierarchy->Find(DrivenKey))
 						{
-							const FString DrivenKeyElementDisplayName = ControlRigHierarchy->GetDisplayNameForUI(DrivenKeyElement).ToString();
+							const FString DrivenKeyElementDisplayName = ControlRigHierarchy->GetDisplayNameForUI(DrivenKeyElement, ElementNameDisplayMode).ToString();
 							if (DrivenKeyElementDisplayName.Equals(ControlTrackLabel))
 							{
 								return true;
