Control Rig: Fix binding to objects on old control rig instead of newly created one when the class changes
#rnx

diff --git a/Engine/Plugins/Animation/ControlRig/Source/ControlRig/Private/AnimNode_ControlRig.cpp b/Engine/Plugins/Animation/ControlRig/Source/ControlRig/Private/AnimNode_ControlRig.cpp
--- a/Engine/Plugins/Animation/ControlRig/Source/ControlRig/Private/AnimNode_ControlRig.cpp
+++ b/Engine/Plugins/Animation/ControlRig/Source/ControlRig/Private/AnimNode_ControlRig.cpp
@@ -343,8 +343,8 @@ bool FAnimNode_ControlRig::UpdateControlRigIfNeeded(const UAnimInstance* InAnimI
 							// garbage collected in the future. 
 							(void)NewControlRig->AtomicallyClearInternalFlags(EInternalObjectFlags::Async);
 														
-							ControlRig->SetObjectBinding(MakeShared<FControlRigObjectBinding>());
-							ControlRig->GetObjectBinding()->BindToObject(InAnimInstance->GetOwningComponent());
+							NewControlRig->SetObjectBinding(MakeShared<FControlRigObjectBinding>());
+							NewControlRig->GetObjectBinding()->BindToObject(InAnimInstance->GetOwningComponent());
 							NewControlRig->Initialize(true);
 							NewControlRig->RequestInit();
 						}
