Constraints: ability to remove constraints from sequencer
- deleting a constraint channel will remove the related constraint
- deleting a section will also remove the constraints related to its constraint channels

#rnx

diff --git a/Engine/Plugins/Animation/ControlRig/Source/ControlRig/Private/Sequencer/MovieSceneControlRigParameterTrack.cpp b/Engine/Plugins/Animation/ControlRig/Source/ControlRig/Private/Sequencer/MovieSceneControlRigParameterTrack.cpp
--- a/Engine/Plugins/Animation/ControlRig/Source/ControlRig/Private/Sequencer/MovieSceneControlRigParameterTrack.cpp
+++ b/Engine/Plugins/Animation/ControlRig/Source/ControlRig/Private/Sequencer/MovieSceneControlRigParameterTrack.cpp
@@ -245,8 +245,27 @@ void UMovieSceneControlRigParameterTrack::HandleOnConstraintAdded(
 	OnConstraintChannelAdded.Broadcast(InSection, InChannel);
 }
 
+void UMovieSceneControlRigParameterTrack::RemoveAllConstraints() const
+{
+	// no need to compensate, we just want to get rig of all the animation data
+	constexpr bool bDoNotCompensate = true;
+	
+	for (UMovieSceneSection* Section : Sections)
+	{
+		if (UMovieSceneControlRigParameterSection* RigSection = Cast<UMovieSceneControlRigParameterSection>(Section))
+		{
+			const UControlRig* Rig = RigSection->GetControlRig();
+			if (UWorld* World = Rig ? Rig->GetWorld() : nullptr)
+			{
+				RigSection->RemoveAllConstraints(World, bDoNotCompensate);
+			}
+		}
+	}
+}
+
 void UMovieSceneControlRigParameterTrack::RemoveAllAnimationData()
 {
+	RemoveAllConstraints();
 	Sections.Empty();
 	SectionToKey = nullptr;
 }
@@ -312,6 +331,10 @@ const TArray<UMovieSceneSection*>& UMovieSceneControlRigParameterTrack::GetAllSe
 	return Sections;
 }
 
+void UMovieSceneControlRigParameterTrack::OnRemovedFromMovieSceneImpl()
+{
+	RemoveAllConstraints();
+}
 
 #if WITH_EDITORONLY_DATA
 FText UMovieSceneControlRigParameterTrack::GetDefaultDisplayName() const
diff --git a/Engine/Plugins/Animation/ControlRig/Source/ControlRig/Public/Sequencer/MovieSceneControlRigParameterSection.h b/Engine/Plugins/Animation/ControlRig/Source/ControlRig/Public/Sequencer/MovieSceneControlRigParameterSection.h
--- a/Engine/Plugins/Animation/ControlRig/Source/ControlRig/Public/Sequencer/MovieSceneControlRigParameterSection.h
+++ b/Engine/Plugins/Animation/ControlRig/Source/ControlRig/Public/Sequencer/MovieSceneControlRigParameterSection.h
@@ -374,10 +374,15 @@ public:
 	UE_API virtual void RemoveConstraintChannel(const UTickableConstraint* InConstraint) override;
 
 	/*
-	*  Get The channels
+	*  Get a reference to the channels
 	*/
 	UE_API virtual TArray<FConstraintAndActiveChannel>& GetConstraintsChannels()  override;
 
+	/*
+	*  Get a const reference to the channels
+	*/
+	UE_API virtual const TArray<FConstraintAndActiveChannel>& GetConstraintsChannels() const override;
+
 	/*
 	*  Replace the constraint with the specified name with the new one
 	*/
@@ -388,9 +393,6 @@ public:
 	*/
 	UE_API virtual void OnConstraintsChanged() override;
 
-	//not override but needed
-	UE_API const TArray<FConstraintAndActiveChannel>& GetConstraintsChannels() const;
-
 #if WITH_EDITOR
 	//Function to save control rig key when recording.
 	UE_API void RecordControlRigKey(FFrameNumber FrameNumber, bool bSetDefault, EMovieSceneKeyInterpolation InInterpMode, bool bOntoSelectedControls = false);
diff --git a/Engine/Plugins/Animation/ControlRig/Source/ControlRig/Public/Sequencer/MovieSceneControlRigParameterTrack.h b/Engine/Plugins/Animation/ControlRig/Source/ControlRig/Public/Sequencer/MovieSceneControlRigParameterTrack.h
--- a/Engine/Plugins/Animation/ControlRig/Source/ControlRig/Public/Sequencer/MovieSceneControlRigParameterTrack.h
+++ b/Engine/Plugins/Animation/ControlRig/Source/ControlRig/Public/Sequencer/MovieSceneControlRigParameterTrack.h
@@ -65,6 +65,7 @@ public:
 	virtual bool IsEmpty() const override;
 	virtual const TArray<UMovieSceneSection*>& GetAllSections() const override;
 	virtual FName GetTrackName() const override { return TrackName; }
+	virtual void OnRemovedFromMovieSceneImpl() override;
 
 	// UObject
 	virtual void PostLoad() override;
@@ -209,6 +210,9 @@ private:
 		IMovieSceneConstrainedSection* InSection,
 		FMovieSceneConstraintChannel* InChannel) const;
 	IMovieSceneConstrainedSection::FConstraintChannelAddedEvent OnConstraintChannelAdded;
+
+	void RemoveAllConstraints() const;
+	
 	void ReconstructControlRig();
 
 	//look for name changes in particularly with modular rigs so we can update the rotation order map
diff --git a/Engine/Plugins/Animation/ControlRig/Source/ControlRigEditor/Private/Editor/Constraints/SConstraintsWidget.cpp b/Engine/Plugins/Animation/ControlRig/Source/ControlRigEditor/Private/Editor/Constraints/SConstraintsWidget.cpp
--- a/Engine/Plugins/Animation/ControlRig/Source/ControlRigEditor/Private/Editor/Constraints/SConstraintsWidget.cpp
+++ b/Engine/Plugins/Animation/ControlRig/Source/ControlRigEditor/Private/Editor/Constraints/SConstraintsWidget.cpp
@@ -1045,33 +1045,65 @@ void SEditableConstraintItem::Construct(
 
 FBaseConstraintListWidget::~FBaseConstraintListWidget()
 {
-	UnregisterSelectionChanged();
+	UnregisterNotifications();
 }
 
-void FBaseConstraintListWidget::RegisterSelectionChanged()
+void FBaseConstraintListWidget::RegisterNotifications()
 {
-	FLevelEditorModule& LevelEditor = FModuleManager::GetModuleChecked<FLevelEditorModule>("LevelEditor");
-	FLevelEditorModule::FActorSelectionChangedEvent& ActorSelectionChangedEvent = LevelEditor.OnActorSelectionChanged();
-
-	// unregister previous one
-	if (OnSelectionChangedHandle.IsValid())
+	// selection changed
 	{
-		ActorSelectionChangedEvent.Remove(OnSelectionChangedHandle);
-		OnSelectionChangedHandle.Reset();
+		FLevelEditorModule& LevelEditor = FModuleManager::GetModuleChecked<FLevelEditorModule>("LevelEditor");
+		FLevelEditorModule::FActorSelectionChangedEvent& ActorSelectionChangedEvent = LevelEditor.OnActorSelectionChanged();
+
+		// unregister previous one
+		if (OnSelectionChangedHandle.IsValid())
+		{
+			ActorSelectionChangedEvent.Remove(OnSelectionChangedHandle);
+			OnSelectionChangedHandle.Reset();
+		}
+
+		// register
+		OnSelectionChangedHandle = ActorSelectionChangedEvent.AddRaw(this, &FBaseConstraintListWidget::OnActorSelectionChanged);
 	}
 
-	// register
-	OnSelectionChangedHandle = ActorSelectionChangedEvent.AddRaw(this, &FBaseConstraintListWidget::OnActorSelectionChanged);
+	// constraints controller updated
+	{
+		FConstraintsManagerController& Controller = FConstraintsManagerController::Get(UE::Private::GetCurrentWorld());
+		
+		if (ConstraintsNotificationHandle.IsValid())
+		{
+			Controller.GetNotifyDelegate().Remove(ConstraintsNotificationHandle);
+			ConstraintsNotificationHandle.Reset();
+		}
+		
+		ConstraintsNotificationHandle =
+			Controller.GetNotifyDelegate().AddLambda([this](EConstraintsManagerNotifyType InNotifyType, UObject* InObject)
+		{
+			if (InNotifyType == EConstraintsManagerNotifyType::GraphUpdated)
+			{
+				InvalidateConstraintList();
+			}
+		});
+	}
 }
 
-void FBaseConstraintListWidget::UnregisterSelectionChanged()
+void FBaseConstraintListWidget::UnregisterNotifications()
 {
 	if (OnSelectionChangedHandle.IsValid())
 	{
-		FLevelEditorModule& LevelEditor = FModuleManager::GetModuleChecked<FLevelEditorModule>("LevelEditor");
-		LevelEditor.OnActorSelectionChanged().Remove(OnSelectionChangedHandle);
+		if (FLevelEditorModule* LevelEditor = FModuleManager::GetModulePtr<FLevelEditorModule>("LevelEditor"))
+		{
+			LevelEditor->OnActorSelectionChanged().Remove(OnSelectionChangedHandle);
+		}
 		OnSelectionChangedHandle.Reset();
 	}
+
+	if (ConstraintsNotificationHandle.IsValid())
+	{
+		FConstraintsManagerController& Controller = FConstraintsManagerController::Get(UE::Private::GetCurrentWorld());
+		Controller.GetNotifyDelegate().Remove(ConstraintsNotificationHandle);
+		ConstraintsNotificationHandle.Reset();
+	}
 }
 
 void FBaseConstraintListWidget::PostUndo(bool bSuccess)
@@ -1559,7 +1591,7 @@ void SConstraintsEditionWidget::Construct(const FArguments& InArgs)
 	];
 
 	RefreshConstraintList();
-	RegisterSelectionChanged();
+	RegisterNotifications();
 }
 
 bool SConstraintsEditionWidget::SequencerTimeChanged() 
@@ -2196,7 +2228,7 @@ void SConstraintBakeWidget::Construct(const FArguments& InArgs)
 		]
 	];
 	RefreshConstraintList();
-	RegisterSelectionChanged();
+	RegisterNotifications();
 }
 
 TSharedRef<ITableRow> SConstraintBakeWidget::OnGenerateWidgetForItem(
diff --git a/Engine/Plugins/Animation/ControlRig/Source/ControlRigEditor/Private/Editor/Constraints/SConstraintsWidget.h b/Engine/Plugins/Animation/ControlRig/Source/ControlRigEditor/Private/Editor/Constraints/SConstraintsWidget.h
--- a/Engine/Plugins/Animation/ControlRig/Source/ControlRigEditor/Private/Editor/Constraints/SConstraintsWidget.h
+++ b/Engine/Plugins/Animation/ControlRig/Source/ControlRigEditor/Private/Editor/Constraints/SConstraintsWidget.h
@@ -269,9 +269,8 @@ protected:
 	using ItemSharedPtr = TSharedPtr<FEditableConstraintItem>;
 	using ConstraintItemListView = SListView< ItemSharedPtr >;
 
-
-	UE_API void RegisterSelectionChanged();
-	UE_API void UnregisterSelectionChanged();
+	void RegisterNotifications();
+	void UnregisterNotifications();
 
 	/** List view that shows constraint types*/
 	TSharedPtr< ConstraintItemListView > ListView;
@@ -283,6 +282,7 @@ protected:
 	bool bNeedsRefresh = false;
 
 	FDelegateHandle OnSelectionChangedHandle;
+	FDelegateHandle ConstraintsNotificationHandle;
 
 public:
 
diff --git a/Engine/Plugins/Animation/ControlRig/Source/ControlRigEditor/Private/Sequencer/ControlRigParameterTrackEditor.cpp b/Engine/Plugins/Animation/ControlRig/Source/ControlRigEditor/Private/Sequencer/ControlRigParameterTrackEditor.cpp
--- a/Engine/Plugins/Animation/ControlRig/Source/ControlRigEditor/Private/Sequencer/ControlRigParameterTrackEditor.cpp
+++ b/Engine/Plugins/Animation/ControlRig/Source/ControlRigEditor/Private/Sequencer/ControlRigParameterTrackEditor.cpp
@@ -6563,7 +6563,6 @@ void FControlRigParameterSection::ShowAllControlsChannels()
 	}
 }
 
-//mz todo
 bool FControlRigParameterSection::RequestDeleteCategory(const TArray<FName>& CategoryNamePaths)
 {
 	UMovieSceneControlRigParameterSection* ParameterSection = CastChecked<UMovieSceneControlRigParameterSection>(WeakSection.Get());
@@ -6589,52 +6588,42 @@ bool FControlRigParameterSection::RequestDeleteCategory(const TArray<FName>& Cat
 		}
 	}
 	
-	/*
-	const FScopedTransaction Transaction(LOCTEXT("DeleteTransformCategory", "Delete transform category"));
-
-	if (ParameterSection->TryModify())
-	{
-	FName CategoryName = CategoryNamePaths[CategoryNamePaths.Num() - 1];
-
-	EMovieSceneTransformChannel Channel = ParameterSection->GetTransformMask().GetChannels();
-	EMovieSceneTransformChannel ChannelToRemove = ParameterSection->GetTransformMaskByName(CategoryName).GetChannels();
-
-	Channel = Channel ^ ChannelToRemove;
-
-	ParameterSection->SetTransformMask(Channel);
-
-	SequencerPtr->NotifyMovieSceneDataChanged(EMovieSceneDataChangeType::MovieSceneStructureItemsChanged);
-	return true;
-	}
-	*/
 	return false;
 }
 
 bool FControlRigParameterSection::RequestDeleteKeyArea(const TArray<FName>& KeyAreaNamePaths)
 {
-	UMovieSceneControlRigParameterSection* ParameterSection = CastChecked<UMovieSceneControlRigParameterSection>(WeakSection.Get());
-	TSharedPtr<ISequencer> SequencerPtr = WeakSequencer.Pin();
-
-	/*
-	const FScopedTransaction Transaction(LOCTEXT("DeleteTransformChannel", "Delete transform channel"));
-
-	if (ParameterSection->TryModify())
+	TSharedPtr<ISequencer> Sequencer = WeakSequencer.IsValid() ? WeakSequencer.Pin() : nullptr;
+	UMovieSceneControlRigParameterSection* ControlRigSection = CastChecked<UMovieSceneControlRigParameterSection>(WeakSection.Get());
+	
+	if (!Sequencer || !ControlRigSection || KeyAreaNamePaths.Num() != 2)
 	{
-	// Only delete the last key area path which is the channel. ie. TranslationX as opposed to Translation
-	FName KeyAreaName = KeyAreaNamePaths[KeyAreaNamePaths.Num() - 1];
-
-	EMovieSceneTransformChannel Channel = ParameterSection->GetTransformMask().GetChannels();
-	EMovieSceneTransformChannel ChannelToRemove = ParameterSection->GetTransformMaskByName(KeyAreaName).GetChannels();
-
-	Channel = Channel ^ ChannelToRemove;
+		return false;
+	}
 
-	ParameterSection->SetTransformMask(Channel);
+	// KeyAreaNamePaths[0] is the control's name
+	const FName ConstraintName = KeyAreaNamePaths[1];
+	UTickableConstraint* Constraint = ControlRigSection->FindConstraint(ConstraintName);
+	if (!Constraint)
+	{
+		return false;
+	}
 
-	SequencerPtr->NotifyMovieSceneDataChanged(EMovieSceneDataChangeType::MovieSceneStructureItemsChanged);
-	return true;
+	if (UWorld* World = GEditor ? GEditor->GetEditorWorldContext().World() : nullptr)
+	{
+		const FScopedTransaction Transaction(LOCTEXT("DeleteConstraintChannel", "Delete Constraint Channel"));
+		if (ControlRigSection->TryModify())
+		{
+			FConstraintsManagerController& Controller = FConstraintsManagerController::Get(World);
+			if (Controller.RemoveConstraint(Constraint))
+			{
+				Sequencer->NotifyMovieSceneDataChanged(EMovieSceneDataChangeType::MovieSceneStructureItemsChanged);
+				return true;
+			}
+		}
 	}
-	*/
-	return true;
+	
+	return false;
 }
 
 
diff --git a/Engine/Source/Editor/MovieSceneTools/Private/Sections/TransformPropertySection.cpp b/Engine/Source/Editor/MovieSceneTools/Private/Sections/TransformPropertySection.cpp
--- a/Engine/Source/Editor/MovieSceneTools/Private/Sections/TransformPropertySection.cpp
+++ b/Engine/Source/Editor/MovieSceneTools/Private/Sections/TransformPropertySection.cpp
@@ -4,6 +4,7 @@
 
 #include "Containers/ArrayView.h"
 #include "Delegates/Delegate.h"
+#include "Editor.h"
 #include "Framework/Commands/UIAction.h"
 #include "Framework/Commands/UICommandInfo.h"
 #include "Framework/MultiBox/MultiBoxBuilder.h"
@@ -14,6 +15,7 @@
 #include "Misc/Guid.h"
 #include "MovieSceneSection.h"
 #include "ScopedTransaction.h"
+#include "Constraints/MovieSceneConstraintChannelHelper.h"
 #include "Sections/MovieScene3DTransformSection.h"
 #include "Styling/SlateTypes.h"
 #include "Templates/Casts.h"
@@ -195,11 +197,39 @@ void FTransformSection::BuildSectionContextMenu(FMenuBuilder& MenuBuilder, const
 	MenuBuilder.EndSection();
 }
 
+namespace UE::Private
+{
+	static const FName ConstraintsSectionName(TEXT("Constraints"));
+}
+
 bool FTransformSection::RequestDeleteCategory(const TArray<FName>& CategoryNamePaths)
 {
+	if (CategoryNamePaths.IsEmpty())
+	{
+		return false;
+	}
+	
 	UMovieScene3DTransformSection* TransformSection = CastChecked<UMovieScene3DTransformSection>(WeakSection.Get());
 	TSharedPtr<ISequencer> SequencerPtr = WeakSequencer.Pin();
-		
+
+	if (CategoryNamePaths[0] == UE::Private::ConstraintsSectionName)
+	{
+		constexpr bool bDoNotCompensate = false;
+		if (UWorld* World = GEditor ? GEditor->GetEditorWorldContext().World() : nullptr)
+		{
+			const FScopedTransaction Transaction(LOCTEXT("DeleteConstraintChannels", "Delete Constraint Channels"));
+			if (TransformSection->TryModify())
+			{
+				if (TransformSection->RemoveAllConstraints(World, bDoNotCompensate))
+				{
+					SequencerPtr->NotifyMovieSceneDataChanged(EMovieSceneDataChangeType::MovieSceneStructureItemsChanged);
+				}
+				return true;
+			}
+		}
+		return false;
+	}
+	
 	const FScopedTransaction Transaction( LOCTEXT( "DeleteTransformCategory", "Delete transform category" ) );
 
 	if (TransformSection->TryModify())
@@ -222,9 +252,39 @@ bool FTransformSection::RequestDeleteCategory(const TArray<FName>& CategoryNameP
 
 bool FTransformSection::RequestDeleteKeyArea(const TArray<FName>& KeyAreaNamePaths)
 {
+	if (KeyAreaNamePaths.IsEmpty())
+	{
+		return false;
+	}
+
 	UMovieScene3DTransformSection* TransformSection = CastChecked<UMovieScene3DTransformSection>(WeakSection.Get());
 	TSharedPtr<ISequencer> SequencerPtr = WeakSequencer.Pin();
-
+	
+	if (KeyAreaNamePaths.Num() == 2 && KeyAreaNamePaths[0] == UE::Private::ConstraintsSectionName)
+	{
+		UTickableConstraint* Constraint = TransformSection->FindConstraint(KeyAreaNamePaths[1]);
+		if (!Constraint)
+		{
+			return false;
+		}
+
+		if (UWorld* World = GEditor ? GEditor->GetEditorWorldContext().World() : nullptr)
+		{
+			const FScopedTransaction Transaction(LOCTEXT("DeleteConstraintChannel", "Delete Constraint Channel"));
+			if (TransformSection->TryModify())
+			{
+				FConstraintsManagerController& Controller = FConstraintsManagerController::Get(World);
+				if (Controller.RemoveConstraint(Constraint))
+				{
+					SequencerPtr->NotifyMovieSceneDataChanged(EMovieSceneDataChangeType::MovieSceneStructureItemsChanged);
+				}
+				return true;
+			}
+		}
+		
+		return false;
+	}
+	
 	const FScopedTransaction Transaction( LOCTEXT( "DeleteTransformChannel", "Delete transform channel" ) );
 
 	if (TransformSection->TryModify())
diff --git a/Engine/Source/Runtime/MovieSceneTracks/Private/Sections/MovieScene3DTransformSection.cpp b/Engine/Source/Runtime/MovieSceneTracks/Private/Sections/MovieScene3DTransformSection.cpp
--- a/Engine/Source/Runtime/MovieSceneTracks/Private/Sections/MovieScene3DTransformSection.cpp
+++ b/Engine/Source/Runtime/MovieSceneTracks/Private/Sections/MovieScene3DTransformSection.cpp
@@ -1159,6 +1159,16 @@ TArray<FConstraintAndActiveChannel>& UMovieScene3DTransformSection::GetConstrain
 	return EmptyChannels;
 }
 
+const TArray<FConstraintAndActiveChannel>& UMovieScene3DTransformSection::GetConstraintsChannels() const
+{
+	static TArray< FConstraintAndActiveChannel> EmptyChannels;
+	if (Constraints)
+	{
+		return Constraints->ConstraintsChannels;
+	}
+	return EmptyChannels;
+}
+
 void UMovieScene3DTransformSection::PostLoad()
 {
 	Super::PostLoad();
diff --git a/Engine/Source/Runtime/MovieSceneTracks/Private/Sections/MovieSceneConstrainedSection.cpp b/Engine/Source/Runtime/MovieSceneTracks/Private/Sections/MovieSceneConstrainedSection.cpp
--- a/Engine/Source/Runtime/MovieSceneTracks/Private/Sections/MovieSceneConstrainedSection.cpp
+++ b/Engine/Source/Runtime/MovieSceneTracks/Private/Sections/MovieSceneConstrainedSection.cpp
@@ -15,4 +15,42 @@ void IMovieSceneConstrainedSection::SetDoNoRemoveChannel(bool bInDoNotRemoveChan
 	bDoNotRemoveChannel = bInDoNotRemoveChannel;
 }
 
+UTickableConstraint* IMovieSceneConstrainedSection::FindConstraint(const FName InConstraintName) const
+{
+	const TArray<FConstraintAndActiveChannel>& Constraints = GetConstraintsChannels();
+	const int32 ConstraintIndex = Constraints.IndexOfByPredicate([InConstraintName](const FConstraintAndActiveChannel& ConstraintChannel)
+	{
+		UTickableConstraint* Constraint = ConstraintChannel.GetConstraint().Get();
+		return Constraint && Constraint->GetFName() == InConstraintName;
+	});
+	return ConstraintIndex != INDEX_NONE ? Constraints[ConstraintIndex].GetConstraint() : nullptr;
+}
 
+bool IMovieSceneConstrainedSection::RemoveConstraint(UWorld* InWorld, const FName InConstraintName, const bool bDoNotCompensate) const
+{
+	UTickableConstraint* Constraint = FindConstraint(InConstraintName);
+	if (InWorld && Constraint)
+	{
+		return FConstraintsManagerController::Get(InWorld).RemoveConstraint(Constraint, bDoNotCompensate);
+	}
+	return false;
+}
+
+bool IMovieSceneConstrainedSection::RemoveAllConstraints(UWorld* InWorld, const bool bDoNotCompensate) const
+{
+	bool bRemoved = false;
+
+	if (InWorld)
+	{
+		const TArray<FConstraintAndActiveChannel> ConstraintAndActiveChannels = GetConstraintsChannels();
+		for (const FConstraintAndActiveChannel& ActiveChannel : ConstraintAndActiveChannels)
+		{
+			if (FConstraintsManagerController::Get(InWorld).RemoveConstraint(ActiveChannel.GetConstraint(), bDoNotCompensate))
+			{
+				bRemoved = true;
+			}
+		}
+	}
+	
+	return bRemoved;
+}
diff --git a/Engine/Source/Runtime/MovieSceneTracks/Public/Sections/MovieScene3DTransformSection.h b/Engine/Source/Runtime/MovieSceneTracks/Public/Sections/MovieScene3DTransformSection.h
--- a/Engine/Source/Runtime/MovieSceneTracks/Public/Sections/MovieScene3DTransformSection.h
+++ b/Engine/Source/Runtime/MovieSceneTracks/Public/Sections/MovieScene3DTransformSection.h
@@ -354,9 +354,14 @@ public:
 	virtual void RemoveConstraintChannel(const UTickableConstraint* InConstraint) override;
 
 	/*
-	*  Get The channels by value
+	*  Get a reference to the channels
 	*/
-	virtual TArray<FConstraintAndActiveChannel>& GetConstraintsChannels()  override;
+	virtual TArray<FConstraintAndActiveChannel>& GetConstraintsChannels() override;
+
+	/*
+	*  Get a const reference to the channels
+	*/
+	virtual const TArray<FConstraintAndActiveChannel>& GetConstraintsChannels() const override;
 
 	/*
 	*  Replace the constraint with the specified name with the new one
diff --git a/Engine/Source/Runtime/MovieSceneTracks/Public/Sections/MovieSceneConstrainedSection.h b/Engine/Source/Runtime/MovieSceneTracks/Public/Sections/MovieSceneConstrainedSection.h
--- a/Engine/Source/Runtime/MovieSceneTracks/Public/Sections/MovieSceneConstrainedSection.h
+++ b/Engine/Source/Runtime/MovieSceneTracks/Public/Sections/MovieSceneConstrainedSection.h
@@ -58,10 +58,15 @@ public:
 	virtual void RemoveConstraintChannel(const UTickableConstraint* InConstraint) = 0;
 
 	/*
-	*  Get The channels
+	*  Get a reference to the channels
 	*/
 	virtual TArray<FConstraintAndActiveChannel>& GetConstraintsChannels() = 0;
 
+	/*
+	*  Get a const reference to the channels
+	*/
+	virtual const TArray<FConstraintAndActiveChannel>& GetConstraintsChannels() const = 0;
+
 	/*
 	*  Replace the constraint with the specified name with the new one
 	*/
@@ -87,6 +92,20 @@ public:
 	*/
 	FDelegateHandle OnConstraintRemovedHandle;
 
+	/*
+	* Find the constraint named ConstraintName
+	*/
+	MOVIESCENETRACKS_API UTickableConstraint* FindConstraint(const FName InConstraintName) const;
+
+	/*
+	* Removes the constraint named ConstraintName
+	*/
+	MOVIESCENETRACKS_API bool RemoveConstraint(UWorld* InWorld, const FName InConstraintName, const bool bDoNotCompensate) const;
+	
+	/*
+    * Removes all the constraints stored in this section
+    */
+	MOVIESCENETRACKS_API bool RemoveAllConstraints(UWorld* InWorld, const bool bDoNotCompensate) const;
 
 protected:
 
