Control Rig Edit Mode: do not diplay the gizmo if nothing is bound to the control rig

#rnx

diff --git a/Engine/Plugins/Animation/ControlRig/Source/ControlRigEditor/Private/EditMode/ControlRigEditMode.cpp b/Engine/Plugins/Animation/ControlRig/Source/ControlRigEditor/Private/EditMode/ControlRigEditMode.cpp
--- a/Engine/Plugins/Animation/ControlRig/Source/ControlRigEditor/Private/EditMode/ControlRigEditMode.cpp
+++ b/Engine/Plugins/Animation/ControlRig/Source/ControlRigEditor/Private/EditMode/ControlRigEditMode.cpp
@@ -2348,12 +2348,17 @@ bool FControlRigEditMode::ShouldDrawWidget() const
 	{
 		if (UControlRig* ControlRig = RuntimeRigPtr.Get())
 		{
-			if (AreRigElementSelectedAndMovable(ControlRig))
+			const TSharedPtr<IControlRigObjectBinding>& Binding = ControlRig->GetObjectBinding();
+			if (Binding && Binding->GetBoundObject())
 			{
-				return true;
+				if (AreRigElementSelectedAndMovable(ControlRig))
+				{
+					return true;
+				}
 			}
 		}
 	}
+
 	return FEdMode::ShouldDrawWidget();
 }
 
