Anim Details: iterate through all the control rig tracks find the right one (instead of the first one)

#rnx

diff --git a/Engine/Plugins/Animation/ControlRig/Source/ControlRig/Private/Sequencer/ControlRigSequencerHelpers.cpp b/Engine/Plugins/Animation/ControlRig/Source/ControlRig/Private/Sequencer/ControlRigSequencerHelpers.cpp
--- a/Engine/Plugins/Animation/ControlRig/Source/ControlRig/Private/Sequencer/ControlRigSequencerHelpers.cpp
+++ b/Engine/Plugins/Animation/ControlRig/Source/ControlRig/Private/Sequencer/ControlRigSequencerHelpers.cpp
@@ -169,23 +169,21 @@ TArrayView<FMovieSceneIntegerChannel*> FControlRigSequencerHelpers::GetIntegerCh
 
 UMovieSceneControlRigParameterTrack* FControlRigSequencerHelpers::FindControlRigTrack(UMovieSceneSequence* MovieSceneSequence, const UControlRig* ControlRig)
 {
-	UMovieSceneControlRigParameterTrack* Track = nullptr;
-	if (MovieSceneSequence == nullptr)
+	if (const UMovieScene* MovieScene = MovieSceneSequence ? MovieSceneSequence->GetMovieScene() : nullptr)
 	{
-		return Track;
-	}
-	const UMovieScene* MovieScene = MovieSceneSequence->GetMovieScene();
-	const TArray<FMovieSceneBinding>& Bindings = MovieScene->GetBindings();
-	bool bRecreateCurves = false;
-	TArray<TPair<UControlRig*, FName>> ControlRigPairsToReselect;
-	for (const FMovieSceneBinding& Binding : Bindings)
-	{
-		UMovieSceneControlRigParameterTrack* ControlRigParameterTrack = Cast<UMovieSceneControlRigParameterTrack>(MovieScene->FindTrack(UMovieSceneControlRigParameterTrack::StaticClass(), Binding.GetObjectGuid(), NAME_None));
-		if (ControlRigParameterTrack && ControlRigParameterTrack->GetControlRig() == ControlRig)
+		for (const FMovieSceneBinding& Binding : MovieScene->GetBindings())
 		{
-			Track = ControlRigParameterTrack;
-			break;
+			const TArray<UMovieSceneTrack*> Tracks = MovieScene->FindTracks(UMovieSceneControlRigParameterTrack::StaticClass(), Binding.GetObjectGuid());
+			for (UMovieSceneTrack* Track : Tracks)
+			{
+				UMovieSceneControlRigParameterTrack* ControlRigTrack = Cast<UMovieSceneControlRigParameterTrack>(Track);
+				if (ControlRigTrack && ControlRigTrack->GetControlRig() == ControlRig)
+				{
+					return ControlRigTrack;
+				}
+			}
 		}
 	}
-	return Track;
+	
+	return nullptr;
 }
