Control Rig: Control Rig Hierarchy Editor cleanup first pass

* Move hierarchy editor related source files into their own folder
* One file per logical type, consistent file names
* No change in implementations and logic, no change in type names
* Minor IWYU pass for new files

#rnx

diff --git a/Engine/Plugins/Animation/ControlRig/Source/ControlRigEditor/Private/ControlRigBlueprintEditorLibrary.cpp b/Engine/Plugins/Animation/ControlRig/Source/ControlRigEditor/Private/ControlRigBlueprintEditorLibrary.cpp
--- a/Engine/Plugins/Animation/ControlRig/Source/ControlRigEditor/Private/ControlRigBlueprintEditorLibrary.cpp
+++ b/Engine/Plugins/Animation/ControlRig/Source/ControlRigEditor/Private/ControlRigBlueprintEditorLibrary.cpp
@@ -1,8 +1,8 @@
 // Copyright Epic Games, Inc. All Rights Reserved.
 
 #include "ControlRigBlueprintEditorLibrary.h"
-#include "Editor/SRigHierarchy.h"
-#include "Editor/SModularRigModel.h"
+#include "Editor/Hierarchy/Widgets/SRigHierarchy.h"
+#include "Editor/Hierarchy/Widgets/SModularRigHierarchy.h"
 #include "AssetRegistry/AssetRegistryModule.h"
 #include "RigVMModel/RigVMPin.h"
 
diff --git a/Engine/Plugins/Animation/ControlRig/Source/ControlRigEditor/Private/ControlRigElementDetails.cpp b/Engine/Plugins/Animation/ControlRig/Source/ControlRigEditor/Private/ControlRigElementDetails.cpp
--- a/Engine/Plugins/Animation/ControlRig/Source/ControlRigEditor/Private/ControlRigElementDetails.cpp
+++ b/Engine/Plugins/Animation/ControlRig/Source/ControlRigEditor/Private/ControlRigElementDetails.cpp
@@ -28,6 +28,7 @@
 #include "IStructureDataProvider.h"
 #include "Editor/ControlRigEditor.h"
 #include "Editor/SRigConnectorTargetWidget.h"
+#include "Editor/Hierarchy/Views/SRigHierarchyItem.h"
 #include "ModularRigRuleManager.h"
 #include "Async/TaskGraphInterfaces.h"
 
diff --git a/Engine/Plugins/Animation/ControlRig/Source/ControlRigEditor/Private/ControlRigElementDetails.h b/Engine/Plugins/Animation/ControlRig/Source/ControlRigEditor/Private/ControlRigElementDetails.h
--- a/Engine/Plugins/Animation/ControlRig/Source/ControlRigEditor/Private/ControlRigElementDetails.h
+++ b/Engine/Plugins/Animation/ControlRig/Source/ControlRigEditor/Private/ControlRigElementDetails.h
@@ -28,7 +28,7 @@
 #include "Algo/Transform.h"
 #include "Editor/ControlRigEditor.h"
 #include "ControlRigAsset.h"
-#include "Editor/SRigHierarchyTreeView.h"
+#include "Editor/Hierarchy/Views/SRigHierarchyTreeView.h"
 
 #include "ControlRigElementDetails.generated.h"
 
diff --git a/Engine/Plugins/Animation/ControlRig/Source/ControlRigEditor/Private/ControlRigModuleDetails.cpp b/Engine/Plugins/Animation/ControlRig/Source/ControlRigEditor/Private/ControlRigModuleDetails.cpp
--- a/Engine/Plugins/Animation/ControlRig/Source/ControlRigEditor/Private/ControlRigModuleDetails.cpp
+++ b/Engine/Plugins/Animation/ControlRig/Source/ControlRigEditor/Private/ControlRigModuleDetails.cpp
@@ -19,7 +19,6 @@
 #include "SEnumCombo.h"
 #include "HAL/PlatformApplicationMisc.h"
 #include "Styling/AppStyle.h"
-#include "Editor/SModularRigTreeView.h"
 #include "StructViewerFilter.h"
 #include "StructViewerModule.h"
 #include "Features/IModularFeatures.h"
@@ -27,7 +26,6 @@
 #include "ModularRigRuleManager.h"
 #include "PropertyPath.h"
 #include "ScopedTransaction.h"
-#include "Editor/SRigHierarchyTreeView.h"
 #include "Widgets/SRigVMVariantTagWidget.h"
 #include "Algo/Sort.h"
 #include "Editor/SRigConnectorTargetWidget.h"
diff --git a/Engine/Plugins/Animation/ControlRig/Source/ControlRigEditor/Private/ControlRigModuleDetails.h b/Engine/Plugins/Animation/ControlRig/Source/ControlRigEditor/Private/ControlRigModuleDetails.h
--- a/Engine/Plugins/Animation/ControlRig/Source/ControlRigEditor/Private/ControlRigModuleDetails.h
+++ b/Engine/Plugins/Animation/ControlRig/Source/ControlRigEditor/Private/ControlRigModuleDetails.h
@@ -23,7 +23,7 @@
 #include "Styling/AppStyle.h"
 #include "Algo/Transform.h"
 #include "IPropertyUtilities.h"
-#include "Editor/SRigHierarchyTreeView.h"
+#include "Editor/Hierarchy/Views/SRigHierarchyTreeView.h"
 
 class IPropertyHandle;
 
diff --git a/Engine/Plugins/Animation/ControlRig/Source/ControlRigEditor/Private/EditMode/SControlRigOutliner.cpp b/Engine/Plugins/Animation/ControlRig/Source/ControlRigEditor/Private/EditMode/SControlRigOutliner.cpp
--- a/Engine/Plugins/Animation/ControlRig/Source/ControlRigEditor/Private/EditMode/SControlRigOutliner.cpp
+++ b/Engine/Plugins/Animation/ControlRig/Source/ControlRigEditor/Private/EditMode/SControlRigOutliner.cpp
@@ -22,7 +22,7 @@
 #include "Editor.h"
 #include "Widgets/Text/STextBlock.h"
 #include "TimerManager.h"
-#include "Editor/SRigHierarchyTreeView.h"
+#include "Editor/Hierarchy/Views/SRigHierarchyTreeView.h"
 #include "Rigs/RigHierarchyController.h"
 #include "ControlRigObjectBinding.h"
 #include "ModularRig.h"
@@ -36,7 +36,8 @@
 #include "Editor/EditorEngine.h"
 #include "ScopedTransaction.h"
 #include "Blueprint/WidgetBlueprintLibrary.h"
-#include "Editor/SRigHierarchy.h"
+#include "Editor/Hierarchy/Widgets/SRigHierarchy.h"
+#include "Editor/Hierarchy/Views/SRigHierarchyItem.h"
 #include "Rigs/FKControlRig.h"
 #include "SEnumCombo.h"
 #include "Components/SkeletalMeshComponent.h"
diff --git a/Engine/Plugins/Animation/ControlRig/Source/ControlRigEditor/Private/Editor/ControlRigEditor.cpp b/Engine/Plugins/Animation/ControlRig/Source/ControlRigEditor/Private/Editor/ControlRigEditor.cpp
--- a/Engine/Plugins/Animation/ControlRig/Source/ControlRigEditor/Private/Editor/ControlRigEditor.cpp
+++ b/Engine/Plugins/Animation/ControlRig/Source/ControlRigEditor/Private/Editor/ControlRigEditor.cpp
@@ -44,7 +44,7 @@
 #include "Widgets/Input/SNumericEntryBox.h"
 #include "Widgets/Notifications/SNotificationList.h"
 #include "Framework/Notifications/NotificationManager.h"
-#include "Editor/SRigHierarchy.h"
+#include "Editor/Hierarchy/Widgets/SRigHierarchy.h"
 #include "Framework/Application/MenuStack.h"
 #include "Framework/Application/SlateApplication.h"
 #include "Units/Hierarchy/RigUnit_BoneName.h"
@@ -128,6 +128,7 @@
 #include "Editor/SRigVMDetailsInspector.h"
 #include "Overrides/OverrideStatusDetailsObjectFilter.h"
 #include "RigVMFunctions/RigVMDispatch_Constant.h"
+#include "Editor/Hierarchy/DragDrop/RigHierarchyElementDragDropOp.h"
 
 #if WITH_RIGVMLEGACYEDITOR
 #include "SKismetInspector.h"
diff --git a/Engine/Plugins/Animation/ControlRig/Source/ControlRigEditor/Private/Editor/ControlRigSchematicModel.cpp b/Engine/Plugins/Animation/ControlRig/Source/ControlRigEditor/Private/Editor/ControlRigSchematicModel.cpp
--- a/Engine/Plugins/Animation/ControlRig/Source/ControlRigEditor/Private/Editor/ControlRigSchematicModel.cpp
+++ b/Engine/Plugins/Animation/ControlRig/Source/ControlRigEditor/Private/Editor/ControlRigSchematicModel.cpp
@@ -12,9 +12,10 @@
 #include "Rigs/RigHierarchyController.h"
 #include "Async/TaskGraphInterfaces.h"
 #include <SchematicGraphPanel/SchematicGraphStyle.h>
-#include "Editor/SModularRigModel.h"
 #include "Blueprint/WidgetBlueprintLibrary.h"
-#include "ControlRigDragOps.h"
+#include "Editor/Hierarchy/DragDrop/ModularRigHierarchyElementDragDropOp.h"
+#include "Editor/Hierarchy/DragDrop/RigHierarchyElementDragDropOp.h"
+#include "Editor/Hierarchy/DragDrop/RigHierarchyTagDragDropOp.h"
 #include "SchematicGraphPanel/SSchematicGraphPanel.h"
 #include "Framework/Application/SlateApplication.h"
 #include "ScopedTransaction.h"
diff --git a/Engine/Plugins/Animation/ControlRig/Source/ControlRigEditor/Private/Editor/Hierarchy/DragDrop/ModularRigHierarchyElementDragDropOp.cpp b/Engine/Plugins/Animation/ControlRig/Source/ControlRigEditor/Private/Editor/Hierarchy/DragDrop/ModularRigHierarchyElementDragDropOp.cpp
new file mode 100644
--- /dev/null
+++ b/Engine/Plugins/Animation/ControlRig/Source/ControlRigEditor/Private/Editor/Hierarchy/DragDrop/ModularRigHierarchyElementDragDropOp.cpp
@@ -0,0 +1,36 @@
+// Copyright Epic Games, Inc. All Rights Reserved.
+
+#include "ModularRigHierarchyElementDragDropOp.h"
+
+#include "Widgets/Layout/SBorder.h"
+#include "Widgets/Text/STextBlock.h"
+
+TSharedRef<FModularRigModuleDragDropOp> FModularRigModuleDragDropOp::New(const TArray<FName>& InModuleNames)
+{
+	TSharedRef<FModularRigModuleDragDropOp> Operation = MakeShared<FModularRigModuleDragDropOp>();
+	Operation->ModuleNames = InModuleNames;
+	Operation->Construct();
+	return Operation;
+}
+
+TSharedPtr<SWidget> FModularRigModuleDragDropOp::GetDefaultDecorator() const
+{
+	return SNew(SBorder)
+		.Visibility(EVisibility::Visible)
+		.BorderImage(FAppStyle::GetBrush("Menu.Background"))
+		[
+			SNew(STextBlock)
+			.Text(FText::FromString(GetJoinedModuleNames()))
+			//.Font(FAppStyle::Get().GetFontStyle("FontAwesome.10"))
+		];
+}
+
+FString FModularRigModuleDragDropOp::GetJoinedModuleNames() const
+{
+	TArray<FString> ElementNameStrings;
+	for (const FName& ModuleName: ModuleNames)
+	{
+		ElementNameStrings.Add(ModuleName.ToString());
+	}
+	return FString::Join(ElementNameStrings, TEXT(","));
+}
diff --git a/Engine/Plugins/Animation/ControlRig/Source/ControlRigEditor/Private/Editor/Hierarchy/DragDrop/ModularRigHierarchyElementDragDropOp.h b/Engine/Plugins/Animation/ControlRig/Source/ControlRigEditor/Private/Editor/Hierarchy/DragDrop/ModularRigHierarchyElementDragDropOp.h
new file mode 100644
--- /dev/null
+++ b/Engine/Plugins/Animation/ControlRig/Source/ControlRigEditor/Private/Editor/Hierarchy/DragDrop/ModularRigHierarchyElementDragDropOp.h
@@ -0,0 +1,34 @@
+// Copyright Epic Games, Inc. All Rights Reserved.
+
+#pragma once
+
+#include "Input/DragAndDrop.h"
+
+class FModularRigModuleDragDropOp : public FDragDropOperation
+{
+public:
+	DRAG_DROP_OPERATOR_TYPE(FModularRigModuleDragDropOp, FDragDropOperation)
+
+	static TSharedRef<FModularRigModuleDragDropOp> New(const TArray<FName>& InModuleNames);
+
+	virtual TSharedPtr<SWidget> GetDefaultDecorator() const override;
+
+	/** @return true if this drag operation contains dragged modules */
+	bool HasModules() const
+	{
+		return ModuleNames.Num() > 0;
+	}
+
+	/** @return The module names from this drag operation */
+	const TArray<FName>& GetModules() const
+	{
+		return ModuleNames;
+	}
+
+	FString GetJoinedModuleNames() const;
+
+private:
+
+	/** Data for the module names this item represents */
+	TArray<FName> ModuleNames;
+};
diff --git a/Engine/Plugins/Animation/ControlRig/Source/ControlRigEditor/Private/Editor/Hierarchy/DragDrop/RigHierarchyElementDragDropOp.cpp b/Engine/Plugins/Animation/ControlRig/Source/ControlRigEditor/Private/Editor/Hierarchy/DragDrop/RigHierarchyElementDragDropOp.cpp
new file mode 100644
--- /dev/null
+++ b/Engine/Plugins/Animation/ControlRig/Source/ControlRigEditor/Private/Editor/Hierarchy/DragDrop/RigHierarchyElementDragDropOp.cpp
@@ -0,0 +1,61 @@
+// Copyright Epic Games, Inc. All Rights Reserved.
+
+#include "RigHierarchyElementDragDropOp.h"
+
+#include "Editor/Hierarchy/Widgets/SRigHierarchyTagWidget.h"
+#include "Widgets/Layout/SBorder.h"
+#include "Widgets/Text/STextBlock.h"
+
+TSharedRef<FRigElementHierarchyDragDropOp> FRigElementHierarchyDragDropOp::New(const TArray<FRigHierarchyKey>& InElements)
+{
+	TSharedRef<FRigElementHierarchyDragDropOp> Operation = MakeShared<FRigElementHierarchyDragDropOp>();
+	Operation->Elements = InElements;
+	Operation->Construct();
+	return Operation;
+}
+
+TSharedPtr<SWidget> FRigElementHierarchyDragDropOp::GetDefaultDecorator() const
+{
+	return SNew(SBorder)
+		.Visibility(EVisibility::Visible)
+		.BorderImage(FAppStyle::GetBrush("Menu.Background"))
+		[
+			SNew(STextBlock)
+			.Text(FText::FromString(GetJoinedElementNames()))
+			//.Font(FAppStyle::Get().GetFontStyle("FontAwesome.10"))
+		];
+}
+
+FString FRigElementHierarchyDragDropOp::GetJoinedElementNames() const
+{
+	TArray<FString> ElementNameStrings;
+	for (const FRigHierarchyKey& Element: Elements)
+	{
+		ElementNameStrings.Add(Element.GetName());
+	}
+	return FString::Join(ElementNameStrings, TEXT(","));
+}
+
+bool FRigElementHierarchyDragDropOp::IsDraggingSingleConnector() const
+{
+	if(Elements.Num() == 1)
+	{
+		if(Elements[0].IsElement())
+		{
+			return Elements[0].GetElement().Type == ERigElementType::Connector;
+		}
+	}
+	return false;
+}
+
+bool FRigElementHierarchyDragDropOp::IsDraggingSingleSocket() const
+{
+	if(Elements.Num() == 1)
+	{
+		if(Elements[0].IsElement())
+		{
+			return Elements[0].GetElement().Type == ERigElementType::Socket;
+		}
+	}
+	return false;
+}
diff --git a/Engine/Plugins/Animation/ControlRig/Source/ControlRigEditor/Private/Editor/Hierarchy/Models/ModularRigHierarchyTreeElement.cpp b/Engine/Plugins/Animation/ControlRig/Source/ControlRigEditor/Private/Editor/Hierarchy/Models/ModularRigHierarchyTreeElement.cpp
new file mode 100644
--- /dev/null
+++ b/Engine/Plugins/Animation/ControlRig/Source/ControlRigEditor/Private/Editor/Hierarchy/Models/ModularRigHierarchyTreeElement.cpp
@@ -0,0 +1,71 @@
+// Copyright Epic Games, Inc. All Rights Reserved.
+
+#include "ModularRigHierarchyTreeElement.h"
+
+#include "Editor/Hierarchy/Views/SModularRigTreeView.h"
+#include "ModularRig.h"
+
+TMap<FSoftObjectPath, TSharedPtr<FSlateBrush>> FModularRigTreeElement::IconPathToBrush;
+
+TArray<TStrongObjectPtr<UTexture2D>> FModularRigTreeElement::Icons;
+
+FModularRigTreeElement::FModularRigTreeElement(const FString& InKey, TWeakPtr<SModularRigTreeView> InTreeView, bool bInIsPrimary)
+{
+	Key = InKey;
+	bIsPrimary = bInIsPrimary;
+
+	FString ModuleNameString;
+	FString ConnectorNameString = Key;
+	(void)FRigHierarchyModulePath(ConnectorNameString).Split(&ModuleNameString, &ConnectorNameString);
+	if (bIsPrimary)
+	{
+		ModuleName = *Key;
+		if (const UModularRig* ModularRig = InTreeView.Pin()->GetRigTreeDelegates().GetModularRig())
+		{
+			if (const FRigModuleInstance* Module = ModularRig->FindModule(ModuleName))
+			{
+				if (const UControlRig* Rig = Module->GetRig())
+				{
+					if (const FRigModuleConnector* PrimaryConnector = Rig->GetRigModuleSettings().FindPrimaryConnector())
+					{
+						ConnectorName = PrimaryConnector->Name;
+					}
+				}
+			}
+		}
+	}
+	else
+	{
+		ConnectorName = ConnectorNameString;
+		ModuleName = *ModuleNameString;
+	}
+
+	ShortName = *ConnectorNameString;
+
+	if (InTreeView.IsValid())
+	{
+		if (const UModularRig* ModularRig = InTreeView.Pin()->GetRigTreeDelegates().GetModularRig())
+		{
+			RefreshDisplaySettings(ModularRig);
+		}
+	}
+}
+
+void FModularRigTreeElement::RefreshDisplaySettings(const UModularRig* InModularRig)
+{
+	const TPair<const FSlateBrush*, FSlateColor> Result = GetBrushAndColor(InModularRig);
+
+	IconBrush = Result.Key;
+	IconColor = Result.Value;
+	TextColor = FSlateColor::UseForeground();
+}
+
+TSharedRef<ITableRow> FModularRigTreeElement::MakeTreeRowWidget(const TSharedRef<STableViewBase>& InOwnerTable, TSharedRef<FModularRigTreeElement> InRigTreeElement, TSharedPtr<SModularRigTreeView> InTreeView, bool bPinned)
+{
+	return SNew(SModularRigModelItem, InOwnerTable, InRigTreeElement, InTreeView, bPinned);
+}
+
+void FModularRigTreeElement::RequestRename()
+{
+	OnRenameRequested.ExecuteIfBound();
+}
diff --git a/Engine/Plugins/Animation/ControlRig/Source/ControlRigEditor/Private/Editor/Hierarchy/Models/ModularRigHierarchyTreeElement.h b/Engine/Plugins/Animation/ControlRig/Source/ControlRigEditor/Private/Editor/Hierarchy/Models/ModularRigHierarchyTreeElement.h
new file mode 100644
--- /dev/null
+++ b/Engine/Plugins/Animation/ControlRig/Source/ControlRigEditor/Private/Editor/Hierarchy/Models/ModularRigHierarchyTreeElement.h
@@ -0,0 +1,49 @@
+// Copyright Epic Games, Inc. All Rights Reserved.
+
+#pragma once
+
+#include "RigHierarchyTreeElement.h"
+
+class SModularRigTreeView;
+class UModularRig;
+class UTexture2D;
+
+/** An item in the tree */
+class FModularRigTreeElement : public TSharedFromThis<FModularRigTreeElement>
+{
+public:
+	FModularRigTreeElement(const FString& InKey, TWeakPtr<SModularRigTreeView> InTreeView, bool bInIsPrimary);
+
+public:
+	/** Element Data to display */
+	FString Key;
+	bool bIsPrimary;
+	FName ModuleName;
+	FString ConnectorName;
+	FName ShortName;
+	TArray<TSharedPtr<FModularRigTreeElement>> Children;
+
+	TSharedRef<ITableRow> MakeTreeRowWidget(const TSharedRef<STableViewBase>& InOwnerTable, TSharedRef<FModularRigTreeElement> InRigTreeElement, TSharedPtr<SModularRigTreeView> InTreeView, bool bPinned);
+
+	void RequestRename();
+
+	void RefreshDisplaySettings(const UModularRig* InModularRig);
+
+	TPair<const FSlateBrush*, FSlateColor> GetBrushAndColor(const UModularRig* InModularRig);
+
+	/** Delegate for when the context menu requests a rename */
+	DECLARE_DELEGATE(FOnRenameRequested);
+	FOnRenameRequested OnRenameRequested;
+
+	static TMap<FSoftObjectPath, TSharedPtr<FSlateBrush>> IconPathToBrush;
+	static TArray<TStrongObjectPtr<UTexture2D>> Icons;
+
+	/** The brush to use when rendering an icon */
+	const FSlateBrush* IconBrush;
+
+	/** The color to use when rendering an icon */
+	FSlateColor IconColor;
+
+	/** The color to use when rendering the label text */
+	FSlateColor TextColor;
+};
diff --git a/Engine/Plugins/Animation/ControlRig/Source/ControlRigEditor/Private/Editor/Hierarchy/Models/RigHierarchyTreeElement.h b/Engine/Plugins/Animation/ControlRig/Source/ControlRigEditor/Private/Editor/Hierarchy/Models/RigHierarchyTreeElement.h
new file mode 100644
--- /dev/null
+++ b/Engine/Plugins/Animation/ControlRig/Source/ControlRigEditor/Private/Editor/Hierarchy/Models/RigHierarchyTreeElement.h
@@ -0,0 +1,82 @@
+// Copyright Epic Games, Inc. All Rights Reserved.
+
+#pragma once
+
+#include "Editor/Hierarchy/RigHierarchyTreeDisplaySettings.h"
+#include "Editor/Hierarchy/Widgets/SRigHierarchyTagWidget.h"
+#include "GenericPlatform/GenericPlatform.h"
+#include "Internationalization/Text.h"
+#include "Styling/SlateColor.h"
+
+class ITableRow;
+class SRigHierarchyTreeView;
+class STableViewBase;
+class URigHierarchy;
+struct FRigHierarchyKey;
+struct FSlateBrush;
+
+/**
+ * Order is important here!
+ * This enum is used internally to the filtering logic and represents an ordering of most filtered (hidden) to
+ * least filtered (highlighted).
+ */
+enum class ERigTreeFilterResult : int32
+{
+	/** Hide the item */
+	Hidden,
+
+	/** Show the item because child items were shown */
+	ShownDescendant,
+
+	/** Show the item */
+	Shown,
+};
+
+/** An item in the tree */
+class FRigTreeElement : public TSharedFromThis<FRigTreeElement>
+{
+public:
+	FRigTreeElement(const FRigHierarchyKey& InKey, TWeakPtr<SRigHierarchyTreeView> InTreeView, bool InSupportsRename, ERigTreeFilterResult InFilterResult);
+public:
+	/** Element Data to display */
+	FRigHierarchyKey Key;
+	FText LongName;
+	FText ShortName;
+	FName ChannelName;
+	bool bIsTransient;
+	bool bIsAnimationChannel;
+	bool bIsProcedural;
+	bool bSupportsRename;
+	TArray<TSharedPtr<FRigTreeElement>> Children;
+
+	TSharedRef<ITableRow> MakeTreeRowWidget(const TSharedRef<STableViewBase>& InOwnerTable, TSharedRef<FRigTreeElement> InRigTreeElement, TSharedPtr<SRigHierarchyTreeView> InTreeView, const FRigTreeDisplaySettings& InSettings, bool bPinned);
+
+	void RequestRename();
+
+	void RefreshDisplaySettings(const URigHierarchy* InHierarchy, const FRigTreeDisplaySettings& InSettings);
+
+	FSlateColor GetIconColor() const;
+	FSlateColor GetTextColor() const;
+
+	/** Delegate for when the context menu requests a rename */
+	DECLARE_DELEGATE(FOnRenameRequested);
+	FOnRenameRequested OnRenameRequested;
+
+	/** The current filter result */
+	ERigTreeFilterResult FilterResult;
+
+	/** The brush to use when rendering an icon */
+	const FSlateBrush* IconBrush;;
+
+	/** The color to use when rendering an icon */
+	FSlateColor IconColor;
+
+	/** The color to use when rendering the label text */
+	FSlateColor TextColor;
+
+	/** If true the item is filtered out during a drag */
+	bool bFadedOutDuringDragDrop;
+
+	/** The tag arguments for this element */
+	TArray<SRigHierarchyTagWidget::FArguments> Tags;
+};
diff --git a/Engine/Plugins/Animation/ControlRig/Source/ControlRigEditor/Private/Editor/Hierarchy/ModularRigHierarchyTreeDelegates.cpp b/Engine/Plugins/Animation/ControlRig/Source/ControlRigEditor/Private/Editor/Hierarchy/ModularRigHierarchyTreeDelegates.cpp
new file mode 100644
--- /dev/null
+++ b/Engine/Plugins/Animation/ControlRig/Source/ControlRigEditor/Private/Editor/Hierarchy/ModularRigHierarchyTreeDelegates.cpp
@@ -0,0 +1,69 @@
+// Copyright Epic Games, Inc. All Rights Reserved.
+
+#include "ModularRigHierarchyTreeDelegates.h"
+
+const UModularRig* FModularRigTreeDelegates::GetModularRig() const
+{
+	if (OnGetModularRig.IsBound())
+	{
+		return OnGetModularRig.Execute();
+	}
+	return nullptr;
+}
+
+FName FModularRigTreeDelegates::HandleRenameElement(const FName& OldModuleName, const FName& NewName) const
+{
+	if (OnRenameElement.IsBound())
+	{
+		return OnRenameElement.Execute(OldModuleName, NewName);
+	}
+	return OldModuleName;
+}
+
+bool FModularRigTreeDelegates::HandleVerifyElementNameChanged(const FName& OldModuleName, const FName& NewName, FText& OutErrorMessage) const
+{
+	if (OnVerifyModuleNameChanged.IsBound())
+	{
+		return OnVerifyModuleNameChanged.Execute(OldModuleName, NewName, OutErrorMessage);
+	}
+	return false;
+}
+
+bool FModularRigTreeDelegates::HandleResolveConnector(const FRigElementKey& InConnector, const TArray<FRigElementKey>& InTargets)
+{
+	if (OnResolveConnector.IsBound())
+	{
+		OnResolveConnector.Execute(InConnector, InTargets);
+		return true;
+	}
+	return false;
+}
+
+bool FModularRigTreeDelegates::HandleDisconnectConnector(const FRigElementKey& InConnector)
+{
+	if (OnDisconnectConnector.IsBound())
+	{
+		OnDisconnectConnector.Execute(InConnector);
+		return true;
+	}
+	return false;
+}
+
+void FModularRigTreeDelegates::HandleSelectionChanged(TSharedPtr<FModularRigTreeElement> Selection, ESelectInfo::Type SelectInfo)
+{
+	if (bSuspendSelectionDelegate)
+	{
+		return;
+	}
+	TGuardValue<bool> Guard(bSuspendSelectionDelegate, true);
+	(void)OnSelectionChanged.ExecuteIfBound(Selection, SelectInfo);
+}
+
+bool FModularRigTreeDelegates::ShouldAlwaysShowConnector(const FName& InConnectorName) const
+{
+	if (OnAlwaysShowConnector.IsBound())
+	{
+		return OnAlwaysShowConnector.Execute(InConnectorName);
+	}
+	return false;
+}
diff --git a/Engine/Plugins/Animation/ControlRig/Source/ControlRigEditor/Private/Editor/Hierarchy/ModularRigHierarchyTreeDelegates.h b/Engine/Plugins/Animation/ControlRig/Source/ControlRigEditor/Private/Editor/Hierarchy/ModularRigHierarchyTreeDelegates.h
new file mode 100644
--- /dev/null
+++ b/Engine/Plugins/Animation/ControlRig/Source/ControlRigEditor/Private/Editor/Hierarchy/ModularRigHierarchyTreeDelegates.h
@@ -0,0 +1,61 @@
+// Copyright Epic Games, Inc. All Rights Reserved.
+
+#pragma once
+
+#include "Delegates/Delegate.h"
+#include "Editor/Hierarchy/Models/ModularRigHierarchyTreeElement.h"
+#include "Editor/Hierarchy/RigHierarchyTreeDelegates.h"
+#include "Editor/Hierarchy/RigHierarchyTreeDisplaySettings.h"
+#include "Editor/Hierarchy/Widgets/SRigHierarchyTagWidget.h"
+#include "Framework/SlateDelegates.h"
+#include "Rigs/RigHierarchy.h"
+#include "Widgets/Views/STableRow.h"
+#include "Widgets/Views/STreeView.h"
+
+DECLARE_DELEGATE_RetVal(const UModularRig*, FOnGetModularRigTreeRig);
+DECLARE_DELEGATE_OneParam(FOnModularRigTreeRequestDetailsInspection, const FName&);
+DECLARE_DELEGATE_RetVal_TwoParams(FName, FOnModularRigTreeRenameElement, const FName& /*OldModuleName*/, const FName& /*NewName*/);
+DECLARE_DELEGATE_TwoParams(FOnModularRigTreeResolveConnector, const FRigElementKey& /*Connector*/, const TArray<FRigElementKey>& /*Targets*/);
+DECLARE_DELEGATE_OneParam(FOnModularRigTreeDisconnectConnector, const FRigElementKey& /*Connector*/);
+DECLARE_DELEGATE_RetVal_ThreeParams(bool, FOnModularRigTreeVerifyElementNameChanged, const FName& /*OldModuleName*/, const FName& /*NewName*/, FText& /*OutErrorMessage*/);
+DECLARE_DELEGATE_RetVal_OneParam(bool, FOnModularRigTreeAlwaysShowConnector, const FName& /*ConnectorName*/);
+
+typedef STreeView<TSharedPtr<FModularRigTreeElement>>::FOnMouseButtonClick FOnModularRigTreeMouseButtonClick;
+typedef STreeView<TSharedPtr<FModularRigTreeElement>>::FOnMouseButtonDoubleClick FOnModularRigTreeMouseButtonDoubleClick;
+typedef STableRow<TSharedPtr<FModularRigTreeElement>>::FOnCanAcceptDrop FOnModularRigTreeCanAcceptDrop;
+typedef STableRow<TSharedPtr<FModularRigTreeElement>>::FOnAcceptDrop FOnModularRigTreeAcceptDrop;
+typedef STreeView<TSharedPtr<FModularRigTreeElement>>::FOnSelectionChanged FOnModularRigTreeSelectionChanged;
+
+struct FModularRigTreeDelegates
+{
+public:
+
+	FOnGetModularRigTreeRig OnGetModularRig;
+	FOnModularRigTreeMouseButtonClick OnMouseButtonClick;
+	FOnModularRigTreeMouseButtonDoubleClick OnMouseButtonDoubleClick;
+	FOnDragDetected OnDragDetected;
+	FOnModularRigTreeCanAcceptDrop OnCanAcceptDrop;
+	FOnModularRigTreeAcceptDrop OnAcceptDrop;
+	FOnContextMenuOpening OnContextMenuOpening;
+	FOnModularRigTreeRequestDetailsInspection OnRequestDetailsInspection;
+	FOnModularRigTreeRenameElement OnRenameElement;
+	FOnModularRigTreeVerifyElementNameChanged OnVerifyModuleNameChanged;
+	FOnModularRigTreeResolveConnector OnResolveConnector;
+	FOnModularRigTreeDisconnectConnector OnDisconnectConnector;
+	FOnModularRigTreeSelectionChanged OnSelectionChanged;
+	FOnModularRigTreeAlwaysShowConnector OnAlwaysShowConnector;
+
+	const UModularRig* GetModularRig() const;
+	FName HandleRenameElement(const FName& OldModuleName, const FName& NewName) const;
+	bool HandleVerifyElementNameChanged(const FName& OldModuleName, const FName& NewName, FText& OutErrorMessage) const;
+	bool HandleResolveConnector(const FRigElementKey& InConnector, const TArray<FRigElementKey>& InTargets);
+	bool HandleDisconnectConnector(const FRigElementKey& InConnector);
+	void HandleSelectionChanged(TSharedPtr<FModularRigTreeElement> Selection, ESelectInfo::Type SelectInfo);
+	bool ShouldAlwaysShowConnector(const FName& InConnectorName) const;
+
+private:
+
+	bool bSuspendSelectionDelegate = false;
+
+	friend class SModularRigTreeView;
+};
diff --git a/Engine/Plugins/Animation/ControlRig/Source/ControlRigEditor/Private/Editor/Hierarchy/RigHierarchyTreeDelegates.cpp b/Engine/Plugins/Animation/ControlRig/Source/ControlRigEditor/Private/Editor/Hierarchy/RigHierarchyTreeDelegates.cpp
new file mode 100644
--- /dev/null
+++ b/Engine/Plugins/Animation/ControlRig/Source/ControlRigEditor/Private/Editor/Hierarchy/RigHierarchyTreeDelegates.cpp
@@ -0,0 +1,90 @@
+// Copyright Epic Games, Inc. All Rights Reserved.
+
+#include "RigHierarchyTreeDelegates.h"
+
+FRigTreeDisplaySettings FRigTreeDelegates::DefaultDisplaySettings;
+
+URigHierarchy* FRigTreeDelegates::GetHierarchy()
+{
+	if (OnGetHierarchy.IsBound())
+	{
+		return const_cast<URigHierarchy*>(OnGetHierarchy.Execute());
+	}
+	return nullptr;
+}
+
+const URigHierarchy* FRigTreeDelegates::GetHierarchy() const
+{
+	if (OnGetHierarchy.IsBound())
+	{
+		return OnGetHierarchy.Execute();
+	}
+	return nullptr;
+}
+
+const FRigTreeDisplaySettings& FRigTreeDelegates::GetDisplaySettings() const
+{
+	if (OnGetDisplaySettings.IsBound())
+	{
+		return OnGetDisplaySettings.Execute();
+	}
+	return DefaultDisplaySettings;
+}
+
+TArray<FRigHierarchyKey> FRigTreeDelegates::GetSelection() const
+{
+	if (OnGetSelection.IsBound())
+	{
+		return OnGetSelection.Execute();
+	}
+	if (const URigHierarchy* Hierarchy = GetHierarchy())
+	{
+		return Hierarchy->GetSelectedHierarchyKeys();
+	}
+	return {};
+}
+
+FName FRigTreeDelegates::HandleRenameElement(const FRigHierarchyKey& OldKey, const FString& NewName) const
+{
+	if (OnRenameElement.IsBound())
+	{
+		return OnRenameElement.Execute(OldKey, NewName);
+	}
+	return OldKey.GetFName();
+}
+
+bool FRigTreeDelegates::HandleVerifyElementNameChanged(const FRigHierarchyKey& OldKey, const FString& NewName, FText& OutErrorMessage) const
+{
+	if (OnVerifyElementNameChanged.IsBound())
+	{
+		return OnVerifyElementNameChanged.Execute(OldKey, NewName, OutErrorMessage);
+	}
+	return false;
+}
+
+void FRigTreeDelegates::HandleSelectionChanged(TSharedPtr<FRigTreeElement> Selection, ESelectInfo::Type SelectInfo)
+{
+	if (bIsChangingRigHierarchy)
+	{
+		return;
+	}
+	TGuardValue<bool> Guard(bIsChangingRigHierarchy, true);
+	(void)OnSelectionChanged.ExecuteIfBound(Selection, SelectInfo);
+}
+
+FRigHierarchyKey FRigTreeDelegates::GetResolvedKey(const FRigHierarchyKey& InKey)
+{
+	if (OnGetResolvedKey.IsBound())
+	{
+		return OnGetResolvedKey.Execute(InKey);
+	}
+	return InKey;
+}
+
+void FRigTreeDelegates::RequestDetailsInspection(const FRigHierarchyKey& InKey)
+{
+	if (OnRequestDetailsInspection.IsBound())
+	{
+		return OnRequestDetailsInspection.Execute(InKey);
+	}
+}
\ No newline at end of file
diff --git a/Engine/Plugins/Animation/ControlRig/Source/ControlRigEditor/Private/Editor/Hierarchy/Views/SModularRigHierarchySearchableTreeView.cpp b/Engine/Plugins/Animation/ControlRig/Source/ControlRigEditor/Private/Editor/Hierarchy/Views/SModularRigHierarchySearchableTreeView.cpp
new file mode 100644
--- /dev/null
+++ b/Engine/Plugins/Animation/ControlRig/Source/ControlRigEditor/Private/Editor/Hierarchy/Views/SModularRigHierarchySearchableTreeView.cpp
@@ -0,0 +1,34 @@
+// Copyright Epic Games, Inc. All Rights Reserved.
+
+#include "SModularRigHierarchySearchableTreeView.h"
+
+#include "SModularRigTreeView.h"
+#include "Widgets/Layout/SScrollBox.h"
+
+void SSearchableModularRigTreeView::Construct(const FArguments& InArgs)
+{
+	FModularRigTreeDelegates TreeDelegates = InArgs._RigTreeDelegates;
+	
+	ChildSlot
+	[
+		SNew(SVerticalBox)
+		+SVerticalBox::Slot()
+		.AutoHeight()
+		.VAlign(VAlign_Top)
+		.HAlign(HAlign_Fill)
+		.Padding(0.0f, 0.0f)
+		[
+			SNew(SScrollBox)
+			+ SScrollBox::Slot()
+			[
+				SNew(SBorder)
+				.Padding(2.0f)
+				.BorderImage(FAppStyle::GetBrush("SCSEditor.TreePanel"))
+				[
+					SAssignNew(TreeView, SModularRigTreeView)
+					.RigTreeDelegates(TreeDelegates)
+				]
+			]
+		]
+	];
+}
diff --git a/Engine/Plugins/Animation/ControlRig/Source/ControlRigEditor/Private/Editor/Hierarchy/Views/SModularRigHierarchySearchableTreeView.h b/Engine/Plugins/Animation/ControlRig/Source/ControlRigEditor/Private/Editor/Hierarchy/Views/SModularRigHierarchySearchableTreeView.h
new file mode 100644
--- /dev/null
+++ b/Engine/Plugins/Animation/ControlRig/Source/ControlRigEditor/Private/Editor/Hierarchy/Views/SModularRigHierarchySearchableTreeView.h
@@ -0,0 +1,26 @@
+// Copyright Epic Games, Inc. All Rights Reserved.
+
+#pragma once
+
+#include "Editor/Hierarchy/ModularRigHierarchyTreeDelegates.h"
+#include "Widgets/SCompoundWidget.h"
+
+class SModularRigTreeView;
+
+class SSearchableModularRigTreeView : public SCompoundWidget
+{
+public:
+
+	SLATE_BEGIN_ARGS(SSearchableModularRigTreeView) {}
+		SLATE_ARGUMENT(FModularRigTreeDelegates, RigTreeDelegates)
+		SLATE_ARGUMENT(FText, InitialFilterText)
+	SLATE_END_ARGS()
+
+	void Construct(const FArguments& InArgs);
+	virtual ~SSearchableModularRigTreeView() {}
+	TSharedRef<SModularRigTreeView> GetTreeView() const { return TreeView.ToSharedRef(); }
+
+private:
+
+	TSharedPtr<SModularRigTreeView> TreeView;
+};
diff --git a/Engine/Plugins/Animation/ControlRig/Source/ControlRigEditor/Private/Editor/Hierarchy/Views/SModularRigHierarchyTreeItem.cpp b/Engine/Plugins/Animation/ControlRig/Source/ControlRigEditor/Private/Editor/Hierarchy/Views/SModularRigHierarchyTreeItem.cpp
new file mode 100644
--- /dev/null
+++ b/Engine/Plugins/Animation/ControlRig/Source/ControlRigEditor/Private/Editor/Hierarchy/Views/SModularRigHierarchyTreeItem.cpp
@@ -0,0 +1,400 @@
+// Copyright Epic Games, Inc. All Rights Reserved.
+
+#include "SModularRigHierarchyTreeItem.h"
+
+#include "ControlRigBlueprintLegacy.h"
+#include "Editor/Hierarchy/Views/SModularRigTreeView.h"
+#include "Editor/SRigConnectorTargetWidget.h"
+#include "Rigs/RigHierarchyController.h"
+#include "ScopedTransaction.h"
+#include "Widgets/Text/SInlineEditableTextBlock.h"
+
+#define LOCTEXT_NAMESPACE "SModularRigModelItem"
+
+void SModularRigModelItem::Construct(const FArguments& InArgs, const TSharedRef<STableViewBase>& OwnerTable, TSharedRef<FModularRigTreeElement> InRigTreeElement, TSharedPtr<SModularRigTreeView> InTreeView, bool bPinned)
+{
+	WeakRigTreeElement = InRigTreeElement;
+	Delegates = InTreeView->GetRigTreeDelegates();
+
+	if (InRigTreeElement->Key.IsEmpty())
+	{
+		SMultiColumnTableRow<TSharedPtr<FModularRigTreeElement>>::Construct(
+			SMultiColumnTableRow<TSharedPtr<FModularRigTreeElement>>::FArguments()
+			.ShowSelection(false)
+			.OnCanAcceptDrop(Delegates.OnCanAcceptDrop)
+			.OnAcceptDrop(Delegates.OnAcceptDrop)
+			, OwnerTable);
+		return;
+	}
+
+	const FString& ModuleName = InRigTreeElement->ModuleName.ToString();
+	const FRigHierarchyModulePath ConnectorModulePath(ModuleName, InRigTreeElement->ConnectorName);
+	ConnectorKey = FRigElementKey(ConnectorModulePath.GetPathFName(), ERigElementType::Connector);
+
+	SMultiColumnTableRow<TSharedPtr<FModularRigTreeElement>>::Construct(
+		SMultiColumnTableRow<TSharedPtr<FModularRigTreeElement>>::FArguments()
+		.OnDragDetected(Delegates.OnDragDetected)
+		.OnCanAcceptDrop(Delegates.OnCanAcceptDrop)
+		.OnAcceptDrop(Delegates.OnAcceptDrop)
+		.ShowWires(true), OwnerTable);
+}
+
+bool SModularRigModelItem::OnConnectorTargetChanged(TArray<FRigElementKey> InTargets, const FRigElementKey InConnectorKey)
+{
+	FScopedTransaction Transaction(LOCTEXT("ModuleHierarchyResolveConnector", "Resolve Connector"));
+	Delegates.HandleResolveConnector(InConnectorKey, InTargets);
+	return false;
+}
+
+void SModularRigModelItem::OnNameCommitted(const FText& InText, ETextCommit::Type InCommitType) const
+{
+	// for now only allow enter
+	// because it is important to keep the unique names per pose
+	if (InCommitType == ETextCommit::OnEnter)
+	{
+		FString NewName = InText.ToString();
+		const FName OldModuleName = WeakRigTreeElement.Pin()->ModuleName;
+
+		Delegates.HandleRenameElement(OldModuleName, *NewName);
+	}
+}
+
+bool SModularRigModelItem::OnVerifyNameChanged(const FText& InText, FText& OutErrorMessage)
+{
+	const FName NewName = *InText.ToString();
+	const FName OldModuleName = WeakRigTreeElement.Pin()->ModuleName;
+	return Delegates.HandleVerifyElementNameChanged(OldModuleName, NewName, OutErrorMessage);
+}
+
+TSharedRef<SWidget> SModularRigModelItem::GenerateWidgetForColumn(const FName& ColumnName)
+{
+	if(ColumnName == SModularRigTreeView::Column_Module)
+	{
+		static constexpr float TopPadding = 2.f;
+		
+		TSharedPtr< SInlineEditableTextBlock > InlineWidget;
+		TSharedRef<SWidget> Widget = SNew(SHorizontalBox)
+
+		+SHorizontalBox::Slot()
+		.AutoWidth()
+		.Padding(6, TopPadding, 0, 0)
+		.VAlign(VAlign_Fill)
+		[
+			SNew(SExpanderArrow, SharedThis(this))
+			.IndentAmount(12)
+			.ShouldDrawWires(true)
+		]
+
+		+SHorizontalBox::Slot()
+		.MaxWidth(25)
+		.HAlign(HAlign_Left)
+		.VAlign(VAlign_Center)
+		.Padding(FMargin(0.f, TopPadding, 3.f, 0.f))
+		[
+			SNew(SVerticalBox)
+			+SVerticalBox::Slot()
+			.MaxHeight(25)
+			[
+				SNew(SImage)
+				.Image_Lambda([this]() -> const FSlateBrush*
+				{
+					if(WeakRigTreeElement.IsValid())
+					{
+						return WeakRigTreeElement.Pin()->IconBrush;
+					}
+					return nullptr;
+				})
+				.ColorAndOpacity_Lambda([this]()
+				{
+					if(WeakRigTreeElement.IsValid())
+					{
+						return WeakRigTreeElement.Pin()->IconColor;
+					}
+					return FSlateColor::UseForeground();
+				})
+				.DesiredSizeOverride(FVector2D(16, 16))
+			]
+		]
+
+		+ SHorizontalBox::Slot()
+		.AutoWidth()
+		.VAlign(VAlign_Center)
+		.Padding(0, TopPadding, 0, 0)
+		[
+			SAssignNew(InlineWidget, SInlineEditableTextBlock)
+			.Text(this, &SModularRigModelItem::GetName, true)
+			.MaximumLength(NAME_SIZE-1)
+			.OnVerifyTextChanged(this, &SModularRigModelItem::OnVerifyNameChanged)
+			.OnTextCommitted(this, &SModularRigModelItem::OnNameCommitted)
+			.ToolTipText(this, &SModularRigModelItem::GetItemTooltip)
+			.MultiLine(false)
+			//.Font(IDetailLayoutBuilder::GetDetailFont())
+			.ColorAndOpacity_Lambda([this]()
+			{
+				if(WeakRigTreeElement.IsValid())
+				{
+					return WeakRigTreeElement.Pin()->TextColor;
+				}
+				return FSlateColor::UseForeground();
+			})
+		];
+
+		if(WeakRigTreeElement.IsValid())
+		{
+			WeakRigTreeElement.Pin()->OnRenameRequested.BindSP(InlineWidget.Get(), &SInlineEditableTextBlock::EnterEditingMode);
+		}
+
+		return Widget;
+	}
+	if (ColumnName == SModularRigTreeView::Column_Tags)
+	{
+		TSharedRef<SWidget> Widget = SNew(SImage)
+		.Visibility_Lambda([this]() -> EVisibility
+		{
+			if (WeakRigTreeElement.IsValid())
+			{
+				if (!WeakRigTreeElement.Pin()->bIsPrimary)
+				{
+					return EVisibility::Hidden;
+				}
+			
+				if (const UModularRig* ModularRig = Delegates.GetModularRig())
+				{
+					if (const FRigModuleInstance* Module = ModularRig->FindModule(WeakRigTreeElement.Pin()->ModuleName))
+					{
+						if (const UControlRig* ModuleRig = Module->GetRig())
+						{
+							if (const UControlRigBlueprint* ModuleBlueprint = Cast<UControlRigBlueprint>(ModuleRig->GetClass()->ClassGeneratedBy))
+							{
+								for (const FRigVMTag& Tag : ModuleBlueprint->GetAssetVariant().Tags)
+								{
+									if (Tag.bMarksSubjectAsInvalid)
+									{
+										return EVisibility::Visible;
+									}
+								}
+							}
+						}
+					}
+				}
+			}
+			return EVisibility::Hidden;
+		})
+		.ToolTipText_Lambda([this]() -> FText
+		{
+			TArray<FString> ToolTip;
+			if (WeakRigTreeElement.IsValid())
+			{
+				if (const UModularRig* ModularRig = Delegates.GetModularRig())
+				{
+					if (const FRigModuleInstance* Module = ModularRig->FindModule(WeakRigTreeElement.Pin()->ModuleName))
+					{
+						if (const UControlRigBlueprint* ModuleBlueprint = Cast<UControlRigBlueprint>(Module->GetRig()->GetClass()->ClassGeneratedBy))
+						{
+							for (const FRigVMTag& Tag : ModuleBlueprint->GetAssetVariant().Tags)
+							{
+								if (Tag.bMarksSubjectAsInvalid)
+								{
+									ToolTip.Add(FString::Printf(TEXT("%s: %s"), *Tag.Label, *Tag.ToolTip.ToString()));
+									ToolTip.Add(TEXT("Right click on the module to swap it to a newer variant.")); 
+								}
+							}
+						}
+					}
+				}
+			}
+
+			return FText::FromString(FString::Join(ToolTip, TEXT("\n")));
+		})
+		.Image_Lambda([this]() -> const FSlateBrush*
+		{
+			const FSlateBrush* WarningBrush = FAppStyle::Get().GetBrush("Icons.WarningWithColor");
+			return WarningBrush;
+		})
+		.DesiredSizeOverride(FVector2D(16, 16));
+
+		return Widget;
+	}
+	if(ColumnName == SModularRigTreeView::Column_Connector)
+	{
+		bool bIsArrayConnector = false;
+		if (const UModularRig* ModularRig = Delegates.GetModularRig())
+		{
+			if (URigHierarchy* Hierarchy = ModularRig->GetHierarchy())
+			{
+				if (const FRigConnectorElement* ConnectorElement = Cast<FRigConnectorElement>(Hierarchy->Find(ConnectorKey)))
+				{
+					bIsArrayConnector = ConnectorElement->IsArrayConnector();
+				}
+			}
+		}
+
+		TSharedPtr<SVerticalBox> ComboButtonBox;
+		
+		FRigTreeDelegates RigTreeDelegates;
+		RigTreeDelegates.OnGetHierarchy = FOnGetRigTreeHierarchy::CreateLambda([this]()
+		{
+			return Delegates.GetModularRig()->GetHierarchy();
+		});
+
+		return SNew(SRigConnectorTargetWidget)
+			.Outer(const_cast<UModularRig*>(Delegates.GetModularRig()))
+			.ConnectorKey(ConnectorKey)
+			.IsArray(bIsArrayConnector)
+			.Targets(GetTargetKeys())
+			.OnSetTargetArray(FRigConnectorTargetWidget_SetTargetArray::CreateSP(this, &SModularRigModelItem::OnConnectorTargetChanged, ConnectorKey))
+			.RigTreeDelegates(RigTreeDelegates);
+	}
+	if(ColumnName == SModularRigTreeView::Column_Buttons)
+	{
+		return SNew(SHorizontalBox)
+
+		// Reset button
+		+SHorizontalBox::Slot()
+		.AutoWidth()
+		.Padding(0.f, 0.f, 0.f, 0.f)
+		[
+			SAssignNew(ResetConnectorButton, SButton)
+			.ButtonStyle( FAppStyle::Get(), "NoBorder" )
+			.ButtonColorAndOpacity_Lambda([this]()
+			{
+				return ResetConnectorButton.IsValid() && ResetConnectorButton->IsHovered()
+					? FSlateColor(FLinearColor(1.f, 1.f, 1.f, 0.8))
+					: FSlateColor(FLinearColor(1.f, 1.f, 1.f, 0.4));
+			})
+			.OnClicked_Lambda([this]()
+			{
+				Delegates.HandleDisconnectConnector(ConnectorKey);
+				return FReply::Handled();
+			})
+			.ContentPadding(1.f)
+			.ToolTipText(NSLOCTEXT("ControlRigModuleDetails", "Reset_Connector", "Reset Connector"))
+			[
+				SNew(SImage)
+				.ColorAndOpacity_Lambda( [this]()
+				{
+					return ResetConnectorButton.IsValid() && ResetConnectorButton->IsHovered()
+					? FSlateColor(FLinearColor(1.f, 1.f, 1.f, 0.8))
+					: FSlateColor(FLinearColor(1.f, 1.f, 1.f, 0.4));
+				})
+				.Image(FSlateIcon(FAppStyle::Get().GetStyleSetName(), "PropertyWindow.DiffersFromDefault").GetIcon())
+			]
+		]
+
+		// Use button
+		+SHorizontalBox::Slot()
+		.AutoWidth()
+		.Padding(0.f, 0.f, 0.f, 0.f)
+		[
+			SAssignNew(UseSelectedButton, SButton)
+			.ButtonStyle( FAppStyle::Get(), "NoBorder" )
+			.ButtonColorAndOpacity_Lambda([this]()
+			{
+				return UseSelectedButton.IsValid() && UseSelectedButton->IsHovered()
+					? FSlateColor(FLinearColor(1.f, 1.f, 1.f, 0.8))
+					: FSlateColor(FLinearColor(1.f, 1.f, 1.f, 0.4));
+			})
+			.OnClicked_Lambda([this]()
+			{
+				if (const UModularRig* ModularRig = Delegates.GetModularRig())
+				{
+					const TArray<FRigElementKey>& Selected = ModularRig->GetHierarchy()->GetSelectedKeys();
+					if (Selected.Num() > 0)
+					{
+						Delegates.HandleResolveConnector(ConnectorKey, Selected);
+					}
+				}
+				return FReply::Handled();
+			})
+			.ContentPadding(1.f)
+			.ToolTipText(NSLOCTEXT("ControlRigModuleDetails", "Use_Selected", "Use Selected"))
+			[
+				SNew(SImage)
+				.ColorAndOpacity_Lambda( [this]()
+				{
+					return UseSelectedButton.IsValid() && UseSelectedButton->IsHovered()
+					? FSlateColor(FLinearColor(1.f, 1.f, 1.f, 0.8))
+					: FSlateColor(FLinearColor(1.f, 1.f, 1.f, 0.4));
+				})
+				.Image(FAppStyle::GetBrush("Icons.CircleArrowLeft"))
+			]
+		]
+		
+		// Select in hierarchy button
+		+SHorizontalBox::Slot()
+		.AutoWidth()
+		.Padding(0.f, 0.f, 0.f, 0.f)
+		[
+			SAssignNew(SelectElementButton, SButton)
+			.ButtonStyle( FAppStyle::Get(), "NoBorder" )
+			.ButtonColorAndOpacity_Lambda([this]()
+			{
+				return SelectElementButton.IsValid() && SelectElementButton->IsHovered()
+					? FSlateColor(FLinearColor(1.f, 1.f, 1.f, 0.8))
+					: FSlateColor(FLinearColor(1.f, 1.f, 1.f, 0.4));
+			})
+			.OnClicked_Lambda([this]()
+			{
+				if (const UModularRig* ModularRig = Delegates.GetModularRig())
+				{
+					const FRigElementKeyRedirector& Redirector = ModularRig->GetElementKeyRedirector();
+					if (const FRigElementKeyRedirector::FKeyArray* TargetKeys = Redirector.FindExternalKey(ConnectorKey))
+					{
+						bool bClearSelection = true;
+						for(const FRigElementKey& TargetKey : *TargetKeys)
+						{
+							ModularRig->GetHierarchy()->GetController()->SelectElement(TargetKey, true, bClearSelection);
+							bClearSelection = false;
+						}
+					}
+				}
+				return FReply::Handled();
+			})
+			.ContentPadding(1.f)
+			.ToolTipText(NSLOCTEXT("ControlRigModuleDetails", "Select_Element", "Select Element"))
+			[
+				SNew(SImage)
+				.ColorAndOpacity_Lambda( [this]()
+				{
+					return SelectElementButton.IsValid() && SelectElementButton->IsHovered()
+					? FSlateColor(FLinearColor(1.f, 1.f, 1.f, 0.8))
+					: FSlateColor(FLinearColor(1.f, 1.f, 1.f, 0.4));
+				})
+				.Image(FAppStyle::GetBrush("Icons.Search"))
+			]
+		];
+	}
+	return SNullWidget::NullWidget;
+}
+
+FText SModularRigModelItem::GetName(bool bUseShortName) const
+{
+	if(bUseShortName)
+	{
+		return (FText::FromName(WeakRigTreeElement.Pin()->ShortName));
+	}
+	return (FText::FromName(WeakRigTreeElement.Pin()->ModuleName));
+}
+
+FText SModularRigModelItem::GetItemTooltip() const
+{
+	const FText FullName = GetName(false);
+	const FText ShortName = GetName(true);
+	if(FullName.EqualTo(ShortName))
+	{
+		return FText();
+	}
+	return FullName;
+}
+
+TArray<FRigElementKey> SModularRigModelItem::GetTargetKeys() const
+{
+	TArray<FRigElementKey> Result;
+	if (const UModularRig* ModularRig = Delegates.GetModularRig())
+	{
+		Result = ModularRig->GetModularRigModel().Connections.FindTargetsFromConnector(ConnectorKey);
+	}
+	return Result;
+}
+
+#undef LOCTEXT_NAMESPACE
diff --git a/Engine/Plugins/Animation/ControlRig/Source/ControlRigEditor/Private/Editor/Hierarchy/Views/SModularRigHierarchyTreeItem.h b/Engine/Plugins/Animation/ControlRig/Source/ControlRigEditor/Private/Editor/Hierarchy/Views/SModularRigHierarchyTreeItem.h
new file mode 100644
--- /dev/null
+++ b/Engine/Plugins/Animation/ControlRig/Source/ControlRigEditor/Private/Editor/Hierarchy/Views/SModularRigHierarchyTreeItem.h
@@ -0,0 +1,40 @@
+// Copyright Epic Games, Inc. All Rights Reserved.
+
+#pragma once
+
+#include "Editor/Hierarchy/Models/ModularRigHierarchyTreeElement.h"
+#include "Editor/Hierarchy/ModularRigHierarchyTreeDelegates.h"
+#include "Internationalization/Text.h"
+#include "Rigs/RigHierarchyDefines.h"
+#include "Widgets/Views/STableRow.h"
+
+class SSearchableRigHierarchyTreeView;
+
+class SModularRigModelItem : public SMultiColumnTableRow<TSharedPtr<FModularRigTreeElement>>
+{
+public:
+
+	void Construct(const FArguments& InArgs, const TSharedRef<STableViewBase>& OwnerTable, TSharedRef<FModularRigTreeElement> InRigTreeElement, TSharedPtr<SModularRigTreeView> InTreeView, bool bPinned);
+	bool OnConnectorTargetChanged(TArray<FRigElementKey> InTargets, const FRigElementKey InConnectorKey);
+
+	void OnNameCommitted(const FText& InText, ETextCommit::Type InCommitType) const;
+	bool OnVerifyNameChanged(const FText& InText, FText& OutErrorMessage);
+
+	virtual TSharedRef<SWidget> GenerateWidgetForColumn(const FName& ColumnName) override;
+
+private:
+	TWeakPtr<FModularRigTreeElement> WeakRigTreeElement;
+	FModularRigTreeDelegates Delegates;
+	TSharedPtr<SSearchableRigHierarchyTreeView> ConnectorComboBox;
+	TSharedPtr<SButton> ResetConnectorButton;
+	TSharedPtr<SButton> UseSelectedButton;
+	TSharedPtr<SButton> SelectElementButton;
+	FRigElementKey ConnectorKey;
+	TOptional<FModularRigResolveResult> ConnectorMatches;
+
+	FText GetName(bool bUseShortName) const;
+	FText GetItemTooltip() const;
+	TArray<FRigElementKey> GetTargetKeys() const;
+
+	friend class SModularRigTreeView;
+};
diff --git a/Engine/Plugins/Animation/ControlRig/Source/ControlRigEditor/Private/Editor/Hierarchy/Views/SRigHierarchySearchableTreeView.cpp b/Engine/Plugins/Animation/ControlRig/Source/ControlRigEditor/Private/Editor/Hierarchy/Views/SRigHierarchySearchableTreeView.cpp
new file mode 100644
--- /dev/null
+++ b/Engine/Plugins/Animation/ControlRig/Source/ControlRigEditor/Private/Editor/Hierarchy/Views/SRigHierarchySearchableTreeView.cpp
@@ -0,0 +1,57 @@
+// Copyright Epic Games, Inc. All Rights Reserved.
+
+#include "SRigHierarchySearchableTreeView.h"
+
+#include "Widgets/Input/SSearchBox.h"
+#include "Widgets/Layout/SScrollBox.h"
+
+void SSearchableRigHierarchyTreeView::Construct(const FArguments& InArgs)
+{
+	FRigTreeDelegates TreeDelegates = InArgs._RigTreeDelegates;
+	SuperGetRigTreeDisplaySettings = TreeDelegates.OnGetDisplaySettings;
+
+	TreeDelegates.OnGetDisplaySettings.BindSP(this, &SSearchableRigHierarchyTreeView::GetDisplaySettings);
+
+	TSharedPtr<SVerticalBox> VerticalBox;
+	ChildSlot
+	[
+		SNew(SBorder)
+		.BorderImage(FAppStyle::GetBrush("SCSEditor.TreePanel"))
+		[
+			SAssignNew(VerticalBox, SVerticalBox)
+			+ SVerticalBox::Slot()
+			.AutoHeight()
+			[
+				SAssignNew(SearchBox, SSearchBox)
+				.InitialText(InArgs._InitialFilterText)
+				.OnTextChanged(this, &SSearchableRigHierarchyTreeView::OnFilterTextChanged)
+			]
+
+			+SVerticalBox::Slot()
+			[
+				SNew(SScrollBox)
+				+ SScrollBox::Slot()
+				[
+					SAssignNew(TreeView, SRigHierarchyTreeView)
+					.RigTreeDelegates(TreeDelegates)
+				]
+			]
+		]
+	];
+}
+
+const FRigTreeDisplaySettings& SSearchableRigHierarchyTreeView::GetDisplaySettings()
+{
+	if(SuperGetRigTreeDisplaySettings.IsBound())
+	{
+		Settings = SuperGetRigTreeDisplaySettings.Execute();
+	}
+	Settings.FilterText = FilterText;
+	return Settings;
+}
+
+void SSearchableRigHierarchyTreeView::OnFilterTextChanged(const FText& SearchText)
+{
+	FilterText = SearchText;
+	GetTreeView()->RefreshTreeView();
+}
diff --git a/Engine/Plugins/Animation/ControlRig/Source/ControlRigEditor/Private/Editor/Hierarchy/Views/SRigHierarchySearchableTreeView.h b/Engine/Plugins/Animation/ControlRig/Source/ControlRigEditor/Private/Editor/Hierarchy/Views/SRigHierarchySearchableTreeView.h
new file mode 100644
--- /dev/null
+++ b/Engine/Plugins/Animation/ControlRig/Source/ControlRigEditor/Private/Editor/Hierarchy/Views/SRigHierarchySearchableTreeView.h
@@ -0,0 +1,40 @@
+// Copyright Epic Games, Inc. All Rights Reserved.
+
+#pragma once
+
+#include "Editor/Hierarchy/RigHierarchyTreeDelegates.h"
+#include "Editor/Hierarchy/RigHierarchyTreeDisplaySettings.h"
+#include "Internationalization/Text.h"
+#include "Rigs/RigHierarchyDefines.h"
+#include "SRigHierarchyTreeView.h"
+#include "Widgets/SCompoundWidget.h"
+
+class SSearchBox;
+
+class SSearchableRigHierarchyTreeView : public SCompoundWidget
+{
+public:
+
+	SLATE_BEGIN_ARGS(SSearchableRigHierarchyTreeView)
+		{
+		}
+		SLATE_ARGUMENT(FRigTreeDelegates, RigTreeDelegates)
+		SLATE_ARGUMENT(FText, InitialFilterText)
+	SLATE_END_ARGS()
+
+	void Construct(const FArguments& InArgs);
+
+	TSharedRef<SSearchBox> GetSearchBox() const { return SearchBox.ToSharedRef(); }
+	TSharedRef<SRigHierarchyTreeView> GetTreeView() const { return TreeView.ToSharedRef(); }
+	const FRigTreeDisplaySettings& GetDisplaySettings();
+
+private:
+
+	void OnFilterTextChanged(const FText& SearchText);
+
+	FOnGetRigTreeDisplaySettings SuperGetRigTreeDisplaySettings;
+	FText FilterText;
+	FRigTreeDisplaySettings Settings;
+	TSharedPtr<SSearchBox> SearchBox;
+	TSharedPtr<SRigHierarchyTreeView> TreeView;
+};
