Control Rig: Control Rig Hierarchy Editor cleanup first pass

diff --git a/Engine/Plugins/Animation/ControlRig/Source/ControlRigEditor/Private/Editor/Hierarchy/Widgets/SModularRigHierarchy.cpp b/Engine/Plugins/Animation/ControlRig/Source/ControlRigEditor/Private/Editor/Hierarchy/Widgets/SModularRigHierarchy.cpp
new file mode 100644
--- /dev/null
+++ b/Engine/Plugins/Animation/ControlRig/Source/ControlRigEditor/Private/Editor/Hierarchy/Widgets/SModularRigHierarchy.cpp
@@ -0,0 +1,1779 @@
+// Copyright Epic Games, Inc. All Rights Reserved.
+
+#include "SModularRigHierarchy.h"
+
+#include "Widgets/Input/SComboButton.h"
+#include "Styling/AppStyle.h"
+#include "Widgets/Input/SButton.h"
+#include "Widgets/Input/SSearchBox.h"
+#include "Kismet2/BlueprintEditorUtils.h"
+#include "ScopedTransaction.h"
+#include "Editor/ControlRigEditor.h"
+#include "BlueprintActionDatabase.h"
+#include "BlueprintVariableNodeSpawner.h"
+#include "Widgets/Layout/SSpacer.h"
+#include "Widgets/Images/SImage.h"
+#include "Widgets/Text/STextBlock.h"
+#include "K2Node_VariableGet.h"
+#include "RigVMBlueprintUtils.h"
+#include "ControlRigModularRigCommands.h"
+#include "ControlRigBlueprintLegacy.h"
+#include "Graph/ControlRigGraph.h"
+#include "Graph/ControlRigGraphNode.h"
+#include "Graph/ControlRigGraphSchema.h"
+#include "GraphEditorModule.h"
+#include "Framework/MultiBox/MultiBoxBuilder.h"
+#include "AnimationRuntime.h"
+#include "ClassViewerFilter.h"
+#include "PropertyCustomizationHelpers.h"
+#include "Framework/Application/SlateApplication.h"
+#include "Editor/EditorEngine.h"
+#include "HelperUtil.h"
+#include "Widgets/Text/SInlineEditableTextBlock.h"
+#include "ControlRig.h"
+#include "HAL/PlatformApplicationMisc.h"
+#include "HAL/PlatformTime.h"
+#include "Dialogs/Dialogs.h"
+#include "IPersonaToolkit.h"
+#include "Types/WidgetActiveTimerDelegate.h"
+#include "Dialog/SCustomDialog.h"
+#include "EditMode/ControlRigEditMode.h"
+#include "ToolMenus.h"
+#include "Editor/ControlRigContextMenuContext.h"
+#include "Editor/SRigSpacePickerWidget.h"
+#include "Settings/ControlRigSettings.h"
+#include "Framework/Notifications/NotificationManager.h"
+#include "Widgets/Notifications/SNotificationList.h"
+#include "Styling/AppStyle.h"
+#include "Editor/ControlRigSkeletalMeshComponent.h"
+#include "Sequencer/ControlRigLayerInstance.h"
+#include "Algo/MinElement.h"
+#include "Algo/MaxElement.h"
+#include "AssetRegistry/AssetRegistryModule.h"
+#include "DragAndDrop/AssetDragDropOp.h"
+#include "Editor/RigVMEditorTools.h"
+#include "Kismet2/SClassPickerDialog.h"
+#include "RigVMFunctions/Math/RigVMMathLibrary.h"
+#include "Preferences/PersonaOptions.h"
+#include "Widgets/SRigVMBulkEditDialog.h"
+#include "Widgets/SRigVMSwapAssetReferencesWidget.h"
+#include "Editor/Hierarchy/DragDrop/ModularRigHierarchyElementDragDropOp.h"
+
+#if WITH_RIGVMLEGACYEDITOR
+#include "SKismetInspector.h"
+#else
+#include "Editor/SRigVMDetailsInspector.h"
+#endif
+
+#define LOCTEXT_NAMESPACE "SModularRigModel"
+
+///////////////////////////////////////////////////////////
+
+const FName SModularRigModel::ContextMenuName = TEXT("ControlRigEditor.ModularRigModel.ContextMenu");
+
+SModularRigModel::~SModularRigModel()
+{
+	IControlRigBaseEditor* Editor = ControlRigEditor.IsValid() ? ControlRigEditor.Pin().Get() : nullptr;
+	OnEditorClose(Editor, ControlRigBlueprint.GetObject());
+}
+
+void SModularRigModel::Construct(const FArguments& InArgs, TSharedRef<IControlRigBaseEditor> InControlRigEditor)
+{
+	ControlRigEditor = InControlRigEditor;
+
+	ControlRigBlueprint = ControlRigEditor.Pin()->GetControlRigAssetInterface();
+
+	ControlRigBlueprint->GetRigVMAssetInterface()->OnRefreshEditor().AddRaw(this, &SModularRigModel::HandleRefreshEditorFromBlueprint);
+	ControlRigBlueprint->GetRigVMAssetInterface()->OnSetObjectBeingDebugged().AddRaw(this, &SModularRigModel::HandleSetObjectBeingDebugged);
+	ControlRigBlueprint->OnModularRigPreCompiled().AddRaw(this, &SModularRigModel::HandlePreCompileModularRigs);
+	ControlRigBlueprint->OnModularRigCompiled().AddRaw(this, &SModularRigModel::HandlePostCompileModularRigs);
+
+	if(UModularRigController* ModularRigController = ControlRigBlueprint->GetModularRigController())
+	{
+		ModularRigController->OnModified().AddSP(this, &SModularRigModel::OnModularRigModified);
+	}
+
+	// for deleting, renaming, dragging
+	CommandList = MakeShared<FUICommandList>();
+
+	UEditorEngine* Editor = Cast<UEditorEngine>(GEngine);
+	if (Editor != nullptr)
+	{
+		Editor->RegisterForUndo(this);
+	}
+
+	BindCommands();
+
+	bShowSecondaryConnectors = false;
+	bShowOptionalConnectors = false;
+	bShowUnresolvedConnectors = true;
+	bIsPerformingSelection = false;
+	bKeepCurrentEditedConnectors = false;
+
+	// setup all delegates for the modular rig model widget
+	FModularRigTreeDelegates Delegates;
+	Delegates.OnGetModularRig = FOnGetModularRigTreeRig::CreateSP(this, &SModularRigModel::GetModularRigForTreeView);
+	Delegates.OnContextMenuOpening = FOnContextMenuOpening::CreateSP(this, &SModularRigModel::CreateContextMenuWidget);
+	Delegates.OnDragDetected = FOnDragDetected::CreateSP(this, &SModularRigModel::OnDragDetected);
+	Delegates.OnCanAcceptDrop = FOnModularRigTreeCanAcceptDrop::CreateSP(this, &SModularRigModel::OnCanAcceptDrop);
+	Delegates.OnAcceptDrop = FOnModularRigTreeAcceptDrop::CreateSP(this, &SModularRigModel::OnAcceptDrop);
+	Delegates.OnMouseButtonClick = FOnModularRigTreeMouseButtonClick::CreateSP(this, &SModularRigModel::OnItemClicked);
+	Delegates.OnMouseButtonDoubleClick = FOnModularRigTreeMouseButtonClick::CreateSP(this, &SModularRigModel::OnItemDoubleClicked);
+	Delegates.OnRequestDetailsInspection = FOnModularRigTreeRequestDetailsInspection::CreateSP(this, &SModularRigModel::OnRequestDetailsInspection);
+	Delegates.OnRenameElement = FOnModularRigTreeRenameElement::CreateSP(this, &SModularRigModel::HandleRenameModule);
+	Delegates.OnVerifyModuleNameChanged = FOnModularRigTreeVerifyElementNameChanged::CreateSP(this, &SModularRigModel::HandleVerifyNameChanged);
+	Delegates.OnResolveConnector = FOnModularRigTreeResolveConnector::CreateSP(this, &SModularRigModel::HandleConnectorResolved);
+	Delegates.OnDisconnectConnector = FOnModularRigTreeDisconnectConnector::CreateSP(this, &SModularRigModel::HandleConnectorDisconnect);
+	Delegates.OnSelectionChanged = FOnModularRigTreeSelectionChanged::CreateSP(this, &SModularRigModel::HandleSelectionChanged);
+	Delegates.OnAlwaysShowConnector = FOnModularRigTreeAlwaysShowConnector::CreateSP(this, &SModularRigModel::ShouldAlwaysShowConnector);
+
+	HeaderRowWidget = SNew(SHeaderRow)
+		.Visibility(EVisibility::Visible);
+
+	HeaderRowWidget->AddColumn(
+		SHeaderRow::Column(SModularRigTreeView::Column_Module)
+		.DefaultLabel(FText::FromName(SModularRigTreeView::Column_Module))
+		.HAlignCell(HAlign_Left)
+		.HAlignHeader(HAlign_Left)
+		.VAlignCell(VAlign_Top)
+	);
+	HeaderRowWidget->AddColumn(
+		SHeaderRow::Column(SModularRigTreeView::Column_Tags)
+		.DefaultLabel(FText())
+		.HAlignCell(HAlign_Fill)
+		.HAlignHeader(HAlign_Fill)
+		.FixedWidth(16.f)
+		.VAlignCell(VAlign_Top)
+	);
+	HeaderRowWidget->AddColumn(
+		SHeaderRow::Column(SModularRigTreeView::Column_Connector)
+		.DefaultLabel(FText::FromName(SModularRigTreeView::Column_Connector))
+		.HAlignCell(HAlign_Left)
+		.HAlignHeader(HAlign_Left)
+		.VAlignCell(VAlign_Top)
+	);
+	HeaderRowWidget->AddColumn(
+		SHeaderRow::Column(SModularRigTreeView::Column_Buttons)
+		.DefaultLabel(FText::FromName(SModularRigTreeView::Column_Buttons))
+		.ManualWidth(60)
+		.HAlignCell(HAlign_Left)
+		.HAlignHeader(HAlign_Left)
+		.VAlignCell(VAlign_Top)
+	);
+	
+	ChildSlot
+	[
+		SNew(SVerticalBox)
+		
+		+SVerticalBox::Slot()
+		.Padding(0.0f, 0.0f)
+		.AutoHeight()
+		[
+			SNew(SHorizontalBox)
+			
+			+SHorizontalBox::Slot()
+			.Padding(0.0f, 0.0f)
+			.HAlign(HAlign_Left)
+			.AutoWidth()
+			[
+				SNew(SComboButton)
+			   .ComboButtonStyle(&FAppStyle::Get().GetWidgetStyle<FComboButtonStyle>("SimpleComboButtonWithIcon"))
+			   .ForegroundColor(FSlateColor::UseStyle())
+			   .ToolTipText(LOCTEXT("OptionsToolTip", "Open the Options Menu ."))
+			   .OnGetMenuContent(this, &SModularRigModel::OnGetOptionsMenu)
+			   .ContentPadding(FMargin(1, 0))
+			   .ButtonContent()
+			   [
+					SNew(SImage)
+					.Image(FAppStyle::Get().GetBrush("Icons.Filter"))
+					.ColorAndOpacity(FSlateColor::UseForeground())
+			   ]
+			]
+
+			+SHorizontalBox::Slot()
+			.Padding(4.0f, 0.0f, 0.0f, 0.0f)
+			.HAlign(HAlign_Fill)
+			[
+				SAssignNew(FilterBox, SSearchBox)
+				.OnTextChanged(this, &SModularRigModel::OnFilterTextChanged)
+			]
+		]
+		
+		+SVerticalBox::Slot()
+		.Padding(0.0f, 0.0f)
+		[
+			SNew(SBorder)
+			.Padding(0.0f)
+			.ShowEffectWhenDisabled(false)
+			[
+				SNew(SBorder)
+				.Padding(2.0f)
+				.BorderImage(FAppStyle::GetBrush("SCSEditor.TreePanel"))
+				[
+					SAssignNew(TreeView, SModularRigTreeView)
+					.HeaderRow(HeaderRowWidget)
+					.RigTreeDelegates(Delegates)
+					.AutoScrollEnabled(true)
+					.FilterText_Lambda([this]() { return FilterText; })
+					.ShowSecondaryConnectors_Lambda([this]() { return bShowSecondaryConnectors; })
+					.ShowOptionalConnectors_Lambda([this]() { return bShowOptionalConnectors; })
+					.ShowUnresolvedConnectors_Lambda([this]() { return bShowUnresolvedConnectors; })
+				]
+			]
+		]
+	];
+
+	RefreshTreeView();
+
+	if (ControlRigEditor.IsValid())
+	{
+		ControlRigEditor.Pin()->GetKeyDownDelegate().BindLambda([&](const FGeometry& MyGeometry, const FKeyEvent& InKeyEvent)->FReply {
+			return OnKeyDown(MyGeometry, InKeyEvent);
+		});
+		ControlRigEditor.Pin()->OnGetViewportContextMenu().BindSP(this, &SModularRigModel::GetContextMenu);
+		ControlRigEditor.Pin()->OnViewportContextMenuCommands().BindSP(this, &SModularRigModel::GetContextMenuCommands);
+		ControlRigEditor.Pin()->OnEditorClosed().AddSP(this, &SModularRigModel::OnEditorClose);
+	}
+	
+	CreateContextMenu();
+
+	if(const UModularRig* Rig = GetModularRigForTreeView())
+	{
+		if(URigHierarchy* Hierarchy = Rig->GetHierarchy())
+		{
+			Hierarchy->OnModified().AddSP(this, &SModularRigModel::OnHierarchyModified);
+		}
+	}
+}
+
+void SModularRigModel::OnEditorClose(IControlRigBaseEditor* InEditor, FControlRigAssetInterfacePtr InBlueprint)
+{
+	if (InEditor)
+	{
+		IControlRigBaseEditor* Editor = (IControlRigBaseEditor*)InEditor;  
+		Editor->OnGetViewportContextMenu().Unbind();
+		Editor->OnViewportContextMenuCommands().Unbind();
+		Editor->OnEditorClosed().RemoveAll(this);
+	}
+
+	if (InBlueprint)
+	{
+		InBlueprint->GetRigVMAssetInterface()->OnRefreshEditor().RemoveAll(this);
+		InBlueprint->GetRigVMAssetInterface()->OnSetObjectBeingDebugged().RemoveAll(this);
+		InBlueprint->OnModularRigPreCompiled().RemoveAll(this);
+		InBlueprint->OnModularRigCompiled().RemoveAll(this);
+		if(UModularRigController* ModularRigController = InBlueprint->GetModularRigController())
+		{
+			ModularRigController->OnModified().RemoveAll(this);
+		}
+	}
+
+	if(const UModularRig* Rig = GetModularRigForTreeView())
+	{
+		if(URigHierarchy* Hierarchy = Rig->GetHierarchy())
+		{
+			Hierarchy->OnModified().RemoveAll(this);
+		}
+	}
+
+	ControlRigEditor.Reset();
+	ControlRigBlueprint.Reset();
+}
+
+void SModularRigModel::BindCommands()
+{
+	// create new command
+	const FControlRigModularRigCommands& Commands = FControlRigModularRigCommands::Get();
+
+	CommandList->MapAction(Commands.AddModuleItem,
+		FExecuteAction::CreateSP(this, &SModularRigModel::HandleNewItem),
+		FCanExecuteAction());
+
+	CommandList->MapAction(Commands.RenameModuleItem,
+		FExecuteAction::CreateSP(this, &SModularRigModel::HandleRenameModule),
+		FCanExecuteAction());
+
+	CommandList->MapAction(Commands.DeleteModuleItem,
+		FExecuteAction::CreateSP(this, &SModularRigModel::HandleDeleteModules),
+		FCanExecuteAction());
+
+	CommandList->MapAction(Commands.MirrorModuleItem,
+		FExecuteAction::CreateSP(this, &SModularRigModel::HandleMirrorModules),
+		FCanExecuteAction());
+
+	CommandList->MapAction(Commands.ReresolveModuleItem,
+		FExecuteAction::CreateSP(this, &SModularRigModel::HandleReresolveModules),
+		FCanExecuteAction());
+
+	CommandList->MapAction(Commands.SwapModuleClassItem,
+		FExecuteAction::CreateSP(this, &SModularRigModel::HandleSwapClassForModules),
+		FCanExecuteAction::CreateSP(this, &SModularRigModel::CanSwapModules));
+
+	CommandList->MapAction(Commands.CopyModuleSettings,
+		FExecuteAction::CreateSP(this, &SModularRigModel::HandleCopyModuleSettings),
+		FCanExecuteAction::CreateSP(this, &SModularRigModel::CanCopyModuleSettings));
+
+	CommandList->MapAction(Commands.PasteModuleSettings,
+		FExecuteAction::CreateSP(this, &SModularRigModel::HandlePasteModuleSettings),
+		FCanExecuteAction::CreateSP(this, &SModularRigModel::CanPasteModuleSettings));
+}
+
+FReply SModularRigModel::OnKeyDown(const FGeometry& MyGeometry, const FKeyEvent& InKeyEvent)
+{
+	if (CommandList.IsValid() && CommandList->ProcessCommandBindings(InKeyEvent))
+	{
+		return FReply::Handled();
+	}
+	return FReply::Unhandled();
+}
+
+FReply SModularRigModel::OnMouseButtonDown(const FGeometry& MyGeometry, const FPointerEvent& MouseEvent)
+{
+	const FReply Reply = SCompoundWidget::OnMouseButtonDown(MyGeometry, MouseEvent);
+	if(Reply.IsEventHandled())
+	{
+		return Reply;
+	}
+
+	if(MouseEvent.GetEffectingButton() == EKeys::MiddleMouseButton)
+	{
+		if(const TSharedPtr<FModularRigTreeElement>* ItemPtr = TreeView->FindItemAtPosition(MouseEvent.GetScreenSpacePosition()))
+		{
+			if(const TSharedPtr<FModularRigTreeElement>& Item = *ItemPtr)
+			{
+				if (ControlRigBlueprint.IsValid())
+				{
+					UModularRigController* Controller = ControlRigBlueprint->GetModularRigController();
+					check(Controller);
+
+					if(const FRigModuleReference* Module = Controller->FindModule(Item->ModuleName))
+					{
+						TArray<const FRigModuleReference*> ModulesToSelect = {Module};
+						TArray<FName> ModuleNames;
+						for(int32 Index = 0; Index < ModulesToSelect.Num(); Index++)
+						{
+							ModuleNames.Add(ModulesToSelect[Index]->Name);
+							for(const FRigModuleReference* ChildModule : ModulesToSelect[Index]->CachedChildren)
+							{
+								ModulesToSelect.AddUnique(ChildModule);
+							}
+						}
+						
+						Controller->SetModuleSelection(ModuleNames);
+					}
+				}
+			}
+		}
+	}
+
+	return FReply::Unhandled();
+}
+
+void SModularRigModel::RefreshTreeView(bool bRebuildContent)
+{
+	bool bDummySuspensionFlag = false;
+	bool* SuspensionFlagPtr = &bDummySuspensionFlag;
+	if (ControlRigEditor.IsValid())
+	{
+		SuspensionFlagPtr = &ControlRigEditor.Pin()->GetSuspendDetailsPanelRefreshFlag();
+	}
+	TGuardValue<bool> SuspendDetailsPanelRefreshGuard(*SuspensionFlagPtr, true);
+
+	TreeView->RefreshTreeView(bRebuildContent);
+}
+
+TArray<TSharedPtr<FModularRigTreeElement>> SModularRigModel::GetSelectedItems() const
+{
+	TArray<TSharedPtr<FModularRigTreeElement>> SelectedItems = TreeView->GetSelectedItems();
+	SelectedItems.Remove(TSharedPtr<FModularRigTreeElement>(nullptr));
+	return SelectedItems;
+}
+
+TArray<FString> SModularRigModel::GetSelectedKeys() const
+{
+	const TArray<TSharedPtr<FModularRigTreeElement>> SelectedItems = GetSelectedItems();
+	
+	TArray<FString> SelectedKeys;
+	for (const TSharedPtr<FModularRigTreeElement>& SelectedItem : SelectedItems)
+	{
+		if (SelectedItem.IsValid())
+		{
+			if(!SelectedItem->Key.IsEmpty())
+			{
+				SelectedKeys.AddUnique(SelectedItem->Key);
+			}
+		}
+	}
+
+	return SelectedKeys;
+}
+
+
+void SModularRigModel::HandlePreCompileModularRigs(FRigVMAssetInterfacePtr InBlueprint)
+{
+}
+
+void SModularRigModel::HandlePostCompileModularRigs(FRigVMAssetInterfacePtr InBlueprint)
+{
+	if(!bKeepCurrentEditedConnectors)
+	{
+		CurrentlyEditedConnectors.Reset();
+	}
+	
+	RefreshTreeView();
+	if (ControlRigEditor.IsValid())
+	{
+		TArray<TSharedPtr<FModularRigTreeElement>> SelectedElements;
+		Algo::Transform(ControlRigEditor.Pin()->GetSelectedModules(), SelectedElements, [this](const FName& ModuleName)
+		{
+			return TreeView->FindElement(ModuleName.ToString());
+		});
+		TreeView->SetSelection(SelectedElements);
+		ControlRigEditor.Pin()->RefreshDetailView();
+	}
+}
+
+void SModularRigModel::HandleRefreshEditorFromBlueprint(FRigVMAssetInterfacePtr InBlueprint)
+{
+	RefreshTreeView();
+}
+
+void SModularRigModel::HandleSetObjectBeingDebugged(UObject* InObject)
+{
+	if(ControlRigBeingDebuggedPtr.Get() == InObject)
+	{
+		return;
+	}
+
+	if(ControlRigBeingDebuggedPtr.IsValid())
+	{
+		if(URigHierarchy* Hierarchy = ControlRigBeingDebuggedPtr->GetHierarchy())
+		{
+			Hierarchy->OnModified().RemoveAll(this);
+		}
+	}
+
+	ControlRigBeingDebuggedPtr.Reset();
+
+	if(UModularRig* ControlRig = Cast<UModularRig>(InObject))
+	{
+		ControlRigBeingDebuggedPtr = ControlRig;
+
+		if(URigHierarchy* Hierarchy = ControlRig->GetHierarchy())
+		{
+			Hierarchy->OnModified().AddSP(this, &SModularRigModel::OnHierarchyModified);
+		}
+	}
+
+	RefreshTreeView();
+}
+
+TSharedRef<SWidget> SModularRigModel::OnGetOptionsMenu()
+{
+	FMenuBuilder MenuBuilder(true, nullptr);
+
+	const FCanExecuteAction CanExecuteAction = FCanExecuteAction::CreateLambda([]() { return true; });
+
+	MenuBuilder.BeginSection("FilterOptions", LOCTEXT("FilterOptions", "Filter Options"));
+	{
+		MenuBuilder.AddMenuEntry(
+			LOCTEXT("SecondaryConnectors", "Secondary Connectors"),
+			LOCTEXT("SecondaryConnectorsToolTip", "Toggle the display of secondary connectors"),
+			FSlateIcon(FControlRigEditorStyle::Get().GetStyleSetName(), TEXT("ControlRig.ConnectorSecondary")),
+			FUIAction(
+				FExecuteAction::CreateLambda([this](){
+					bShowSecondaryConnectors = !bShowSecondaryConnectors;
+					RefreshTreeView(true);
+				}),
+				CanExecuteAction,
+				FIsActionChecked::CreateLambda([this]()
+				{
+					return bShowSecondaryConnectors;
+				})
+			),
+			NAME_None,
+			EUserInterfaceActionType::Check
+		);
+
+		MenuBuilder.AddMenuEntry(
+			LOCTEXT("OptionalConnectors", "Optional Connectors"),
+			LOCTEXT("OptionalConnectorsToolTip", "Toggle the display of secondary connectors"),
+			FSlateIcon(FControlRigEditorStyle::Get().GetStyleSetName(), TEXT("ControlRig.ConnectorOptional")),
+			FUIAction(
+				FExecuteAction::CreateLambda([this](){
+					bShowOptionalConnectors = !bShowOptionalConnectors;
+					RefreshTreeView(true);
+				}),
+				CanExecuteAction,
+				FIsActionChecked::CreateLambda([this]()
+				{
+					return bShowOptionalConnectors;
+				})
+			),
+			NAME_None,
+			EUserInterfaceActionType::Check
+		);
+
+		MenuBuilder.AddMenuEntry(
+			LOCTEXT("UnresolvedConnectors", "Unresolved Connectors"),
+			LOCTEXT("UnresolvedConnectorsToolTip", "Toggle the display of unresolved connectors"),
+			FSlateIcon(FControlRigEditorStyle::Get().GetStyleSetName(), TEXT("ControlRig.ConnectorWarning")),
+			FUIAction(
+				FExecuteAction::CreateLambda([this](){
+					bShowUnresolvedConnectors = !bShowUnresolvedConnectors;
+					RefreshTreeView(true);
+				}),
+				CanExecuteAction,
+				FIsActionChecked::CreateLambda([this]()
+				{
+					return bShowUnresolvedConnectors;
+				})
+			),
+			NAME_None,
+			EUserInterfaceActionType::Check
+		);
+	}
+	MenuBuilder.EndSection();
+
+	return MenuBuilder.MakeWidget();
+}
+
+void SModularRigModel::OnFilterTextChanged(const FText& SearchText)
+{
+	FilterText = SearchText;
+	RefreshTreeView(true);
+}
+
+TSharedPtr< SWidget > SModularRigModel::CreateContextMenuWidget()
+{
+	UToolMenus* ToolMenus = UToolMenus::Get();
+
+	if (UToolMenu* Menu = GetContextMenu())
+	{
+		return ToolMenus->GenerateWidget(Menu);
+	}
+	
+	return SNullWidget::NullWidget;
+}
+
+void SModularRigModel::OnItemClicked(TSharedPtr<FModularRigTreeElement> InItem)
+{
+}
+
+void SModularRigModel::OnItemDoubleClicked(TSharedPtr<FModularRigTreeElement> InItem)
+{
+
+}
+
+void SModularRigModel::CreateContextMenu()
+{
+	static bool bCreatedMenu = false;
+	if(bCreatedMenu)
+	{
+		return;
+	}
+	bCreatedMenu = true;
+	
+	const FName MenuName = ContextMenuName;
+
+	UToolMenus* ToolMenus = UToolMenus::Get();
+	
+	if (!ensure(ToolMenus))
+	{
+		return;
+	}
+
+	if (UToolMenu* Menu = ToolMenus->ExtendMenu(MenuName))
+	{
+		Menu->AddDynamicSection(NAME_None, FNewToolMenuDelegate::CreateLambda([](UToolMenu* InMenu)
+			{
+				UControlRigContextMenuContext* MainContext = InMenu->FindContext<UControlRigContextMenuContext>();
+				
+				if (SModularRigModel* ModelPanel = MainContext->GetModularRigModelPanel())
+				{
+					const FControlRigModularRigCommands& Commands = FControlRigModularRigCommands::Get(); 
+				
+					FToolMenuSection& ModulesSection = InMenu->AddSection(TEXT("Modules"), LOCTEXT("ModulesHeader", "Modules"));
+					ModulesSection.AddSubMenu(TEXT("New"), LOCTEXT("New", "New"), LOCTEXT("New_ToolTip", "Create New Modules"),
+						FNewToolMenuDelegate::CreateLambda([Commands, ModelPanel](UToolMenu* InSubMenu)
+						{
+							FToolMenuSection& DefaultSection = InSubMenu->AddSection(NAME_None);
+							DefaultSection.AddMenuEntry(Commands.AddModuleItem);
+						})
+					);
+					ModulesSection.AddMenuEntry(Commands.RenameModuleItem);
+					ModulesSection.AddMenuEntry(Commands.DeleteModuleItem);
+					ModulesSection.AddMenuEntry(Commands.MirrorModuleItem);
+					ModulesSection.AddMenuEntry(Commands.ReresolveModuleItem);
+					ModulesSection.AddMenuEntry(Commands.SwapModuleClassItem);
+					ModulesSection.AddMenuEntry(Commands.CopyModuleSettings);
+					ModulesSection.AddMenuEntry(Commands.PasteModuleSettings);
+				}
+			})
+		);
+	}
+}
+
+UToolMenu* SModularRigModel::GetContextMenu()
+{
+	const FName MenuName = ContextMenuName;
+	UToolMenus* ToolMenus = UToolMenus::Get();
+
+	if(!ensure(ToolMenus))
+	{
+		return nullptr;
+	}
+
+	// individual entries in this menu can access members of this context, particularly useful for editor scripting
+	UControlRigContextMenuContext* ContextMenuContext = NewObject<UControlRigContextMenuContext>();
+	FControlRigMenuSpecificContext MenuSpecificContext;
+	MenuSpecificContext.ModularRigModelPanel = SharedThis(this);
+	ContextMenuContext->Init(ControlRigEditor, MenuSpecificContext);
+
+	FToolMenuContext MenuContext(CommandList);
+	MenuContext.AddObject(ContextMenuContext);
+
+	UToolMenu* Menu = ToolMenus->GenerateMenu(MenuName, MenuContext);
+
+	return Menu;
+}
+
+TSharedPtr<FUICommandList> SModularRigModel::GetContextMenuCommands() const
+{
+	return CommandList;
+}
+
+bool SModularRigModel::IsSingleSelected() const
+{
+	if(GetSelectedKeys().Num() == 1)
+	{
+		return true;
+	}
+	return false;
+}
+
+/** Filter class to show only RigModules. */
+class FClassViewerRigModulesFilter : public IClassViewerFilter
+{
+public:
+	FClassViewerRigModulesFilter()
+		: AssetRegistry(FModuleManager::GetModuleChecked<FAssetRegistryModule>(TEXT("AssetRegistry")).Get())
+	{}
+	
+	virtual bool IsClassAllowed(const FClassViewerInitializationOptions& InInitOptions, const UClass* InClass, TSharedRef< FClassViewerFilterFuncs > InFilterFuncs) override
+	{
+		if(InClass)
+		{
+			const bool bChildOfObjectClass = InClass->IsChildOf(UControlRig::StaticClass());
+			const bool bMatchesFlags = !InClass->HasAnyClassFlags(CLASS_Hidden | CLASS_HideDropDown | CLASS_Deprecated | CLASS_Abstract);
+			const bool bNotNative = !InClass->IsNative();
+
+			// Allow any class contained in the extra picker common classes array
+			if (InInitOptions.ExtraPickerCommonClasses.Contains(InClass))
+			{
+				return true;
+			}
+			
+			if (bChildOfObjectClass && bMatchesFlags && bNotNative)
+			{
+				const FAssetData AssetData(InClass);
+				return MatchesFilter(AssetData);
+			}
+		}
+		return false;
+	}
+
+	virtual bool IsUnloadedClassAllowed(const FClassViewerInitializationOptions& InInitOptions, const TSharedRef< const IUnloadedBlueprintData > InUnloadedClassData, TSharedRef< FClassViewerFilterFuncs > InFilterFuncs) override
+	{
+		const bool bChildOfObjectClass = InUnloadedClassData->IsChildOf(UControlRig::StaticClass());
+		const bool bMatchesFlags = !InUnloadedClassData->HasAnyClassFlags(CLASS_Hidden | CLASS_HideDropDown | CLASS_Deprecated | CLASS_Abstract);
+		if (bChildOfObjectClass && bMatchesFlags)
+		{
+			const FString GeneratedClassPathString = InUnloadedClassData->GetClassPathName().ToString();
+			const FString BlueprintPath = GeneratedClassPathString.LeftChop(2); // Chop off _C
+			const FAssetData AssetData = AssetRegistry.GetAssetByObjectPath(FSoftObjectPath(BlueprintPath));
+			return MatchesFilter(AssetData);
+
+		}
+		return false;
+	}
+
+private:
+	bool MatchesFilter(const FAssetData& AssetData)
+	{
+		static const UEnum* ControlTypeEnum = StaticEnum<EControlRigType>();
+		const FString ControlRigTypeStr = AssetData.GetTagValueRef<FString>(TEXT("ControlRigType"));
+		if (ControlRigTypeStr.IsEmpty())
+		{
+			return false;
+		}
+
+		const EControlRigType ControlRigType = (EControlRigType)(ControlTypeEnum->GetValueByName(*ControlRigTypeStr));
+		return ControlRigType == EControlRigType::RigModule;
+	}
+
+	const IAssetRegistry& AssetRegistry;
+};
+
+/** Create Item */
+void SModularRigModel::HandleNewItem()
+{
+	if(!ControlRigEditor.IsValid())
+	{
+		return;
+	}
+
+	FName ParentModuleName = NAME_None;
+	if (IsSingleSelected())
+	{
+		const TSharedPtr<FModularRigTreeElement> ParentElement = TreeView->FindElement(GetSelectedKeys()[0]);
+		if (ParentElement.IsValid())
+		{
+			ParentModuleName = ParentElement->ModuleName;
+		}
+	}
+	
+	FClassViewerInitializationOptions Options;
+	Options.bShowUnloadedBlueprints = true;
+	Options.NameTypeToDisplay = EClassViewerNameTypeToDisplay::DisplayName;
+
+	TSharedPtr<FClassViewerRigModulesFilter> ClassFilter = MakeShareable(new FClassViewerRigModulesFilter());
+	Options.ClassFilters.Add(ClassFilter.ToSharedRef());
+	Options.bShowNoneOption = false;
+	
+	UClass* ChosenClass;
+	const FText TitleText = LOCTEXT("ModularRigModelPickModuleClass", "Pick Rig Module Class");
+	const bool bPressedOk = SClassPickerDialog::PickClass(TitleText, Options, ChosenClass, UControlRig::StaticClass());
+	if (bPressedOk)
+	{
+		HandleNewItem(ChosenClass, ParentModuleName);
+	}
+}
+
+void SModularRigModel::HandleNewItem(UClass* InClass, const FName &InParentModuleName)
+{
+	UControlRig* ControlRig = InClass->GetDefaultObject<UControlRig>();
+	if (!ControlRig)
+	{
+		return;
+	}
+
+	FSlateApplication::Get().DismissAllMenus();
+	
+	if (ControlRigBlueprint.IsValid())
+	{
+		UModularRigController* Controller = ControlRigBlueprint->GetModularRigController();
+		
+		FString ClassName = InClass->GetName();
+		ClassName.RemoveFromEnd(TEXT("_C"));
+		const FRigName Name = Controller->GetSafeNewName(FRigName(ClassName));
+		const FName NewModuleName = Controller->AddModule(Name, InClass, InParentModuleName);
+		TSharedPtr<FModularRigTreeElement> Element = TreeView->FindElement(NewModuleName.ToString());
+		if (Element.IsValid())
+		{
+			TreeView->SetSelection({Element});
+			TreeView->bRequestRenameSelected = true;
+		}
+	}
+}
+
+bool SModularRigModel::CanRenameModule() const
+{
+	return IsSingleSelected() && TreeView->FindElement(GetSelectedKeys()[0])->bIsPrimary;
+}
+
+void SModularRigModel::HandleRenameModule()
+{
+	if(!ControlRigEditor.IsValid())
+	{
+		return;
+	}
+
+	if (!CanRenameModule())
+	{
+		return;
+	}
+
+	UModularRig* Rig = GetDefaultModularRig();
+	if (Rig)
+	{
+		FScopedTransaction Transaction(LOCTEXT("ModularRigModelRenameSelected", "Rename selected module"));
+
+		const TArray<TSharedPtr<FModularRigTreeElement>> SelectedItems = GetSelectedItems();
+		if (SelectedItems.Num() == 1)
+		{
+			SelectedItems[0]->RequestRename();
+		}
+	}
+
+	return;
+}
+
+FName SModularRigModel::HandleRenameModule(const FName& InOldModuleName, const FName& InNewName)
+{
+	if (ControlRigBlueprint.IsValid())
+	{
+		FScopedTransaction Transaction(LOCTEXT("ModularRigModelRename", "Rename Module"));
+
+		UModularRigController* Controller = ControlRigBlueprint->GetModularRigController();
+		check(Controller);
+
+		const FName NewModuleName = Controller->RenameModule(InOldModuleName, InNewName);
+		if (!NewModuleName.IsNone())
+		{
+			return InNewName;
+		}
+	}
+
+	return NAME_None;
+}
+
+bool SModularRigModel::HandleVerifyNameChanged(const FName& InOldModuleName, const FName& InNewName, FText& OutErrorMessage)
+{
+	if (InNewName.IsNone())
+	{
+		return false;
+	}
+
+	if (ControlRigBlueprint.IsValid())
+	{
+		UModularRigController* Controller = ControlRigBlueprint->GetModularRigController();
+		check(Controller);
+
+		return Controller->CanRenameModule(InOldModuleName, InNewName, OutErrorMessage);
+	}
+
+	return false;
+}
+
+void SModularRigModel::HandleDeleteModules()
+{
+	if(!ControlRigEditor.IsValid())
+	{
+		return;
+	}
+
+	UModularRig* Rig = GetDefaultModularRig();
+	if (Rig)
+	{
+		FScopedTransaction Transaction(LOCTEXT("ModularRigModelDeleteSelected", "Delete selected modules"));
+
+		const TArray<TSharedPtr<FModularRigTreeElement>> SelectedItems = GetSelectedItems();
+		TArray<FName> SelectedModuleNames;
+		Algo::Transform(SelectedItems, SelectedModuleNames, [](const TSharedPtr<FModularRigTreeElement>& Element)
+		{
+			if (Element.IsValid())
+			{
+				return Element->ModuleName;
+			}
+			return FName(NAME_None);
+		});
+		HandleDeleteModules(SelectedModuleNames);
+	}
+
+	return;
+}
+
+void SModularRigModel::HandleDeleteModules(const TArray<FName>& InModuleNames)
+{
+	if (ControlRigBlueprint.IsValid())
+	{
+		FScopedTransaction Transaction(LOCTEXT("ModularRigModelDelete", "Delete Modules"));
+
+		UModularRigController* Controller = ControlRigBlueprint->GetModularRigController();
+		check(Controller);
+
+		// Make sure we delete the modules from children to root
+		TArray<FName> SortedModuleNames = Controller->Model->SortModuleNames(InModuleNames);
+		Algo::Reverse(SortedModuleNames);
+		for (const FName& ModuleName : SortedModuleNames)
+		{
+			Controller->DeleteModule(ModuleName);
+		}
+	}
+}
+
+void SModularRigModel::HandleReparentModules(const TArray<FName>& InModuleNames, const FName& InParentModuleName, int32 NewModuleIndex)
+{
+	if (ControlRigBlueprint.IsValid())
+	{
+		FScopedTransaction Transaction(LOCTEXT("ModularRigModelReparent", "Reparent Modules"));
+
+		UModularRigController* Controller = ControlRigBlueprint->GetModularRigController();
+		check(Controller);
+
+		for (const FName& ModuleName : InModuleNames)
+		{
+			Controller->ReparentModule(ModuleName, InParentModuleName);
+			if(NewModuleIndex != INDEX_NONE)
+			{
+				Controller->ReorderModule(ModuleName, NewModuleIndex);
+			}
+		}
+	}
+}
+
+void SModularRigModel::HandleMirrorModules()
+{
+	if(!ControlRigEditor.IsValid())
+	{
+		return;
+	}
+
+	UModularRig* Rig = GetDefaultModularRig();
+	if (Rig)
+	{
+		const TArray<TSharedPtr<FModularRigTreeElement>> SelectedItems = GetSelectedItems();
+		TArray<FName> SelectedModuleNames;
+		Algo::Transform(SelectedItems, SelectedModuleNames, [](const TSharedPtr<FModularRigTreeElement>& Element)
+		{
+			if (Element.IsValid())
+			{
+				return Element->ModuleName;
+			}
+			return FName(NAME_None);
+		});
+		HandleMirrorModules(SelectedModuleNames);
+	}
+}
+
+void SModularRigModel::HandleMirrorModules(const TArray<FName>& InModuleNames)
+{
+	if (ControlRigBlueprint.IsValid())
+	{
+		FRigVMMirrorSettings Settings;
+		TSharedPtr<FStructOnScope> StructToDisplay = MakeShareable(new FStructOnScope(FRigVMMirrorSettings::StaticStruct(), (uint8*)&Settings));
+
+#if WITH_RIGVMLEGACYEDITOR
+		TSharedRef<SKismetInspector> DetailsInspector = SNew(SKismetInspector);
+#else
+		TSharedRef<SRigVMDetailsInspector> DetailsInspector = SNew(SRigVMDetailsInspector);
+#endif
+		DetailsInspector->ShowSingleStruct(StructToDisplay);
+		
+		TSharedRef<SCustomDialog> MirrorDialog = SNew(SCustomDialog)
+			.Title(FText(LOCTEXT("ControlModularModelMirror", "Mirror Selected Modules")))
+			.Content()
+			[
+				DetailsInspector
+			]
+			.Buttons({
+				SCustomDialog::FButton(LOCTEXT("OK", "OK")),
+				SCustomDialog::FButton(LOCTEXT("Cancel", "Cancel"))
+		});
+
+		if (MirrorDialog->ShowModal() == 0)
+		{
+			FScopedTransaction Transaction(LOCTEXT("ModularRigModelMirror", "Mirror Modules"));
+
+			UModularRigController* Controller = ControlRigBlueprint->GetModularRigController();
+			check(Controller);
+
+			// Make sure we mirror the modules from root to children
+			TArray<FName> SortedModuleNames = Controller->Model->SortModuleNames(InModuleNames);
+			for (const FName& ModuleName : SortedModuleNames)
+			{
+				Controller->MirrorModule(ModuleName, Settings);
+			}
+		}
+	}
+}
+
+void SModularRigModel::HandleReresolveModules()
+{
+	if(!ControlRigEditor.IsValid())
+	{
+		return;
+	}
+
+	UModularRig* Rig = GetDefaultModularRig();
+	if (Rig)
+	{
+		const TArray<TSharedPtr<FModularRigTreeElement>> SelectedItems = GetSelectedItems();
+		TArray<FName> SelectedModuleNames;
+		Algo::Transform(SelectedItems, SelectedModuleNames, [](const TSharedPtr<FModularRigTreeElement>& Element)
+		{
+			if (Element.IsValid())
+			{
+				if(Element->ConnectorName.IsEmpty())
+				{
+					return Element->ModuleName;
+				}
+				return FRigHierarchyModulePath(Element->ModuleName.ToString(), Element->ConnectorName).GetPathFName();
+			}
+			return FName(NAME_None);
+		});
+		HandleReresolveModules(SelectedModuleNames);
+	}
+}
+
+void SModularRigModel::HandleReresolveModules(const TArray<FName>& InModuleNames)
+{
+	if (ControlRigBlueprint.IsValid())
+	{
+		UModularRigController* Controller = ControlRigBlueprint->GetModularRigController();
+		check(Controller);
+
+		const UModularRig* Rig = GetDefaultModularRig();
+		if (Rig == nullptr)
+		{
+			return;
+		}
+
+		const URigHierarchy* Hierarchy = Rig->GetHierarchy();
+		if(Hierarchy == nullptr)
+		{
+			return;
+		}
+
+		TArray<FRigElementKey> ConnectorKeys;
+		for (const FName& PathAndConnector : InModuleNames)
+		{
+			FString ModuleNameString = PathAndConnector.ToString();
+			FString ConnectorName;
+			(void)FRigHierarchyModulePath(PathAndConnector).Split(&ModuleNameString, &ConnectorName);
+
+			const FRigModuleReference* Module = Controller->Model->FindModule(*ModuleNameString);
+			if (!Module)
+			{
+				UE_LOG(LogControlRig, Error, TEXT("Could not find module %s"), *ModuleNameString);
+				return;
+			}
+
+			if(!ConnectorName.IsEmpty())
+			{
+				// if we are executing this on a primary connector we want to re-resolve all secondaries
+				const FRigConnectorElement* PrimaryConnector = Module->FindPrimaryConnector(Hierarchy);
+				const FName DesiredName = Hierarchy->GetNameMetadata(PrimaryConnector->GetKey(), URigHierarchy::DesiredNameMetadataName, NAME_None);
+				if(!DesiredName.IsNone() && DesiredName.ToString().Equals(ConnectorName, ESearchCase::IgnoreCase))
+				{
+					ConnectorName.Reset();
+				}
+			}
+
+			const TArray<const FRigConnectorElement*> Connectors = Module->FindConnectors(Hierarchy);
+			for(const FRigConnectorElement* Connector : Connectors)
+			{
+				if(Connector->IsSecondary())
+				{
+					if(ConnectorName.IsEmpty())
+					{
+						ConnectorKeys.AddUnique(Connector->GetKey());
+					}
+					else
+					{
+						const FName DesiredName = Hierarchy->GetNameMetadata(Connector->GetKey(), URigHierarchy::DesiredNameMetadataName, NAME_None);
+						if(!DesiredName.IsNone() && DesiredName.ToString().Equals(ConnectorName, ESearchCase::IgnoreCase))
+						{
+							ConnectorKeys.AddUnique(Connector->GetKey());
+							break;
+						}
+					}
+				}
+			}
+		}
+
+		Controller->AutoConnectSecondaryConnectors(ConnectorKeys, true, true);
+	}
+}
+
+bool SModularRigModel::CanSwapModules() const
+{
+	// Only if all modules selected have the same module class
+	if(!ControlRigEditor.IsValid())
+	{
+		return false;
+	}
+
+	UModularRig* Rig = GetDefaultModularRig();
+	if (Rig)
+	{
+		TSoftClassPtr<UControlRig> CommonClass = nullptr;
+		const TArray<TSharedPtr<FModularRigTreeElement>> SelectedItems = GetSelectedItems();
+		for (const TSharedPtr<FModularRigTreeElement>& SelectedItem : SelectedItems)
+		{
+			if(!SelectedItem.IsValid())
+			{
+				continue;
+			}
+			TSoftClassPtr<UControlRig> ModuleClass;
+			if (const FRigModuleReference* Module = ControlRigBlueprint->GetModularRigModel().FindModule(SelectedItem->ModuleName))
+			{
+				if (Module->Class.IsValid())
+				{
+					ModuleClass = Module->Class;
+				}
+			}
+			if(!ModuleClass)
+			{
+				return false;
+			}
+			if(!CommonClass)
+			{
+				CommonClass = ModuleClass;
+			}
+			if(ModuleClass != CommonClass)
+			{
+				return false;
+			}
+		}
+		return true;
+	}
+	return false;
+}
+
+void SModularRigModel::HandleSwapClassForModules()
+{
+	if(!ControlRigEditor.IsValid())
+	{
+		return;
+	}
+
+	UModularRig* Rig = GetDefaultModularRig();
+	if (Rig)
+	{
+		const TArray<TSharedPtr<FModularRigTreeElement>> SelectedItems = GetSelectedItems();
+		TArray<FName> SelectedModuleNames;
+		Algo::Transform(SelectedItems, SelectedModuleNames, [](const TSharedPtr<FModularRigTreeElement>& Element)
+		{
+			if (Element.IsValid())
+			{
+				return Element->ModuleName;
+			}
+			return FName(NAME_None);
+		});
+
+		if (SelectedModuleNames.IsEmpty())
+		{
+			return;
+		}
+		
+		HandleSwapClassForModules(SelectedModuleNames);
+	}
+}
+
+void SModularRigModel::HandleSwapClassForModules(const TArray<FName>& InModuleNames)
+{
+	TArray<FSoftObjectPath> ModulePaths;
+	Algo::Transform(InModuleNames, ModulePaths, [this](const FName& InModuleName)
+	{
+		FSoftObjectPath ModulePath(ControlRigBlueprint.GetObject()->GetPathName());
+		ModulePath.SetSubPathString(InModuleName.ToString());
+		return ModulePath;
+	});
+
+	TSoftClassPtr<UControlRig> SourceClass = nullptr;
+	if (FRigModuleReference* Module = ControlRigBlueprint->GetModularRigModel().FindModule(InModuleNames[0]))
+	{
+		SourceClass = Module->Class;
+	}
+
+	if (!SourceClass)
+	{
+		return;
+	}
+
+	FAssetData SourceAsset;
+	if (UPackage* Package = SourceClass->GetPackage())
+	{
+		SourceAsset = FAssetData(Package);
+	}
+	
+	SRigVMSwapAssetReferencesWidget::FArguments WidgetArgs;
+	FRigVMAssetDataFilter FilterModules = FRigVMAssetDataFilter::CreateLambda([](const FAssetData& AssetData)
+		{
+			return UControlRigBlueprint::GetRigType(AssetData) == EControlRigType::RigModule;
+		});
+	TArray<FRigVMAssetDataFilter> SourceFilters = {FilterModules};
+	TArray<FRigVMAssetDataFilter> TargetFilters = {FilterModules};
+	
+	WidgetArgs
+		.EnableUndo(true)
+		.CloseOnSuccess(true)
+		.Source(SourceAsset)
+		.ReferencePaths(ModulePaths)
+		.SkipPickingRefs(true)
+		.OnSwapReference_Lambda([](const FSoftObjectPath& ModulePath, const FAssetData& NewModuleAsset) -> bool
+		{
+			TSubclassOf<UControlRig> NewModuleClass = nullptr;
+			if (const UControlRigBlueprint* ModuleBlueprint = Cast<UControlRigBlueprint>(NewModuleAsset.GetAsset()))
+			{
+				NewModuleClass = ModuleBlueprint->GetRigVMBlueprintGeneratedClass();
+			}
+			else if (UControlRigBlueprintGeneratedClass* GeneratedClass = Cast<UControlRigBlueprintGeneratedClass>(NewModuleAsset.GetAsset()))
+			{
+				NewModuleClass = GeneratedClass;
+			}
+			if (NewModuleClass)
+			{
+				if (UControlRigBlueprint* RigBlueprint = Cast<UControlRigBlueprint>(ModulePath.GetWithoutSubPath().ResolveObject()))
+				{
+					return RigBlueprint->GetModularRigController()->SwapModuleClass(*ModulePath.GetSubPathString(), NewModuleClass);
+				}
+			}
+			return false;
+		})
+		.SourceAssetFilters(SourceFilters)
+		.TargetAssetFilters(TargetFilters);
+
+	const TSharedRef<SRigVMBulkEditDialog<SRigVMSwapAssetReferencesWidget>> SwapModulesDialog =
+		SNew(SRigVMBulkEditDialog<SRigVMSwapAssetReferencesWidget>)
+		.WindowSize(FVector2D(800.0f, 640.0f))
+		.WidgetArgs(WidgetArgs);
+	
+	SwapModulesDialog->ShowNormal();
+}
+
+bool SModularRigModel::CanCopyModuleSettings() const
+{
+	if(!ControlRigEditor.IsValid())
+	{
+		return false;
+	}
+	return !GetSelectedItems().IsEmpty();
+}
+
+void SModularRigModel::HandleCopyModuleSettings()
+{
+	if (ControlRigBlueprint.IsValid())
+	{
+		TArray<FName> SelectedModuleNames;
+		const TArray<TSharedPtr<FModularRigTreeElement>> SelectedItems = GetSelectedItems();
+		for (const TSharedPtr<FModularRigTreeElement>& SelectedItem : SelectedItems)
+		{
+			if (SelectedItem.IsValid())
+			{
+				SelectedModuleNames.AddUnique(SelectedItem->ModuleName);
+			}
+		}
+		
+		const UModularRigController* Controller = ControlRigBlueprint->GetModularRigController();
+		check(Controller);
+
+		const FString ContentAsString = Controller->ExportModuleSettingsToString(SelectedModuleNames);
+		if(!ContentAsString.IsEmpty())
+		{
+			FPlatformApplicationMisc::ClipboardCopy(*ContentAsString);
+		}
+	}
+}
+
+bool SModularRigModel::CanPasteModuleSettings() const
+{
+	if(!ControlRigEditor.IsValid())
+	{
+		return false;
+	}
+
+	FString ContentAsString;
+	FPlatformApplicationMisc::ClipboardPaste(ContentAsString);
+	if(ContentAsString.IsEmpty())
+	{
+		return false;
+	}
+
+	FControlRigOverrideValueErrorPipe ErrorPipe;
+	FModularRigModuleSettingsSetForClipboard Content;
+	FModularRigModuleSettingsSetForClipboard::StaticStruct()->ImportText(*ContentAsString, &Content, nullptr, PPF_None, &ErrorPipe, FModularRigModuleSettingsForClipboard::StaticStruct()->GetName(), true);
+	if(ErrorPipe.GetNumErrors() > 0)
+	{
+		return false;
+	}
+	
+	return GetSelectedItems().Num() == Content.Settings.Num();
+}
+
+void SModularRigModel::HandlePasteModuleSettings()
+{
+	if (ControlRigBlueprint.IsValid())
+	{
+		FString ContentAsString;
+		FPlatformApplicationMisc::ClipboardPaste(ContentAsString);
+		if(ContentAsString.IsEmpty())
+		{
+			return;
+		}
+
+		TArray<FName> SelectedModuleNames;
+		const TArray<TSharedPtr<FModularRigTreeElement>> SelectedItems = GetSelectedItems();
+		for (const TSharedPtr<FModularRigTreeElement>& SelectedItem : SelectedItems)
+		{
+			if (SelectedItem.IsValid())
+			{
+				SelectedModuleNames.AddUnique(SelectedItem->ModuleName);
+			}
+		}
+
+		FScopedTransaction Transaction(LOCTEXT("ModularRigModelResolveConnector", "Resolve Connector"));
+
+		UModularRigController* Controller = ControlRigBlueprint->GetModularRigController();
+		check(Controller);
+
+		(void)Controller->ImportModuleSettingsFromString(ContentAsString, SelectedModuleNames);
+	}
+
+}
+
+void SModularRigModel::HandleConnectorResolved(const FRigElementKey& InConnector, const TArray<FRigElementKey>& InTargets)
+{
+	if (ControlRigBlueprint.IsValid())
+	{
+		FScopedTransaction Transaction(LOCTEXT("ModularRigModelResolveConnector", "Resolve Connector"));
+
+		UModularRigController* Controller = ControlRigBlueprint->GetModularRigController();
+		check(Controller);
+
+		if (const UModularRig* ModularRig = GetModularRig())
+		{
+			if(!bKeepCurrentEditedConnectors)
+			{
+				CurrentlyEditedConnectors.Reset();
+			}
+			const TGuardValue<bool> KeepCurrentEditedConnectorsGuard(bKeepCurrentEditedConnectors, true); 
+			CurrentlyEditedConnectors.Add(InConnector.Name);
+			Controller->ConnectConnectorToElements(InConnector, InTargets, true, ModularRig->GetModularRigSettings().bAutoResolve);
+		}
+	}
+}
+
+void SModularRigModel::HandleConnectorDisconnect(const FRigElementKey& InConnector)
+{
+	if (ControlRigBlueprint.IsValid())
+	{
+		FScopedTransaction Transaction(LOCTEXT("ModularRigModelDisconnectConnector", "Disconnect Connector"));
+
+		UModularRigController* Controller = ControlRigBlueprint->GetModularRigController();
+		check(Controller);
+
+		if(!bKeepCurrentEditedConnectors)
+		{
+			CurrentlyEditedConnectors.Reset();
+		}
+		const TGuardValue<bool> KeepCurrentEditedConnectorsGuard(bKeepCurrentEditedConnectors, true); 
+		CurrentlyEditedConnectors.Add(InConnector.Name);
+		Controller->DisconnectConnector(InConnector, false, true);
+	}
+}
+
+void SModularRigModel::HandleSelectionChanged(TSharedPtr<FModularRigTreeElement> Selection, ESelectInfo::Type SelectInfo)
+{
+	if(bIsPerformingSelection)
+	{
+		return;
+	}
+
+	TreeView->ClearHighlightedItems();
+	
+	if (ControlRigBlueprint.IsValid())
+	{
+		UModularRigController* Controller = ControlRigBlueprint->GetModularRigController();
+		check(Controller);
+
+		const TGuardValue<bool> GuardSelection(bIsPerformingSelection, true);
+		const TArray<FName> NewSelection = TreeView->GetSelectedModuleNames();
+		Controller->SetModuleSelection(NewSelection);
+	}
+}
+
+bool SModularRigModel::ShouldAlwaysShowConnector(const FName& InConnectorName) const
+{
+	return CurrentlyEditedConnectors.Contains(InConnectorName);
+}
+
+void SModularRigModel::OnModularRigModified(EModularRigNotification InNotif, const FRigModuleReference* InModule)
+{
+	if (!ControlRigBlueprint.IsValid())
+	{
+		return;
+	}
+
+	switch(InNotif)
+	{
+		case EModularRigNotification::ModuleSelected:
+		case EModularRigNotification::ModuleDeselected:
+		{
+			if(!bIsPerformingSelection)
+			{
+				const TGuardValue<bool> GuardSelection(bIsPerformingSelection, true);
+				if(const UModularRigController* ModularRigController = ControlRigBlueprint->GetModularRigController())
+				{
+					const TArray<FName> SelectedModuleNames = ModularRigController->GetSelectedModules();
+					TArray<TSharedPtr<FModularRigTreeElement>> NewSelection;
+					for(const FName& SelectedModuleName : SelectedModuleNames)
+					{
+						if(const TSharedPtr<FModularRigTreeElement> Module = TreeView->FindElement(SelectedModuleName.ToString()))
+						{
+							NewSelection.Add(Module);
+						}
+					}
+					TreeView->SetSelection(NewSelection);
+				}
+			}
+			break;
+		}
+		case EModularRigNotification::ModuleAdded:
+		case EModularRigNotification::ModuleRenamed:
+		case EModularRigNotification::ModuleRemoved:
+		case EModularRigNotification::ModuleReparented:
+		case EModularRigNotification::ModuleReordered:
+		case EModularRigNotification::ConnectionChanged:
+		case EModularRigNotification::ModuleConfigValueChanged:
+		case EModularRigNotification::ModuleShortNameChanged:
+		case EModularRigNotification::ModuleClassChanged:
+		{
+			TreeView->RefreshTreeView();
+			break;
+		}
+		default:
+		{
+			break;
+		}
+	}
+}
+
+void SModularRigModel::OnHierarchyModified(ERigHierarchyNotification InNotif, URigHierarchy* InHierarchy, const FRigNotificationSubject& InSubject)
+{
+	if(!ControlRigBlueprint.IsValid())
+	{
+		return;
+	}
+	
+	const FRigBaseElement* InElement = InSubject.Element;
+	const FRigBaseComponent* InComponent = InSubject.Component;
+	
+	switch(InNotif)
+	{
+		case ERigHierarchyNotification::ElementSelected:
+		case ERigHierarchyNotification::ElementDeselected:
+		{
+			FString ModuleOrConnectorName = InHierarchy->GetModuleName(InElement->GetKey());
+
+			if(const FRigConnectorElement* Connector = Cast<FRigConnectorElement>(InElement))
+			{
+				if(Connector->IsPrimary())
+            	{
+					ModuleOrConnectorName = Connector->GetName();
+            	}
+			}
+
+			if(!ModuleOrConnectorName.IsEmpty())
+			{
+				if(TSharedPtr<FModularRigTreeElement> Item = TreeView->FindElement(ModuleOrConnectorName))
+				{
+					const bool bSelected = InNotif == ERigHierarchyNotification::ElementSelected;
+					TreeView->SetItemHighlighted(Item, bSelected);
+					TreeView->RequestScrollIntoView(Item);
+				}
+			}
+		}
+		default:
+		{
+			break;
+		}
+	}
+}
+
+class SModularRigModelPasteTransformsErrorPipe : public FOutputDevice
+{
+public:
+
+	int32 NumErrors;
+
+	SModularRigModelPasteTransformsErrorPipe()
+		: FOutputDevice()
+		, NumErrors(0)
+	{
+	}
+
+	virtual void Serialize(const TCHAR* V, ELogVerbosity::Type Verbosity, const class FName& Category) override
+	{
+		UE_LOG(LogControlRig, Error, TEXT("Error importing transforms to Model: %s"), V);
+		NumErrors++;
+	}
+};
+
+UModularRig* SModularRigModel::GetModularRig() const
+{
+	if (ControlRigBlueprint.IsValid())
+	{
+		if (UControlRig* DebuggedRig = ControlRigBeingDebuggedPtr.Get())
+		{
+			return Cast<UModularRig>(DebuggedRig);
+		}
+		if(UControlRig* DebuggedRig = ControlRigBlueprint->GetDebuggedControlRig())
+		{
+			return Cast<UModularRig>(DebuggedRig);
+		}
+	}
+	if (ControlRigEditor.IsValid())
+	{
+		if (UControlRig* CurrentRig = ControlRigEditor.Pin()->GetControlRig())
+		{
+			return Cast<UModularRig>(CurrentRig);
+		}
+	}
+	return nullptr;
+}
+
+UModularRig* SModularRigModel::GetDefaultModularRig() const
+{
+	if (ControlRigBlueprint.IsValid())
+	{
+		UControlRig* DebuggedRig = ControlRigBeingDebuggedPtr.Get();
+		if (!DebuggedRig)
+		{
+			DebuggedRig = ControlRigBlueprint->GetDebuggedControlRig();
+		}
+		
+		if (DebuggedRig)
+		{
+			return Cast<UModularRig>(DebuggedRig);
+		}
+	}
+	return nullptr;
+}
+
+
+void SModularRigModel::OnRequestDetailsInspection(const FName& InModuleName)
+{
+	if(!ControlRigEditor.IsValid())
+	{
+		return;
+	}
+	ControlRigEditor.Pin()->SetDetailViewForRigModules({InModuleName});
+}
+
+void SModularRigModel::PostRedo(bool bSuccess) 
+{
+	if (bSuccess)
+	{
+		RefreshTreeView();
+	}
+}
+
+void SModularRigModel::PostUndo(bool bSuccess) 
+{
+	if (bSuccess)
+	{
+		RefreshTreeView();
+	}
+}
+
+FReply SModularRigModel::OnDragDetected(const FGeometry& MyGeometry, const FPointerEvent& MouseEvent)
+{
+	TArray<FString> DraggedKeys = GetSelectedKeys();
+	TArray<FName> ModuleNames;
+	for(const FString& DraggedKey : DraggedKeys)
+	{
+		ModuleNames.Add(*DraggedKey);
+	}
+	
+	if (MouseEvent.IsMouseButtonDown(EKeys::LeftMouseButton) && ModuleNames.Num() > 0)
+	{
+		if (ControlRigEditor.IsValid())
+		{
+			TSharedRef<FModularRigModuleDragDropOp> DragDropOp = FModularRigModuleDragDropOp::New(MoveTemp(ModuleNames));
+			return FReply::Handled().BeginDragDrop(DragDropOp);
+		}
+	}
+
+	return FReply::Unhandled();
+}
+
+TOptional<EItemDropZone> SModularRigModel::OnCanAcceptDrop(const FDragDropEvent& DragDropEvent, EItemDropZone DropZone, TSharedPtr<FModularRigTreeElement> TargetItem)
+{
+	const TOptional<EItemDropZone> InvalidDropZone;
+	TOptional<EItemDropZone> ReturnDropZone = DropZone;
+	int32 NewModuleIndex = INDEX_NONE;
+
+	if(DropZone == EItemDropZone::BelowItem || DropZone == EItemDropZone::AboveItem)
+	{
+		DropZone = EItemDropZone::OntoItem;
+		
+		if(TargetItem.IsValid())
+		{
+			TSharedPtr<FModularRigTreeElement> ChildTargetItem;
+			Swap(TargetItem, ChildTargetItem);
+			
+			if(ControlRigBlueprint.IsValid())
+			{
+				FModularRigModel& Model = ControlRigBlueprint.Get()->GetModularRigModel();
+				if(const FRigModuleReference* ChildModule = Model.FindModule(ChildTargetItem->ModuleName))
+				{
+					if(const FRigModuleReference* ParentModule = ChildModule->GetParentModule())
+					{
+						if(const TSharedPtr<FModularRigTreeElement> ParentItem = TreeView->FindElement(ParentModule->Name.ToString()))
+						{
+							TargetItem = ParentItem;
+						}
+					}
+				}
+			}
+		}
+	}
+	
+	if(DropZone != EItemDropZone::OntoItem)
+	{
+		return InvalidDropZone;
+	}
+
+	TSharedPtr<FAssetDragDropOp> AssetDragDropOperation = DragDropEvent.GetOperationAs<FAssetDragDropOp>();
+	TSharedPtr<FModularRigModuleDragDropOp> ModuleDragDropOperation = DragDropEvent.GetOperationAs<FModularRigModuleDragDropOp>();
+	if (AssetDragDropOperation)
+	{
+		for (const FAssetData& AssetData : AssetDragDropOperation->GetAssets())
+		{
+			static const UEnum* ControlTypeEnum = StaticEnum<EControlRigType>();
+			const FString ControlRigTypeStr = AssetData.GetTagValueRef<FString>(TEXT("ControlRigType"));
+			if (ControlRigTypeStr.IsEmpty())
+			{
+				ReturnDropZone.Reset();
+				break;
+			}
+
+			const EControlRigType ControlRigType = (EControlRigType)(ControlTypeEnum->GetValueByName(*ControlRigTypeStr));
+			if (ControlRigType != EControlRigType::RigModule)
+			{
+				ReturnDropZone.Reset();
+				break;
+			}
+		}
+	}
+	else if(ModuleDragDropOperation)
+	{
+		if(TargetItem.IsValid())
+		{
+			// we cannot drag a module onto itself
+			if(ModuleDragDropOperation->GetModules().Contains(TargetItem->ModuleName))
+			{
+				return InvalidDropZone;
+			}
+		}
+	}
+	else
+	{
+		ReturnDropZone.Reset();
+	}
+
+	return ReturnDropZone;
+}
+
+FReply SModularRigModel::OnAcceptDrop(const FDragDropEvent& DragDropEvent, EItemDropZone DropZone, TSharedPtr<FModularRigTreeElement> TargetItem)
+{
+	int32 NewModuleIndex = INDEX_NONE;
+	
+	if(DropZone == EItemDropZone::BelowItem || DropZone == EItemDropZone::AboveItem)
+	{
+		if(TargetItem.IsValid())
+		{
+			TSharedPtr<FModularRigTreeElement> ChildTargetItem;
+			Swap(TargetItem, ChildTargetItem);
+			
+			if(ControlRigBlueprint.IsValid())
+			{
+				FModularRigModel& Model = ControlRigBlueprint.Get()->GetModularRigModel();
+				if(const FRigModuleReference* ChildModule = Model.FindModule(ChildTargetItem->ModuleName))
+				{
+					TArray<FRigModuleReference*> Children;
+					if(const FRigModuleReference* ParentModule = ChildModule->GetParentModule())
+					{
+						Children = ParentModule->CachedChildren;
+						
+						if(const TSharedPtr<FModularRigTreeElement> ParentItem = TreeView->FindElement(ParentModule->Name.ToString()))
+						{
+							TargetItem = ParentItem;
+						}
+					}
+					else
+					{
+						Children = Model.RootModules;
+					}
+
+					for(int32 ChildIndex = 0; ChildIndex < Children.Num(); ChildIndex++)
+					{
+						if(Children[ChildIndex] == ChildModule)
+						{
+							NewModuleIndex = ChildIndex;
+							break;
+						}
+					}
+				}
+			}
+		}
+
+		DropZone = EItemDropZone::OntoItem;
+	}
+
+	FName ParentModuleName = NAME_None;
+	if (TargetItem.IsValid())
+	{
+		ParentModuleName = TargetItem->ModuleName;
+	}
+
+	TSharedPtr<FAssetDragDropOp> AssetDragDropOperation = DragDropEvent.GetOperationAs<FAssetDragDropOp>();
+	TSharedPtr<FModularRigModuleDragDropOp> ModuleDragDropOperation = DragDropEvent.GetOperationAs<FModularRigModuleDragDropOp>();
+	if (AssetDragDropOperation)
+	{
+		for (const FAssetData& AssetData : AssetDragDropOperation->GetAssets())
+		{
+			static const UEnum* ControlTypeEnum = StaticEnum<EControlRigType>();
+			const FString ControlRigTypeStr = AssetData.GetTagValueRef<FString>(TEXT("ControlRigType"));
+			if (ControlRigTypeStr.IsEmpty())
+			{
+				continue;
+			}
+
+			const EControlRigType ControlRigType = (EControlRigType)(ControlTypeEnum->GetValueByName(*ControlRigTypeStr));
+			if (ControlRigType != EControlRigType::RigModule)
+			{
+				continue;
+			}
+
+			if(UControlRigBlueprint* AssetBlueprint = Cast<UControlRigBlueprint>(AssetData.GetAsset()))
+			{
+				HandleNewItem(AssetBlueprint->GetControlRigClass(), ParentModuleName);
+			}
+			else if (UControlRigBlueprintGeneratedClass* GeneratedClass = Cast<UControlRigBlueprintGeneratedClass>(AssetData.GetAsset()))
+			{
+				HandleNewItem(GeneratedClass, ParentModuleName);
+			}
+			else
+			{
+				continue;
+			}
+		}
+
+		return FReply::Handled();
+	}
+	else if(ModuleDragDropOperation)
+	{
+		const TArray<FName> ModuleNames = ModuleDragDropOperation->GetModules();
+		HandleReparentModules(ModuleNames, ParentModuleName, NewModuleIndex);
+	}
+	
+	return FReply::Unhandled();
+}
+
+FReply SModularRigModel::OnDrop(const FGeometry& MyGeometry, const FDragDropEvent& DragDropEvent)
+{
+	// only allow drops onto empty space of the widget (when there's no target item under the mouse)
+	// when dropped onto an item SModularRigModel::OnAcceptDrop will deal with the event
+	const TSharedPtr<FModularRigTreeElement>* ItemAtMouse = TreeView->FindItemAtPosition(DragDropEvent.GetScreenSpacePosition());
+	FString ParentPath;
+	if (ItemAtMouse && ItemAtMouse->IsValid())
+	{
+		return SCompoundWidget::OnDrop(MyGeometry, DragDropEvent);
+	}
+	
+	if (OnCanAcceptDrop(DragDropEvent, EItemDropZone::BelowItem, nullptr))
+	{
+		if (OnAcceptDrop(DragDropEvent, EItemDropZone::BelowItem, nullptr).IsEventHandled())
+		{
+			return FReply::Handled();
+		}
+	}
+	return SCompoundWidget::OnDrop(MyGeometry, DragDropEvent);
+}
+
+#undef LOCTEXT_NAMESPACE
+
diff --git a/Engine/Plugins/Animation/ControlRig/Source/ControlRigEditor/Private/Editor/Hierarchy/Widgets/SModularRigHierarchy.h b/Engine/Plugins/Animation/ControlRig/Source/ControlRigEditor/Private/Editor/Hierarchy/Widgets/SModularRigHierarchy.h
new file mode 100644
--- /dev/null
+++ b/Engine/Plugins/Animation/ControlRig/Source/ControlRigEditor/Private/Editor/Hierarchy/Widgets/SModularRigHierarchy.h
@@ -0,0 +1,178 @@
+// Copyright Epic Games, Inc. All Rights Reserved.
+
+#pragma once
+
+#include "EditorUndoClient.h"
+#include "Editor/Hierarchy/Views/SModularRigTreeView.h"
+#include "ControlRigBlueprintLegacy.h"
+#include "Editor/RigVMNewEditor.h"
+
+class SModularRigModel;
+class IControlRigBaseEditor;
+class SSearchBox;
+class FUICommandList;
+class URigVMBlueprint;
+class UControlRig;
+struct FAssetData;
+class FMenuBuilder;
+class UToolMenu;
+struct FToolMenuContext;
+
+/** Widget allowing editing of a control rig's structure */
+class SModularRigModel : public SCompoundWidget, public FEditorUndoClient
+{
+public:
+	SLATE_BEGIN_ARGS(SModularRigModel) {}
+	SLATE_END_ARGS()
+
+	~SModularRigModel();
+
+	void Construct(const FArguments& InArgs, TSharedRef<IControlRigBaseEditor> InControlRigEditor);
+
+	IControlRigBaseEditor* GetControlRigEditor() const
+	{
+		if(ControlRigEditor.IsValid())
+		{
+			return ControlRigEditor.Pin().Get();
+		}
+		return nullptr;
+	}
+
+private:
+
+	void OnEditorClose(IControlRigBaseEditor* InEditor, FControlRigAssetInterfacePtr InBlueprint);
+
+	/** Bind commands that this widget handles */
+	void BindCommands();
+
+	/** SWidget interface */
+	virtual FReply OnKeyDown(const FGeometry& MyGeometry, const FKeyEvent& InKeyEvent);
+	virtual FReply OnMouseButtonDown(const FGeometry& MyGeometry, const FPointerEvent& MouseEvent) override;
+
+	/** Rebuild the tree view */
+	void RefreshTreeView(bool bRebuildContent = true);
+
+	/** Returns all selected items */
+	TArray<TSharedPtr<FModularRigTreeElement>> GetSelectedItems() const;
+
+	/** Return all selected keys */
+	TArray<FString> GetSelectedKeys() const;
+
+	/** Create a new item */
+	void HandleNewItem();
+	void HandleNewItem(UClass* InClass, const FName& InParentModuleName);
+
+	/** Rename item */
+	bool CanRenameModule() const;
+	void HandleRenameModule();
+	FName HandleRenameModule(const FName& InOldModuleName, const FName& InNewName);
+	bool HandleVerifyNameChanged(const FName& InOldModuleName, const FName& InNewName, FText& OutErrorMessage);
+
+	/** Delete items */
+	void HandleDeleteModules();
+	void HandleDeleteModules(const TArray<FName>& InModuleNames);
+
+	/** Reparent items */
+	void HandleReparentModules(const TArray<FName>& InModuleNames, const FName& InParentModuleName, int32 NewModuleIndex);
+
+	/** Mirror items */
+	void HandleMirrorModules();
+	void HandleMirrorModules(const TArray<FName>& InModuleNames);
+
+	/** Reresolve items */
+	void HandleReresolveModules();
+	void HandleReresolveModules(const TArray<FName>& InModuleNames);
+
+	/** Swap module class for items */
+	bool CanSwapModules() const;
+	void HandleSwapClassForModules();
+	void HandleSwapClassForModules(const TArray<FName>& InModuleNames);
+
+	/** Copy & paste module settings */
+	bool CanCopyModuleSettings() const;
+	void HandleCopyModuleSettings();
+	bool CanPasteModuleSettings() const;
+	void HandlePasteModuleSettings();
+
+	/** Resolve connector */
+	void HandleConnectorResolved(const FRigElementKey& InConnector, const TArray<FRigElementKey>& InTargets);
+
+	/** UnResolve connector */
+	void HandleConnectorDisconnect(const FRigElementKey& InConnector);
+
+	/** Set Selection Changed */
+	void HandleSelectionChanged(TSharedPtr<FModularRigTreeElement> Selection, ESelectInfo::Type SelectInfo);
+
+	/** Returns true if a given connector should always be shown */
+	bool ShouldAlwaysShowConnector(const FName& InConnectorName) const;
+
+	TSharedPtr< SWidget > CreateContextMenuWidget();
+	void OnItemClicked(TSharedPtr<FModularRigTreeElement> InItem);
+	void OnItemDoubleClicked(TSharedPtr<FModularRigTreeElement> InItem);
+	
+	// FEditorUndoClient
+	virtual void PostUndo(bool bSuccess) override;
+	virtual void PostRedo(bool bSuccess) override;
+
+	// reply to a drag operation
+	FReply OnDragDetected(const FGeometry& MyGeometry, const FPointerEvent& MouseEvent);
+
+	// reply to a drop operation on item
+	TOptional<EItemDropZone> OnCanAcceptDrop(const FDragDropEvent& DragDropEvent, EItemDropZone DropZone, TSharedPtr<FModularRigTreeElement> TargetItem);
+	FReply OnAcceptDrop(const FDragDropEvent& DragDropEvent, EItemDropZone DropZone, TSharedPtr<FModularRigTreeElement> TargetItem);
+
+	// SWidget Overrides
+	virtual FReply OnDrop(const FGeometry& MyGeometry, const FDragDropEvent& DragDropEvent) override;
+
+	static const FName ContextMenuName;
+	static void CreateContextMenu();
+	UToolMenu* GetContextMenu();
+	TSharedPtr<FUICommandList> GetContextMenuCommands() const;
+	
+	/** Our owning control rig editor */
+	TWeakPtr<IControlRigBaseEditor> ControlRigEditor;
+
+	/** Tree view widget */
+	TSharedPtr<SModularRigTreeView> TreeView;
+	TSharedPtr<SHeaderRow> HeaderRowWidget;
+
+	TWeakInterfacePtr<IControlRigAssetInterface> ControlRigBlueprint;
+	TWeakObjectPtr<UModularRig> ControlRigBeingDebuggedPtr;
+	
+	/** Command list we bind to */
+	TSharedPtr<FUICommandList> CommandList;
+
+	bool IsSingleSelected() const;
+	
+	UModularRig* GetModularRig() const;
+	UModularRig* GetDefaultModularRig() const;
+	const UModularRig* GetModularRigForTreeView() const { return GetModularRig(); }
+	void OnRequestDetailsInspection(const FName& InModuleName);
+	void HandlePreCompileModularRigs(FRigVMAssetInterfacePtr InBlueprint);
+	void HandlePostCompileModularRigs(FRigVMAssetInterfacePtr InBlueprint);
+	void OnModularRigModified(EModularRigNotification InNotif, const FRigModuleReference* InModule);
+	void OnHierarchyModified(ERigHierarchyNotification InNotif, URigHierarchy* InHierarchy, const FRigNotificationSubject& InSubject);
+
+	void HandleRefreshEditorFromBlueprint(FRigVMAssetInterfacePtr InBlueprint);
+	void HandleSetObjectBeingDebugged(UObject* InObject);
+
+	TSharedRef<SWidget> OnGetOptionsMenu();
+	void OnFilterTextChanged(const FText& SearchText);
+
+	bool bShowSecondaryConnectors;
+	bool bShowOptionalConnectors;
+	bool bShowUnresolvedConnectors;
+	FText FilterText;
+
+	TSharedPtr<SSearchBox> FilterBox;
+	bool bIsPerformingSelection;
+	bool bKeepCurrentEditedConnectors;
+	TSet<FName> CurrentlyEditedConnectors;
+
+public:
+
+	friend class FModularRigTreeElement;
+	friend class SModularRigModelItem;
+	friend class UControlRigBlueprintEditorLibrary;
+};
+
diff --git a/Engine/Plugins/Animation/ControlRig/Source/ControlRigEditor/Private/Editor/ModularRigModelTabSummoner.cpp b/Engine/Plugins/Animation/ControlRig/Source/ControlRigEditor/Private/Editor/ModularRigModelTabSummoner.cpp
--- a/Engine/Plugins/Animation/ControlRig/Source/ControlRigEditor/Private/Editor/ModularRigModelTabSummoner.cpp
+++ b/Engine/Plugins/Animation/ControlRig/Source/ControlRigEditor/Private/Editor/ModularRigModelTabSummoner.cpp
@@ -1,7 +1,7 @@
 // Copyright Epic Games, Inc. All Rights Reserved.
 
 #include "Editor/ModularRigModelTabSummoner.h"
-#include "Editor/SModularRigModel.h"
+#include "Editor/Hierarchy/Widgets/SModularRigHierarchy.h"
 #include "ControlRigEditorStyle.h"
 #include "Editor/ControlRigEditor.h"
 #include "Widgets/Docking/SDockTab.h"
diff --git a/Engine/Plugins/Animation/ControlRig/Source/ControlRigEditor/Private/Editor/RigHierarchyTabSummoner.cpp b/Engine/Plugins/Animation/ControlRig/Source/ControlRigEditor/Private/Editor/RigHierarchyTabSummoner.cpp
--- a/Engine/Plugins/Animation/ControlRig/Source/ControlRigEditor/Private/Editor/RigHierarchyTabSummoner.cpp
+++ b/Engine/Plugins/Animation/ControlRig/Source/ControlRigEditor/Private/Editor/RigHierarchyTabSummoner.cpp
@@ -1,10 +1,10 @@
 // Copyright Epic Games, Inc. All Rights Reserved.
 
 #include "Editor/RigHierarchyTabSummoner.h"
-#include "Editor/SRigHierarchy.h"
+
 #include "ControlRigEditorStyle.h"
 #include "Editor/ControlRigEditor.h"
-#include "Editor/SRigHierarchy.h"
+#include "Editor/Hierarchy/Widgets/SRigHierarchy.h"
 #include "Widgets/Docking/SDockTab.h"
 
 #define LOCTEXT_NAMESPACE "RigHierarchyTabSummoner"
diff --git a/Engine/Plugins/Animation/ControlRig/Source/ControlRigEditor/Private/Editor/RigInfluenceMapTabSummoner.cpp b/Engine/Plugins/Animation/ControlRig/Source/ControlRigEditor/Private/Editor/RigInfluenceMapTabSummoner.cpp
--- a/Engine/Plugins/Animation/ControlRig/Source/ControlRigEditor/Private/Editor/RigInfluenceMapTabSummoner.cpp
+++ b/Engine/Plugins/Animation/ControlRig/Source/ControlRigEditor/Private/Editor/RigInfluenceMapTabSummoner.cpp
@@ -1,7 +1,7 @@
 // Copyright Epic Games, Inc. All Rights Reserved.
 
 #include "Editor/RigInfluenceMapTabSummoner.h"
-#include "Editor/SRigHierarchy.h"
+#include "Editor/Hierarchy/Widgets/SRigHierarchy.h"
 #include "Editor/RigVMEditorStyle.h"
 #include "Editor/ControlRigEditor.h"
 #include "ControlRigBlueprintLegacy.h"
diff --git a/Engine/Plugins/Animation/ControlRig/Source/ControlRigEditor/Private/Editor/RigValidationTabSummoner.cpp b/Engine/Plugins/Animation/ControlRig/Source/ControlRigEditor/Private/Editor/RigValidationTabSummoner.cpp
--- a/Engine/Plugins/Animation/ControlRig/Source/ControlRigEditor/Private/Editor/RigValidationTabSummoner.cpp
+++ b/Engine/Plugins/Animation/ControlRig/Source/ControlRigEditor/Private/Editor/RigValidationTabSummoner.cpp
@@ -1,7 +1,7 @@
 // Copyright Epic Games, Inc. All Rights Reserved.
 
 #include "Editor/RigValidationTabSummoner.h"
-#include "Editor/SRigHierarchy.h"
+#include "Editor/Hierarchy/Widgets/SRigHierarchy.h"
 #include "ControlRigEditorStyle.h"
 #include "Editor/ControlRigEditor.h"
 #include "ControlRigBlueprintLegacy.h"
diff --git a/Engine/Plugins/Animation/ControlRig/Source/ControlRigEditor/Private/Editor/SRigConnectorTargetWidget.cpp b/Engine/Plugins/Animation/ControlRig/Source/ControlRigEditor/Private/Editor/SRigConnectorTargetWidget.cpp
--- a/Engine/Plugins/Animation/ControlRig/Source/ControlRigEditor/Private/Editor/SRigConnectorTargetWidget.cpp
+++ b/Engine/Plugins/Animation/ControlRig/Source/ControlRigEditor/Private/Editor/SRigConnectorTargetWidget.cpp
@@ -4,6 +4,7 @@
 
 #include "DetailWidgetRow.h"
 #include "DetailLayoutBuilder.h"
+#include "Editor/Hierarchy/Views/SRigHierarchyItem.h"
 #include "IDetailChildrenBuilder.h"
 #include "ISinglePropertyView.h"
 #include "Modules/ModuleManager.h"
diff --git a/Engine/Plugins/Animation/ControlRig/Source/ControlRigEditor/Private/Editor/SRigConnectorTargetWidget.h b/Engine/Plugins/Animation/ControlRig/Source/ControlRigEditor/Private/Editor/SRigConnectorTargetWidget.h
--- a/Engine/Plugins/Animation/ControlRig/Source/ControlRigEditor/Private/Editor/SRigConnectorTargetWidget.h
+++ b/Engine/Plugins/Animation/ControlRig/Source/ControlRigEditor/Private/Editor/SRigConnectorTargetWidget.h
@@ -10,7 +10,7 @@
 #include "PropertyEditorDelegates.h"
 #include "UObject/StrongObjectPtrTemplates.h"
 #include "IPropertyTypeCustomization.h"
-#include "SRigHierarchyTreeView.h"
+#include "Editor/Hierarchy/Views/SRigHierarchySearchableTreeView.h"
 #include "SRigConnectorTargetWidget.generated.h"
 
 DECLARE_DELEGATE_RetVal_OneParam(bool, FRigConnectorTargetWidget_SetTarget, FRigElementKey);
diff --git a/Engine/Plugins/Animation/ControlRig/Source/ControlRigEditor/Private/Editor/SRigSpacePickerWidget.cpp b/Engine/Plugins/Animation/ControlRig/Source/ControlRigEditor/Private/Editor/SRigSpacePickerWidget.cpp
--- a/Engine/Plugins/Animation/ControlRig/Source/ControlRigEditor/Private/Editor/SRigSpacePickerWidget.cpp
+++ b/Engine/Plugins/Animation/ControlRig/Source/ControlRigEditor/Private/Editor/SRigSpacePickerWidget.cpp
@@ -8,7 +8,8 @@
 #include "Widgets/Input/SButton.h"
 #include "Widgets/Input/SComboButton.h"
 #include "Widgets/Input/SSearchBox.h"
-#include "Editor/SRigHierarchyTreeView.h"
+#include "Editor/Hierarchy/Views/SRigHierarchyItem.h"
+#include "Editor/Hierarchy/Views/SRigHierarchySearchableTreeView.h"
 #include "ControlRigEditorStyle.h"
 #include "PropertyCustomizationHelpers.h"
 #include "ISequencer.h"
