Anim Details: follow-up to 46583619

#rnx

diff --git a/Engine/Plugins/Animation/ControlRig/Source/ControlRig/Private/Sequencer/ControlRigSequencerHelpers.cpp b/Engine/Plugins/Animation/ControlRig/Source/ControlRig/Private/Sequencer/ControlRigSequencerHelpers.cpp
--- a/Engine/Plugins/Animation/ControlRig/Source/ControlRig/Private/Sequencer/ControlRigSequencerHelpers.cpp
+++ b/Engine/Plugins/Animation/ControlRig/Source/ControlRig/Private/Sequencer/ControlRigSequencerHelpers.cpp
@@ -169,7 +169,7 @@ TArrayView<FMovieSceneIntegerChannel*> FControlRigSequencerHelpers::GetIntegerCh
 
 UMovieSceneControlRigParameterTrack* FControlRigSequencerHelpers::FindControlRigTrack(UMovieSceneSequence* MovieSceneSequence, const UControlRig* ControlRig)
 {
-	if (const UMovieScene* MovieScene = MovieSceneSequence ? MovieSceneSequence->GetMovieScene() : nullptr)
+	if (const UMovieScene* MovieScene = (ControlRig && MovieSceneSequence) ? MovieSceneSequence->GetMovieScene() : nullptr)
 	{
 		for (const FMovieSceneBinding& Binding : MovieScene->GetBindings())
 		{
