Control Rig: Control Rig Hierarchy Editor cleanup second pass

* Change type names to match file names
* No change in implementations and logic

diff --git a/Engine/Plugins/Animation/ControlRig/Source/ControlRigEditor/Private/ControlRigBlueprintEditorLibrary.cpp b/Engine/Plugins/Animation/ControlRig/Source/ControlRigEditor/Private/ControlRigBlueprintEditorLibrary.cpp
--- a/Engine/Plugins/Animation/ControlRig/Source/ControlRigEditor/Private/ControlRigBlueprintEditorLibrary.cpp
+++ b/Engine/Plugins/Animation/ControlRig/Source/ControlRigEditor/Private/ControlRigBlueprintEditorLibrary.cpp
@@ -89,7 +89,7 @@ void UControlRigBlueprintEditorLibrary::SetupAllEditorMenus()
 {
 	SRigHierarchy::CreateContextMenu();
 	SRigHierarchy::CreateDragDropMenu();
-	SModularRigModel::CreateContextMenu();
+	SModularRigHierarchy::CreateContextMenu();
 }
 
 TArray<FRigModuleDescription> UControlRigBlueprintEditorLibrary::GetAvailableRigModules()
diff --git a/Engine/Plugins/Animation/ControlRig/Source/ControlRigEditor/Private/ControlRigElementDetails.cpp b/Engine/Plugins/Animation/ControlRig/Source/ControlRigEditor/Private/ControlRigElementDetails.cpp
--- a/Engine/Plugins/Animation/ControlRig/Source/ControlRigEditor/Private/ControlRigElementDetails.cpp
+++ b/Engine/Plugins/Animation/ControlRig/Source/ControlRigEditor/Private/ControlRigElementDetails.cpp
@@ -4691,7 +4691,7 @@ TSharedRef<SWidget> FRigControlElementDetails::GetAddSpaceContent(const TSharedR
 		return SNullWidget::NullWidget;
 	}
 
-	FRigTreeDelegates RigTreeDelegates;
+	FRigHierarchyTreeDelegates RigTreeDelegates;
 	RigTreeDelegates.OnGetHierarchy = FOnGetRigTreeHierarchy::CreateLambda([this]()
 	{
 		return PerElementInfos[0].GetHierarchy();
@@ -4717,7 +4717,7 @@ FReply FRigControlElementDetails::OnAddSpaceMouseDown(const FGeometry& InGeometr
 	return FReply::Handled();
 }
 
-void FRigControlElementDetails::OnAddSpaceSelection(TSharedPtr<FRigTreeElement> Selection, ESelectInfo::Type SelectInfo, const TSharedRef<IPropertyUtilities> PropertyUtilities)
+void FRigControlElementDetails::OnAddSpaceSelection(TSharedPtr<FRigHierarchyTreeElement> Selection, ESelectInfo::Type SelectInfo, const TSharedRef<IPropertyUtilities> PropertyUtilities)
 {
 	if(Selection)
 	{
@@ -6364,7 +6364,7 @@ void FRigConnectorElementDetails::CustomizeConnectorTargets(IDetailLayoutBuilder
 	URigHierarchy* Hierarchy = PerElementInfos[0].GetHierarchy();
 	check(Hierarchy);
 
-	FRigTreeDelegates RigTreeDelegates;
+	FRigHierarchyTreeDelegates RigTreeDelegates;
 	RigTreeDelegates.OnGetHierarchy.BindLambda([Hierarchy]()
 	{
 		return Hierarchy;
diff --git a/Engine/Plugins/Animation/ControlRig/Source/ControlRigEditor/Private/ControlRigElementDetails.h b/Engine/Plugins/Animation/ControlRig/Source/ControlRigEditor/Private/ControlRigElementDetails.h
--- a/Engine/Plugins/Animation/ControlRig/Source/ControlRigEditor/Private/ControlRigElementDetails.h
+++ b/Engine/Plugins/Animation/ControlRig/Source/ControlRigEditor/Private/ControlRigElementDetails.h
@@ -814,10 +814,10 @@ private:
 	TSharedRef<ITableRow> HandleGenerateAnimationChannelTypeRow(TSharedPtr<ERigControlType> ControlType, const TSharedRef<STableViewBase>& OwnerTable, FRigElementKey ControlKey);
 
 	// multi parent related callbacks
-	const FRigTreeDisplaySettings& GetDisplaySettings() const { return DisplaySettings; }
+	const FRigHierarchyTreeDisplaySettings& GetDisplaySettings() const { return DisplaySettings; }
 	TSharedRef<SWidget> GetAddSpaceContent(const TSharedRef<IPropertyUtilities> PropertyUtilities);
 	FReply OnAddSpaceMouseDown(const FGeometry& InGeometry, const FPointerEvent& InPointerEvent, const TSharedRef<IPropertyUtilities> PropertyUtilities);
-	void OnAddSpaceSelection(TSharedPtr<FRigTreeElement> Selection, ESelectInfo::Type SelectInfo, const TSharedRef<IPropertyUtilities> PropertyUtilities);
+	void OnAddSpaceSelection(TSharedPtr<FRigHierarchyTreeElement> Selection, ESelectInfo::Type SelectInfo, const TSharedRef<IPropertyUtilities> PropertyUtilities);
 
 	void HandleControlTypeChanged(TSharedPtr<ERigControlType> ControlType, ESelectInfo::Type SelectInfo, FRigElementKey ControlKey, const TSharedRef<IPropertyUtilities> PropertyUtilities);
 	void HandleControlTypeChanged(ERigControlType ControlType, TArray<FRigElementKey> ControlKeys, const TSharedRef<IPropertyUtilities> PropertyUtilities);
@@ -833,7 +833,7 @@ private:
 	TSharedPtr<SControlRigShapeNameList> ShapeNameListWidget; 
 	static TSharedPtr<TArray<ERigControlValueType>> PickedValueTypes;
 
-	FRigTreeDisplaySettings DisplaySettings;
+	FRigHierarchyTreeDisplaySettings DisplaySettings;
 	TSharedPtr<SMenuAnchor> AddSpaceMenuAnchor;
 };
 
diff --git a/Engine/Plugins/Animation/ControlRig/Source/ControlRigEditor/Private/ControlRigModuleDetails.cpp b/Engine/Plugins/Animation/ControlRig/Source/ControlRigEditor/Private/ControlRigModuleDetails.cpp
--- a/Engine/Plugins/Animation/ControlRig/Source/ControlRigEditor/Private/ControlRigModuleDetails.cpp
+++ b/Engine/Plugins/Animation/ControlRig/Source/ControlRigEditor/Private/ControlRigModuleDetails.cpp
@@ -373,7 +373,7 @@ void FRigModuleInstanceDetails::CustomizeDetails(IDetailLayoutBuilder& DetailBui
 				const FText Label = FText::FromString(Connector.Name);
 				bool bIsArrayConnector = IsArrayConnector[ConnectorIndex].Get(false);
 
-				FRigTreeDelegates RigTreeDelegates;
+				FRigHierarchyTreeDelegates RigTreeDelegates;
 				RigTreeDelegates.OnGetHierarchy = FOnGetRigTreeHierarchy::CreateLambda([WeakThisPtr]() -> const URigHierarchy*
 				{
 					if (TSharedPtr<FRigModuleInstanceDetails> StrongThis = WeakThisPtr.Pin())
diff --git a/Engine/Plugins/Animation/ControlRig/Source/ControlRigEditor/Private/EditMode/SControlRigOutliner.cpp b/Engine/Plugins/Animation/ControlRig/Source/ControlRigEditor/Private/EditMode/SControlRigOutliner.cpp
--- a/Engine/Plugins/Animation/ControlRig/Source/ControlRigEditor/Private/EditMode/SControlRigOutliner.cpp
+++ b/Engine/Plugins/Animation/ControlRig/Source/ControlRigEditor/Private/EditMode/SControlRigOutliner.cpp
@@ -46,7 +46,7 @@
 
 #define LOCTEXT_NAMESPACE "ControlRigOutliner"
 
-FRigTreeDisplaySettings FMultiRigTreeDelegates::DefaultDisplaySettings;
+FRigHierarchyTreeDisplaySettings FMultiRigTreeDelegates::DefaultDisplaySettings;
 
 
 //////////////////////////////////////////////////////////////
@@ -89,7 +89,7 @@ FText FMultiRigData::GetName() const
 	return FText::FromName(NAME_None);
 }
 
-FText FMultiRigData::GetDisplayName(const FRigTreeDisplaySettings& InSettings) const
+FText FMultiRigData::GetDisplayName(const FRigHierarchyTreeDisplaySettings& InSettings) const
 {
 	if (!CachedDisplayName.IsSet())
 	{
@@ -143,7 +143,7 @@ void FMultiRigData::InvalidateDisplayName() const
 	CachedDisplayName.Reset();
 }
 
-FText FMultiRigData::GetToolTipText(const FRigTreeDisplaySettings& Settings) const
+FText FMultiRigData::GetToolTipText(const FRigHierarchyTreeDisplaySettings& Settings) const
 {
 	if (IsControlElement())
 	{
@@ -236,18 +236,18 @@ FMultiRigTreeElement::FMultiRigTreeElement(const FMultiRigData& InData, TWeakPtr
 
 	if (InTreeView.IsValid() && Data.IsValid())
 	{
-		const FRigTreeDisplaySettings& Settings = InTreeView.Pin()->GetTreeDelegates().GetDisplaySettings();
+		const FRigHierarchyTreeDisplaySettings& Settings = InTreeView.Pin()->GetTreeDelegates().GetDisplaySettings();
 		RefreshDisplaySettings(Data.GetHierarchy(), Settings);
 	}
 }
 
-TSharedRef<ITableRow> FMultiRigTreeElement::MakeTreeRowWidget(const TSharedRef<STableViewBase>& InOwnerTable, TSharedRef<FMultiRigTreeElement> InRigTreeElement, TSharedPtr<SMultiRigHierarchyTreeView> InTreeView, const FRigTreeDisplaySettings& InSettings, bool bPinned)
+TSharedRef<ITableRow> FMultiRigTreeElement::MakeTreeRowWidget(const TSharedRef<STableViewBase>& InOwnerTable, TSharedRef<FMultiRigTreeElement> InRigTreeElement, TSharedPtr<SMultiRigHierarchyTreeView> InTreeView, const FRigHierarchyTreeDisplaySettings& InSettings, bool bPinned)
 {
 	return SNew(SMultiRigHierarchyItem, InOwnerTable, InRigTreeElement, InTreeView, InSettings, bPinned);
 
 }
 
-void FMultiRigTreeElement::RefreshDisplaySettings(const URigHierarchy* InHierarchy, const FRigTreeDisplaySettings& InSettings)
+void FMultiRigTreeElement::RefreshDisplaySettings(const URigHierarchy* InHierarchy, const FRigHierarchyTreeDisplaySettings& InSettings)
 {
 	TPair<const FSlateBrush*, FSlateColor> Result;
 	if (InHierarchy)
@@ -313,7 +313,7 @@ bool FMultiRigTreeElement::AreControlsVisible() const
 /// SMultiRigHierarchyItem
 ///////////////////////////////////////////////////////////
 TMap<FSoftObjectPath, TSharedPtr<FSlateBrush>> SMultiRigHierarchyItem::IconPathToBrush;
-void SMultiRigHierarchyItem::Construct(const FArguments& InArgs, const TSharedRef<STableViewBase>& OwnerTable, TSharedRef<FMultiRigTreeElement> InRigTreeElement, TSharedPtr<SMultiRigHierarchyTreeView> InTreeView, const FRigTreeDisplaySettings& InSettings, bool bPinned)
+void SMultiRigHierarchyItem::Construct(const FArguments& InArgs, const TSharedRef<STableViewBase>& OwnerTable, TSharedRef<FMultiRigTreeElement> InRigTreeElement, TSharedPtr<SMultiRigHierarchyTreeView> InTreeView, const FRigHierarchyTreeDisplaySettings& InSettings, bool bPinned)
 {
 	WeakRigTreeElement = InRigTreeElement;
 	TreeView = InTreeView;
@@ -426,13 +426,13 @@ FReply SMultiRigHierarchyItem::OnToggleVisibilityClicked()
 
 FText SMultiRigHierarchyItem::GetDisplayName() const
 {
-	const FRigTreeDisplaySettings Settings = Delegates.GetDisplaySettings();
+	const FRigHierarchyTreeDisplaySettings Settings = Delegates.GetDisplaySettings();
 	return (WeakRigTreeElement.Pin()->Data.GetDisplayName(Settings));
 }
 
 FText SMultiRigHierarchyItem::GetToolTipText() const
 {
-	const FRigTreeDisplaySettings Settings = Delegates.GetDisplaySettings();
+	const FRigHierarchyTreeDisplaySettings Settings = Delegates.GetDisplaySettings();
 	return (WeakRigTreeElement.Pin()->Data.GetToolTipText(Settings));
 }
 
@@ -764,7 +764,7 @@ bool SMultiRigHierarchyTreeView::AddElement(const FMultiRigData& InData, const F
 		return false;
 	}
 
-	const FRigTreeDisplaySettings& Settings = Delegates.GetDisplaySettings();
+	const FRigHierarchyTreeDisplaySettings& Settings = Delegates.GetDisplaySettings();
 
 	const FString FilteredString = Settings.FilterText.ToString();
 	if (FilteredString.IsEmpty() || !InData.IsValid())
@@ -849,7 +849,7 @@ bool SMultiRigHierarchyTreeView::AddElement(UControlRig* InControlRig, const FRi
 		return false;
 	}
 
-	const FRigTreeDisplaySettings& Settings = Delegates.GetDisplaySettings();
+	const FRigHierarchyTreeDisplaySettings& Settings = Delegates.GetDisplaySettings();
 
 	auto IsElementShown = [Settings](const FRigBaseElement* InElement) -> bool
 	{
@@ -1067,7 +1067,7 @@ bool SMultiRigHierarchyTreeView::ReparentElement(const FMultiRigData& InData, co
 		return false;
 	}
 
-	const FRigTreeDisplaySettings& Settings = Delegates.GetDisplaySettings();
+	const FRigHierarchyTreeDisplaySettings& Settings = Delegates.GetDisplaySettings();
 
 	TSharedPtr<FMultiRigTreeElement>* FoundItem = ElementMap.Find(InData);
 	if (FoundItem == nullptr)
@@ -1234,7 +1234,7 @@ void SMultiRigHierarchyTreeView::RefreshTreeView()
 	
 	TSharedPtr<ISequencer> Sequencer = WeakSequencer.Pin();
 
-	const FRigTreeDisplaySettings& Settings = Delegates.GetDisplaySettings();
+	const FRigHierarchyTreeDisplaySettings& Settings = Delegates.GetDisplaySettings();
 
 	TMap<UControlRig*, TArray<FRigElementKey>> SelectedControls;
 	TArray<FRigModuleInstance*> SelectedModules;
@@ -1584,7 +1584,7 @@ void SMultiRigHierarchyTreeView::SetExpansionRecursive(TSharedPtr<FMultiRigTreeE
 TSharedRef<ITableRow> SMultiRigHierarchyTreeView::MakeTableRowWidget(TSharedPtr<FMultiRigTreeElement> InItem,
 	const TSharedRef<STableViewBase>& OwnerTable, bool bPinned)
 {
-	const FRigTreeDisplaySettings& Settings = Delegates.GetDisplaySettings();
+	const FRigHierarchyTreeDisplaySettings& Settings = Delegates.GetDisplaySettings();
 	return InItem->MakeTreeRowWidget(OwnerTable, InItem.ToSharedRef(), SharedThis(this), Settings, bPinned);
 }
 
@@ -1724,7 +1724,7 @@ void SSearchableMultiRigHierarchyTreeView::Construct(const FArguments& InArgs)
 		];
 }
 
-const FRigTreeDisplaySettings& SSearchableMultiRigHierarchyTreeView::GetDisplaySettings()
+const FRigHierarchyTreeDisplaySettings& SSearchableMultiRigHierarchyTreeView::GetDisplaySettings()
 {
 	if (SuperGetRigTreeDisplaySettings.IsBound())
 	{
@@ -2240,7 +2240,7 @@ void SControlRigOutliner::HandleSelectionChanged(TSharedPtr<FMultiRigTreeElement
 		if (!Data.IsControlElement())
 		{
 			TArray<FRigElementKey>& Keys = SelectedRigAndKeys.FindOrAdd(Data.WeakControlRig.Get());
-			const FRigTreeDisplaySettings& Settings = HierarchyTreeView->GetDisplaySettings();
+			const FRigHierarchyTreeDisplaySettings& Settings = HierarchyTreeView->GetDisplaySettings();
 			if (Settings.FilterText.IsEmpty() || !Settings.bFlattenHierarchyOnFilter)
 			{
 				TSharedPtr<FMultiRigTreeElement> Element = HierarchyTreeView->GetTreeView()->FindElement(Data);
diff --git a/Engine/Plugins/Animation/ControlRig/Source/ControlRigEditor/Private/EditMode/SControlRigOutliner.h b/Engine/Plugins/Animation/ControlRig/Source/ControlRigEditor/Private/EditMode/SControlRigOutliner.h
--- a/Engine/Plugins/Animation/ControlRig/Source/ControlRigEditor/Private/EditMode/SControlRigOutliner.h
+++ b/Engine/Plugins/Animation/ControlRig/Source/ControlRigEditor/Private/EditMode/SControlRigOutliner.h
@@ -19,7 +19,7 @@
 
 class ISequencer;
 class SExpandableArea;
-class SSearchableRigHierarchyTreeView;
+class SRigHierarchySearchableTreeView;
 
 class SMultiRigHierarchyTreeView;
 class SMultiRigHierarchyItem;
@@ -53,7 +53,7 @@ struct FMultiRigTreeDelegates
 	}
 
 
-	const FRigTreeDisplaySettings& GetDisplaySettings() const
+	const FRigHierarchyTreeDisplaySettings& GetDisplaySettings() const
 	{
 		if (OnGetDisplaySettings.IsBound())
 		{
@@ -81,7 +81,7 @@ struct FMultiRigTreeDelegates
 		OnSelectionChanged.ExecuteIfBound(Selection, SelectInfo);
 	}
 
-	static UE_API FRigTreeDisplaySettings DefaultDisplaySettings;
+	static UE_API FRigHierarchyTreeDisplaySettings DefaultDisplaySettings;
 	bool bIsChangingRigHierarchy;
 };
 
@@ -106,9 +106,9 @@ public:
 	FMultiRigData(UControlRig* InControlRig, FRigElementKey InKey) : WeakControlRig(InControlRig), Type(EMultiRigDataType_Element), Key(InKey) {};
 	FMultiRigData(UControlRig* InControlRig, EMultiRigDataType InType, FName InName) : WeakControlRig(InControlRig), Type(InType), Name(InName) {};
 	FText GetName() const;
-	FText GetDisplayName(const FRigTreeDisplaySettings& DisplaySettings) const;
+	FText GetDisplayName(const FRigHierarchyTreeDisplaySettings& DisplaySettings) const;
 	void InvalidateDisplayName() const;
-	FText GetToolTipText(const FRigTreeDisplaySettings& Settings) const;
+	FText GetToolTipText(const FRigHierarchyTreeDisplaySettings& Settings) const;
 	bool operator == (const FMultiRigData & Other) const;
 	bool IsValid() const;
 	void Reset() { Type = EMultiRigDataType_Invalid; }
@@ -151,9 +151,9 @@ public:
 	FMultiRigData Data;
 	TArray<TSharedPtr<FMultiRigTreeElement>> Children;
 
-	TSharedRef<ITableRow> MakeTreeRowWidget(const TSharedRef<STableViewBase>& InOwnerTable, TSharedRef<FMultiRigTreeElement> InRigTreeElement, TSharedPtr<SMultiRigHierarchyTreeView> InTreeView, const FRigTreeDisplaySettings& InSettings, bool bPinned);
+	TSharedRef<ITableRow> MakeTreeRowWidget(const TSharedRef<STableViewBase>& InOwnerTable, TSharedRef<FMultiRigTreeElement> InRigTreeElement, TSharedPtr<SMultiRigHierarchyTreeView> InTreeView, const FRigHierarchyTreeDisplaySettings& InSettings, bool bPinned);
 
-	void RefreshDisplaySettings(const URigHierarchy* InHierarchy, const FRigTreeDisplaySettings& InSettings);
+	void RefreshDisplaySettings(const URigHierarchy* InHierarchy, const FRigHierarchyTreeDisplaySettings& InSettings);
 	bool AreControlsVisible() const;
 
 	/** The current filter result */
@@ -180,7 +180,7 @@ class SMultiRigHierarchyItem : public SMultiColumnTableRow<TSharedPtr<FMultiRigT
 {
 public:
 
-	void Construct(const FArguments& InArgs, const TSharedRef<STableViewBase>& OwnerTable, TSharedRef<FMultiRigTreeElement> InRigTreeElement, TSharedPtr<SMultiRigHierarchyTreeView> InTreeView, const FRigTreeDisplaySettings& InSettings, bool bPinned);
+	void Construct(const FArguments& InArgs, const TSharedRef<STableViewBase>& OwnerTable, TSharedRef<FMultiRigTreeElement> InRigTreeElement, TSharedPtr<SMultiRigHierarchyTreeView> InTreeView, const FRigHierarchyTreeDisplaySettings& InSettings, bool bPinned);
 	static TPair<const FSlateBrush*, FSlateColor> GetBrushForElementType(const URigHierarchy* InHierarchy, const FMultiRigData& InData);
 	virtual const FSlateBrush* GetBorder() const override;
 	virtual TSharedRef<SWidget> GenerateWidgetForColumn( const FName& InColumnName ) override;
@@ -310,7 +310,7 @@ public:
 	void Construct(const FArguments& InArgs);
 	virtual ~SSearchableMultiRigHierarchyTreeView() {}
 	TSharedRef<SMultiRigHierarchyTreeView> GetTreeView() const { return TreeView.ToSharedRef(); }
-	const FRigTreeDisplaySettings& GetDisplaySettings();
+	const FRigHierarchyTreeDisplaySettings& GetDisplaySettings();
 	TSharedRef<SWidget> OnGetOptionsMenu();
 	void BindCommands();
 	void OnSettingChanged(UObject* Object, FPropertyChangedEvent& PropertyChangedEvent);
@@ -339,7 +339,7 @@ private:
 	FOnGetRigTreeDisplaySettings SuperGetRigTreeDisplaySettings;
 	FOnMultiRigTreeGetEditMode GetEditMode;
 	FText FilterText;
-	FRigTreeDisplaySettings Settings;
+	FRigHierarchyTreeDisplaySettings Settings;
 	TSharedPtr<SMultiRigHierarchyTreeView> TreeView;
 };
 
@@ -373,8 +373,8 @@ private:
 
 	/** Hierarchy picker for controls*/
 	TSharedPtr<SSearchableMultiRigHierarchyTreeView> HierarchyTreeView;
-	FRigTreeDisplaySettings DisplaySettings;
-	const FRigTreeDisplaySettings& GetDisplaySettings() const { return DisplaySettings; }
+	FRigHierarchyTreeDisplaySettings DisplaySettings;
+	const FRigHierarchyTreeDisplaySettings& GetDisplaySettings() const { return DisplaySettings; }
 	bool bIsChangingRigHierarchy = false;
 
 	//set of control rigs we are bound too and need to clear delegates from
diff --git a/Engine/Plugins/Animation/ControlRig/Source/ControlRigEditor/Private/EditMode/SControlRigSpacePicker.h b/Engine/Plugins/Animation/ControlRig/Source/ControlRigEditor/Private/EditMode/SControlRigSpacePicker.h
--- a/Engine/Plugins/Animation/ControlRig/Source/ControlRigEditor/Private/EditMode/SControlRigSpacePicker.h
+++ b/Engine/Plugins/Animation/ControlRig/Source/ControlRigEditor/Private/EditMode/SControlRigSpacePicker.h
@@ -14,7 +14,7 @@
 
 class ISequencer;
 class SExpandableArea;
-class SSearchableRigHierarchyTreeView;
+class SRigHierarchySearchableTreeView;
 class UControlRig;
 namespace UE::ControlRigEditor { class FRigSelectionViewModel; }
 
diff --git a/Engine/Plugins/Animation/ControlRig/Source/ControlRigEditor/Private/Editor/ControlRigContextMenuContext.cpp b/Engine/Plugins/Animation/ControlRig/Source/ControlRigEditor/Private/Editor/ControlRigContextMenuContext.cpp
--- a/Engine/Plugins/Animation/ControlRig/Source/ControlRigEditor/Private/Editor/ControlRigContextMenuContext.cpp
+++ b/Engine/Plugins/Animation/ControlRig/Source/ControlRigEditor/Private/Editor/ControlRigContextMenuContext.cpp
@@ -79,9 +79,9 @@ SRigHierarchy* UControlRigContextMenuContext::GetRigHierarchyPanel() const
 	return nullptr;
 }
 
-SModularRigModel* UControlRigContextMenuContext::GetModularRigModelPanel() const
+SModularRigHierarchy* UControlRigContextMenuContext::GetModularRigModelPanel() const
 {
-	if (const TSharedPtr<SModularRigModel> ModularRigModelPanel = MenuSpecificContext.ModularRigModelPanel.Pin())
+	if (const TSharedPtr<SModularRigHierarchy> ModularRigModelPanel = MenuSpecificContext.ModularRigModelPanel.Pin())
 	{
 		return ModularRigModelPanel.Get();
 	}
diff --git a/Engine/Plugins/Animation/ControlRig/Source/ControlRigEditor/Private/Editor/ControlRigContextMenuContext.h b/Engine/Plugins/Animation/ControlRig/Source/ControlRigEditor/Private/Editor/ControlRigContextMenuContext.h
--- a/Engine/Plugins/Animation/ControlRig/Source/ControlRigEditor/Private/Editor/ControlRigContextMenuContext.h
+++ b/Engine/Plugins/Animation/ControlRig/Source/ControlRigEditor/Private/Editor/ControlRigContextMenuContext.h
@@ -17,7 +17,7 @@ class URigVMGraph;
 class URigVMNode;
 class URigVMPin;
 class SRigHierarchy;
-class SModularRigModel;
+class SModularRigHierarchy;
 
 USTRUCT(BlueprintType)
 struct FControlRigRigHierarchyDragAndDropContext
@@ -100,7 +100,7 @@ struct FControlRigMenuSpecificContext
 	
 	FControlRigRigHierarchyDragAndDropContext RigHierarchyDragAndDropContext;
 
-	TWeakPtr<SModularRigModel> ModularRigModelPanel;
+	TWeakPtr<SModularRigHierarchy> ModularRigModelPanel;
 	
 	FControlRigGraphNodeContextMenuContext GraphNodeContextMenuContext;
 
@@ -149,7 +149,7 @@ public:
 
 	SRigHierarchy* GetRigHierarchyPanel() const;
 
-	SModularRigModel* GetModularRigModelPanel() const;
+	SModularRigHierarchy* GetModularRigModelPanel() const;
 	
 	IControlRigBaseEditor* GetControlRigEditor() const;
 
diff --git a/Engine/Plugins/Animation/ControlRig/Source/ControlRigEditor/Private/Editor/ControlRigEditor.cpp b/Engine/Plugins/Animation/ControlRig/Source/ControlRigEditor/Private/Editor/ControlRigEditor.cpp
--- a/Engine/Plugins/Animation/ControlRig/Source/ControlRigEditor/Private/Editor/ControlRigEditor.cpp
+++ b/Engine/Plugins/Animation/ControlRig/Source/ControlRigEditor/Private/Editor/ControlRigEditor.cpp
@@ -5512,9 +5512,9 @@ void FControlRigBaseEditor::OnGraphNodeDropToPerformImpl(TSharedPtr<FDragDropOpe
 {
 	DECLARE_SCOPE_HIERARCHICAL_COUNTER_FUNC()
 
-	if (InDragDropOp->IsOfType<FRigElementHierarchyDragDropOp>())
+	if (InDragDropOp->IsOfType<FRigHierarchyElementDragDropOp>())
 	{
-		TSharedPtr<FRigElementHierarchyDragDropOp> RigHierarchyOp = StaticCastSharedPtr<FRigElementHierarchyDragDropOp>(InDragDropOp);
+		TSharedPtr<FRigHierarchyElementDragDropOp> RigHierarchyOp = StaticCastSharedPtr<FRigHierarchyElementDragDropOp>(InDragDropOp);
 
 		if (RigHierarchyOp->GetElements().Num() > 0 && GetFocusedGraphEd().IsValid())
 		{
diff --git a/Engine/Plugins/Animation/ControlRig/Source/ControlRigEditor/Private/Editor/ControlRigEditor.h b/Engine/Plugins/Animation/ControlRig/Source/ControlRigEditor/Private/Editor/ControlRigEditor.h
--- a/Engine/Plugins/Animation/ControlRig/Source/ControlRigEditor/Private/Editor/ControlRigEditor.h
+++ b/Engine/Plugins/Animation/ControlRig/Source/ControlRigEditor/Private/Editor/ControlRigEditor.h
@@ -549,7 +549,7 @@ protected:
 	friend class FModularRigEditorMode;
 	friend class SControlRigStackView;
 	friend class SRigHierarchy;
-	friend class SModularRigModel;
+	friend class SModularRigHierarchy;
 	friend struct FRigHierarchyTabSummoner;
 	friend struct FModularRigModelTabSummoner;
 };
diff --git a/Engine/Plugins/Animation/ControlRig/Source/ControlRigEditor/Private/Editor/ControlRigSchematicModel.cpp b/Engine/Plugins/Animation/ControlRig/Source/ControlRigEditor/Private/Editor/ControlRigSchematicModel.cpp
--- a/Engine/Plugins/Animation/ControlRig/Source/ControlRigEditor/Private/Editor/ControlRigSchematicModel.cpp
+++ b/Engine/Plugins/Animation/ControlRig/Source/ControlRigEditor/Private/Editor/ControlRigSchematicModel.cpp
@@ -1203,9 +1203,9 @@ TArray<FRigElementKey> FControlRigSchematicModel::GetElementKeysFromDragDropEven
 {
 	TArray<FRigElementKey> DraggedKeys;
 
-	if(InDragDropOperation.IsOfType<FRigElementHierarchyDragDropOp>())
+	if(InDragDropOperation.IsOfType<FRigHierarchyElementDragDropOp>())
 	{
-		const FRigElementHierarchyDragDropOp* RigDragDropOp = StaticCast<const FRigElementHierarchyDragDropOp*>(&InDragDropOperation);
+		const FRigHierarchyElementDragDropOp* RigDragDropOp = StaticCast<const FRigHierarchyElementDragDropOp*>(&InDragDropOperation);
 		const TArray<FRigHierarchyKey>& HierarchyKeys = RigDragDropOp->GetElements();
 		TArray<FRigElementKey> Elements;
 		Elements.Reserve(HierarchyKeys.Num());
@@ -1227,9 +1227,9 @@ TArray<FRigElementKey> FControlRigSchematicModel::GetElementKeysFromDragDropEven
 		return {DraggedKey};
 	}
 	
-	if(InDragDropOperation.IsOfType<FModularRigModuleDragDropOp>())
+	if(InDragDropOperation.IsOfType<FModularRigHierarchyElementDragDropOp>())
 	{
-		const FModularRigModuleDragDropOp* ModuleDropOp = StaticCast<const FModularRigModuleDragDropOp*>(&InDragDropOperation);
+		const FModularRigHierarchyElementDragDropOp* ModuleDropOp = StaticCast<const FModularRigHierarchyElementDragDropOp*>(&InDragDropOperation);
 		const TArray<FName> ModuleNames = ModuleDropOp->GetModules();
 
 		URigHierarchy* Hierarchy = InControlRig->GetHierarchy();
@@ -1523,7 +1523,7 @@ void FControlRigSchematicModel::HandleSchematicDrop(SSchematicGraphPanel* InPane
 
 	TSharedPtr<FAssetDragDropOp> AssetDragDropOp = InDragDropEvent.GetOperationAs<FAssetDragDropOp>();
 	TSharedPtr<FSchematicGraphNodeDragDropOp> SchematicDragDropOp = InDragDropEvent.GetOperationAs<FSchematicGraphNodeDragDropOp>();
-	TSharedPtr<FModularRigModuleDragDropOp> ModuleDragDropOperation = InDragDropEvent.GetOperationAs<FModularRigModuleDragDropOp>();
+	TSharedPtr<FModularRigHierarchyElementDragDropOp> ModuleDragDropOperation = InDragDropEvent.GetOperationAs<FModularRigHierarchyElementDragDropOp>();
 	if (AssetDragDropOp.IsValid())
 	{
 		for (const FAssetData& AssetData : AssetDragDropOp->GetAssets())
diff --git a/Engine/Plugins/Animation/ControlRig/Source/ControlRigEditor/Private/Editor/Hierarchy/DragDrop/ModularRigHierarchyElementDragDropOp.cpp b/Engine/Plugins/Animation/ControlRig/Source/ControlRigEditor/Private/Editor/Hierarchy/DragDrop/ModularRigHierarchyElementDragDropOp.cpp
--- a/Engine/Plugins/Animation/ControlRig/Source/ControlRigEditor/Private/Editor/Hierarchy/DragDrop/ModularRigHierarchyElementDragDropOp.cpp
+++ b/Engine/Plugins/Animation/ControlRig/Source/ControlRigEditor/Private/Editor/Hierarchy/DragDrop/ModularRigHierarchyElementDragDropOp.cpp
@@ -5,15 +5,15 @@
 #include "Widgets/Layout/SBorder.h"
 #include "Widgets/Text/STextBlock.h"
 
-TSharedRef<FModularRigModuleDragDropOp> FModularRigModuleDragDropOp::New(const TArray<FName>& InModuleNames)
+TSharedRef<FModularRigHierarchyElementDragDropOp> FModularRigHierarchyElementDragDropOp::New(const TArray<FName>& InModuleNames)
 {
-	TSharedRef<FModularRigModuleDragDropOp> Operation = MakeShared<FModularRigModuleDragDropOp>();
+	TSharedRef<FModularRigHierarchyElementDragDropOp> Operation = MakeShared<FModularRigHierarchyElementDragDropOp>();
 	Operation->ModuleNames = InModuleNames;
 	Operation->Construct();
 	return Operation;
 }
 
-TSharedPtr<SWidget> FModularRigModuleDragDropOp::GetDefaultDecorator() const
+TSharedPtr<SWidget> FModularRigHierarchyElementDragDropOp::GetDefaultDecorator() const
 {
 	return SNew(SBorder)
 		.Visibility(EVisibility::Visible)
@@ -25,7 +25,7 @@ TSharedPtr<SWidget> FModularRigModuleDragDropOp::GetDefaultDecorator() const
 		];
 }
 
-FString FModularRigModuleDragDropOp::GetJoinedModuleNames() const
+FString FModularRigHierarchyElementDragDropOp::GetJoinedModuleNames() const
 {
 	TArray<FString> ElementNameStrings;
 	for (const FName& ModuleName: ModuleNames)
diff --git a/Engine/Plugins/Animation/ControlRig/Source/ControlRigEditor/Private/Editor/Hierarchy/DragDrop/ModularRigHierarchyElementDragDropOp.h b/Engine/Plugins/Animation/ControlRig/Source/ControlRigEditor/Private/Editor/Hierarchy/DragDrop/ModularRigHierarchyElementDragDropOp.h
--- a/Engine/Plugins/Animation/ControlRig/Source/ControlRigEditor/Private/Editor/Hierarchy/DragDrop/ModularRigHierarchyElementDragDropOp.h
+++ b/Engine/Plugins/Animation/ControlRig/Source/ControlRigEditor/Private/Editor/Hierarchy/DragDrop/ModularRigHierarchyElementDragDropOp.h
@@ -4,12 +4,12 @@
 
 #include "Input/DragAndDrop.h"
 
-class FModularRigModuleDragDropOp : public FDragDropOperation
+class FModularRigHierarchyElementDragDropOp : public FDragDropOperation
 {
 public:
-	DRAG_DROP_OPERATOR_TYPE(FModularRigModuleDragDropOp, FDragDropOperation)
+	DRAG_DROP_OPERATOR_TYPE(FModularRigHierarchyElementDragDropOp, FDragDropOperation)
 
-	static TSharedRef<FModularRigModuleDragDropOp> New(const TArray<FName>& InModuleNames);
+	static TSharedRef<FModularRigHierarchyElementDragDropOp> New(const TArray<FName>& InModuleNames);
 
 	virtual TSharedPtr<SWidget> GetDefaultDecorator() const override;
 
diff --git a/Engine/Plugins/Animation/ControlRig/Source/ControlRigEditor/Private/Editor/Hierarchy/DragDrop/RigHierarchyElementDragDropOp.cpp b/Engine/Plugins/Animation/ControlRig/Source/ControlRigEditor/Private/Editor/Hierarchy/DragDrop/RigHierarchyElementDragDropOp.cpp
--- a/Engine/Plugins/Animation/ControlRig/Source/ControlRigEditor/Private/Editor/Hierarchy/DragDrop/RigHierarchyElementDragDropOp.cpp
+++ b/Engine/Plugins/Animation/ControlRig/Source/ControlRigEditor/Private/Editor/Hierarchy/DragDrop/RigHierarchyElementDragDropOp.cpp
@@ -6,15 +6,15 @@
 #include "Widgets/Layout/SBorder.h"
 #include "Widgets/Text/STextBlock.h"
 
-TSharedRef<FRigElementHierarchyDragDropOp> FRigElementHierarchyDragDropOp::New(const TArray<FRigHierarchyKey>& InElements)
+TSharedRef<FRigHierarchyElementDragDropOp> FRigHierarchyElementDragDropOp::New(const TArray<FRigHierarchyKey>& InElements)
 {
-	TSharedRef<FRigElementHierarchyDragDropOp> Operation = MakeShared<FRigElementHierarchyDragDropOp>();
+	TSharedRef<FRigHierarchyElementDragDropOp> Operation = MakeShared<FRigHierarchyElementDragDropOp>();
 	Operation->Elements = InElements;
 	Operation->Construct();
 	return Operation;
 }
 
-TSharedPtr<SWidget> FRigElementHierarchyDragDropOp::GetDefaultDecorator() const
+TSharedPtr<SWidget> FRigHierarchyElementDragDropOp::GetDefaultDecorator() const
 {
 	return SNew(SBorder)
 		.Visibility(EVisibility::Visible)
@@ -26,7 +26,7 @@ TSharedPtr<SWidget> FRigElementHierarchyDragDropOp::GetDefaultDecorator() const
 		];
 }
 
-FString FRigElementHierarchyDragDropOp::GetJoinedElementNames() const
+FString FRigHierarchyElementDragDropOp::GetJoinedElementNames() const
 {
 	TArray<FString> ElementNameStrings;
 	for (const FRigHierarchyKey& Element: Elements)
@@ -36,7 +36,7 @@ FString FRigElementHierarchyDragDropOp::GetJoinedElementNames() const
 	return FString::Join(ElementNameStrings, TEXT(","));
 }
 
-bool FRigElementHierarchyDragDropOp::IsDraggingSingleConnector() const
+bool FRigHierarchyElementDragDropOp::IsDraggingSingleConnector() const
 {
 	if(Elements.Num() == 1)
 	{
@@ -48,7 +48,7 @@ bool FRigElementHierarchyDragDropOp::IsDraggingSingleConnector() const
 	return false;
 }
 
-bool FRigElementHierarchyDragDropOp::IsDraggingSingleSocket() const
+bool FRigHierarchyElementDragDropOp::IsDraggingSingleSocket() const
 {
 	if(Elements.Num() == 1)
 	{
diff --git a/Engine/Plugins/Animation/ControlRig/Source/ControlRigEditor/Private/Editor/Hierarchy/DragDrop/RigHierarchyElementDragDropOp.h b/Engine/Plugins/Animation/ControlRig/Source/ControlRigEditor/Private/Editor/Hierarchy/DragDrop/RigHierarchyElementDragDropOp.h
--- a/Engine/Plugins/Animation/ControlRig/Source/ControlRigEditor/Private/Editor/Hierarchy/DragDrop/RigHierarchyElementDragDropOp.h
+++ b/Engine/Plugins/Animation/ControlRig/Source/ControlRigEditor/Private/Editor/Hierarchy/DragDrop/RigHierarchyElementDragDropOp.h
@@ -7,12 +7,12 @@
 
 class SRigHierarchyTagWidget;
 
-class FRigElementHierarchyDragDropOp : public FGraphNodeDragDropOp
+class FRigHierarchyElementDragDropOp : public FGraphNodeDragDropOp
 {
 public:
-	DRAG_DROP_OPERATOR_TYPE(FRigElementHierarchyDragDropOp, FGraphNodeDragDropOp)
+	DRAG_DROP_OPERATOR_TYPE(FRigHierarchyElementDragDropOp, FGraphNodeDragDropOp)
 
-	static TSharedRef<FRigElementHierarchyDragDropOp> New(const TArray<FRigHierarchyKey>& InElements);
+	static TSharedRef<FRigHierarchyElementDragDropOp> New(const TArray<FRigHierarchyKey>& InElements);
 
 	virtual TSharedPtr<SWidget> GetDefaultDecorator() const override;
 
diff --git a/Engine/Plugins/Animation/ControlRig/Source/ControlRigEditor/Private/Editor/Hierarchy/Models/ModularRigHierarchyTreeElement.cpp b/Engine/Plugins/Animation/ControlRig/Source/ControlRigEditor/Private/Editor/Hierarchy/Models/ModularRigHierarchyTreeElement.cpp
--- a/Engine/Plugins/Animation/ControlRig/Source/ControlRigEditor/Private/Editor/Hierarchy/Models/ModularRigHierarchyTreeElement.cpp
+++ b/Engine/Plugins/Animation/ControlRig/Source/ControlRigEditor/Private/Editor/Hierarchy/Models/ModularRigHierarchyTreeElement.cpp
@@ -5,11 +5,11 @@
 #include "Editor/Hierarchy/Views/SModularRigTreeView.h"
 #include "ModularRig.h"
 
-TMap<FSoftObjectPath, TSharedPtr<FSlateBrush>> FModularRigTreeElement::IconPathToBrush;
+TMap<FSoftObjectPath, TSharedPtr<FSlateBrush>> FModularRigHierarchyTreeElement::IconPathToBrush;
 
-TArray<TStrongObjectPtr<UTexture2D>> FModularRigTreeElement::Icons;
+TArray<TStrongObjectPtr<UTexture2D>> FModularRigHierarchyTreeElement::Icons;
 
-FModularRigTreeElement::FModularRigTreeElement(const FString& InKey, TWeakPtr<SModularRigTreeView> InTreeView, bool bInIsPrimary)
+FModularRigHierarchyTreeElement::FModularRigHierarchyTreeElement(const FString& InKey, TWeakPtr<SModularRigTreeView> InTreeView, bool bInIsPrimary)
 {
 	Key = InKey;
 	bIsPrimary = bInIsPrimary;
@@ -51,7 +51,7 @@ FModularRigTreeElement::FModularRigTreeElement(const FString& InKey, TWeakPtr<SM
 	}
 }
 
-void FModularRigTreeElement::RefreshDisplaySettings(const UModularRig* InModularRig)
+void FModularRigHierarchyTreeElement::RefreshDisplaySettings(const UModularRig* InModularRig)
 {
 	const TPair<const FSlateBrush*, FSlateColor> Result = GetBrushAndColor(InModularRig);
 
@@ -60,12 +60,12 @@ void FModularRigTreeElement::RefreshDisplaySettings(const UModularRig* InModular
 	TextColor = FSlateColor::UseForeground();
 }
 
-TSharedRef<ITableRow> FModularRigTreeElement::MakeTreeRowWidget(const TSharedRef<STableViewBase>& InOwnerTable, TSharedRef<FModularRigTreeElement> InRigTreeElement, TSharedPtr<SModularRigTreeView> InTreeView, bool bPinned)
+TSharedRef<ITableRow> FModularRigHierarchyTreeElement::MakeTreeRowWidget(const TSharedRef<STableViewBase>& InOwnerTable, TSharedRef<FModularRigHierarchyTreeElement> InRigTreeElement, TSharedPtr<SModularRigTreeView> InTreeView, bool bPinned)
 {
-	return SNew(SModularRigModelItem, InOwnerTable, InRigTreeElement, InTreeView, bPinned);
+	return SNew(SModularRigHierarchyTreeItem, InOwnerTable, InRigTreeElement, InTreeView, bPinned);
 }
 
-void FModularRigTreeElement::RequestRename()
+void FModularRigHierarchyTreeElement::RequestRename()
 {
 	OnRenameRequested.ExecuteIfBound();
 }
diff --git a/Engine/Plugins/Animation/ControlRig/Source/ControlRigEditor/Private/Editor/Hierarchy/Models/ModularRigHierarchyTreeElement.h b/Engine/Plugins/Animation/ControlRig/Source/ControlRigEditor/Private/Editor/Hierarchy/Models/ModularRigHierarchyTreeElement.h
--- a/Engine/Plugins/Animation/ControlRig/Source/ControlRigEditor/Private/Editor/Hierarchy/Models/ModularRigHierarchyTreeElement.h
+++ b/Engine/Plugins/Animation/ControlRig/Source/ControlRigEditor/Private/Editor/Hierarchy/Models/ModularRigHierarchyTreeElement.h
@@ -9,10 +9,10 @@ class UModularRig;
 class UTexture2D;
 
 /** An item in the tree */
-class FModularRigTreeElement : public TSharedFromThis<FModularRigTreeElement>
+class FModularRigHierarchyTreeElement : public TSharedFromThis<FModularRigHierarchyTreeElement>
 {
 public:
-	FModularRigTreeElement(const FString& InKey, TWeakPtr<SModularRigTreeView> InTreeView, bool bInIsPrimary);
+	FModularRigHierarchyTreeElement(const FString& InKey, TWeakPtr<SModularRigTreeView> InTreeView, bool bInIsPrimary);
 
 public:
 	/** Element Data to display */
@@ -21,9 +21,9 @@ public:
 	FName ModuleName;
 	FString ConnectorName;
 	FName ShortName;
-	TArray<TSharedPtr<FModularRigTreeElement>> Children;
+	TArray<TSharedPtr<FModularRigHierarchyTreeElement>> Children;
 
-	TSharedRef<ITableRow> MakeTreeRowWidget(const TSharedRef<STableViewBase>& InOwnerTable, TSharedRef<FModularRigTreeElement> InRigTreeElement, TSharedPtr<SModularRigTreeView> InTreeView, bool bPinned);
+	TSharedRef<ITableRow> MakeTreeRowWidget(const TSharedRef<STableViewBase>& InOwnerTable, TSharedRef<FModularRigHierarchyTreeElement> InRigTreeElement, TSharedPtr<SModularRigTreeView> InTreeView, bool bPinned);
 
 	void RequestRename();
 
diff --git a/Engine/Plugins/Animation/ControlRig/Source/ControlRigEditor/Private/Editor/Hierarchy/Models/RigHierarchyTreeElement.cpp b/Engine/Plugins/Animation/ControlRig/Source/ControlRigEditor/Private/Editor/Hierarchy/Models/RigHierarchyTreeElement.cpp
--- a/Engine/Plugins/Animation/ControlRig/Source/ControlRigEditor/Private/Editor/Hierarchy/Models/RigHierarchyTreeElement.cpp
+++ b/Engine/Plugins/Animation/ControlRig/Source/ControlRigEditor/Private/Editor/Hierarchy/Models/RigHierarchyTreeElement.cpp
@@ -10,7 +10,7 @@
 #include "Styling/SlateColor.h"
 #include "Widgets/Views/ITableRow.h"
 
-FRigTreeElement::FRigTreeElement(const FRigHierarchyKey& InKey, TWeakPtr<SRigHierarchyTreeView> InTreeView, bool InSupportsRename, ERigTreeFilterResult InFilterResult)
+FRigHierarchyTreeElement::FRigHierarchyTreeElement(const FRigHierarchyKey& InKey, TWeakPtr<SRigHierarchyTreeView> InTreeView, bool InSupportsRename, ERigTreeFilterResult InFilterResult)
 {
 	Key = InKey;
 	ShortName = LongName = FText::FromName(InKey.GetFName());
@@ -32,19 +32,19 @@ FRigTreeElement::FRigTreeElement(const FRigHierarchyKey& InKey, TWeakPtr<SRigHie
 				ShortName = Hierarchy->GetDisplayNameForUI(InKey.GetElement(), EElementNameDisplayMode::ForceShort);
 			}
 
-			const FRigTreeDisplaySettings& Settings = InTreeView.Pin()->GetRigTreeDelegates().GetDisplaySettings();
+			const FRigHierarchyTreeDisplaySettings& Settings = InTreeView.Pin()->GetRigTreeDelegates().GetDisplaySettings();
 			RefreshDisplaySettings(Hierarchy, Settings);
 		}
 	}
 }
 
 
-TSharedRef<ITableRow> FRigTreeElement::MakeTreeRowWidget(const TSharedRef<STableViewBase>& InOwnerTable, TSharedRef<FRigTreeElement> InRigTreeElement, TSharedPtr<SRigHierarchyTreeView> InTreeView, const FRigTreeDisplaySettings& InSettings, bool bPinned)
+TSharedRef<ITableRow> FRigHierarchyTreeElement::MakeTreeRowWidget(const TSharedRef<STableViewBase>& InOwnerTable, TSharedRef<FRigHierarchyTreeElement> InRigTreeElement, TSharedPtr<SRigHierarchyTreeView> InTreeView, const FRigHierarchyTreeDisplaySettings& InSettings, bool bPinned)
 {
 	return SNew(SRigHierarchyItem, InOwnerTable, InRigTreeElement, InTreeView, InSettings, bPinned);
 }
 
-void FRigTreeElement::RequestRename()
+void FRigHierarchyTreeElement::RequestRename()
 {
 	if (bSupportsRename)
 	{
@@ -52,7 +52,7 @@ void FRigTreeElement::RequestRename()
 	}
 }
 
-void FRigTreeElement::RefreshDisplaySettings(const URigHierarchy* InHierarchy, const FRigTreeDisplaySettings& InSettings)
+void FRigHierarchyTreeElement::RefreshDisplaySettings(const URigHierarchy* InHierarchy, const FRigHierarchyTreeDisplaySettings& InSettings)
 {
 	const TPair<const FSlateBrush*, FSlateColor> Result = SRigHierarchyItem::GetBrushForElementType(InHierarchy, Key);
 
@@ -95,7 +95,7 @@ void FRigTreeElement::RefreshDisplaySettings(const URigHierarchy* InHierarchy, c
 		(bIsProcedural ? FSlateColor(FLinearColor(0.9f, 0.8f, 0.4f) * 0.5f) : FSlateColor(FLinearColor::Gray * 0.5f));
 }
 
-FSlateColor FRigTreeElement::GetIconColor() const
+FSlateColor FRigHierarchyTreeElement::GetIconColor() const
 {
 	if (bFadedOutDuringDragDrop)
 	{
@@ -107,7 +107,7 @@ FSlateColor FRigTreeElement::GetIconColor() const
 	return IconColor;
 }
 
-FSlateColor FRigTreeElement::GetTextColor() const
+FSlateColor FRigHierarchyTreeElement::GetTextColor() const
 {
 	if (bFadedOutDuringDragDrop)
 	{
diff --git a/Engine/Plugins/Animation/ControlRig/Source/ControlRigEditor/Private/Editor/Hierarchy/Models/RigHierarchyTreeElement.h b/Engine/Plugins/Animation/ControlRig/Source/ControlRigEditor/Private/Editor/Hierarchy/Models/RigHierarchyTreeElement.h
--- a/Engine/Plugins/Animation/ControlRig/Source/ControlRigEditor/Private/Editor/Hierarchy/Models/RigHierarchyTreeElement.h
+++ b/Engine/Plugins/Animation/ControlRig/Source/ControlRigEditor/Private/Editor/Hierarchy/Models/RigHierarchyTreeElement.h
@@ -33,10 +33,10 @@ enum class ERigTreeFilterResult : int32
 };
 
 /** An item in the tree */
-class FRigTreeElement : public TSharedFromThis<FRigTreeElement>
+class FRigHierarchyTreeElement : public TSharedFromThis<FRigHierarchyTreeElement>
 {
 public:
-	FRigTreeElement(const FRigHierarchyKey& InKey, TWeakPtr<SRigHierarchyTreeView> InTreeView, bool InSupportsRename, ERigTreeFilterResult InFilterResult);
+	FRigHierarchyTreeElement(const FRigHierarchyKey& InKey, TWeakPtr<SRigHierarchyTreeView> InTreeView, bool InSupportsRename, ERigTreeFilterResult InFilterResult);
 public:
 	/** Element Data to display */
 	FRigHierarchyKey Key;
@@ -47,13 +47,13 @@ public:
 	bool bIsAnimationChannel;
 	bool bIsProcedural;
 	bool bSupportsRename;
-	TArray<TSharedPtr<FRigTreeElement>> Children;
+	TArray<TSharedPtr<FRigHierarchyTreeElement>> Children;
 
-	TSharedRef<ITableRow> MakeTreeRowWidget(const TSharedRef<STableViewBase>& InOwnerTable, TSharedRef<FRigTreeElement> InRigTreeElement, TSharedPtr<SRigHierarchyTreeView> InTreeView, const FRigTreeDisplaySettings& InSettings, bool bPinned);
+	TSharedRef<ITableRow> MakeTreeRowWidget(const TSharedRef<STableViewBase>& InOwnerTable, TSharedRef<FRigHierarchyTreeElement> InRigTreeElement, TSharedPtr<SRigHierarchyTreeView> InTreeView, const FRigHierarchyTreeDisplaySettings& InSettings, bool bPinned);
 
 	void RequestRename();
 
-	void RefreshDisplaySettings(const URigHierarchy* InHierarchy, const FRigTreeDisplaySettings& InSettings);
+	void RefreshDisplaySettings(const URigHierarchy* InHierarchy, const FRigHierarchyTreeDisplaySettings& InSettings);
 
 	FSlateColor GetIconColor() const;
 	FSlateColor GetTextColor() const;
diff --git a/Engine/Plugins/Animation/ControlRig/Source/ControlRigEditor/Private/Editor/Hierarchy/ModularRigHierarchyTreeDelegates.cpp b/Engine/Plugins/Animation/ControlRig/Source/ControlRigEditor/Private/Editor/Hierarchy/ModularRigHierarchyTreeDelegates.cpp
--- a/Engine/Plugins/Animation/ControlRig/Source/ControlRigEditor/Private/Editor/Hierarchy/ModularRigHierarchyTreeDelegates.cpp
+++ b/Engine/Plugins/Animation/ControlRig/Source/ControlRigEditor/Private/Editor/Hierarchy/ModularRigHierarchyTreeDelegates.cpp
@@ -2,7 +2,7 @@
 
 #include "ModularRigHierarchyTreeDelegates.h"
 
-const UModularRig* FModularRigTreeDelegates::GetModularRig() const
+const UModularRig* FModularRigHierarchyTreeDelegates::GetModularRig() const
 {
 	if (OnGetModularRig.IsBound())
 	{
@@ -11,7 +11,7 @@ const UModularRig* FModularRigTreeDelegates::GetModularRig() const
 	return nullptr;
 }
 
-FName FModularRigTreeDelegates::HandleRenameElement(const FName& OldModuleName, const FName& NewName) const
+FName FModularRigHierarchyTreeDelegates::HandleRenameElement(const FName& OldModuleName, const FName& NewName) const
 {
 	if (OnRenameElement.IsBound())
 	{
@@ -20,7 +20,7 @@ FName FModularRigTreeDelegates::HandleRenameElement(const FName& OldModuleName,
 	return OldModuleName;
 }
 
-bool FModularRigTreeDelegates::HandleVerifyElementNameChanged(const FName& OldModuleName, const FName& NewName, FText& OutErrorMessage) const
+bool FModularRigHierarchyTreeDelegates::HandleVerifyElementNameChanged(const FName& OldModuleName, const FName& NewName, FText& OutErrorMessage) const
 {
 	if (OnVerifyModuleNameChanged.IsBound())
 	{
@@ -29,7 +29,7 @@ bool FModularRigTreeDelegates::HandleVerifyElementNameChanged(const FName& OldMo
 	return false;
 }
 
-bool FModularRigTreeDelegates::HandleResolveConnector(const FRigElementKey& InConnector, const TArray<FRigElementKey>& InTargets)
+bool FModularRigHierarchyTreeDelegates::HandleResolveConnector(const FRigElementKey& InConnector, const TArray<FRigElementKey>& InTargets)
 {
 	if (OnResolveConnector.IsBound())
 	{
@@ -39,7 +39,7 @@ bool FModularRigTreeDelegates::HandleResolveConnector(const FRigElementKey& InCo
 	return false;
 }
 
-bool FModularRigTreeDelegates::HandleDisconnectConnector(const FRigElementKey& InConnector)
+bool FModularRigHierarchyTreeDelegates::HandleDisconnectConnector(const FRigElementKey& InConnector)
 {
 	if (OnDisconnectConnector.IsBound())
 	{
@@ -49,7 +49,7 @@ bool FModularRigTreeDelegates::HandleDisconnectConnector(const FRigElementKey& I
 	return false;
 }
 
-void FModularRigTreeDelegates::HandleSelectionChanged(TSharedPtr<FModularRigTreeElement> Selection, ESelectInfo::Type SelectInfo)
+void FModularRigHierarchyTreeDelegates::HandleSelectionChanged(TSharedPtr<FModularRigHierarchyTreeElement> Selection, ESelectInfo::Type SelectInfo)
 {
 	if (bSuspendSelectionDelegate)
 	{
@@ -59,7 +59,7 @@ void FModularRigTreeDelegates::HandleSelectionChanged(TSharedPtr<FModularRigTree
 	(void)OnSelectionChanged.ExecuteIfBound(Selection, SelectInfo);
 }
 
-bool FModularRigTreeDelegates::ShouldAlwaysShowConnector(const FName& InConnectorName) const
+bool FModularRigHierarchyTreeDelegates::ShouldAlwaysShowConnector(const FName& InConnectorName) const
 {
 	if (OnAlwaysShowConnector.IsBound())
 	{
diff --git a/Engine/Plugins/Animation/ControlRig/Source/ControlRigEditor/Private/Editor/Hierarchy/ModularRigHierarchyTreeDelegates.h b/Engine/Plugins/Animation/ControlRig/Source/ControlRigEditor/Private/Editor/Hierarchy/ModularRigHierarchyTreeDelegates.h
--- a/Engine/Plugins/Animation/ControlRig/Source/ControlRigEditor/Private/Editor/Hierarchy/ModularRigHierarchyTreeDelegates.h
+++ b/Engine/Plugins/Animation/ControlRig/Source/ControlRigEditor/Private/Editor/Hierarchy/ModularRigHierarchyTreeDelegates.h
@@ -20,13 +20,13 @@ DECLARE_DELEGATE_OneParam(FOnModularRigTreeDisconnectConnector, const FRigElemen
 DECLARE_DELEGATE_RetVal_ThreeParams(bool, FOnModularRigTreeVerifyElementNameChanged, const FName& /*OldModuleName*/, const FName& /*NewName*/, FText& /*OutErrorMessage*/);
 DECLARE_DELEGATE_RetVal_OneParam(bool, FOnModularRigTreeAlwaysShowConnector, const FName& /*ConnectorName*/);
 
-typedef STreeView<TSharedPtr<FModularRigTreeElement>>::FOnMouseButtonClick FOnModularRigTreeMouseButtonClick;
-typedef STreeView<TSharedPtr<FModularRigTreeElement>>::FOnMouseButtonDoubleClick FOnModularRigTreeMouseButtonDoubleClick;
-typedef STableRow<TSharedPtr<FModularRigTreeElement>>::FOnCanAcceptDrop FOnModularRigTreeCanAcceptDrop;
-typedef STableRow<TSharedPtr<FModularRigTreeElement>>::FOnAcceptDrop FOnModularRigTreeAcceptDrop;
-typedef STreeView<TSharedPtr<FModularRigTreeElement>>::FOnSelectionChanged FOnModularRigTreeSelectionChanged;
+typedef STreeView<TSharedPtr<FModularRigHierarchyTreeElement>>::FOnMouseButtonClick FOnModularRigTreeMouseButtonClick;
+typedef STreeView<TSharedPtr<FModularRigHierarchyTreeElement>>::FOnMouseButtonDoubleClick FOnModularRigTreeMouseButtonDoubleClick;
+typedef STableRow<TSharedPtr<FModularRigHierarchyTreeElement>>::FOnCanAcceptDrop FOnModularRigTreeCanAcceptDrop;
+typedef STableRow<TSharedPtr<FModularRigHierarchyTreeElement>>::FOnAcceptDrop FOnModularRigTreeAcceptDrop;
+typedef STreeView<TSharedPtr<FModularRigHierarchyTreeElement>>::FOnSelectionChanged FOnModularRigTreeSelectionChanged;
 
-struct FModularRigTreeDelegates
+struct FModularRigHierarchyTreeDelegates
 {
 public:
 
@@ -50,7 +50,7 @@ public:
 	bool HandleVerifyElementNameChanged(const FName& OldModuleName, const FName& NewName, FText& OutErrorMessage) const;
 	bool HandleResolveConnector(const FRigElementKey& InConnector, const TArray<FRigElementKey>& InTargets);
 	bool HandleDisconnectConnector(const FRigElementKey& InConnector);
-	void HandleSelectionChanged(TSharedPtr<FModularRigTreeElement> Selection, ESelectInfo::Type SelectInfo);
+	void HandleSelectionChanged(TSharedPtr<FModularRigHierarchyTreeElement> Selection, ESelectInfo::Type SelectInfo);
 	bool ShouldAlwaysShowConnector(const FName& InConnectorName) const;
 
 private:
diff --git a/Engine/Plugins/Animation/ControlRig/Source/ControlRigEditor/Private/Editor/Hierarchy/RigHierarchyTreeDelegates.cpp b/Engine/Plugins/Animation/ControlRig/Source/ControlRigEditor/Private/Editor/Hierarchy/RigHierarchyTreeDelegates.cpp
--- a/Engine/Plugins/Animation/ControlRig/Source/ControlRigEditor/Private/Editor/Hierarchy/RigHierarchyTreeDelegates.cpp
+++ b/Engine/Plugins/Animation/ControlRig/Source/ControlRigEditor/Private/Editor/Hierarchy/RigHierarchyTreeDelegates.cpp
@@ -2,9 +2,9 @@
 
 #include "RigHierarchyTreeDelegates.h"
 
-FRigTreeDisplaySettings FRigTreeDelegates::DefaultDisplaySettings;
+FRigHierarchyTreeDisplaySettings FRigHierarchyTreeDelegates::DefaultDisplaySettings;
 
-URigHierarchy* FRigTreeDelegates::GetHierarchy()
+URigHierarchy* FRigHierarchyTreeDelegates::GetHierarchy()
 {
 	if (OnGetHierarchy.IsBound())
 	{
@@ -13,7 +13,7 @@ URigHierarchy* FRigTreeDelegates::GetHierarchy()
 	return nullptr;
 }
 
-const URigHierarchy* FRigTreeDelegates::GetHierarchy() const
+const URigHierarchy* FRigHierarchyTreeDelegates::GetHierarchy() const
 {
 	if (OnGetHierarchy.IsBound())
 	{
@@ -22,7 +22,7 @@ const URigHierarchy* FRigTreeDelegates::GetHierarchy() const
 	return nullptr;
 }
 
-const FRigTreeDisplaySettings& FRigTreeDelegates::GetDisplaySettings() const
+const FRigHierarchyTreeDisplaySettings& FRigHierarchyTreeDelegates::GetDisplaySettings() const
 {
 	if (OnGetDisplaySettings.IsBound())
 	{
@@ -31,7 +31,7 @@ const FRigTreeDisplaySettings& FRigTreeDelegates::GetDisplaySettings() const
 	return DefaultDisplaySettings;
 }
 
-TArray<FRigHierarchyKey> FRigTreeDelegates::GetSelection() const
+TArray<FRigHierarchyKey> FRigHierarchyTreeDelegates::GetSelection() const
 {
 	if (OnGetSelection.IsBound())
 	{
@@ -44,7 +44,7 @@ TArray<FRigHierarchyKey> FRigTreeDelegates::GetSelection() const
 	return {};
 }
 
-FName FRigTreeDelegates::HandleRenameElement(const FRigHierarchyKey& OldKey, const FString& NewName) const
+FName FRigHierarchyTreeDelegates::HandleRenameElement(const FRigHierarchyKey& OldKey, const FString& NewName) const
 {
 	if (OnRenameElement.IsBound())
 	{
@@ -53,7 +53,7 @@ FName FRigTreeDelegates::HandleRenameElement(const FRigHierarchyKey& OldKey, con
 	return OldKey.GetFName();
 }
 
-bool FRigTreeDelegates::HandleVerifyElementNameChanged(const FRigHierarchyKey& OldKey, const FString& NewName, FText& OutErrorMessage) const
+bool FRigHierarchyTreeDelegates::HandleVerifyElementNameChanged(const FRigHierarchyKey& OldKey, const FString& NewName, FText& OutErrorMessage) const
 {
 	if (OnVerifyElementNameChanged.IsBound())
 	{
@@ -62,7 +62,7 @@ bool FRigTreeDelegates::HandleVerifyElementNameChanged(const FRigHierarchyKey& O
 	return false;
 }
 
-void FRigTreeDelegates::HandleSelectionChanged(TSharedPtr<FRigTreeElement> Selection, ESelectInfo::Type SelectInfo)
+void FRigHierarchyTreeDelegates::HandleSelectionChanged(TSharedPtr<FRigHierarchyTreeElement> Selection, ESelectInfo::Type SelectInfo)
 {
 	if (bIsChangingRigHierarchy)
 	{
@@ -72,7 +72,7 @@ void FRigTreeDelegates::HandleSelectionChanged(TSharedPtr<FRigTreeElement> Selec
 	(void)OnSelectionChanged.ExecuteIfBound(Selection, SelectInfo);
 }
 
-FRigHierarchyKey FRigTreeDelegates::GetResolvedKey(const FRigHierarchyKey& InKey)
+FRigHierarchyKey FRigHierarchyTreeDelegates::GetResolvedKey(const FRigHierarchyKey& InKey)
 {
 	if (OnGetResolvedKey.IsBound())
 	{
@@ -81,7 +81,7 @@ FRigHierarchyKey FRigTreeDelegates::GetResolvedKey(const FRigHierarchyKey& InKey
 	return InKey;
 }
 
-void FRigTreeDelegates::RequestDetailsInspection(const FRigHierarchyKey& InKey)
+void FRigHierarchyTreeDelegates::RequestDetailsInspection(const FRigHierarchyKey& InKey)
 {
 	if (OnRequestDetailsInspection.IsBound())
 	{
diff --git a/Engine/Plugins/Animation/ControlRig/Source/ControlRigEditor/Private/Editor/Hierarchy/RigHierarchyTreeDelegates.h b/Engine/Plugins/Animation/ControlRig/Source/ControlRigEditor/Private/Editor/Hierarchy/RigHierarchyTreeDelegates.h
--- a/Engine/Plugins/Animation/ControlRig/Source/ControlRigEditor/Private/Editor/Hierarchy/RigHierarchyTreeDelegates.h
+++ b/Engine/Plugins/Animation/ControlRig/Source/ControlRigEditor/Private/Editor/Hierarchy/RigHierarchyTreeDelegates.h
@@ -11,18 +11,18 @@
 #include "Widgets/Views/STableRow.h"
 #include "Widgets/Views/STreeView.h"
 
-class FRigTreeElement;
+class FRigHierarchyTreeElement;
 class URigHierarchy;
 
-typedef STableRow<TSharedPtr<FRigTreeElement>>::FOnCanAcceptDrop FOnRigTreeCanAcceptDrop;
-typedef STableRow<TSharedPtr<FRigTreeElement>>::FOnAcceptDrop FOnRigTreeAcceptDrop;
-typedef STreeView<TSharedPtr<FRigTreeElement>>::FOnSelectionChanged FOnRigTreeSelectionChanged;
-typedef STreeView<TSharedPtr<FRigTreeElement>>::FOnMouseButtonClick FOnRigTreeMouseButtonClick;
-typedef STreeView<TSharedPtr<FRigTreeElement>>::FOnMouseButtonDoubleClick FOnRigTreeMouseButtonDoubleClick;
-typedef STreeView<TSharedPtr<FRigTreeElement>>::FOnSetExpansionRecursive FOnRigTreeSetExpansionRecursive;
+typedef STableRow<TSharedPtr<FRigHierarchyTreeElement>>::FOnCanAcceptDrop FOnRigTreeCanAcceptDrop;
+typedef STableRow<TSharedPtr<FRigHierarchyTreeElement>>::FOnAcceptDrop FOnRigTreeAcceptDrop;
+typedef STreeView<TSharedPtr<FRigHierarchyTreeElement>>::FOnSelectionChanged FOnRigTreeSelectionChanged;
+typedef STreeView<TSharedPtr<FRigHierarchyTreeElement>>::FOnMouseButtonClick FOnRigTreeMouseButtonClick;
+typedef STreeView<TSharedPtr<FRigHierarchyTreeElement>>::FOnMouseButtonDoubleClick FOnRigTreeMouseButtonDoubleClick;
+typedef STreeView<TSharedPtr<FRigHierarchyTreeElement>>::FOnSetExpansionRecursive FOnRigTreeSetExpansionRecursive;
 
 DECLARE_DELEGATE_RetVal(const URigHierarchy*, FOnGetRigTreeHierarchy);
-DECLARE_DELEGATE_RetVal(const FRigTreeDisplaySettings&, FOnGetRigTreeDisplaySettings);
+DECLARE_DELEGATE_RetVal(const FRigHierarchyTreeDisplaySettings&, FOnGetRigTreeDisplaySettings);
 DECLARE_DELEGATE_RetVal(const TArray<FRigHierarchyKey>, FOnRigTreeGetSelection);
 DECLARE_DELEGATE_RetVal_TwoParams(FName, FOnRigTreeRenameElement, const FRigHierarchyKey& /*OldKey*/, const FString& /*NewName*/);
 DECLARE_DELEGATE_RetVal_ThreeParams(bool, FOnRigTreeVerifyElementNameChanged, const FRigHierarchyKey& /*OldKey*/, const FString& /*NewName*/, FText& /*OutErrorMessage*/);
@@ -34,7 +34,7 @@ DECLARE_DELEGATE_RetVal_OneParam(bool, FOnRigTreeIsItemVisible, const FRigHierar
 
 #define UE_API CONTROLRIGEDITOR_API
 
-struct FRigTreeDelegates
+struct FRigHierarchyTreeDelegates
 {
 	FOnGetRigTreeHierarchy OnGetHierarchy;
 	FOnGetRigTreeDisplaySettings OnGetDisplaySettings;
@@ -58,15 +58,15 @@ struct FRigTreeDelegates
 
 	URigHierarchy* GetHierarchy();
 	const URigHierarchy* GetHierarchy() const;
-	const FRigTreeDisplaySettings& GetDisplaySettings() const;
+	const FRigHierarchyTreeDisplaySettings& GetDisplaySettings() const;
 	TArray<FRigHierarchyKey> GetSelection() const;
 	FName HandleRenameElement(const FRigHierarchyKey& OldKey, const FString& NewName) const;
 	bool HandleVerifyElementNameChanged(const FRigHierarchyKey& OldKey, const FString& NewName, FText& OutErrorMessage) const;
-	void HandleSelectionChanged(TSharedPtr<FRigTreeElement> Selection, ESelectInfo::Type SelectInfo);
+	void HandleSelectionChanged(TSharedPtr<FRigHierarchyTreeElement> Selection, ESelectInfo::Type SelectInfo);
 	FRigHierarchyKey GetResolvedKey(const FRigHierarchyKey& InKey);
 	void RequestDetailsInspection(const FRigHierarchyKey& InKey);
 
-	static UE_API FRigTreeDisplaySettings DefaultDisplaySettings;
+	static UE_API FRigHierarchyTreeDisplaySettings DefaultDisplaySettings;
 
 private:
 	bool bIsChangingRigHierarchy = false;
diff --git a/Engine/Plugins/Animation/ControlRig/Source/ControlRigEditor/Private/Editor/Hierarchy/RigHierarchyTreeDisplaySettings.h b/Engine/Plugins/Animation/ControlRig/Source/ControlRigEditor/Private/Editor/Hierarchy/RigHierarchyTreeDisplaySettings.h
--- a/Engine/Plugins/Animation/ControlRig/Source/ControlRigEditor/Private/Editor/Hierarchy/RigHierarchyTreeDisplaySettings.h
+++ b/Engine/Plugins/Animation/ControlRig/Source/ControlRigEditor/Private/Editor/Hierarchy/RigHierarchyTreeDisplaySettings.h
@@ -5,9 +5,9 @@
 #include "Rigs/RigHierarchyDefines.h"
 #include "Settings/ControlRigSettings.h"
 
-struct FRigTreeDisplaySettings
+struct FRigHierarchyTreeDisplaySettings
 {
-	FRigTreeDisplaySettings()
+	FRigHierarchyTreeDisplaySettings()
 	{
 		FilterText = FText();
 
diff --git a/Engine/Plugins/Animation/ControlRig/Source/ControlRigEditor/Private/Editor/Hierarchy/Views/SModularRigHierarchySearchableTreeView.cpp b/Engine/Plugins/Animation/ControlRig/Source/ControlRigEditor/Private/Editor/Hierarchy/Views/SModularRigHierarchySearchableTreeView.cpp
--- a/Engine/Plugins/Animation/ControlRig/Source/ControlRigEditor/Private/Editor/Hierarchy/Views/SModularRigHierarchySearchableTreeView.cpp
+++ b/Engine/Plugins/Animation/ControlRig/Source/ControlRigEditor/Private/Editor/Hierarchy/Views/SModularRigHierarchySearchableTreeView.cpp
@@ -5,9 +5,9 @@
 #include "SModularRigTreeView.h"
 #include "Widgets/Layout/SScrollBox.h"
 
-void SSearchableModularRigTreeView::Construct(const FArguments& InArgs)
+void SModularRigHierarchySearchableTreeView::Construct(const FArguments& InArgs)
 {
-	FModularRigTreeDelegates TreeDelegates = InArgs._RigTreeDelegates;
+	FModularRigHierarchyTreeDelegates TreeDelegates = InArgs._RigTreeDelegates;
 	
 	ChildSlot
 	[
diff --git a/Engine/Plugins/Animation/ControlRig/Source/ControlRigEditor/Private/Editor/Hierarchy/Views/SModularRigHierarchySearchableTreeView.h b/Engine/Plugins/Animation/ControlRig/Source/ControlRigEditor/Private/Editor/Hierarchy/Views/SModularRigHierarchySearchableTreeView.h
--- a/Engine/Plugins/Animation/ControlRig/Source/ControlRigEditor/Private/Editor/Hierarchy/Views/SModularRigHierarchySearchableTreeView.h
+++ b/Engine/Plugins/Animation/ControlRig/Source/ControlRigEditor/Private/Editor/Hierarchy/Views/SModularRigHierarchySearchableTreeView.h
@@ -7,17 +7,17 @@
 
 class SModularRigTreeView;
 
-class SSearchableModularRigTreeView : public SCompoundWidget
+class SModularRigHierarchySearchableTreeView : public SCompoundWidget
 {
 public:
 
-	SLATE_BEGIN_ARGS(SSearchableModularRigTreeView) {}
-		SLATE_ARGUMENT(FModularRigTreeDelegates, RigTreeDelegates)
+	SLATE_BEGIN_ARGS(SModularRigHierarchySearchableTreeView) {}
+		SLATE_ARGUMENT(FModularRigHierarchyTreeDelegates, RigTreeDelegates)
 		SLATE_ARGUMENT(FText, InitialFilterText)
 	SLATE_END_ARGS()
 
 	void Construct(const FArguments& InArgs);
-	virtual ~SSearchableModularRigTreeView() {}
+	virtual ~SModularRigHierarchySearchableTreeView() {}
 	TSharedRef<SModularRigTreeView> GetTreeView() const { return TreeView.ToSharedRef(); }
 
 private:
diff --git a/Engine/Plugins/Animation/ControlRig/Source/ControlRigEditor/Private/Editor/Hierarchy/Views/SModularRigHierarchyTreeItem.cpp b/Engine/Plugins/Animation/ControlRig/Source/ControlRigEditor/Private/Editor/Hierarchy/Views/SModularRigHierarchyTreeItem.cpp
--- a/Engine/Plugins/Animation/ControlRig/Source/ControlRigEditor/Private/Editor/Hierarchy/Views/SModularRigHierarchyTreeItem.cpp
+++ b/Engine/Plugins/Animation/ControlRig/Source/ControlRigEditor/Private/Editor/Hierarchy/Views/SModularRigHierarchyTreeItem.cpp
@@ -9,17 +9,17 @@
 #include "ScopedTransaction.h"
 #include "Widgets/Text/SInlineEditableTextBlock.h"
 
-#define LOCTEXT_NAMESPACE "SModularRigModelItem"
+#define LOCTEXT_NAMESPACE "SModularRigHierarchyTreeItem"
 
-void SModularRigModelItem::Construct(const FArguments& InArgs, const TSharedRef<STableViewBase>& OwnerTable, TSharedRef<FModularRigTreeElement> InRigTreeElement, TSharedPtr<SModularRigTreeView> InTreeView, bool bPinned)
+void SModularRigHierarchyTreeItem::Construct(const FArguments& InArgs, const TSharedRef<STableViewBase>& OwnerTable, TSharedRef<FModularRigHierarchyTreeElement> InRigTreeElement, TSharedPtr<SModularRigTreeView> InTreeView, bool bPinned)
 {
 	WeakRigTreeElement = InRigTreeElement;
 	Delegates = InTreeView->GetRigTreeDelegates();
 
 	if (InRigTreeElement->Key.IsEmpty())
 	{
-		SMultiColumnTableRow<TSharedPtr<FModularRigTreeElement>>::Construct(
-			SMultiColumnTableRow<TSharedPtr<FModularRigTreeElement>>::FArguments()
+		SMultiColumnTableRow<TSharedPtr<FModularRigHierarchyTreeElement>>::Construct(
+			SMultiColumnTableRow<TSharedPtr<FModularRigHierarchyTreeElement>>::FArguments()
 			.ShowSelection(false)
 			.OnCanAcceptDrop(Delegates.OnCanAcceptDrop)
 			.OnAcceptDrop(Delegates.OnAcceptDrop)
@@ -31,22 +31,22 @@ void SModularRigModelItem::Construct(const FArguments& InArgs, const TSharedRef<
 	const FRigHierarchyModulePath ConnectorModulePath(ModuleName, InRigTreeElement->ConnectorName);
 	ConnectorKey = FRigElementKey(ConnectorModulePath.GetPathFName(), ERigElementType::Connector);
 
-	SMultiColumnTableRow<TSharedPtr<FModularRigTreeElement>>::Construct(
-		SMultiColumnTableRow<TSharedPtr<FModularRigTreeElement>>::FArguments()
+	SMultiColumnTableRow<TSharedPtr<FModularRigHierarchyTreeElement>>::Construct(
+		SMultiColumnTableRow<TSharedPtr<FModularRigHierarchyTreeElement>>::FArguments()
 		.OnDragDetected(Delegates.OnDragDetected)
 		.OnCanAcceptDrop(Delegates.OnCanAcceptDrop)
 		.OnAcceptDrop(Delegates.OnAcceptDrop)
 		.ShowWires(true), OwnerTable);
 }
 
-bool SModularRigModelItem::OnConnectorTargetChanged(TArray<FRigElementKey> InTargets, const FRigElementKey InConnectorKey)
+bool SModularRigHierarchyTreeItem::OnConnectorTargetChanged(TArray<FRigElementKey> InTargets, const FRigElementKey InConnectorKey)
 {
 	FScopedTransaction Transaction(LOCTEXT("ModuleHierarchyResolveConnector", "Resolve Connector"));
 	Delegates.HandleResolveConnector(InConnectorKey, InTargets);
 	return false;
 }
 
-void SModularRigModelItem::OnNameCommitted(const FText& InText, ETextCommit::Type InCommitType) const
+void SModularRigHierarchyTreeItem::OnNameCommitted(const FText& InText, ETextCommit::Type InCommitType) const
 {
 	// for now only allow enter
 	// because it is important to keep the unique names per pose
@@ -59,14 +59,14 @@ void SModularRigModelItem::OnNameCommitted(const FText& InText, ETextCommit::Typ
 	}
 }
 
-bool SModularRigModelItem::OnVerifyNameChanged(const FText& InText, FText& OutErrorMessage)
+bool SModularRigHierarchyTreeItem::OnVerifyNameChanged(const FText& InText, FText& OutErrorMessage)
 {
 	const FName NewName = *InText.ToString();
 	const FName OldModuleName = WeakRigTreeElement.Pin()->ModuleName;
 	return Delegates.HandleVerifyElementNameChanged(OldModuleName, NewName, OutErrorMessage);
 }
 
-TSharedRef<SWidget> SModularRigModelItem::GenerateWidgetForColumn(const FName& ColumnName)
+TSharedRef<SWidget> SModularRigHierarchyTreeItem::GenerateWidgetForColumn(const FName& ColumnName)
 {
 	if(ColumnName == SModularRigTreeView::Column_Module)
 	{
@@ -122,11 +122,11 @@ TSharedRef<SWidget> SModularRigModelItem::GenerateWidgetForColumn(const FName& C
 		.Padding(0, TopPadding, 0, 0)
 		[
 			SAssignNew(InlineWidget, SInlineEditableTextBlock)
-			.Text(this, &SModularRigModelItem::GetName, true)
+			.Text(this, &SModularRigHierarchyTreeItem::GetName, true)
 			.MaximumLength(NAME_SIZE-1)
-			.OnVerifyTextChanged(this, &SModularRigModelItem::OnVerifyNameChanged)
-			.OnTextCommitted(this, &SModularRigModelItem::OnNameCommitted)
-			.ToolTipText(this, &SModularRigModelItem::GetItemTooltip)
+			.OnVerifyTextChanged(this, &SModularRigHierarchyTreeItem::OnVerifyNameChanged)
+			.OnTextCommitted(this, &SModularRigHierarchyTreeItem::OnNameCommitted)
+			.ToolTipText(this, &SModularRigHierarchyTreeItem::GetItemTooltip)
 			.MultiLine(false)
 			//.Font(IDetailLayoutBuilder::GetDetailFont())
 			.ColorAndOpacity_Lambda([this]()
@@ -231,7 +231,7 @@ TSharedRef<SWidget> SModularRigModelItem::GenerateWidgetForColumn(const FName& C
 
 		TSharedPtr<SVerticalBox> ComboButtonBox;
 		
-		FRigTreeDelegates RigTreeDelegates;
+		FRigHierarchyTreeDelegates RigTreeDelegates;
 		RigTreeDelegates.OnGetHierarchy = FOnGetRigTreeHierarchy::CreateLambda([this]()
 		{
 			return Delegates.GetModularRig()->GetHierarchy();
@@ -242,7 +242,7 @@ TSharedRef<SWidget> SModularRigModelItem::GenerateWidgetForColumn(const FName& C
 			.ConnectorKey(ConnectorKey)
 			.IsArray(bIsArrayConnector)
 			.Targets(GetTargetKeys())
-			.OnSetTargetArray(FRigConnectorTargetWidget_SetTargetArray::CreateSP(this, &SModularRigModelItem::OnConnectorTargetChanged, ConnectorKey))
+			.OnSetTargetArray(FRigConnectorTargetWidget_SetTargetArray::CreateSP(this, &SModularRigHierarchyTreeItem::OnConnectorTargetChanged, ConnectorKey))
 			.RigTreeDelegates(RigTreeDelegates);
 	}
 	if(ColumnName == SModularRigTreeView::Column_Buttons)
@@ -367,7 +367,7 @@ TSharedRef<SWidget> SModularRigModelItem::GenerateWidgetForColumn(const FName& C
 	return SNullWidget::NullWidget;
 }
 
-FText SModularRigModelItem::GetName(bool bUseShortName) const
+FText SModularRigHierarchyTreeItem::GetName(bool bUseShortName) const
 {
 	if(bUseShortName)
 	{
@@ -376,7 +376,7 @@ FText SModularRigModelItem::GetName(bool bUseShortName) const
 	return (FText::FromName(WeakRigTreeElement.Pin()->ModuleName));
 }
 
-FText SModularRigModelItem::GetItemTooltip() const
+FText SModularRigHierarchyTreeItem::GetItemTooltip() const
 {
 	const FText FullName = GetName(false);
 	const FText ShortName = GetName(true);
@@ -387,7 +387,7 @@ FText SModularRigModelItem::GetItemTooltip() const
 	return FullName;
 }
 
-TArray<FRigElementKey> SModularRigModelItem::GetTargetKeys() const
+TArray<FRigElementKey> SModularRigHierarchyTreeItem::GetTargetKeys() const
 {
 	TArray<FRigElementKey> Result;
 	if (const UModularRig* ModularRig = Delegates.GetModularRig())
diff --git a/Engine/Plugins/Animation/ControlRig/Source/ControlRigEditor/Private/Editor/Hierarchy/Views/SModularRigHierarchyTreeItem.h b/Engine/Plugins/Animation/ControlRig/Source/ControlRigEditor/Private/Editor/Hierarchy/Views/SModularRigHierarchyTreeItem.h
--- a/Engine/Plugins/Animation/ControlRig/Source/ControlRigEditor/Private/Editor/Hierarchy/Views/SModularRigHierarchyTreeItem.h
+++ b/Engine/Plugins/Animation/ControlRig/Source/ControlRigEditor/Private/Editor/Hierarchy/Views/SModularRigHierarchyTreeItem.h
@@ -8,13 +8,13 @@
 #include "Rigs/RigHierarchyDefines.h"
 #include "Widgets/Views/STableRow.h"
 
-class SSearchableRigHierarchyTreeView;
+class SRigHierarchySearchableTreeView;
 
-class SModularRigModelItem : public SMultiColumnTableRow<TSharedPtr<FModularRigTreeElement>>
+class SModularRigHierarchyTreeItem : public SMultiColumnTableRow<TSharedPtr<FModularRigHierarchyTreeElement>>
 {
 public:
 
-	void Construct(const FArguments& InArgs, const TSharedRef<STableViewBase>& OwnerTable, TSharedRef<FModularRigTreeElement> InRigTreeElement, TSharedPtr<SModularRigTreeView> InTreeView, bool bPinned);
+	void Construct(const FArguments& InArgs, const TSharedRef<STableViewBase>& OwnerTable, TSharedRef<FModularRigHierarchyTreeElement> InRigTreeElement, TSharedPtr<SModularRigTreeView> InTreeView, bool bPinned);
 	bool OnConnectorTargetChanged(TArray<FRigElementKey> InTargets, const FRigElementKey InConnectorKey);
 
 	void OnNameCommitted(const FText& InText, ETextCommit::Type InCommitType) const;
@@ -23,9 +23,9 @@ public:
 	virtual TSharedRef<SWidget> GenerateWidgetForColumn(const FName& ColumnName) override;
 
 private:
-	TWeakPtr<FModularRigTreeElement> WeakRigTreeElement;
-	FModularRigTreeDelegates Delegates;
-	TSharedPtr<SSearchableRigHierarchyTreeView> ConnectorComboBox;
+	TWeakPtr<FModularRigHierarchyTreeElement> WeakRigTreeElement;
+	FModularRigHierarchyTreeDelegates Delegates;
+	TSharedPtr<SRigHierarchySearchableTreeView> ConnectorComboBox;
 	TSharedPtr<SButton> ResetConnectorButton;
 	TSharedPtr<SButton> UseSelectedButton;
 	TSharedPtr<SButton> SelectElementButton;
diff --git a/Engine/Plugins/Animation/ControlRig/Source/ControlRigEditor/Private/Editor/Hierarchy/Views/SModularRigTreeView.cpp b/Engine/Plugins/Animation/ControlRig/Source/ControlRigEditor/Private/Editor/Hierarchy/Views/SModularRigTreeView.cpp
--- a/Engine/Plugins/Animation/ControlRig/Source/ControlRigEditor/Private/Editor/Hierarchy/Views/SModularRigTreeView.cpp
+++ b/Engine/Plugins/Animation/ControlRig/Source/ControlRigEditor/Private/Editor/Hierarchy/Views/SModularRigTreeView.cpp
@@ -52,13 +52,13 @@ void SModularRigTreeView::Construct(const FArguments& InArgs)
 	ShowOptionalConnectors = InArgs._ShowOptionalConnectors;
 	ShowUnresolvedConnectors = InArgs._ShowUnresolvedConnectors;
 
-	STreeView<TSharedPtr<FModularRigTreeElement>>::FArguments SuperArgs;
+	STreeView<TSharedPtr<FModularRigHierarchyTreeElement>>::FArguments SuperArgs;
 	SuperArgs.HeaderRow(InArgs._HeaderRow);
 	SuperArgs.TreeItemsSource(&RootElements);
 	SuperArgs.SelectionMode(ESelectionMode::Multi);
 	SuperArgs.OnGenerateRow(this, &SModularRigTreeView::MakeTableRowWidget, false);
 	SuperArgs.OnGetChildren(this, &SModularRigTreeView::HandleGetChildrenForTree);
-	SuperArgs.OnSelectionChanged(FOnModularRigTreeSelectionChanged::CreateRaw(&Delegates, &FModularRigTreeDelegates::HandleSelectionChanged));
+	SuperArgs.OnSelectionChanged(FOnModularRigTreeSelectionChanged::CreateRaw(&Delegates, &FModularRigHierarchyTreeDelegates::HandleSelectionChanged));
 	SuperArgs.OnContextMenuOpening(Delegates.OnContextMenuOpening);
 	SuperArgs.HighlightParentNodesForSelection(true);
 	SuperArgs.AllowInvisibleItemSelection(true);  //without this we deselect everything when we filter or we collapse
@@ -74,7 +74,7 @@ void SModularRigTreeView::Construct(const FArguments& InArgs)
 		return FMath::Max<int32>(1, UControlRigEditorSettings::Get()->MaxStackSize);
 	});
 
-	STreeView<TSharedPtr<FModularRigTreeElement>>::Construct(SuperArgs);
+	STreeView<TSharedPtr<FModularRigHierarchyTreeElement>>::Construct(SuperArgs);
 
 	LastMousePosition = FVector2D::ZeroVector;
 	TimeAtMousePosition = 0.0;
@@ -82,7 +82,7 @@ void SModularRigTreeView::Construct(const FArguments& InArgs)
 
 void SModularRigTreeView::Tick(const FGeometry& AllottedGeometry, const double InCurrentTime, const float InDeltaTime)
 {
-	STreeView<TSharedPtr<FModularRigTreeElement, ESPMode::ThreadSafe>>::Tick(AllottedGeometry, InCurrentTime, InDeltaTime);
+	STreeView<TSharedPtr<FModularRigHierarchyTreeElement, ESPMode::ThreadSafe>>::Tick(AllottedGeometry, InCurrentTime, InDeltaTime);
 
 	const FGeometry PaintGeometry = GetPaintSpaceGeometry();
 	const FVector2D MousePosition = FSlateApplication::Get().GetCursorPos();
@@ -121,7 +121,7 @@ void SModularRigTreeView::Tick(const FGeometry& AllottedGeometry, const double I
 			}
 			else
 			{
-				const TSharedPtr<FModularRigTreeElement>* Item = FindItemAtPosition(MousePosition);
+				const TSharedPtr<FModularRigHierarchyTreeElement>* Item = FindItemAtPosition(MousePosition);
 				if(Item && Item->IsValid())
 				{
 					if(!IsItemExpanded(*Item))
@@ -136,7 +136,7 @@ void SModularRigTreeView::Tick(const FGeometry& AllottedGeometry, const double I
 	if (bRequestRenameSelected)
 	{
 		RegisterActiveTimer(0.f, FWidgetActiveTimerDelegate::CreateLambda([this](double, float) {
-			TArray<TSharedPtr<FModularRigTreeElement>> SelectedItems = GetSelectedItems();
+			TArray<TSharedPtr<FModularRigHierarchyTreeElement>> SelectedItems = GetSelectedItems();
 			if (SelectedItems.Num() == 1)
 			{
 				SelectedItems[0]->RequestRename();
@@ -147,20 +147,20 @@ void SModularRigTreeView::Tick(const FGeometry& AllottedGeometry, const double I
 	}
 }
 
-TSharedPtr<FModularRigTreeElement> SModularRigTreeView::FindElement(const FString& InElementKey)
+TSharedPtr<FModularRigHierarchyTreeElement> SModularRigTreeView::FindElement(const FString& InElementKey)
 {
-	for (TSharedPtr<FModularRigTreeElement> Root : RootElements)
+	for (TSharedPtr<FModularRigHierarchyTreeElement> Root : RootElements)
 	{
-		if (TSharedPtr<FModularRigTreeElement> Found = FindElement(InElementKey, Root))
+		if (TSharedPtr<FModularRigHierarchyTreeElement> Found = FindElement(InElementKey, Root))
 		{
 			return Found;
 		}
 	}
 
-	return TSharedPtr<FModularRigTreeElement>();
+	return TSharedPtr<FModularRigHierarchyTreeElement>();
 }
 
-TSharedPtr<FModularRigTreeElement> SModularRigTreeView::FindElement(const FString& InElementKey, TSharedPtr<FModularRigTreeElement> CurrentItem)
+TSharedPtr<FModularRigHierarchyTreeElement> SModularRigTreeView::FindElement(const FString& InElementKey, TSharedPtr<FModularRigHierarchyTreeElement> CurrentItem)
 {
 	if (CurrentItem->Key == InElementKey)
 	{
@@ -169,14 +169,14 @@ TSharedPtr<FModularRigTreeElement> SModularRigTreeView::FindElement(const FStrin
 
 	for (int32 ChildIndex = 0; ChildIndex < CurrentItem->Children.Num(); ++ChildIndex)
 	{
-		TSharedPtr<FModularRigTreeElement> Found = FindElement(InElementKey, CurrentItem->Children[ChildIndex]);
+		TSharedPtr<FModularRigHierarchyTreeElement> Found = FindElement(InElementKey, CurrentItem->Children[ChildIndex]);
 		if (Found.IsValid())
 		{
 			return Found;
 		}
 	}
 
-	return TSharedPtr<FModularRigTreeElement>();
+	return TSharedPtr<FModularRigHierarchyTreeElement>();
 }
 
 bool SModularRigTreeView::AddElement(FString InKey, FString InParentKey, bool bApplyFilterText)
@@ -279,14 +279,14 @@ bool SModularRigTreeView::AddElement(FString InKey, FString InParentKey, bool bA
 			return false;
 		}
 		
-		TSharedPtr<FModularRigTreeElement> NewItem = MakeShared<FModularRigTreeElement>(InKey, SharedThis(this), true);
+		TSharedPtr<FModularRigHierarchyTreeElement> NewItem = MakeShared<FModularRigHierarchyTreeElement>(InKey, SharedThis(this), true);
 
 		ElementMap.Add(InKey, NewItem);
 		if (!InParentKey.IsEmpty())
 		{
 			ParentMap.Add(InKey, InParentKey);
 
-			TSharedPtr<FModularRigTreeElement>* FoundItem = ElementMap.Find(InParentKey);
+			TSharedPtr<FModularRigHierarchyTreeElement>* FoundItem = ElementMap.Find(InParentKey);
 			check(FoundItem);
 			FoundItem->Get()->Children.Add(NewItem);
 		}
@@ -299,7 +299,7 @@ bool SModularRigTreeView::AddElement(FString InKey, FString InParentKey, bool bA
 
 		for (const FString& Key : FilteredConnectors)
 		{
-			TSharedPtr<FModularRigTreeElement> ConnectorItem = MakeShared<FModularRigTreeElement>(Key, SharedThis(this), false);
+			TSharedPtr<FModularRigHierarchyTreeElement> ConnectorItem = MakeShared<FModularRigHierarchyTreeElement>(Key, SharedThis(this), false);
 			NewItem.Get()->Children.Add(ConnectorItem);
 			ElementMap.Add(Key, ConnectorItem);
 			ParentMap.Add(Key, InKey);
@@ -360,7 +360,7 @@ bool SModularRigTreeView::ReparentElement(const FString InKey, const FString InP
 		return false;
 	}
 
-	TSharedPtr<FModularRigTreeElement>* FoundItem = ElementMap.Find(InKey);
+	TSharedPtr<FModularRigHierarchyTreeElement>* FoundItem = ElementMap.Find(InKey);
 	if (FoundItem == nullptr)
 	{
 		return false;
@@ -373,7 +373,7 @@ bool SModularRigTreeView::ReparentElement(const FString InKey, const FString InP
 			return false;
 		}
 
-		if (TSharedPtr<FModularRigTreeElement>* ExistingParent = ElementMap.Find(*ExistingParentKey))
+		if (TSharedPtr<FModularRigHierarchyTreeElement>* ExistingParent = ElementMap.Find(*ExistingParentKey))
 		{
 			(*ExistingParent)->Children.Remove(*FoundItem);
 		}
@@ -394,7 +394,7 @@ bool SModularRigTreeView::ReparentElement(const FString InKey, const FString InP
 	{
 		ParentMap.Add(InKey, InParentKey);
 
-		TSharedPtr<FModularRigTreeElement>* FoundParent = ElementMap.Find(InParentKey);
+		TSharedPtr<FModularRigHierarchyTreeElement>* FoundParent = ElementMap.Find(InParentKey);
 		if(FoundParent)
 		{
 			FoundParent->Get()->Children.Add(*FoundItem);
@@ -415,7 +415,7 @@ void SModularRigTreeView::RefreshTreeView(bool bRebuildContent)
 
 	if(bRebuildContent)
 	{
-		for (TPair<FString, TSharedPtr<FModularRigTreeElement>> Pair : ElementMap)
+		for (TPair<FString, TSharedPtr<FModularRigHierarchyTreeElement>> Pair : ElementMap)
 		{
 			ExpansionState.FindOrAdd(Pair.Key) = IsItemExpanded(Pair.Value);
 		}
@@ -444,7 +444,7 @@ void SModularRigTreeView::RefreshTreeView(bool bRebuildContent)
 			// expand all elements upon the initial construction of the tree
 			if (ExpansionState.Num() < ElementMap.Num())
 			{
-				for (const TPair<FString, TSharedPtr<FModularRigTreeElement>>& Element : ElementMap)
+				for (const TPair<FString, TSharedPtr<FModularRigHierarchyTreeElement>>& Element : ElementMap)
 				{
 					if (!ExpansionState.Contains(Element.Key))
 					{
@@ -470,7 +470,7 @@ void SModularRigTreeView::RefreshTreeView(bool bRebuildContent)
 		{
 			// elements may be added at the end of the list after a spacer element
 			// we need to remove the spacer element and re-add it at the end
-			RootElements.RemoveAll([](TSharedPtr<FModularRigTreeElement> InElement)
+			RootElements.RemoveAll([](TSharedPtr<FModularRigHierarchyTreeElement> InElement)
 			{
 				return InElement.Get()->Key == FString();
 			});
@@ -485,10 +485,10 @@ void SModularRigTreeView::RefreshTreeView(bool bRebuildContent)
 
 		if(!Selection.IsEmpty())
 		{
-			TArray<TSharedPtr<FModularRigTreeElement>> SelectedElements;
+			TArray<TSharedPtr<FModularRigHierarchyTreeElement>> SelectedElements;
 			for(const FName& SelectedModuleName : Selection)
 			{
-				if(const TSharedPtr<FModularRigTreeElement> ElementToSelect = FindElement(SelectedModuleName.ToString()))
+				if(const TSharedPtr<FModularRigHierarchyTreeElement> ElementToSelect = FindElement(SelectedModuleName.ToString()))
 				{
 					SelectedElements.Add(ElementToSelect);
 				}
@@ -501,14 +501,14 @@ void SModularRigTreeView::RefreshTreeView(bool bRebuildContent)
 	}
 }
 
-TSharedRef<ITableRow> SModularRigTreeView::MakeTableRowWidget(TSharedPtr<FModularRigTreeElement> InItem,
+TSharedRef<ITableRow> SModularRigTreeView::MakeTableRowWidget(TSharedPtr<FModularRigHierarchyTreeElement> InItem,
 	const TSharedRef<STableViewBase>& OwnerTable, bool bPinned)
 {
 	return InItem->MakeTreeRowWidget(OwnerTable, InItem.ToSharedRef(), SharedThis(this), bPinned);
 }
 
-void SModularRigTreeView::HandleGetChildrenForTree(TSharedPtr<FModularRigTreeElement> InItem,
-	TArray<TSharedPtr<FModularRigTreeElement>>& OutChildren)
+void SModularRigTreeView::HandleGetChildrenForTree(TSharedPtr<FModularRigHierarchyTreeElement> InItem,
+	TArray<TSharedPtr<FModularRigHierarchyTreeElement>>& OutChildren)
 {
 	OutChildren = InItem.Get()->Children;
 }
@@ -516,8 +516,8 @@ void SModularRigTreeView::HandleGetChildrenForTree(TSharedPtr<FModularRigTreeEle
 TArray<FName> SModularRigTreeView::GetSelectedModuleNames() const
 {
 	TArray<FName> ModuleNames;
-	TArray<TSharedPtr<FModularRigTreeElement>> SelectedElements = GetSelectedItems();
-	for(const TSharedPtr<FModularRigTreeElement>& SelectedElement : SelectedElements)
+	TArray<TSharedPtr<FModularRigHierarchyTreeElement>> SelectedElements = GetSelectedItems();
+	for(const TSharedPtr<FModularRigHierarchyTreeElement>& SelectedElement : SelectedElements)
 	{
 		if (SelectedElement.IsValid())
 		{
@@ -527,25 +527,25 @@ TArray<FName> SModularRigTreeView::GetSelectedModuleNames() const
 	return ModuleNames;
 }
 
-void SModularRigTreeView::SetSelection(const TArray<TSharedPtr<FModularRigTreeElement>>& InSelection) 
+void SModularRigTreeView::SetSelection(const TArray<TSharedPtr<FModularRigHierarchyTreeElement>>& InSelection) 
 {
 	ClearSelection();
 	SetItemSelection(InSelection, true, ESelectInfo::Direct);
 }
 
-const TSharedPtr<FModularRigTreeElement>* SModularRigTreeView::FindItemAtPosition(FVector2D InScreenSpacePosition) const
+const TSharedPtr<FModularRigHierarchyTreeElement>* SModularRigTreeView::FindItemAtPosition(FVector2D InScreenSpacePosition) const
 {
-	if (ItemsPanel.IsValid() && SListView<TSharedPtr<FModularRigTreeElement>>::HasValidItemsSource())
+	if (ItemsPanel.IsValid() && SListView<TSharedPtr<FModularRigHierarchyTreeElement>>::HasValidItemsSource())
 	{
 		FArrangedChildren ArrangedChildren(EVisibility::Visible);
 		const int32 Index = FindChildUnderPosition(ArrangedChildren, InScreenSpacePosition);
 		if (ArrangedChildren.IsValidIndex(Index))
 		{
-			TSharedRef<SModularRigModelItem> ItemWidget = StaticCastSharedRef<SModularRigModelItem>(ArrangedChildren[Index].Widget);
+			TSharedRef<SModularRigHierarchyTreeItem> ItemWidget = StaticCastSharedRef<SModularRigHierarchyTreeItem>(ArrangedChildren[Index].Widget);
 			if (ItemWidget->WeakRigTreeElement.IsValid())
 			{
 				const FString Key = ItemWidget->WeakRigTreeElement.Pin()->Key;
-				const TSharedPtr<FModularRigTreeElement>* ResultPtr = SListView<TSharedPtr<FModularRigTreeElement>>::GetItems().FindByPredicate([Key](const TSharedPtr<FModularRigTreeElement>& Item) -> bool
+				const TSharedPtr<FModularRigHierarchyTreeElement>* ResultPtr = SListView<TSharedPtr<FModularRigHierarchyTreeElement>>::GetItems().FindByPredicate([Key](const TSharedPtr<FModularRigHierarchyTreeElement>& Item) -> bool
 					{
 						return Item->Key == Key;
 					});
@@ -560,7 +560,7 @@ const TSharedPtr<FModularRigTreeElement>* SModularRigTreeView::FindItemAtPositio
 	return nullptr;
 }
 
-TPair<const FSlateBrush*, FSlateColor> FModularRigTreeElement::GetBrushAndColor(const UModularRig* InModularRig)
+TPair<const FSlateBrush*, FSlateColor> FModularRigHierarchyTreeElement::GetBrushAndColor(const UModularRig* InModularRig)
 {
 	const FSlateBrush* Brush = nullptr;
 	FLinearColor Color = FSlateColor(EStyleColor::Foreground).GetColor(FWidgetStyle());
diff --git a/Engine/Plugins/Animation/ControlRig/Source/ControlRigEditor/Private/Editor/Hierarchy/Views/SModularRigTreeView.h b/Engine/Plugins/Animation/ControlRig/Source/ControlRigEditor/Private/Editor/Hierarchy/Views/SModularRigTreeView.h
--- a/Engine/Plugins/Animation/ControlRig/Source/ControlRigEditor/Private/Editor/Hierarchy/Views/SModularRigTreeView.h
+++ b/Engine/Plugins/Animation/ControlRig/Source/ControlRigEditor/Private/Editor/Hierarchy/Views/SModularRigTreeView.h
@@ -10,9 +10,9 @@
 
 class SSearchBox;
 class SModularRigTreeView;
-class FModularRigTreeElement;
+class FModularRigHierarchyTreeElement;
 
-class SModularRigTreeView : public STreeView<TSharedPtr<FModularRigTreeElement>>
+class SModularRigTreeView : public STreeView<TSharedPtr<FModularRigHierarchyTreeElement>>
 {
 public:
 
@@ -29,7 +29,7 @@ public:
 		, _ShowUnresolvedConnectors(true)
 	{}
 		SLATE_ARGUMENT( TSharedPtr<SHeaderRow>, HeaderRow )
-		SLATE_ARGUMENT(FModularRigTreeDelegates, RigTreeDelegates)
+		SLATE_ARGUMENT(FModularRigHierarchyTreeDelegates, RigTreeDelegates)
 		SLATE_ARGUMENT(bool, AutoScrollEnabled)
 		SLATE_ATTRIBUTE(FText, FilterText)
 		SLATE_ATTRIBUTE(bool, ShowSecondaryConnectors)
@@ -57,7 +57,7 @@ public:
 	}
 
 	/** Restore the expansion infos map from the saved snapshot after tree reconstruction */
-	void RestoreSparseItemInfos(TSharedPtr<FModularRigTreeElement> ItemPtr)
+	void RestoreSparseItemInfos(TSharedPtr<FModularRigHierarchyTreeElement> ItemPtr)
 	{
 		for (const auto& Pair : OldSparseItemInfos)
 		{
@@ -71,23 +71,23 @@ public:
 	}
 
 
-	TSharedPtr<FModularRigTreeElement> FindElement(const FString& InElementKey);
-	static TSharedPtr<FModularRigTreeElement> FindElement(const FString& InElementKey, TSharedPtr<FModularRigTreeElement> CurrentItem);
+	TSharedPtr<FModularRigHierarchyTreeElement> FindElement(const FString& InElementKey);
+	static TSharedPtr<FModularRigHierarchyTreeElement> FindElement(const FString& InElementKey, TSharedPtr<FModularRigHierarchyTreeElement> CurrentItem);
 	bool AddElement(FString InKey, FString InParentKey = FString(), bool bApplyFilterText = true);
 	bool AddElement(const FRigModuleInstance* InElement, bool bApplyFilterText);
 	void AddSpacerElement();
 	bool ReparentElement(const FString InKey, const FString InParentKey);
 	void RefreshTreeView(bool bRebuildContent = true);
-	TSharedRef<ITableRow> MakeTableRowWidget(TSharedPtr<FModularRigTreeElement> InItem, const TSharedRef<STableViewBase>& OwnerTable, bool bPinned);
-	void HandleGetChildrenForTree(TSharedPtr<FModularRigTreeElement> InItem, TArray<TSharedPtr<FModularRigTreeElement>>& OutChildren);
+	TSharedRef<ITableRow> MakeTableRowWidget(TSharedPtr<FModularRigHierarchyTreeElement> InItem, const TSharedRef<STableViewBase>& OwnerTable, bool bPinned);
+	void HandleGetChildrenForTree(TSharedPtr<FModularRigHierarchyTreeElement> InItem, TArray<TSharedPtr<FModularRigHierarchyTreeElement>>& OutChildren);
 
 	TArray<FName> GetSelectedModuleNames() const;
-	void SetSelection(const TArray<TSharedPtr<FModularRigTreeElement>>& InSelection);
-	const TArray<TSharedPtr<FModularRigTreeElement>>& GetRootElements() const { return RootElements; }
-	FModularRigTreeDelegates& GetRigTreeDelegates() { return Delegates; }
+	void SetSelection(const TArray<TSharedPtr<FModularRigHierarchyTreeElement>>& InSelection);
+	const TArray<TSharedPtr<FModularRigHierarchyTreeElement>>& GetRootElements() const { return RootElements; }
+	FModularRigHierarchyTreeDelegates& GetRigTreeDelegates() { return Delegates; }
 
 	/** Given a position, return the item under that position. If nothing is there, return null. */
-	const TSharedPtr<FModularRigTreeElement>* FindItemAtPosition(FVector2D InScreenSpacePosition) const;
+	const TSharedPtr<FModularRigHierarchyTreeElement>* FindItemAtPosition(FVector2D InScreenSpacePosition) const;
 
 private:
 
@@ -95,15 +95,15 @@ private:
 	TSparseItemMap OldSparseItemInfos;
 
 	/** Backing array for tree view */
-	TArray<TSharedPtr<FModularRigTreeElement>> RootElements;
+	TArray<TSharedPtr<FModularRigHierarchyTreeElement>> RootElements;
 	
 	/** A map for looking up items based on their key */
-	TMap<FString, TSharedPtr<FModularRigTreeElement>> ElementMap;
+	TMap<FString, TSharedPtr<FModularRigHierarchyTreeElement>> ElementMap;
 
 	/** A map for looking up a parent based on their key */
 	TMap<FString, FString> ParentMap;
 
-	FModularRigTreeDelegates Delegates;
+	FModularRigHierarchyTreeDelegates Delegates;
 
 	bool bAutoScrollEnabled;
 	FVector2D LastMousePosition;
@@ -114,6 +114,6 @@ private:
 	TAttribute<bool> ShowOptionalConnectors;
 	TAttribute<bool> ShowUnresolvedConnectors;
 
-	friend class SModularRigModel;
+	friend class SModularRigHierarchy;
 };
 
diff --git a/Engine/Plugins/Animation/ControlRig/Source/ControlRigEditor/Private/Editor/Hierarchy/Views/SRigHierarchyItem.cpp b/Engine/Plugins/Animation/ControlRig/Source/ControlRigEditor/Private/Editor/Hierarchy/Views/SRigHierarchyItem.cpp
--- a/Engine/Plugins/Animation/ControlRig/Source/ControlRigEditor/Private/Editor/Hierarchy/Views/SRigHierarchyItem.cpp
+++ b/Engine/Plugins/Animation/ControlRig/Source/ControlRigEditor/Private/Editor/Hierarchy/Views/SRigHierarchyItem.cpp
@@ -8,16 +8,16 @@
 #include "Widgets/Layout/SSpacer.h"
 #include "Widgets/Text/SInlineEditableTextBlock.h"
 
-void SRigHierarchyItem::Construct(const FArguments& InArgs, const TSharedRef<STableViewBase>& OwnerTable, TSharedRef<FRigTreeElement> InRigTreeElement, TSharedPtr<SRigHierarchyTreeView> InTreeView, const FRigTreeDisplaySettings& InSettings, bool bPinned)
+void SRigHierarchyItem::Construct(const FArguments& InArgs, const TSharedRef<STableViewBase>& OwnerTable, TSharedRef<FRigHierarchyTreeElement> InRigTreeElement, TSharedPtr<SRigHierarchyTreeView> InTreeView, const FRigHierarchyTreeDisplaySettings& InSettings, bool bPinned)
 {
 	WeakRigTreeElement = InRigTreeElement;
 	Delegates = InTreeView->GetRigTreeDelegates();
-	FRigTreeDisplaySettings DisplaySettings = Delegates.GetDisplaySettings();
+	FRigHierarchyTreeDisplaySettings DisplaySettings = Delegates.GetDisplaySettings();
 
 	if (!InRigTreeElement->Key.IsValid())
 	{
-		STableRow<TSharedPtr<FRigTreeElement>>::Construct(
-			STableRow<TSharedPtr<FRigTreeElement>>::FArguments()
+		STableRow<TSharedPtr<FRigHierarchyTreeElement>>::Construct(
+			STableRow<TSharedPtr<FRigHierarchyTreeElement>>::FArguments()
 			.ShowSelection(false)
 			.OnCanAcceptDrop(Delegates.OnCanAcceptDrop)
 			.OnAcceptDrop(Delegates.OnAcceptDrop)
@@ -38,8 +38,8 @@ void SRigHierarchyItem::Construct(const FArguments& InArgs, const TSharedRef<STa
 	TSharedPtr< SInlineEditableTextBlock > InlineWidget;
 	TSharedPtr< SHorizontalBox > HorizontalBox;
 
-	STableRow<TSharedPtr<FRigTreeElement>>::Construct(
-		STableRow<TSharedPtr<FRigTreeElement>>::FArguments()
+	STableRow<TSharedPtr<FRigHierarchyTreeElement>>::Construct(
+		STableRow<TSharedPtr<FRigHierarchyTreeElement>>::FArguments()
 		.Padding(FMargin(0, 1, 1, 1))
 		.OnDragDetected(Delegates.OnDragDetected)
 		.OnCanAcceptDrop(Delegates.OnCanAcceptDrop)
diff --git a/Engine/Plugins/Animation/ControlRig/Source/ControlRigEditor/Private/Editor/Hierarchy/Views/SRigHierarchyItem.h b/Engine/Plugins/Animation/ControlRig/Source/ControlRigEditor/Private/Editor/Hierarchy/Views/SRigHierarchyItem.h
--- a/Engine/Plugins/Animation/ControlRig/Source/ControlRigEditor/Private/Editor/Hierarchy/Views/SRigHierarchyItem.h
+++ b/Engine/Plugins/Animation/ControlRig/Source/ControlRigEditor/Private/Editor/Hierarchy/Views/SRigHierarchyItem.h
@@ -11,19 +11,19 @@
 class SRigHierarchyTreeView;
 enum class EElementNameDisplayMode : uint8;
 
-class SRigHierarchyItem : public STableRow<TSharedPtr<FRigTreeElement>>
+class SRigHierarchyItem : public STableRow<TSharedPtr<FRigHierarchyTreeElement>>
 {
 public:
 
-	void Construct(const FArguments& InArgs, const TSharedRef<STableViewBase>& OwnerTable, TSharedRef<FRigTreeElement> InRigTreeElement, TSharedPtr<SRigHierarchyTreeView> InTreeView, const FRigTreeDisplaySettings& InSettings, bool bPinned);
+	void Construct(const FArguments& InArgs, const TSharedRef<STableViewBase>& OwnerTable, TSharedRef<FRigHierarchyTreeElement> InRigTreeElement, TSharedPtr<SRigHierarchyTreeView> InTreeView, const FRigHierarchyTreeDisplaySettings& InSettings, bool bPinned);
 	void OnNameCommitted(const FText& InText, ETextCommit::Type InCommitType) const;
 	bool OnVerifyNameChanged(const FText& InText, FText& OutErrorMessage);
 	static TPair<const FSlateBrush*, FSlateColor> GetBrushForElementType(const URigHierarchy* InHierarchy, const FRigHierarchyKey& InKey);
 	static FLinearColor GetColorForControlType(ERigControlType InControlType, UEnum* InControlEnum);
 
 private:
-	TWeakPtr<FRigTreeElement> WeakRigTreeElement;
-	FRigTreeDelegates Delegates;
+	TWeakPtr<FRigHierarchyTreeElement> WeakRigTreeElement;
+	FRigHierarchyTreeDelegates Delegates;
 
 	FText GetNameForUI() const;
 	FText GetName(EElementNameDisplayMode InNameDisplayMode) const;
diff --git a/Engine/Plugins/Animation/ControlRig/Source/ControlRigEditor/Private/Editor/Hierarchy/Views/SRigHierarchySearchableTreeView.cpp b/Engine/Plugins/Animation/ControlRig/Source/ControlRigEditor/Private/Editor/Hierarchy/Views/SRigHierarchySearchableTreeView.cpp
--- a/Engine/Plugins/Animation/ControlRig/Source/ControlRigEditor/Private/Editor/Hierarchy/Views/SRigHierarchySearchableTreeView.cpp
+++ b/Engine/Plugins/Animation/ControlRig/Source/ControlRigEditor/Private/Editor/Hierarchy/Views/SRigHierarchySearchableTreeView.cpp
@@ -5,12 +5,12 @@
 #include "Widgets/Input/SSearchBox.h"
 #include "Widgets/Layout/SScrollBox.h"
 
-void SSearchableRigHierarchyTreeView::Construct(const FArguments& InArgs)
+void SRigHierarchySearchableTreeView::Construct(const FArguments& InArgs)
 {
-	FRigTreeDelegates TreeDelegates = InArgs._RigTreeDelegates;
+	FRigHierarchyTreeDelegates TreeDelegates = InArgs._RigTreeDelegates;
 	SuperGetRigTreeDisplaySettings = TreeDelegates.OnGetDisplaySettings;
 
-	TreeDelegates.OnGetDisplaySettings.BindSP(this, &SSearchableRigHierarchyTreeView::GetDisplaySettings);
+	TreeDelegates.OnGetDisplaySettings.BindSP(this, &SRigHierarchySearchableTreeView::GetDisplaySettings);
 
 	TSharedPtr<SVerticalBox> VerticalBox;
 	ChildSlot
@@ -24,7 +24,7 @@ void SSearchableRigHierarchyTreeView::Construct(const FArguments& InArgs)
 			[
 				SAssignNew(SearchBox, SSearchBox)
 				.InitialText(InArgs._InitialFilterText)
-				.OnTextChanged(this, &SSearchableRigHierarchyTreeView::OnFilterTextChanged)
+				.OnTextChanged(this, &SRigHierarchySearchableTreeView::OnFilterTextChanged)
 			]
 
 			+SVerticalBox::Slot()
@@ -40,7 +40,7 @@ void SSearchableRigHierarchyTreeView::Construct(const FArguments& InArgs)
 	];
 }
 
-const FRigTreeDisplaySettings& SSearchableRigHierarchyTreeView::GetDisplaySettings()
+const FRigHierarchyTreeDisplaySettings& SRigHierarchySearchableTreeView::GetDisplaySettings()
 {
 	if(SuperGetRigTreeDisplaySettings.IsBound())
 	{
@@ -50,7 +50,7 @@ const FRigTreeDisplaySettings& SSearchableRigHierarchyTreeView::GetDisplaySettin
 	return Settings;
 }
 
-void SSearchableRigHierarchyTreeView::OnFilterTextChanged(const FText& SearchText)
+void SRigHierarchySearchableTreeView::OnFilterTextChanged(const FText& SearchText)
 {
 	FilterText = SearchText;
 	GetTreeView()->RefreshTreeView();
diff --git a/Engine/Plugins/Animation/ControlRig/Source/ControlRigEditor/Private/Editor/Hierarchy/Views/SRigHierarchySearchableTreeView.h b/Engine/Plugins/Animation/ControlRig/Source/ControlRigEditor/Private/Editor/Hierarchy/Views/SRigHierarchySearchableTreeView.h
--- a/Engine/Plugins/Animation/ControlRig/Source/ControlRigEditor/Private/Editor/Hierarchy/Views/SRigHierarchySearchableTreeView.h
+++ b/Engine/Plugins/Animation/ControlRig/Source/ControlRigEditor/Private/Editor/Hierarchy/Views/SRigHierarchySearchableTreeView.h
@@ -11,14 +11,14 @@
 
 class SSearchBox;
 
-class SSearchableRigHierarchyTreeView : public SCompoundWidget
+class SRigHierarchySearchableTreeView : public SCompoundWidget
 {
 public:
 
-	SLATE_BEGIN_ARGS(SSearchableRigHierarchyTreeView)
+	SLATE_BEGIN_ARGS(SRigHierarchySearchableTreeView)
 		{
 		}
-		SLATE_ARGUMENT(FRigTreeDelegates, RigTreeDelegates)
+		SLATE_ARGUMENT(FRigHierarchyTreeDelegates, RigTreeDelegates)
 		SLATE_ARGUMENT(FText, InitialFilterText)
 	SLATE_END_ARGS()
 
@@ -26,7 +26,7 @@ public:
 
 	TSharedRef<SSearchBox> GetSearchBox() const { return SearchBox.ToSharedRef(); }
 	TSharedRef<SRigHierarchyTreeView> GetTreeView() const { return TreeView.ToSharedRef(); }
-	const FRigTreeDisplaySettings& GetDisplaySettings();
+	const FRigHierarchyTreeDisplaySettings& GetDisplaySettings();
 
 private:
 
@@ -34,7 +34,7 @@ private:
 
 	FOnGetRigTreeDisplaySettings SuperGetRigTreeDisplaySettings;
 	FText FilterText;
-	FRigTreeDisplaySettings Settings;
+	FRigHierarchyTreeDisplaySettings Settings;
 	TSharedPtr<SSearchBox> SearchBox;
 	TSharedPtr<SRigHierarchyTreeView> TreeView;
 };
diff --git a/Engine/Plugins/Animation/ControlRig/Source/ControlRigEditor/Private/Editor/Hierarchy/Views/SRigHierarchyTreeView.cpp b/Engine/Plugins/Animation/ControlRig/Source/ControlRigEditor/Private/Editor/Hierarchy/Views/SRigHierarchyTreeView.cpp
--- a/Engine/Plugins/Animation/ControlRig/Source/ControlRigEditor/Private/Editor/Hierarchy/Views/SRigHierarchyTreeView.cpp
+++ b/Engine/Plugins/Animation/ControlRig/Source/ControlRigEditor/Private/Editor/Hierarchy/Views/SRigHierarchyTreeView.cpp
@@ -33,12 +33,12 @@ void SRigHierarchyTreeView::Construct(const FArguments& InArgs)
 	Delegates = InArgs._RigTreeDelegates;
 	bAutoScrollEnabled = InArgs._AutoScrollEnabled;
 
-	STreeView<TSharedPtr<FRigTreeElement>>::FArguments SuperArgs;
+	STreeView<TSharedPtr<FRigHierarchyTreeElement>>::FArguments SuperArgs;
 	SuperArgs.TreeItemsSource(&RootElements);
 	SuperArgs.SelectionMode(ESelectionMode::Multi);
 	SuperArgs.OnGenerateRow(this, &SRigHierarchyTreeView::MakeTableRowWidget, false);
 	SuperArgs.OnGetChildren(this, &SRigHierarchyTreeView::HandleGetChildrenForTree);
-	SuperArgs.OnSelectionChanged(FOnRigTreeSelectionChanged::CreateRaw(&Delegates, &FRigTreeDelegates::HandleSelectionChanged));
+	SuperArgs.OnSelectionChanged(FOnRigTreeSelectionChanged::CreateRaw(&Delegates, &FRigHierarchyTreeDelegates::HandleSelectionChanged));
 	SuperArgs.OnContextMenuOpening(Delegates.OnContextMenuOpening);
 	SuperArgs.OnMouseButtonClick(Delegates.OnMouseButtonClick);
 	SuperArgs.OnMouseButtonDoubleClick(Delegates.OnMouseButtonDoubleClick);
@@ -55,7 +55,7 @@ void SRigHierarchyTreeView::Construct(const FArguments& InArgs)
 		return FMath::Max<int32>(1, UControlRigEditorSettings::Get()->MaxStackSize);
 	});
 
-	STreeView<TSharedPtr<FRigTreeElement>>::Construct(SuperArgs);
+	STreeView<TSharedPtr<FRigHierarchyTreeElement>>::Construct(SuperArgs);
 
 	LastMousePosition = FVector2D::ZeroVector;
 	TimeAtMousePosition = 0.0;
@@ -68,7 +68,7 @@ void SRigHierarchyTreeView::Construct(const FArguments& InArgs)
 
 void SRigHierarchyTreeView::Tick(const FGeometry& AllottedGeometry, const double InCurrentTime, const float InDeltaTime)
 {
-	STreeView<TSharedPtr<FRigTreeElement, ESPMode::ThreadSafe>>::Tick(AllottedGeometry, InCurrentTime, InDeltaTime);
+	STreeView<TSharedPtr<FRigHierarchyTreeElement, ESPMode::ThreadSafe>>::Tick(AllottedGeometry, InCurrentTime, InDeltaTime);
 
 	const FGeometry PaintGeometry = GetPaintSpaceGeometry();
 	const FVector2D MousePosition = FSlateApplication::Get().GetCursorPos();
@@ -107,7 +107,7 @@ void SRigHierarchyTreeView::Tick(const FGeometry& AllottedGeometry, const double
 			}
 			else
 			{
-				const TSharedPtr<FRigTreeElement>* Item = FindItemAtPosition(MousePosition);
+				const TSharedPtr<FRigHierarchyTreeElement>* Item = FindItemAtPosition(MousePosition);
 				if(Item && Item->IsValid())
 				{
 					if(!IsItemExpanded(*Item))
@@ -120,20 +120,20 @@ void SRigHierarchyTreeView::Tick(const FGeometry& AllottedGeometry, const double
 	}
 }
 
-TSharedPtr<FRigTreeElement> SRigHierarchyTreeView::FindElement(const FRigHierarchyKey& InElementKey) const
+TSharedPtr<FRigHierarchyTreeElement> SRigHierarchyTreeView::FindElement(const FRigHierarchyKey& InElementKey) const
 {
 	for (int32 RootIndex = 0; RootIndex < RootElements.Num(); ++RootIndex)
 	{
-		TSharedPtr<FRigTreeElement> Found = FindElement(InElementKey, RootElements[RootIndex]);
+		TSharedPtr<FRigHierarchyTreeElement> Found = FindElement(InElementKey, RootElements[RootIndex]);
 		if (Found.IsValid())
 		{
 			return Found;
 		}
 	}
-	return TSharedPtr<FRigTreeElement>();
+	return TSharedPtr<FRigHierarchyTreeElement>();
 }
 
-TSharedPtr<FRigTreeElement> SRigHierarchyTreeView::FindElement(const FRigHierarchyKey& InElementKey, TSharedPtr<FRigTreeElement> CurrentItem)
+TSharedPtr<FRigHierarchyTreeElement> SRigHierarchyTreeView::FindElement(const FRigHierarchyKey& InElementKey, TSharedPtr<FRigHierarchyTreeElement> CurrentItem)
 {
 	if (CurrentItem->Key == InElementKey)
 	{
@@ -142,14 +142,14 @@ TSharedPtr<FRigTreeElement> SRigHierarchyTreeView::FindElement(const FRigHierarc
 
 	for (int32 ChildIndex = 0; ChildIndex < CurrentItem->Children.Num(); ++ChildIndex)
 	{
-		TSharedPtr<FRigTreeElement> Found = FindElement(InElementKey, CurrentItem->Children[ChildIndex]);
+		TSharedPtr<FRigHierarchyTreeElement> Found = FindElement(InElementKey, CurrentItem->Children[ChildIndex]);
 		if (Found.IsValid())
 		{
 			return Found;
 		}
 	}
 
-	return TSharedPtr<FRigTreeElement>();
+	return TSharedPtr<FRigHierarchyTreeElement>();
 }
 
 bool SRigHierarchyTreeView::AddElement(FRigHierarchyKey InKey, FRigHierarchyKey InParentKey)
@@ -174,7 +174,7 @@ bool SRigHierarchyTreeView::AddElement(FRigHierarchyKey InKey, FRigHierarchyKey
 		}
 	}
 
-	const FRigTreeDisplaySettings& Settings = Delegates.GetDisplaySettings();
+	const FRigHierarchyTreeDisplaySettings& Settings = Delegates.GetDisplaySettings();
 	const bool bSupportsRename = Delegates.OnRenameElement.IsBound();
 
 	const FString FilteredString = Settings.FilterText.ToString();
@@ -186,7 +186,7 @@ bool SRigHierarchyTreeView::AddElement(FRigHierarchyKey InKey, FRigHierarchyKey
 
 	if (!bAnyFilteredOut)
 	{
-		TSharedPtr<FRigTreeElement> NewItem = MakeShared<FRigTreeElement>(InKey, SharedThis(this), bSupportsRename, ERigTreeFilterResult::Shown);
+		TSharedPtr<FRigHierarchyTreeElement> NewItem = MakeShared<FRigHierarchyTreeElement>(InKey, SharedThis(this), bSupportsRename, ERigTreeFilterResult::Shown);
 
 		if (InKey.IsValid())
 		{
@@ -198,7 +198,7 @@ bool SRigHierarchyTreeView::AddElement(FRigHierarchyKey InKey, FRigHierarchyKey
 
 			if (InParentKey.IsValid())
 			{
-				TSharedPtr<FRigTreeElement>* FoundItem = ElementMap.Find(InParentKey);
+				TSharedPtr<FRigHierarchyTreeElement>* FoundItem = ElementMap.Find(InParentKey);
 				check(FoundItem);
 				FoundItem->Get()->Children.Add(NewItem);
 			}
@@ -237,7 +237,7 @@ bool SRigHierarchyTreeView::AddElement(FRigHierarchyKey InKey, FRigHierarchyKey
 		FString FilteredStringUnderScores = FilteredString.Replace(TEXT(" "), TEXT("_"));
 		if (!bIsFilteredOut && (InKey.GetName().Contains(FilteredString) || InKey.GetName().Contains(FilteredStringUnderScores)))
 		{
-			TSharedPtr<FRigTreeElement> NewItem = MakeShared<FRigTreeElement>(InKey, SharedThis(this), bSupportsRename, ERigTreeFilterResult::Shown);
+			TSharedPtr<FRigHierarchyTreeElement> NewItem = MakeShared<FRigHierarchyTreeElement>(InKey, SharedThis(this), bSupportsRename, ERigTreeFilterResult::Shown);
 			ElementMap.Add(InKey, NewItem);
 			RootElements.Add(NewItem);
 
@@ -245,13 +245,13 @@ bool SRigHierarchyTreeView::AddElement(FRigHierarchyKey InKey, FRigHierarchyKey
 			{
 				if(Hierarchy)
 				{
-					TSharedPtr<FRigTreeElement> ChildItem = NewItem;
+					TSharedPtr<FRigHierarchyTreeElement> ChildItem = NewItem;
 					FRigHierarchyKey ParentKey = GetFirstParent(InKey);
 					while (ParentKey.IsValid())
 					{
 						if (!ElementMap.Contains(ParentKey))
 						{
-							TSharedPtr<FRigTreeElement> ParentItem = MakeShared<FRigTreeElement>(ParentKey, SharedThis(this), bSupportsRename, ERigTreeFilterResult::ShownDescendant);							
+							TSharedPtr<FRigHierarchyTreeElement> ParentItem = MakeShared<FRigHierarchyTreeElement>(ParentKey, SharedThis(this), bSupportsRename, ERigTreeFilterResult::ShownDescendant);							
 							ElementMap.Add(ParentKey, ParentItem);
 							RootElements.Add(ParentItem);
 
@@ -283,7 +283,7 @@ bool SRigHierarchyTreeView::AddElement(const FRigBaseElement* InElement)
 		return false;
 	}
 
-	const FRigTreeDisplaySettings& Settings = Delegates.GetDisplaySettings();
+	const FRigHierarchyTreeDisplaySettings& Settings = Delegates.GetDisplaySettings();
 	const URigHierarchy* Hierarchy = Delegates.GetHierarchy();
 
 	switch(InElement->GetType())
@@ -345,14 +345,14 @@ bool SRigHierarchyTreeView::AddElement(const FRigBaseElement* InElement)
 					if(const FRigElementKeyRedirector::FCachedKeyArray* Cache = Redirector.Find(InElement->GetKey()))
 					{
 						bool bValid = true;
-						TArray<TSharedPtr<FRigTreeElement>> TargetElementPtrs;
+						TArray<TSharedPtr<FRigHierarchyTreeElement>> TargetElementPtrs;
 						for(const FCachedRigElement& CachedRigElement : *Cache)
 						{
 							if(!const_cast<FCachedRigElement*>(&CachedRigElement)->UpdateCache(Hierarchy))
 							{
 								bValid = false;
 							}
-							if(const TSharedPtr<FRigTreeElement>* TargetElementPtr = ElementMap.Find(CachedRigElement.GetKey()))
+							if(const TSharedPtr<FRigHierarchyTreeElement>* TargetElementPtr = ElementMap.Find(CachedRigElement.GetKey()))
 							{
 								if(TargetElementPtr->IsValid())
 								{
@@ -415,7 +415,7 @@ bool SRigHierarchyTreeView::AddElement(const FRigBaseElement* InElement)
 								});
 							}
 
-							for(TSharedPtr<FRigTreeElement> TargetElementPtr : TargetElementPtrs)
+							for(TSharedPtr<FRigHierarchyTreeElement> TargetElementPtr : TargetElementPtrs)
 							{
 								TargetElementPtr->Tags.Add(TagArguments);
 							}
@@ -512,7 +512,7 @@ bool SRigHierarchyTreeView::AddComponent(const FRigBaseComponent* InComponent)
 		return false;
 	}
 
-	const FRigTreeDisplaySettings& Settings = Delegates.GetDisplaySettings();
+	const FRigHierarchyTreeDisplaySettings& Settings = Delegates.GetDisplaySettings();
 	if(!Settings.bShowComponents)
 	{
 		return false;
@@ -546,9 +546,9 @@ bool SRigHierarchyTreeView::ReparentElement(FRigHierarchyKey InKey, FRigHierarch
 		}
 	}
 
-	const FRigTreeDisplaySettings& Settings = Delegates.GetDisplaySettings();
+	const FRigHierarchyTreeDisplaySettings& Settings = Delegates.GetDisplaySettings();
 
-	TSharedPtr<FRigTreeElement>* FoundItem = ElementMap.Find(InKey);
+	TSharedPtr<FRigHierarchyTreeElement>* FoundItem = ElementMap.Find(InKey);
 	if (FoundItem == nullptr)
 	{
 		return false;
@@ -566,7 +566,7 @@ bool SRigHierarchyTreeView::ReparentElement(FRigHierarchyKey InKey, FRigHierarch
 			return false;
 		}
 
-		if (TSharedPtr<FRigTreeElement>* ExistingParent = ElementMap.Find(*ExistingParentKey))
+		if (TSharedPtr<FRigHierarchyTreeElement>* ExistingParent = ElementMap.Find(*ExistingParentKey))
 		{
 			(*ExistingParent)->Children.Remove(*FoundItem);
 		}
@@ -587,7 +587,7 @@ bool SRigHierarchyTreeView::ReparentElement(FRigHierarchyKey InKey, FRigHierarch
 	{
 		ParentMap.Add(InKey, InParentKey);
 
-		TSharedPtr<FRigTreeElement>* FoundParent = ElementMap.Find(InParentKey);
+		TSharedPtr<FRigHierarchyTreeElement>* FoundParent = ElementMap.Find(InParentKey);
 		if(FoundParent)
 		{
 			FoundParent->Get()->Children.Add(*FoundItem);
@@ -603,7 +603,7 @@ bool SRigHierarchyTreeView::ReparentElement(FRigHierarchyKey InKey, FRigHierarch
 
 bool SRigHierarchyTreeView::RemoveElement(FRigHierarchyKey InKey)
 {
-	TSharedPtr<FRigTreeElement>* FoundItem = ElementMap.Find(InKey);
+	TSharedPtr<FRigHierarchyTreeElement>* FoundItem = ElementMap.Find(InKey);
 	if (FoundItem == nullptr)
 	{
 		return false;
@@ -621,7 +621,7 @@ void SRigHierarchyTreeView::RefreshTreeView(bool bRebuildContent)
 
 	if(bRebuildContent)
 	{
-		for (TPair<FRigHierarchyKey, TSharedPtr<FRigTreeElement>> Pair : ElementMap)
+		for (TPair<FRigHierarchyKey, TSharedPtr<FRigHierarchyTreeElement>> Pair : ElementMap)
 		{
 			ExpansionState.FindOrAdd(Pair.Key) = IsItemExpanded(Pair.Value);
 		}
@@ -698,14 +698,14 @@ void SRigHierarchyTreeView::RefreshTreeView(bool bRebuildContent)
 			// expand all elements upon the initial construction of the tree
 			if (ExpansionState.Num() == 0)
 			{
-				for (TSharedPtr<FRigTreeElement> RootElement : RootElements)
+				for (TSharedPtr<FRigHierarchyTreeElement> RootElement : RootElements)
 				{
 					SetExpansionRecursive(RootElement, false, true);
 				}
 			}
 			else if (ExpansionState.Num() < ElementMap.Num())
 			{
-				for (const TPair<FRigHierarchyKey, TSharedPtr<FRigTreeElement>>& Element : ElementMap)
+				for (const TPair<FRigHierarchyKey, TSharedPtr<FRigHierarchyTreeElement>>& Element : ElementMap)
 				{
 					if (!ExpansionState.Contains(Element.Key))
 					{
@@ -721,7 +721,7 @@ void SRigHierarchyTreeView::RefreshTreeView(bool bRebuildContent)
 
 			if(Delegates.OnCompareKeys.IsBound())
 			{
-				Algo::Sort(RootElements, [&](const TSharedPtr<FRigTreeElement>& A, const TSharedPtr<FRigTreeElement>& B)
+				Algo::Sort(RootElements, [&](const TSharedPtr<FRigHierarchyTreeElement>& A, const TSharedPtr<FRigHierarchyTreeElement>& B)
 				{
 					return Delegates.OnCompareKeys.Execute(A->Key, B->Key);
 				});
@@ -739,7 +739,7 @@ void SRigHierarchyTreeView::RefreshTreeView(bool bRebuildContent)
 		{
 			// elements may be added at the end of the list after a spacer element
 			// we need to remove the spacer element and re-add it at the end
-			RootElements.RemoveAll([](TSharedPtr<FRigTreeElement> InElement)
+			RootElements.RemoveAll([](TSharedPtr<FRigHierarchyTreeElement> InElement)
 			{
 				return InElement.Get()->Key == FRigElementKey();
 			});
@@ -756,7 +756,7 @@ void SRigHierarchyTreeView::RefreshTreeView(bool bRebuildContent)
 		{
 			for (int32 RootIndex = 0; RootIndex < RootElements.Num(); ++RootIndex)
 			{
-				TSharedPtr<FRigTreeElement> Found = FindElement(Key, RootElements[RootIndex]);
+				TSharedPtr<FRigHierarchyTreeElement> Found = FindElement(Key, RootElements[RootIndex]);
 				if (Found.IsValid())
 				{
 					SetItemSelection(Found, true, ESelectInfo::OnNavigation);
@@ -766,7 +766,7 @@ void SRigHierarchyTreeView::RefreshTreeView(bool bRebuildContent)
 	}
 }
 
-void SRigHierarchyTreeView::SetExpansionRecursive(TSharedPtr<FRigTreeElement> InElement, bool bTowardsParent,
+void SRigHierarchyTreeView::SetExpansionRecursive(TSharedPtr<FRigHierarchyTreeElement> InElement, bool bTowardsParent,
 	bool bShouldBeExpanded)
 {
 	SetItemExpansion(InElement, bShouldBeExpanded);
@@ -775,7 +775,7 @@ void SRigHierarchyTreeView::SetExpansionRecursive(TSharedPtr<FRigTreeElement> In
 	{
 		if (const FRigHierarchyKey* ParentKey = ParentMap.Find(InElement->Key))
 		{
-			if (TSharedPtr<FRigTreeElement>* ParentItem = ElementMap.Find(*ParentKey))
+			if (TSharedPtr<FRigHierarchyTreeElement>* ParentItem = ElementMap.Find(*ParentKey))
 			{
 				SetExpansionRecursive(*ParentItem, bTowardsParent, bShouldBeExpanded);
 			}
@@ -790,15 +790,15 @@ void SRigHierarchyTreeView::SetExpansionRecursive(TSharedPtr<FRigTreeElement> In
 	}
 }
 
-TSharedRef<ITableRow> SRigHierarchyTreeView::MakeTableRowWidget(TSharedPtr<FRigTreeElement> InItem,
+TSharedRef<ITableRow> SRigHierarchyTreeView::MakeTableRowWidget(TSharedPtr<FRigHierarchyTreeElement> InItem,
 	const TSharedRef<STableViewBase>& OwnerTable, bool bPinned)
 {
-	const FRigTreeDisplaySettings& Settings = Delegates.GetDisplaySettings();
+	const FRigHierarchyTreeDisplaySettings& Settings = Delegates.GetDisplaySettings();
 	return InItem->MakeTreeRowWidget(OwnerTable, InItem.ToSharedRef(), SharedThis(this), Settings, bPinned);
 }
 
-void SRigHierarchyTreeView::HandleGetChildrenForTree(TSharedPtr<FRigTreeElement> InItem,
-	TArray<TSharedPtr<FRigTreeElement>>& OutChildren)
+void SRigHierarchyTreeView::HandleGetChildrenForTree(TSharedPtr<FRigHierarchyTreeElement> InItem,
+	TArray<TSharedPtr<FRigHierarchyTreeElement>>& OutChildren)
 {
 	OutChildren = InItem.Get()->Children;
 }
@@ -811,17 +811,17 @@ void SRigHierarchyTreeView::OnElementKeyTagDragDetected(const FRigElementKey& In
 TArray<FRigHierarchyKey> SRigHierarchyTreeView::GetSelectedKeys() const
 {
 	TArray<FRigHierarchyKey> Keys;
-	TArray<TSharedPtr<FRigTreeElement>> SelectedElements = GetSelectedItems();
-	for(const TSharedPtr<FRigTreeElement>& SelectedElement : SelectedElements)
+	TArray<TSharedPtr<FRigHierarchyTreeElement>> SelectedElements = GetSelectedItems();
+	for(const TSharedPtr<FRigHierarchyTreeElement>& SelectedElement : SelectedElements)
 	{
 		Keys.Add(SelectedElement->Key);
 	}
 	return Keys;
 }
 
-const TSharedPtr<FRigTreeElement>* SRigHierarchyTreeView::FindItemAtPosition(FVector2D InScreenSpacePosition) const
+const TSharedPtr<FRigHierarchyTreeElement>* SRigHierarchyTreeView::FindItemAtPosition(FVector2D InScreenSpacePosition) const
 {
-	if (ItemsPanel.IsValid() && SListView<TSharedPtr<FRigTreeElement>>::HasValidItemsSource())
+	if (ItemsPanel.IsValid() && SListView<TSharedPtr<FRigHierarchyTreeElement>>::HasValidItemsSource())
 	{
 		FArrangedChildren ArrangedChildren(EVisibility::Visible);
 		const int32 Index = FindChildUnderPosition(ArrangedChildren, InScreenSpacePosition);
@@ -831,7 +831,7 @@ const TSharedPtr<FRigTreeElement>* SRigHierarchyTreeView::FindItemAtPosition(FVe
 			if (ItemWidget->WeakRigTreeElement.IsValid())
 			{
 				const FRigHierarchyKey Key = ItemWidget->WeakRigTreeElement.Pin()->Key;
-				const TSharedPtr<FRigTreeElement>* ResultPtr = SListView<TSharedPtr<FRigTreeElement>>::GetItems().FindByPredicate([Key](const TSharedPtr<FRigTreeElement>& Item) -> bool
+				const TSharedPtr<FRigHierarchyTreeElement>* ResultPtr = SListView<TSharedPtr<FRigHierarchyTreeElement>>::GetItems().FindByPredicate([Key](const TSharedPtr<FRigHierarchyTreeElement>& Item) -> bool
 					{
 						return Item->Key == Key;
 					});
@@ -846,7 +846,7 @@ const TSharedPtr<FRigTreeElement>* SRigHierarchyTreeView::FindItemAtPosition(FVe
 	return nullptr;
 }
 
-void SRigHierarchyTreeView::AddConnectorResolveWarningTag(TSharedPtr<FRigTreeElement> InTreeElement,
+void SRigHierarchyTreeView::AddConnectorResolveWarningTag(TSharedPtr<FRigHierarchyTreeElement> InTreeElement,
 	const FRigBaseElement* InRigElement, const URigHierarchy* InHierarchy)
 {
 	check(InTreeElement.IsValid());
@@ -890,7 +890,7 @@ void SRigHierarchyTreeView::AddConnectorResolveWarningTag(TSharedPtr<FRigTreeEle
 	}
 }
 
-FText SRigHierarchyTreeView::GetConnectorWarningMessage(TSharedPtr<FRigTreeElement> InTreeElement,
+FText SRigHierarchyTreeView::GetConnectorWarningMessage(TSharedPtr<FRigHierarchyTreeElement> InTreeElement,
 	TWeakObjectPtr<UControlRig> InControlRigPtr, const FRigElementKey InConnectorKey) const
 {
 	if(UControlRig* ControlRig = InControlRigPtr.Get())
diff --git a/Engine/Plugins/Animation/ControlRig/Source/ControlRigEditor/Private/Editor/Hierarchy/Views/SRigHierarchyTreeView.h b/Engine/Plugins/Animation/ControlRig/Source/ControlRigEditor/Private/Editor/Hierarchy/Views/SRigHierarchyTreeView.h
--- a/Engine/Plugins/Animation/ControlRig/Source/ControlRigEditor/Private/Editor/Hierarchy/Views/SRigHierarchyTreeView.h
+++ b/Engine/Plugins/Animation/ControlRig/Source/ControlRigEditor/Private/Editor/Hierarchy/Views/SRigHierarchyTreeView.h
@@ -9,11 +9,11 @@
 
 class SRigHierarchyTreeView;
 class SRigHierarchyItem;
-class FRigTreeElement;
+class FRigHierarchyTreeElement;
 
 #define UE_API CONTROLRIGEDITOR_API
 
-class SRigHierarchyTreeView : public STreeView<TSharedPtr<FRigTreeElement>>
+class SRigHierarchyTreeView : public STreeView<TSharedPtr<FRigHierarchyTreeElement>>
 {
 public:
 
@@ -21,7 +21,7 @@ public:
 		: _AutoScrollEnabled(false)
 		, _PopulateOnConstruct(false)
 	{}
-		SLATE_ARGUMENT(FRigTreeDelegates, RigTreeDelegates)
+		SLATE_ARGUMENT(FRigHierarchyTreeDelegates, RigTreeDelegates)
 		SLATE_ARGUMENT(bool, AutoScrollEnabled)
 		SLATE_ARGUMENT(bool, PopulateOnConstruct)
 	SLATE_END_ARGS()
@@ -34,7 +34,7 @@ public:
 
 	virtual FReply OnFocusReceived(const FGeometry& MyGeometry, const FFocusEvent& InFocusEvent) override
 	{
-		FReply Reply = STreeView<TSharedPtr<FRigTreeElement>>::OnFocusReceived(MyGeometry, InFocusEvent);
+		FReply Reply = STreeView<TSharedPtr<FRigHierarchyTreeElement>>::OnFocusReceived(MyGeometry, InFocusEvent);
 
 		LastClickCycles = FPlatformTime::Cycles();
 
@@ -55,7 +55,7 @@ public:
 	}
 
 	/** Restore the expansion infos map from the saved snapshot after tree reconstruction */
-	void RestoreSparseItemInfos(TSharedPtr<FRigTreeElement> ItemPtr)
+	void RestoreSparseItemInfos(TSharedPtr<FRigHierarchyTreeElement> ItemPtr)
 	{
 		for (const auto& Pair : OldSparseItemInfos)
 		{
@@ -68,8 +68,8 @@ public:
 		}
 	}
 
-	TSharedPtr<FRigTreeElement> FindElement(const FRigHierarchyKey& InElementKey) const;
-	static TSharedPtr<FRigTreeElement> FindElement(const FRigHierarchyKey& InElementKey, TSharedPtr<FRigTreeElement> CurrentItem);
+	TSharedPtr<FRigHierarchyTreeElement> FindElement(const FRigHierarchyKey& InElementKey) const;
+	static TSharedPtr<FRigHierarchyTreeElement> FindElement(const FRigHierarchyKey& InElementKey, TSharedPtr<FRigHierarchyTreeElement> CurrentItem);
 	bool AddElement(FRigHierarchyKey InKey, FRigHierarchyKey InParentKey = FRigHierarchyKey());
 	bool AddElement(const FRigBaseElement* InElement);
 	bool AddComponent(const FRigBaseComponent* InComponent);
@@ -77,36 +77,36 @@ public:
 	bool ReparentElement(FRigHierarchyKey InKey, FRigHierarchyKey InParentKey);
 	bool RemoveElement(FRigHierarchyKey InKey);
 	void RefreshTreeView(bool bRebuildContent = true);
-	void SetExpansionRecursive(TSharedPtr<FRigTreeElement> InElement, bool bTowardsParent, bool bShouldBeExpanded);
-	TSharedRef<ITableRow> MakeTableRowWidget(TSharedPtr<FRigTreeElement> InItem, const TSharedRef<STableViewBase>& OwnerTable, bool bPinned);
-	void HandleGetChildrenForTree(TSharedPtr<FRigTreeElement> InItem, TArray<TSharedPtr<FRigTreeElement>>& OutChildren);
+	void SetExpansionRecursive(TSharedPtr<FRigHierarchyTreeElement> InElement, bool bTowardsParent, bool bShouldBeExpanded);
+	TSharedRef<ITableRow> MakeTableRowWidget(TSharedPtr<FRigHierarchyTreeElement> InItem, const TSharedRef<STableViewBase>& OwnerTable, bool bPinned);
+	void HandleGetChildrenForTree(TSharedPtr<FRigHierarchyTreeElement> InItem, TArray<TSharedPtr<FRigHierarchyTreeElement>>& OutChildren);
 	void OnElementKeyTagDragDetected(const FRigElementKey& InDraggedTag);
 
 	TArray<FRigHierarchyKey> GetSelectedKeys() const;
-	const TArray<TSharedPtr<FRigTreeElement>>& GetRootElements() const { return RootElements; }
-	FRigTreeDelegates& GetRigTreeDelegates() { return Delegates; }
+	const TArray<TSharedPtr<FRigHierarchyTreeElement>>& GetRootElements() const { return RootElements; }
+	FRigHierarchyTreeDelegates& GetRigTreeDelegates() { return Delegates; }
 
 	/** Given a position, return the item under that position. If nothing is there, return null. */
-	const TSharedPtr<FRigTreeElement>* FindItemAtPosition(FVector2D InScreenSpacePosition) const;
+	const TSharedPtr<FRigHierarchyTreeElement>* FindItemAtPosition(FVector2D InScreenSpacePosition) const;
 
 private:
 
-	void AddConnectorResolveWarningTag(TSharedPtr<FRigTreeElement> InTreeElement, const FRigBaseElement* InRigElement, const URigHierarchy* InHierarchy);
-	FText GetConnectorWarningMessage(TSharedPtr<FRigTreeElement> InTreeElement, TWeakObjectPtr<UControlRig> InControlRigPtr, const FRigElementKey InConnectorKey) const;
+	void AddConnectorResolveWarningTag(TSharedPtr<FRigHierarchyTreeElement> InTreeElement, const FRigBaseElement* InRigElement, const URigHierarchy* InHierarchy);
+	FText GetConnectorWarningMessage(TSharedPtr<FRigHierarchyTreeElement> InTreeElement, TWeakObjectPtr<UControlRig> InControlRigPtr, const FRigElementKey InConnectorKey) const;
 
 	/** A temporary snapshot of the SparseItemInfos in STreeView, used during RefreshTreeView() */
 	TSparseItemMap OldSparseItemInfos;
 
 	/** Backing array for tree view */
-	TArray<TSharedPtr<FRigTreeElement>> RootElements;
+	TArray<TSharedPtr<FRigHierarchyTreeElement>> RootElements;
 	
 	/** A map for looking up items based on their key */
-	TMap<FRigHierarchyKey, TSharedPtr<FRigTreeElement>> ElementMap;
+	TMap<FRigHierarchyKey, TSharedPtr<FRigHierarchyTreeElement>> ElementMap;
 
 	/** A map for looking up a parent based on their key */
 	TMap<FRigHierarchyKey, FRigHierarchyKey> ParentMap;
 
-	FRigTreeDelegates Delegates;
+	FRigHierarchyTreeDelegates Delegates;
 
 	bool bAutoScrollEnabled;
 	FVector2D LastMousePosition;
diff --git a/Engine/Plugins/Animation/ControlRig/Source/ControlRigEditor/Private/Editor/Hierarchy/Widgets/SModularRigHierarchy.cpp b/Engine/Plugins/Animation/ControlRig/Source/ControlRigEditor/Private/Editor/Hierarchy/Widgets/SModularRigHierarchy.cpp
--- a/Engine/Plugins/Animation/ControlRig/Source/ControlRigEditor/Private/Editor/Hierarchy/Widgets/SModularRigHierarchy.cpp
+++ b/Engine/Plugins/Animation/ControlRig/Source/ControlRigEditor/Private/Editor/Hierarchy/Widgets/SModularRigHierarchy.cpp
@@ -65,32 +65,32 @@
 #include "Editor/SRigVMDetailsInspector.h"
 #endif
 
-#define LOCTEXT_NAMESPACE "SModularRigModel"
+#define LOCTEXT_NAMESPACE "SModularRigHierarchy"
 
 ///////////////////////////////////////////////////////////
 
-const FName SModularRigModel::ContextMenuName = TEXT("ControlRigEditor.ModularRigModel.ContextMenu");
+const FName SModularRigHierarchy::ContextMenuName = TEXT("ControlRigEditor.ModularRigModel.ContextMenu");
 
-SModularRigModel::~SModularRigModel()
+SModularRigHierarchy::~SModularRigHierarchy()
 {
 	IControlRigBaseEditor* Editor = ControlRigEditor.IsValid() ? ControlRigEditor.Pin().Get() : nullptr;
 	OnEditorClose(Editor, ControlRigBlueprint.GetObject());
 }
 
-void SModularRigModel::Construct(const FArguments& InArgs, TSharedRef<IControlRigBaseEditor> InControlRigEditor)
+void SModularRigHierarchy::Construct(const FArguments& InArgs, TSharedRef<IControlRigBaseEditor> InControlRigEditor)
 {
 	ControlRigEditor = InControlRigEditor;
 
 	ControlRigBlueprint = ControlRigEditor.Pin()->GetControlRigAssetInterface();
 
-	ControlRigBlueprint->GetRigVMAssetInterface()->OnRefreshEditor().AddRaw(this, &SModularRigModel::HandleRefreshEditorFromBlueprint);
-	ControlRigBlueprint->GetRigVMAssetInterface()->OnSetObjectBeingDebugged().AddRaw(this, &SModularRigModel::HandleSetObjectBeingDebugged);
-	ControlRigBlueprint->OnModularRigPreCompiled().AddRaw(this, &SModularRigModel::HandlePreCompileModularRigs);
-	ControlRigBlueprint->OnModularRigCompiled().AddRaw(this, &SModularRigModel::HandlePostCompileModularRigs);
+	ControlRigBlueprint->GetRigVMAssetInterface()->OnRefreshEditor().AddRaw(this, &SModularRigHierarchy::HandleRefreshEditorFromBlueprint);
+	ControlRigBlueprint->GetRigVMAssetInterface()->OnSetObjectBeingDebugged().AddRaw(this, &SModularRigHierarchy::HandleSetObjectBeingDebugged);
+	ControlRigBlueprint->OnModularRigPreCompiled().AddRaw(this, &SModularRigHierarchy::HandlePreCompileModularRigs);
+	ControlRigBlueprint->OnModularRigCompiled().AddRaw(this, &SModularRigHierarchy::HandlePostCompileModularRigs);
 
 	if(UModularRigController* ModularRigController = ControlRigBlueprint->GetModularRigController())
 	{
-		ModularRigController->OnModified().AddSP(this, &SModularRigModel::OnModularRigModified);
+		ModularRigController->OnModified().AddSP(this, &SModularRigHierarchy::OnModularRigModified);
 	}
 
 	// for deleting, renaming, dragging
@@ -111,21 +111,21 @@ void SModularRigModel::Construct(const FArguments& InArgs, TSharedRef<IControlRi
 	bKeepCurrentEditedConnectors = false;
 
 	// setup all delegates for the modular rig model widget
-	FModularRigTreeDelegates Delegates;
-	Delegates.OnGetModularRig = FOnGetModularRigTreeRig::CreateSP(this, &SModularRigModel::GetModularRigForTreeView);
-	Delegates.OnContextMenuOpening = FOnContextMenuOpening::CreateSP(this, &SModularRigModel::CreateContextMenuWidget);
-	Delegates.OnDragDetected = FOnDragDetected::CreateSP(this, &SModularRigModel::OnDragDetected);
-	Delegates.OnCanAcceptDrop = FOnModularRigTreeCanAcceptDrop::CreateSP(this, &SModularRigModel::OnCanAcceptDrop);
-	Delegates.OnAcceptDrop = FOnModularRigTreeAcceptDrop::CreateSP(this, &SModularRigModel::OnAcceptDrop);
-	Delegates.OnMouseButtonClick = FOnModularRigTreeMouseButtonClick::CreateSP(this, &SModularRigModel::OnItemClicked);
-	Delegates.OnMouseButtonDoubleClick = FOnModularRigTreeMouseButtonClick::CreateSP(this, &SModularRigModel::OnItemDoubleClicked);
-	Delegates.OnRequestDetailsInspection = FOnModularRigTreeRequestDetailsInspection::CreateSP(this, &SModularRigModel::OnRequestDetailsInspection);
-	Delegates.OnRenameElement = FOnModularRigTreeRenameElement::CreateSP(this, &SModularRigModel::HandleRenameModule);
-	Delegates.OnVerifyModuleNameChanged = FOnModularRigTreeVerifyElementNameChanged::CreateSP(this, &SModularRigModel::HandleVerifyNameChanged);
-	Delegates.OnResolveConnector = FOnModularRigTreeResolveConnector::CreateSP(this, &SModularRigModel::HandleConnectorResolved);
-	Delegates.OnDisconnectConnector = FOnModularRigTreeDisconnectConnector::CreateSP(this, &SModularRigModel::HandleConnectorDisconnect);
-	Delegates.OnSelectionChanged = FOnModularRigTreeSelectionChanged::CreateSP(this, &SModularRigModel::HandleSelectionChanged);
-	Delegates.OnAlwaysShowConnector = FOnModularRigTreeAlwaysShowConnector::CreateSP(this, &SModularRigModel::ShouldAlwaysShowConnector);
+	FModularRigHierarchyTreeDelegates Delegates;
+	Delegates.OnGetModularRig = FOnGetModularRigTreeRig::CreateSP(this, &SModularRigHierarchy::GetModularRigForTreeView);
+	Delegates.OnContextMenuOpening = FOnContextMenuOpening::CreateSP(this, &SModularRigHierarchy::CreateContextMenuWidget);
+	Delegates.OnDragDetected = FOnDragDetected::CreateSP(this, &SModularRigHierarchy::OnDragDetected);
+	Delegates.OnCanAcceptDrop = FOnModularRigTreeCanAcceptDrop::CreateSP(this, &SModularRigHierarchy::OnCanAcceptDrop);
+	Delegates.OnAcceptDrop = FOnModularRigTreeAcceptDrop::CreateSP(this, &SModularRigHierarchy::OnAcceptDrop);
+	Delegates.OnMouseButtonClick = FOnModularRigTreeMouseButtonClick::CreateSP(this, &SModularRigHierarchy::OnItemClicked);
+	Delegates.OnMouseButtonDoubleClick = FOnModularRigTreeMouseButtonClick::CreateSP(this, &SModularRigHierarchy::OnItemDoubleClicked);
+	Delegates.OnRequestDetailsInspection = FOnModularRigTreeRequestDetailsInspection::CreateSP(this, &SModularRigHierarchy::OnRequestDetailsInspection);
+	Delegates.OnRenameElement = FOnModularRigTreeRenameElement::CreateSP(this, &SModularRigHierarchy::HandleRenameModule);
+	Delegates.OnVerifyModuleNameChanged = FOnModularRigTreeVerifyElementNameChanged::CreateSP(this, &SModularRigHierarchy::HandleVerifyNameChanged);
+	Delegates.OnResolveConnector = FOnModularRigTreeResolveConnector::CreateSP(this, &SModularRigHierarchy::HandleConnectorResolved);
+	Delegates.OnDisconnectConnector = FOnModularRigTreeDisconnectConnector::CreateSP(this, &SModularRigHierarchy::HandleConnectorDisconnect);
+	Delegates.OnSelectionChanged = FOnModularRigTreeSelectionChanged::CreateSP(this, &SModularRigHierarchy::HandleSelectionChanged);
+	Delegates.OnAlwaysShowConnector = FOnModularRigTreeAlwaysShowConnector::CreateSP(this, &SModularRigHierarchy::ShouldAlwaysShowConnector);
 
 	HeaderRowWidget = SNew(SHeaderRow)
 		.Visibility(EVisibility::Visible);
@@ -180,7 +180,7 @@ void SModularRigModel::Construct(const FArguments& InArgs, TSharedRef<IControlRi
 			   .ComboButtonStyle(&FAppStyle::Get().GetWidgetStyle<FComboButtonStyle>("SimpleComboButtonWithIcon"))
 			   .ForegroundColor(FSlateColor::UseStyle())
 			   .ToolTipText(LOCTEXT("OptionsToolTip", "Open the Options Menu ."))
-			   .OnGetMenuContent(this, &SModularRigModel::OnGetOptionsMenu)
+			   .OnGetMenuContent(this, &SModularRigHierarchy::OnGetOptionsMenu)
 			   .ContentPadding(FMargin(1, 0))
 			   .ButtonContent()
 			   [
@@ -195,7 +195,7 @@ void SModularRigModel::Construct(const FArguments& InArgs, TSharedRef<IControlRi
 			.HAlign(HAlign_Fill)
 			[
 				SAssignNew(FilterBox, SSearchBox)
-				.OnTextChanged(this, &SModularRigModel::OnFilterTextChanged)
+				.OnTextChanged(this, &SModularRigHierarchy::OnFilterTextChanged)
 			]
 		]
 		
@@ -230,9 +230,9 @@ void SModularRigModel::Construct(const FArguments& InArgs, TSharedRef<IControlRi
 		ControlRigEditor.Pin()->GetKeyDownDelegate().BindLambda([&](const FGeometry& MyGeometry, const FKeyEvent& InKeyEvent)->FReply {
 			return OnKeyDown(MyGeometry, InKeyEvent);
 		});
-		ControlRigEditor.Pin()->OnGetViewportContextMenu().BindSP(this, &SModularRigModel::GetContextMenu);
-		ControlRigEditor.Pin()->OnViewportContextMenuCommands().BindSP(this, &SModularRigModel::GetContextMenuCommands);
-		ControlRigEditor.Pin()->OnEditorClosed().AddSP(this, &SModularRigModel::OnEditorClose);
+		ControlRigEditor.Pin()->OnGetViewportContextMenu().BindSP(this, &SModularRigHierarchy::GetContextMenu);
+		ControlRigEditor.Pin()->OnViewportContextMenuCommands().BindSP(this, &SModularRigHierarchy::GetContextMenuCommands);
+		ControlRigEditor.Pin()->OnEditorClosed().AddSP(this, &SModularRigHierarchy::OnEditorClose);
 	}
 	
 	CreateContextMenu();
@@ -241,12 +241,12 @@ void SModularRigModel::Construct(const FArguments& InArgs, TSharedRef<IControlRi
 	{
 		if(URigHierarchy* Hierarchy = Rig->GetHierarchy())
 		{
-			Hierarchy->OnModified().AddSP(this, &SModularRigModel::OnHierarchyModified);
+			Hierarchy->OnModified().AddSP(this, &SModularRigHierarchy::OnHierarchyModified);
 		}
 	}
 }
 
-void SModularRigModel::OnEditorClose(IControlRigBaseEditor* InEditor, FControlRigAssetInterfacePtr InBlueprint)
+void SModularRigHierarchy::OnEditorClose(IControlRigBaseEditor* InEditor, FControlRigAssetInterfacePtr InBlueprint)
 {
 	if (InEditor)
 	{
@@ -280,45 +280,45 @@ void SModularRigModel::OnEditorClose(IControlRigBaseEditor* InEditor, FControlRi
 	ControlRigBlueprint.Reset();
 }
 
-void SModularRigModel::BindCommands()
+void SModularRigHierarchy::BindCommands()
 {
 	// create new command
 	const FControlRigModularRigCommands& Commands = FControlRigModularRigCommands::Get();
 
 	CommandList->MapAction(Commands.AddModuleItem,
-		FExecuteAction::CreateSP(this, &SModularRigModel::HandleNewItem),
+		FExecuteAction::CreateSP(this, &SModularRigHierarchy::HandleNewItem),
 		FCanExecuteAction());
 
 	CommandList->MapAction(Commands.RenameModuleItem,
-		FExecuteAction::CreateSP(this, &SModularRigModel::HandleRenameModule),
+		FExecuteAction::CreateSP(this, &SModularRigHierarchy::HandleRenameModule),
 		FCanExecuteAction());
 
 	CommandList->MapAction(Commands.DeleteModuleItem,
-		FExecuteAction::CreateSP(this, &SModularRigModel::HandleDeleteModules),
+		FExecuteAction::CreateSP(this, &SModularRigHierarchy::HandleDeleteModules),
 		FCanExecuteAction());
 
 	CommandList->MapAction(Commands.MirrorModuleItem,
-		FExecuteAction::CreateSP(this, &SModularRigModel::HandleMirrorModules),
+		FExecuteAction::CreateSP(this, &SModularRigHierarchy::HandleMirrorModules),
 		FCanExecuteAction());
 
 	CommandList->MapAction(Commands.ReresolveModuleItem,
-		FExecuteAction::CreateSP(this, &SModularRigModel::HandleReresolveModules),
+		FExecuteAction::CreateSP(this, &SModularRigHierarchy::HandleReresolveModules),
 		FCanExecuteAction());
 
 	CommandList->MapAction(Commands.SwapModuleClassItem,
-		FExecuteAction::CreateSP(this, &SModularRigModel::HandleSwapClassForModules),
-		FCanExecuteAction::CreateSP(this, &SModularRigModel::CanSwapModules));
+		FExecuteAction::CreateSP(this, &SModularRigHierarchy::HandleSwapClassForModules),
+		FCanExecuteAction::CreateSP(this, &SModularRigHierarchy::CanSwapModules));
 
 	CommandList->MapAction(Commands.CopyModuleSettings,
-		FExecuteAction::CreateSP(this, &SModularRigModel::HandleCopyModuleSettings),
-		FCanExecuteAction::CreateSP(this, &SModularRigModel::CanCopyModuleSettings));
+		FExecuteAction::CreateSP(this, &SModularRigHierarchy::HandleCopyModuleSettings),
+		FCanExecuteAction::CreateSP(this, &SModularRigHierarchy::CanCopyModuleSettings));
 
 	CommandList->MapAction(Commands.PasteModuleSettings,
-		FExecuteAction::CreateSP(this, &SModularRigModel::HandlePasteModuleSettings),
-		FCanExecuteAction::CreateSP(this, &SModularRigModel::CanPasteModuleSettings));
+		FExecuteAction::CreateSP(this, &SModularRigHierarchy::HandlePasteModuleSettings),
+		FCanExecuteAction::CreateSP(this, &SModularRigHierarchy::CanPasteModuleSettings));
 }
 
-FReply SModularRigModel::OnKeyDown(const FGeometry& MyGeometry, const FKeyEvent& InKeyEvent)
+FReply SModularRigHierarchy::OnKeyDown(const FGeometry& MyGeometry, const FKeyEvent& InKeyEvent)
 {
 	if (CommandList.IsValid() && CommandList->ProcessCommandBindings(InKeyEvent))
 	{
@@ -327,7 +327,7 @@ FReply SModularRigModel::OnKeyDown(const FGeometry& MyGeometry, const FKeyEvent&
 	return FReply::Unhandled();
 }
 
-FReply SModularRigModel::OnMouseButtonDown(const FGeometry& MyGeometry, const FPointerEvent& MouseEvent)
+FReply SModularRigHierarchy::OnMouseButtonDown(const FGeometry& MyGeometry, const FPointerEvent& MouseEvent)
 {
 	const FReply Reply = SCompoundWidget::OnMouseButtonDown(MyGeometry, MouseEvent);
 	if(Reply.IsEventHandled())
@@ -337,9 +337,9 @@ FReply SModularRigModel::OnMouseButtonDown(const FGeometry& MyGeometry, const FP
 
 	if(MouseEvent.GetEffectingButton() == EKeys::MiddleMouseButton)
 	{
-		if(const TSharedPtr<FModularRigTreeElement>* ItemPtr = TreeView->FindItemAtPosition(MouseEvent.GetScreenSpacePosition()))
+		if(const TSharedPtr<FModularRigHierarchyTreeElement>* ItemPtr = TreeView->FindItemAtPosition(MouseEvent.GetScreenSpacePosition()))
 		{
-			if(const TSharedPtr<FModularRigTreeElement>& Item = *ItemPtr)
+			if(const TSharedPtr<FModularRigHierarchyTreeElement>& Item = *ItemPtr)
 			{
 				if (ControlRigBlueprint.IsValid())
 				{
@@ -369,7 +369,7 @@ FReply SModularRigModel::OnMouseButtonDown(const FGeometry& MyGeometry, const FP
 	return FReply::Unhandled();
 }
 
-void SModularRigModel::RefreshTreeView(bool bRebuildContent)
+void SModularRigHierarchy::RefreshTreeView(bool bRebuildContent)
 {
 	bool bDummySuspensionFlag = false;
 	bool* SuspensionFlagPtr = &bDummySuspensionFlag;
@@ -382,19 +382,19 @@ void SModularRigModel::RefreshTreeView(bool bRebuildContent)
 	TreeView->RefreshTreeView(bRebuildContent);
 }
 
-TArray<TSharedPtr<FModularRigTreeElement>> SModularRigModel::GetSelectedItems() const
+TArray<TSharedPtr<FModularRigHierarchyTreeElement>> SModularRigHierarchy::GetSelectedItems() const
 {
-	TArray<TSharedPtr<FModularRigTreeElement>> SelectedItems = TreeView->GetSelectedItems();
-	SelectedItems.Remove(TSharedPtr<FModularRigTreeElement>(nullptr));
+	TArray<TSharedPtr<FModularRigHierarchyTreeElement>> SelectedItems = TreeView->GetSelectedItems();
+	SelectedItems.Remove(TSharedPtr<FModularRigHierarchyTreeElement>(nullptr));
 	return SelectedItems;
 }
 
-TArray<FString> SModularRigModel::GetSelectedKeys() const
+TArray<FString> SModularRigHierarchy::GetSelectedKeys() const
 {
-	const TArray<TSharedPtr<FModularRigTreeElement>> SelectedItems = GetSelectedItems();
+	const TArray<TSharedPtr<FModularRigHierarchyTreeElement>> SelectedItems = GetSelectedItems();
 	
 	TArray<FString> SelectedKeys;
-	for (const TSharedPtr<FModularRigTreeElement>& SelectedItem : SelectedItems)
+	for (const TSharedPtr<FModularRigHierarchyTreeElement>& SelectedItem : SelectedItems)
 	{
 		if (SelectedItem.IsValid())
 		{
@@ -409,11 +409,11 @@ TArray<FString> SModularRigModel::GetSelectedKeys() const
 }
 
 
-void SModularRigModel::HandlePreCompileModularRigs(FRigVMAssetInterfacePtr InBlueprint)
+void SModularRigHierarchy::HandlePreCompileModularRigs(FRigVMAssetInterfacePtr InBlueprint)
 {
 }
 
-void SModularRigModel::HandlePostCompileModularRigs(FRigVMAssetInterfacePtr InBlueprint)
+void SModularRigHierarchy::HandlePostCompileModularRigs(FRigVMAssetInterfacePtr InBlueprint)
 {
 	if(!bKeepCurrentEditedConnectors)
 	{
@@ -423,7 +423,7 @@ void SModularRigModel::HandlePostCompileModularRigs(FRigVMAssetInterfacePtr InBl
 	RefreshTreeView();
 	if (ControlRigEditor.IsValid())
 	{
-		TArray<TSharedPtr<FModularRigTreeElement>> SelectedElements;
+		TArray<TSharedPtr<FModularRigHierarchyTreeElement>> SelectedElements;
 		Algo::Transform(ControlRigEditor.Pin()->GetSelectedModules(), SelectedElements, [this](const FName& ModuleName)
 		{
 			return TreeView->FindElement(ModuleName.ToString());
@@ -433,12 +433,12 @@ void SModularRigModel::HandlePostCompileModularRigs(FRigVMAssetInterfacePtr InBl
 	}
 }
 
-void SModularRigModel::HandleRefreshEditorFromBlueprint(FRigVMAssetInterfacePtr InBlueprint)
+void SModularRigHierarchy::HandleRefreshEditorFromBlueprint(FRigVMAssetInterfacePtr InBlueprint)
 {
 	RefreshTreeView();
 }
 
-void SModularRigModel::HandleSetObjectBeingDebugged(UObject* InObject)
+void SModularRigHierarchy::HandleSetObjectBeingDebugged(UObject* InObject)
 {
 	if(ControlRigBeingDebuggedPtr.Get() == InObject)
 	{
@@ -461,14 +461,14 @@ void SModularRigModel::HandleSetObjectBeingDebugged(UObject* InObject)
 
 		if(URigHierarchy* Hierarchy = ControlRig->GetHierarchy())
 		{
-			Hierarchy->OnModified().AddSP(this, &SModularRigModel::OnHierarchyModified);
+			Hierarchy->OnModified().AddSP(this, &SModularRigHierarchy::OnHierarchyModified);
 		}
 	}
 
 	RefreshTreeView();
 }
 
-TSharedRef<SWidget> SModularRigModel::OnGetOptionsMenu()
+TSharedRef<SWidget> SModularRigHierarchy::OnGetOptionsMenu()
 {
 	FMenuBuilder MenuBuilder(true, nullptr);
 
@@ -538,13 +538,13 @@ TSharedRef<SWidget> SModularRigModel::OnGetOptionsMenu()
 	return MenuBuilder.MakeWidget();
 }
 
-void SModularRigModel::OnFilterTextChanged(const FText& SearchText)
+void SModularRigHierarchy::OnFilterTextChanged(const FText& SearchText)
 {
 	FilterText = SearchText;
 	RefreshTreeView(true);
 }
 
-TSharedPtr< SWidget > SModularRigModel::CreateContextMenuWidget()
+TSharedPtr< SWidget > SModularRigHierarchy::CreateContextMenuWidget()
 {
 	UToolMenus* ToolMenus = UToolMenus::Get();
 
@@ -556,16 +556,16 @@ TSharedPtr< SWidget > SModularRigModel::CreateContextMenuWidget()
 	return SNullWidget::NullWidget;
 }
 
-void SModularRigModel::OnItemClicked(TSharedPtr<FModularRigTreeElement> InItem)
+void SModularRigHierarchy::OnItemClicked(TSharedPtr<FModularRigHierarchyTreeElement> InItem)
 {
 }
 
-void SModularRigModel::OnItemDoubleClicked(TSharedPtr<FModularRigTreeElement> InItem)
+void SModularRigHierarchy::OnItemDoubleClicked(TSharedPtr<FModularRigHierarchyTreeElement> InItem)
 {
 
 }
 
-void SModularRigModel::CreateContextMenu()
+void SModularRigHierarchy::CreateContextMenu()
 {
 	static bool bCreatedMenu = false;
 	if(bCreatedMenu)
@@ -589,7 +589,7 @@ void SModularRigModel::CreateContextMenu()
 			{
 				UControlRigContextMenuContext* MainContext = InMenu->FindContext<UControlRigContextMenuContext>();
 				
-				if (SModularRigModel* ModelPanel = MainContext->GetModularRigModelPanel())
+				if (SModularRigHierarchy* ModelPanel = MainContext->GetModularRigModelPanel())
 				{
 					const FControlRigModularRigCommands& Commands = FControlRigModularRigCommands::Get(); 
 				
@@ -614,7 +614,7 @@ void SModularRigModel::CreateContextMenu()
 	}
 }
 
-UToolMenu* SModularRigModel::GetContextMenu()
+UToolMenu* SModularRigHierarchy::GetContextMenu()
 {
 	const FName MenuName = ContextMenuName;
 	UToolMenus* ToolMenus = UToolMenus::Get();
@@ -638,12 +638,12 @@ UToolMenu* SModularRigModel::GetContextMenu()
 	return Menu;
 }
 
-TSharedPtr<FUICommandList> SModularRigModel::GetContextMenuCommands() const
+TSharedPtr<FUICommandList> SModularRigHierarchy::GetContextMenuCommands() const
 {
 	return CommandList;
 }
 
-bool SModularRigModel::IsSingleSelected() const
+bool SModularRigHierarchy::IsSingleSelected() const
 {
 	if(GetSelectedKeys().Num() == 1)
 	{
@@ -716,7 +716,7 @@ private:
 };
 
 /** Create Item */
-void SModularRigModel::HandleNewItem()
+void SModularRigHierarchy::HandleNewItem()
 {
 	if(!ControlRigEditor.IsValid())
 	{
@@ -726,7 +726,7 @@ void SModularRigModel::HandleNewItem()
 	FName ParentModuleName = NAME_None;
 	if (IsSingleSelected())
 	{
-		const TSharedPtr<FModularRigTreeElement> ParentElement = TreeView->FindElement(GetSelectedKeys()[0]);
+		const TSharedPtr<FModularRigHierarchyTreeElement> ParentElement = TreeView->FindElement(GetSelectedKeys()[0]);
 		if (ParentElement.IsValid())
 		{
 			ParentModuleName = ParentElement->ModuleName;
@@ -750,7 +750,7 @@ void SModularRigModel::HandleNewItem()
 	}
 }
 
-void SModularRigModel::HandleNewItem(UClass* InClass, const FName &InParentModuleName)
+void SModularRigHierarchy::HandleNewItem(UClass* InClass, const FName &InParentModuleName)
 {
 	UControlRig* ControlRig = InClass->GetDefaultObject<UControlRig>();
 	if (!ControlRig)
@@ -768,7 +768,7 @@ void SModularRigModel::HandleNewItem(UClass* InClass, const FName &InParentModul
 		ClassName.RemoveFromEnd(TEXT("_C"));
 		const FRigName Name = Controller->GetSafeNewName(FRigName(ClassName));
 		const FName NewModuleName = Controller->AddModule(Name, InClass, InParentModuleName);
-		TSharedPtr<FModularRigTreeElement> Element = TreeView->FindElement(NewModuleName.ToString());
+		TSharedPtr<FModularRigHierarchyTreeElement> Element = TreeView->FindElement(NewModuleName.ToString());
 		if (Element.IsValid())
 		{
 			TreeView->SetSelection({Element});
@@ -777,12 +777,12 @@ void SModularRigModel::HandleNewItem(UClass* InClass, const FName &InParentModul
 	}
 }
 
-bool SModularRigModel::CanRenameModule() const
+bool SModularRigHierarchy::CanRenameModule() const
 {
 	return IsSingleSelected() && TreeView->FindElement(GetSelectedKeys()[0])->bIsPrimary;
 }
 
-void SModularRigModel::HandleRenameModule()
+void SModularRigHierarchy::HandleRenameModule()
 {
 	if(!ControlRigEditor.IsValid())
 	{
@@ -799,7 +799,7 @@ void SModularRigModel::HandleRenameModule()
 	{
 		FScopedTransaction Transaction(LOCTEXT("ModularRigModelRenameSelected", "Rename selected module"));
 
-		const TArray<TSharedPtr<FModularRigTreeElement>> SelectedItems = GetSelectedItems();
+		const TArray<TSharedPtr<FModularRigHierarchyTreeElement>> SelectedItems = GetSelectedItems();
 		if (SelectedItems.Num() == 1)
 		{
 			SelectedItems[0]->RequestRename();
@@ -809,7 +809,7 @@ void SModularRigModel::HandleRenameModule()
 	return;
 }
 
-FName SModularRigModel::HandleRenameModule(const FName& InOldModuleName, const FName& InNewName)
+FName SModularRigHierarchy::HandleRenameModule(const FName& InOldModuleName, const FName& InNewName)
 {
 	if (ControlRigBlueprint.IsValid())
 	{
@@ -828,7 +828,7 @@ FName SModularRigModel::HandleRenameModule(const FName& InOldModuleName, const F
 	return NAME_None;
 }
 
-bool SModularRigModel::HandleVerifyNameChanged(const FName& InOldModuleName, const FName& InNewName, FText& OutErrorMessage)
+bool SModularRigHierarchy::HandleVerifyNameChanged(const FName& InOldModuleName, const FName& InNewName, FText& OutErrorMessage)
 {
 	if (InNewName.IsNone())
 	{
@@ -846,7 +846,7 @@ bool SModularRigModel::HandleVerifyNameChanged(const FName& InOldModuleName, con
 	return false;
 }
 
-void SModularRigModel::HandleDeleteModules()
+void SModularRigHierarchy::HandleDeleteModules()
 {
 	if(!ControlRigEditor.IsValid())
 	{
@@ -858,9 +858,9 @@ void SModularRigModel::HandleDeleteModules()
 	{
 		FScopedTransaction Transaction(LOCTEXT("ModularRigModelDeleteSelected", "Delete selected modules"));
 
-		const TArray<TSharedPtr<FModularRigTreeElement>> SelectedItems = GetSelectedItems();
+		const TArray<TSharedPtr<FModularRigHierarchyTreeElement>> SelectedItems = GetSelectedItems();
 		TArray<FName> SelectedModuleNames;
-		Algo::Transform(SelectedItems, SelectedModuleNames, [](const TSharedPtr<FModularRigTreeElement>& Element)
+		Algo::Transform(SelectedItems, SelectedModuleNames, [](const TSharedPtr<FModularRigHierarchyTreeElement>& Element)
 		{
 			if (Element.IsValid())
 			{
@@ -874,7 +874,7 @@ void SModularRigModel::HandleDeleteModules()
 	return;
 }
 
-void SModularRigModel::HandleDeleteModules(const TArray<FName>& InModuleNames)
+void SModularRigHierarchy::HandleDeleteModules(const TArray<FName>& InModuleNames)
 {
 	if (ControlRigBlueprint.IsValid())
 	{
@@ -893,7 +893,7 @@ void SModularRigModel::HandleDeleteModules(const TArray<FName>& InModuleNames)
 	}
 }
 
-void SModularRigModel::HandleReparentModules(const TArray<FName>& InModuleNames, const FName& InParentModuleName, int32 NewModuleIndex)
+void SModularRigHierarchy::HandleReparentModules(const TArray<FName>& InModuleNames, const FName& InParentModuleName, int32 NewModuleIndex)
 {
 	if (ControlRigBlueprint.IsValid())
 	{
@@ -913,7 +913,7 @@ void SModularRigModel::HandleReparentModules(const TArray<FName>& InModuleNames,
 	}
 }
 
-void SModularRigModel::HandleMirrorModules()
+void SModularRigHierarchy::HandleMirrorModules()
 {
 	if(!ControlRigEditor.IsValid())
 	{
@@ -923,9 +923,9 @@ void SModularRigModel::HandleMirrorModules()
 	UModularRig* Rig = GetDefaultModularRig();
 	if (Rig)
 	{
-		const TArray<TSharedPtr<FModularRigTreeElement>> SelectedItems = GetSelectedItems();
+		const TArray<TSharedPtr<FModularRigHierarchyTreeElement>> SelectedItems = GetSelectedItems();
 		TArray<FName> SelectedModuleNames;
-		Algo::Transform(SelectedItems, SelectedModuleNames, [](const TSharedPtr<FModularRigTreeElement>& Element)
+		Algo::Transform(SelectedItems, SelectedModuleNames, [](const TSharedPtr<FModularRigHierarchyTreeElement>& Element)
 		{
 			if (Element.IsValid())
 			{
@@ -937,7 +937,7 @@ void SModularRigModel::HandleMirrorModules()
 	}
 }
 
-void SModularRigModel::HandleMirrorModules(const TArray<FName>& InModuleNames)
+void SModularRigHierarchy::HandleMirrorModules(const TArray<FName>& InModuleNames)
 {
 	if (ControlRigBlueprint.IsValid())
 	{
@@ -979,7 +979,7 @@ void SModularRigModel::HandleMirrorModules(const TArray<FName>& InModuleNames)
 	}
 }
 
-void SModularRigModel::HandleReresolveModules()
+void SModularRigHierarchy::HandleReresolveModules()
 {
 	if(!ControlRigEditor.IsValid())
 	{
@@ -989,9 +989,9 @@ void SModularRigModel::HandleReresolveModules()
 	UModularRig* Rig = GetDefaultModularRig();
 	if (Rig)
 	{
-		const TArray<TSharedPtr<FModularRigTreeElement>> SelectedItems = GetSelectedItems();
+		const TArray<TSharedPtr<FModularRigHierarchyTreeElement>> SelectedItems = GetSelectedItems();
 		TArray<FName> SelectedModuleNames;
-		Algo::Transform(SelectedItems, SelectedModuleNames, [](const TSharedPtr<FModularRigTreeElement>& Element)
+		Algo::Transform(SelectedItems, SelectedModuleNames, [](const TSharedPtr<FModularRigHierarchyTreeElement>& Element)
 		{
 			if (Element.IsValid())
 			{
@@ -1007,7 +1007,7 @@ void SModularRigModel::HandleReresolveModules()
 	}
 }
 
-void SModularRigModel::HandleReresolveModules(const TArray<FName>& InModuleNames)
+void SModularRigHierarchy::HandleReresolveModules(const TArray<FName>& InModuleNames)
 {
 	if (ControlRigBlueprint.IsValid())
 	{
@@ -1077,7 +1077,7 @@ void SModularRigModel::HandleReresolveModules(const TArray<FName>& InModuleNames
 	}
 }
 
-bool SModularRigModel::CanSwapModules() const
+bool SModularRigHierarchy::CanSwapModules() const
 {
 	// Only if all modules selected have the same module class
 	if(!ControlRigEditor.IsValid())
@@ -1089,8 +1089,8 @@ bool SModularRigModel::CanSwapModules() const
 	if (Rig)
 	{
 		TSoftClassPtr<UControlRig> CommonClass = nullptr;
-		const TArray<TSharedPtr<FModularRigTreeElement>> SelectedItems = GetSelectedItems();
-		for (const TSharedPtr<FModularRigTreeElement>& SelectedItem : SelectedItems)
+		const TArray<TSharedPtr<FModularRigHierarchyTreeElement>> SelectedItems = GetSelectedItems();
+		for (const TSharedPtr<FModularRigHierarchyTreeElement>& SelectedItem : SelectedItems)
 		{
 			if(!SelectedItem.IsValid())
 			{
@@ -1122,7 +1122,7 @@ bool SModularRigModel::CanSwapModules() const
 	return false;
 }
 
-void SModularRigModel::HandleSwapClassForModules()
+void SModularRigHierarchy::HandleSwapClassForModules()
 {
 	if(!ControlRigEditor.IsValid())
 	{
@@ -1132,9 +1132,9 @@ void SModularRigModel::HandleSwapClassForModules()
 	UModularRig* Rig = GetDefaultModularRig();
 	if (Rig)
 	{
-		const TArray<TSharedPtr<FModularRigTreeElement>> SelectedItems = GetSelectedItems();
+		const TArray<TSharedPtr<FModularRigHierarchyTreeElement>> SelectedItems = GetSelectedItems();
 		TArray<FName> SelectedModuleNames;
-		Algo::Transform(SelectedItems, SelectedModuleNames, [](const TSharedPtr<FModularRigTreeElement>& Element)
+		Algo::Transform(SelectedItems, SelectedModuleNames, [](const TSharedPtr<FModularRigHierarchyTreeElement>& Element)
 		{
 			if (Element.IsValid())
 			{
@@ -1152,7 +1152,7 @@ void SModularRigModel::HandleSwapClassForModules()
 	}
 }
 
-void SModularRigModel::HandleSwapClassForModules(const TArray<FName>& InModuleNames)
+void SModularRigHierarchy::HandleSwapClassForModules(const TArray<FName>& InModuleNames)
 {
 	TArray<FSoftObjectPath> ModulePaths;
 	Algo::Transform(InModuleNames, ModulePaths, [this](const FName& InModuleName)
@@ -1224,7 +1224,7 @@ void SModularRigModel::HandleSwapClassForModules(const TArray<FName>& InModuleNa
 	SwapModulesDialog->ShowNormal();
 }
 
-bool SModularRigModel::CanCopyModuleSettings() const
+bool SModularRigHierarchy::CanCopyModuleSettings() const
 {
 	if(!ControlRigEditor.IsValid())
 	{
@@ -1233,13 +1233,13 @@ bool SModularRigModel::CanCopyModuleSettings() const
 	return !GetSelectedItems().IsEmpty();
 }
 
-void SModularRigModel::HandleCopyModuleSettings()
+void SModularRigHierarchy::HandleCopyModuleSettings()
 {
 	if (ControlRigBlueprint.IsValid())
 	{
 		TArray<FName> SelectedModuleNames;
-		const TArray<TSharedPtr<FModularRigTreeElement>> SelectedItems = GetSelectedItems();
-		for (const TSharedPtr<FModularRigTreeElement>& SelectedItem : SelectedItems)
+		const TArray<TSharedPtr<FModularRigHierarchyTreeElement>> SelectedItems = GetSelectedItems();
+		for (const TSharedPtr<FModularRigHierarchyTreeElement>& SelectedItem : SelectedItems)
 		{
 			if (SelectedItem.IsValid())
 			{
@@ -1258,7 +1258,7 @@ void SModularRigModel::HandleCopyModuleSettings()
 	}
 }
 
-bool SModularRigModel::CanPasteModuleSettings() const
+bool SModularRigHierarchy::CanPasteModuleSettings() const
 {
 	if(!ControlRigEditor.IsValid())
 	{
@@ -1283,7 +1283,7 @@ bool SModularRigModel::CanPasteModuleSettings() const
 	return GetSelectedItems().Num() == Content.Settings.Num();
 }
 
-void SModularRigModel::HandlePasteModuleSettings()
+void SModularRigHierarchy::HandlePasteModuleSettings()
 {
 	if (ControlRigBlueprint.IsValid())
 	{
@@ -1295,8 +1295,8 @@ void SModularRigModel::HandlePasteModuleSettings()
 		}
 
 		TArray<FName> SelectedModuleNames;
-		const TArray<TSharedPtr<FModularRigTreeElement>> SelectedItems = GetSelectedItems();
-		for (const TSharedPtr<FModularRigTreeElement>& SelectedItem : SelectedItems)
+		const TArray<TSharedPtr<FModularRigHierarchyTreeElement>> SelectedItems = GetSelectedItems();
+		for (const TSharedPtr<FModularRigHierarchyTreeElement>& SelectedItem : SelectedItems)
 		{
 			if (SelectedItem.IsValid())
 			{
@@ -1314,7 +1314,7 @@ void SModularRigModel::HandlePasteModuleSettings()
 
 }
 
-void SModularRigModel::HandleConnectorResolved(const FRigElementKey& InConnector, const TArray<FRigElementKey>& InTargets)
+void SModularRigHierarchy::HandleConnectorResolved(const FRigElementKey& InConnector, const TArray<FRigElementKey>& InTargets)
 {
 	if (ControlRigBlueprint.IsValid())
 	{
@@ -1336,7 +1336,7 @@ void SModularRigModel::HandleConnectorResolved(const FRigElementKey& InConnector
 	}
 }
 
-void SModularRigModel::HandleConnectorDisconnect(const FRigElementKey& InConnector)
+void SModularRigHierarchy::HandleConnectorDisconnect(const FRigElementKey& InConnector)
 {
 	if (ControlRigBlueprint.IsValid())
 	{
@@ -1355,7 +1355,7 @@ void SModularRigModel::HandleConnectorDisconnect(const FRigElementKey& InConnect
 	}
 }
 
-void SModularRigModel::HandleSelectionChanged(TSharedPtr<FModularRigTreeElement> Selection, ESelectInfo::Type SelectInfo)
+void SModularRigHierarchy::HandleSelectionChanged(TSharedPtr<FModularRigHierarchyTreeElement> Selection, ESelectInfo::Type SelectInfo)
 {
 	if(bIsPerformingSelection)
 	{
@@ -1375,12 +1375,12 @@ void SModularRigModel::HandleSelectionChanged(TSharedPtr<FModularRigTreeElement>
 	}
 }
 
-bool SModularRigModel::ShouldAlwaysShowConnector(const FName& InConnectorName) const
+bool SModularRigHierarchy::ShouldAlwaysShowConnector(const FName& InConnectorName) const
 {
 	return CurrentlyEditedConnectors.Contains(InConnectorName);
 }
 
-void SModularRigModel::OnModularRigModified(EModularRigNotification InNotif, const FRigModuleReference* InModule)
+void SModularRigHierarchy::OnModularRigModified(EModularRigNotification InNotif, const FRigModuleReference* InModule)
 {
 	if (!ControlRigBlueprint.IsValid())
 	{
@@ -1398,10 +1398,10 @@ void SModularRigModel::OnModularRigModified(EModularRigNotification InNotif, con
 				if(const UModularRigController* ModularRigController = ControlRigBlueprint->GetModularRigController())
 				{
 					const TArray<FName> SelectedModuleNames = ModularRigController->GetSelectedModules();
-					TArray<TSharedPtr<FModularRigTreeElement>> NewSelection;
+					TArray<TSharedPtr<FModularRigHierarchyTreeElement>> NewSelection;
 					for(const FName& SelectedModuleName : SelectedModuleNames)
 					{
-						if(const TSharedPtr<FModularRigTreeElement> Module = TreeView->FindElement(SelectedModuleName.ToString()))
+						if(const TSharedPtr<FModularRigHierarchyTreeElement> Module = TreeView->FindElement(SelectedModuleName.ToString()))
 						{
 							NewSelection.Add(Module);
 						}
@@ -1431,7 +1431,7 @@ void SModularRigModel::OnModularRigModified(EModularRigNotification InNotif, con
 	}
 }
 
-void SModularRigModel::OnHierarchyModified(ERigHierarchyNotification InNotif, URigHierarchy* InHierarchy, const FRigNotificationSubject& InSubject)
+void SModularRigHierarchy::OnHierarchyModified(ERigHierarchyNotification InNotif, URigHierarchy* InHierarchy, const FRigNotificationSubject& InSubject)
 {
 	if(!ControlRigBlueprint.IsValid())
 	{
@@ -1458,7 +1458,7 @@ void SModularRigModel::OnHierarchyModified(ERigHierarchyNotification InNotif, UR
 
 			if(!ModuleOrConnectorName.IsEmpty())
 			{
-				if(TSharedPtr<FModularRigTreeElement> Item = TreeView->FindElement(ModuleOrConnectorName))
+				if(TSharedPtr<FModularRigHierarchyTreeElement> Item = TreeView->FindElement(ModuleOrConnectorName))
 				{
 					const bool bSelected = InNotif == ERigHierarchyNotification::ElementSelected;
 					TreeView->SetItemHighlighted(Item, bSelected);
@@ -1492,7 +1492,7 @@ public:
 	}
 };
 
-UModularRig* SModularRigModel::GetModularRig() const
+UModularRig* SModularRigHierarchy::GetModularRig() const
 {
 	if (ControlRigBlueprint.IsValid())
 	{
@@ -1515,7 +1515,7 @@ UModularRig* SModularRigModel::GetModularRig() const
 	return nullptr;
 }
 
-UModularRig* SModularRigModel::GetDefaultModularRig() const
+UModularRig* SModularRigHierarchy::GetDefaultModularRig() const
 {
 	if (ControlRigBlueprint.IsValid())
 	{
@@ -1534,7 +1534,7 @@ UModularRig* SModularRigModel::GetDefaultModularRig() const
 }
 
 
-void SModularRigModel::OnRequestDetailsInspection(const FName& InModuleName)
+void SModularRigHierarchy::OnRequestDetailsInspection(const FName& InModuleName)
 {
 	if(!ControlRigEditor.IsValid())
 	{
@@ -1543,7 +1543,7 @@ void SModularRigModel::OnRequestDetailsInspection(const FName& InModuleName)
 	ControlRigEditor.Pin()->SetDetailViewForRigModules({InModuleName});
 }
 
-void SModularRigModel::PostRedo(bool bSuccess) 
+void SModularRigHierarchy::PostRedo(bool bSuccess) 
 {
 	if (bSuccess)
 	{
@@ -1551,7 +1551,7 @@ void SModularRigModel::PostRedo(bool bSuccess)
 	}
 }
 
-void SModularRigModel::PostUndo(bool bSuccess) 
+void SModularRigHierarchy::PostUndo(bool bSuccess) 
 {
 	if (bSuccess)
 	{
@@ -1559,7 +1559,7 @@ void SModularRigModel::PostUndo(bool bSuccess)
 	}
 }
 
-FReply SModularRigModel::OnDragDetected(const FGeometry& MyGeometry, const FPointerEvent& MouseEvent)
+FReply SModularRigHierarchy::OnDragDetected(const FGeometry& MyGeometry, const FPointerEvent& MouseEvent)
 {
 	TArray<FString> DraggedKeys = GetSelectedKeys();
 	TArray<FName> ModuleNames;
@@ -1572,7 +1572,7 @@ FReply SModularRigModel::OnDragDetected(const FGeometry& MyGeometry, const FPoin
 	{
 		if (ControlRigEditor.IsValid())
 		{
-			TSharedRef<FModularRigModuleDragDropOp> DragDropOp = FModularRigModuleDragDropOp::New(MoveTemp(ModuleNames));
+			TSharedRef<FModularRigHierarchyElementDragDropOp> DragDropOp = FModularRigHierarchyElementDragDropOp::New(MoveTemp(ModuleNames));
 			return FReply::Handled().BeginDragDrop(DragDropOp);
 		}
 	}
@@ -1580,7 +1580,7 @@ FReply SModularRigModel::OnDragDetected(const FGeometry& MyGeometry, const FPoin
 	return FReply::Unhandled();
 }
 
-TOptional<EItemDropZone> SModularRigModel::OnCanAcceptDrop(const FDragDropEvent& DragDropEvent, EItemDropZone DropZone, TSharedPtr<FModularRigTreeElement> TargetItem)
+TOptional<EItemDropZone> SModularRigHierarchy::OnCanAcceptDrop(const FDragDropEvent& DragDropEvent, EItemDropZone DropZone, TSharedPtr<FModularRigHierarchyTreeElement> TargetItem)
 {
 	const TOptional<EItemDropZone> InvalidDropZone;
 	TOptional<EItemDropZone> ReturnDropZone = DropZone;
@@ -1592,7 +1592,7 @@ TOptional<EItemDropZone> SModularRigModel::OnCanAcceptDrop(const FDragDropEvent&
 		
 		if(TargetItem.IsValid())
 		{
-			TSharedPtr<FModularRigTreeElement> ChildTargetItem;
+			TSharedPtr<FModularRigHierarchyTreeElement> ChildTargetItem;
 			Swap(TargetItem, ChildTargetItem);
 			
 			if(ControlRigBlueprint.IsValid())
@@ -1602,7 +1602,7 @@ TOptional<EItemDropZone> SModularRigModel::OnCanAcceptDrop(const FDragDropEvent&
 				{
 					if(const FRigModuleReference* ParentModule = ChildModule->GetParentModule())
 					{
-						if(const TSharedPtr<FModularRigTreeElement> ParentItem = TreeView->FindElement(ParentModule->Name.ToString()))
+						if(const TSharedPtr<FModularRigHierarchyTreeElement> ParentItem = TreeView->FindElement(ParentModule->Name.ToString()))
 						{
 							TargetItem = ParentItem;
 						}
@@ -1618,7 +1618,7 @@ TOptional<EItemDropZone> SModularRigModel::OnCanAcceptDrop(const FDragDropEvent&
 	}
 
 	TSharedPtr<FAssetDragDropOp> AssetDragDropOperation = DragDropEvent.GetOperationAs<FAssetDragDropOp>();
-	TSharedPtr<FModularRigModuleDragDropOp> ModuleDragDropOperation = DragDropEvent.GetOperationAs<FModularRigModuleDragDropOp>();
+	TSharedPtr<FModularRigHierarchyElementDragDropOp> ModuleDragDropOperation = DragDropEvent.GetOperationAs<FModularRigHierarchyElementDragDropOp>();
 	if (AssetDragDropOperation)
 	{
 		for (const FAssetData& AssetData : AssetDragDropOperation->GetAssets())
@@ -1658,7 +1658,7 @@ TOptional<EItemDropZone> SModularRigModel::OnCanAcceptDrop(const FDragDropEvent&
 	return ReturnDropZone;
 }
 
-FReply SModularRigModel::OnAcceptDrop(const FDragDropEvent& DragDropEvent, EItemDropZone DropZone, TSharedPtr<FModularRigTreeElement> TargetItem)
+FReply SModularRigHierarchy::OnAcceptDrop(const FDragDropEvent& DragDropEvent, EItemDropZone DropZone, TSharedPtr<FModularRigHierarchyTreeElement> TargetItem)
 {
 	int32 NewModuleIndex = INDEX_NONE;
 	
@@ -1666,7 +1666,7 @@ FReply SModularRigModel::OnAcceptDrop(const FDragDropEvent& DragDropEvent, EItem
 	{
 		if(TargetItem.IsValid())
 		{
-			TSharedPtr<FModularRigTreeElement> ChildTargetItem;
+			TSharedPtr<FModularRigHierarchyTreeElement> ChildTargetItem;
 			Swap(TargetItem, ChildTargetItem);
 			
 			if(ControlRigBlueprint.IsValid())
@@ -1679,7 +1679,7 @@ FReply SModularRigModel::OnAcceptDrop(const FDragDropEvent& DragDropEvent, EItem
 					{
 						Children = ParentModule->CachedChildren;
 						
-						if(const TSharedPtr<FModularRigTreeElement> ParentItem = TreeView->FindElement(ParentModule->Name.ToString()))
+						if(const TSharedPtr<FModularRigHierarchyTreeElement> ParentItem = TreeView->FindElement(ParentModule->Name.ToString()))
 						{
 							TargetItem = ParentItem;
 						}
@@ -1711,7 +1711,7 @@ FReply SModularRigModel::OnAcceptDrop(const FDragDropEvent& DragDropEvent, EItem
 	}
 
 	TSharedPtr<FAssetDragDropOp> AssetDragDropOperation = DragDropEvent.GetOperationAs<FAssetDragDropOp>();
-	TSharedPtr<FModularRigModuleDragDropOp> ModuleDragDropOperation = DragDropEvent.GetOperationAs<FModularRigModuleDragDropOp>();
+	TSharedPtr<FModularRigHierarchyElementDragDropOp> ModuleDragDropOperation = DragDropEvent.GetOperationAs<FModularRigHierarchyElementDragDropOp>();
 	if (AssetDragDropOperation)
 	{
 		for (const FAssetData& AssetData : AssetDragDropOperation->GetAssets())
@@ -1754,11 +1754,11 @@ FReply SModularRigModel::OnAcceptDrop(const FDragDropEvent& DragDropEvent, EItem
 	return FReply::Unhandled();
 }
 
-FReply SModularRigModel::OnDrop(const FGeometry& MyGeometry, const FDragDropEvent& DragDropEvent)
+FReply SModularRigHierarchy::OnDrop(const FGeometry& MyGeometry, const FDragDropEvent& DragDropEvent)
 {
 	// only allow drops onto empty space of the widget (when there's no target item under the mouse)
-	// when dropped onto an item SModularRigModel::OnAcceptDrop will deal with the event
-	const TSharedPtr<FModularRigTreeElement>* ItemAtMouse = TreeView->FindItemAtPosition(DragDropEvent.GetScreenSpacePosition());
+	// when dropped onto an item SModularRigHierarchy::OnAcceptDrop will deal with the event
+	const TSharedPtr<FModularRigHierarchyTreeElement>* ItemAtMouse = TreeView->FindItemAtPosition(DragDropEvent.GetScreenSpacePosition());
 	FString ParentPath;
 	if (ItemAtMouse && ItemAtMouse->IsValid())
 	{
diff --git a/Engine/Plugins/Animation/ControlRig/Source/ControlRigEditor/Private/Editor/Hierarchy/Widgets/SModularRigHierarchy.h b/Engine/Plugins/Animation/ControlRig/Source/ControlRigEditor/Private/Editor/Hierarchy/Widgets/SModularRigHierarchy.h
--- a/Engine/Plugins/Animation/ControlRig/Source/ControlRigEditor/Private/Editor/Hierarchy/Widgets/SModularRigHierarchy.h
+++ b/Engine/Plugins/Animation/ControlRig/Source/ControlRigEditor/Private/Editor/Hierarchy/Widgets/SModularRigHierarchy.h
@@ -7,7 +7,7 @@
 #include "ControlRigBlueprintLegacy.h"
 #include "Editor/RigVMNewEditor.h"
 
-class SModularRigModel;
+class SModularRigHierarchy;
 class IControlRigBaseEditor;
 class SSearchBox;
 class FUICommandList;
@@ -19,13 +19,13 @@ class UToolMenu;
 struct FToolMenuContext;
 
 /** Widget allowing editing of a control rig's structure */
-class SModularRigModel : public SCompoundWidget, public FEditorUndoClient
+class SModularRigHierarchy : public SCompoundWidget, public FEditorUndoClient
 {
 public:
-	SLATE_BEGIN_ARGS(SModularRigModel) {}
+	SLATE_BEGIN_ARGS(SModularRigHierarchy) {}
 	SLATE_END_ARGS()
 
-	~SModularRigModel();
+	~SModularRigHierarchy();
 
 	void Construct(const FArguments& InArgs, TSharedRef<IControlRigBaseEditor> InControlRigEditor);
 
@@ -53,7 +53,7 @@ private:
 	void RefreshTreeView(bool bRebuildContent = true);
 
 	/** Returns all selected items */
-	TArray<TSharedPtr<FModularRigTreeElement>> GetSelectedItems() const;
+	TArray<TSharedPtr<FModularRigHierarchyTreeElement>> GetSelectedItems() const;
 
 	/** Return all selected keys */
 	TArray<FString> GetSelectedKeys() const;
@@ -101,14 +101,14 @@ private:
 	void HandleConnectorDisconnect(const FRigElementKey& InConnector);
 
 	/** Set Selection Changed */
-	void HandleSelectionChanged(TSharedPtr<FModularRigTreeElement> Selection, ESelectInfo::Type SelectInfo);
+	void HandleSelectionChanged(TSharedPtr<FModularRigHierarchyTreeElement> Selection, ESelectInfo::Type SelectInfo);
 
 	/** Returns true if a given connector should always be shown */
 	bool ShouldAlwaysShowConnector(const FName& InConnectorName) const;
 
 	TSharedPtr< SWidget > CreateContextMenuWidget();
-	void OnItemClicked(TSharedPtr<FModularRigTreeElement> InItem);
-	void OnItemDoubleClicked(TSharedPtr<FModularRigTreeElement> InItem);
+	void OnItemClicked(TSharedPtr<FModularRigHierarchyTreeElement> InItem);
+	void OnItemDoubleClicked(TSharedPtr<FModularRigHierarchyTreeElement> InItem);
 	
 	// FEditorUndoClient
 	virtual void PostUndo(bool bSuccess) override;
@@ -118,8 +118,8 @@ private:
 	FReply OnDragDetected(const FGeometry& MyGeometry, const FPointerEvent& MouseEvent);
 
 	// reply to a drop operation on item
-	TOptional<EItemDropZone> OnCanAcceptDrop(const FDragDropEvent& DragDropEvent, EItemDropZone DropZone, TSharedPtr<FModularRigTreeElement> TargetItem);
-	FReply OnAcceptDrop(const FDragDropEvent& DragDropEvent, EItemDropZone DropZone, TSharedPtr<FModularRigTreeElement> TargetItem);
+	TOptional<EItemDropZone> OnCanAcceptDrop(const FDragDropEvent& DragDropEvent, EItemDropZone DropZone, TSharedPtr<FModularRigHierarchyTreeElement> TargetItem);
+	FReply OnAcceptDrop(const FDragDropEvent& DragDropEvent, EItemDropZone DropZone, TSharedPtr<FModularRigHierarchyTreeElement> TargetItem);
 
 	// SWidget Overrides
 	virtual FReply OnDrop(const FGeometry& MyGeometry, const FDragDropEvent& DragDropEvent) override;
@@ -171,8 +171,8 @@ private:
 
 public:
 
-	friend class FModularRigTreeElement;
-	friend class SModularRigModelItem;
+	friend class FModularRigHierarchyTreeElement;
+	friend class SModularRigHierarchyTreeItem;
 	friend class UControlRigBlueprintEditorLibrary;
 };
 
diff --git a/Engine/Plugins/Animation/ControlRig/Source/ControlRigEditor/Private/Editor/Hierarchy/Widgets/SRigHierarchy.cpp b/Engine/Plugins/Animation/ControlRig/Source/ControlRigEditor/Private/Editor/Hierarchy/Widgets/SRigHierarchy.cpp
--- a/Engine/Plugins/Animation/ControlRig/Source/ControlRigEditor/Private/Editor/Hierarchy/Widgets/SRigHierarchy.cpp
+++ b/Engine/Plugins/Animation/ControlRig/Source/ControlRigEditor/Private/Editor/Hierarchy/Widgets/SRigHierarchy.cpp
@@ -105,7 +105,7 @@ void SRigHierarchy::Construct(const FArguments& InArgs, TSharedRef<IControlRigBa
 	BindCommands();
 
 	// setup all delegates for the rig hierarchy widget
-	FRigTreeDelegates Delegates;
+	FRigHierarchyTreeDelegates Delegates;
 	Delegates.OnGetHierarchy = FOnGetRigTreeHierarchy::CreateSP(this, &SRigHierarchy::GetHierarchyForTreeView);
 	Delegates.OnGetDisplaySettings = FOnGetRigTreeDisplaySettings::CreateSP(this, &SRigHierarchy::GetDisplaySettings);
 	Delegates.OnRenameElement = FOnRigTreeRenameElement::CreateSP(this, &SRigHierarchy::HandleRenameElement);
@@ -488,9 +488,9 @@ FReply SRigHierarchy::OnMouseButtonUp(const FGeometry& MyGeometry, const FPointe
 
 	if(MouseEvent.GetEffectingButton() == EKeys::MiddleMouseButton)
 	{
-		if(const TSharedPtr<FRigTreeElement>* ItemPtr = TreeView->FindItemAtPosition(MouseEvent.GetScreenSpacePosition()))
+		if(const TSharedPtr<FRigHierarchyTreeElement>* ItemPtr = TreeView->FindItemAtPosition(MouseEvent.GetScreenSpacePosition()))
 		{
-			if(const TSharedPtr<FRigTreeElement>& Item = *ItemPtr)
+			if(const TSharedPtr<FRigHierarchyTreeElement>& Item = *ItemPtr)
 			{
 				if (URigHierarchy* Hierarchy = GetHierarchy())
 				{
@@ -629,10 +629,10 @@ void SRigHierarchy::RefreshTreeView(bool bRebuildContent)
 
 TArray<FRigHierarchyKey> SRigHierarchy::GetSelectedKeys() const
 {
-	TArray<TSharedPtr<FRigTreeElement>> SelectedItems = TreeView->GetSelectedItems();
+	TArray<TSharedPtr<FRigHierarchyTreeElement>> SelectedItems = TreeView->GetSelectedItems();
 	
 	TArray<FRigHierarchyKey> SelectedKeys;
-	for (const TSharedPtr<FRigTreeElement>& SelectedItem : SelectedItems)
+	for (const TSharedPtr<FRigHierarchyTreeElement>& SelectedItem : SelectedItems)
 	{
 		if(SelectedItem->Key.IsValid())
 		{
@@ -658,7 +658,7 @@ TArray<FRigElementKey> SRigHierarchy::GetSelectedElementKeys() const
 	return ElementKeys;
 }
 
-void SRigHierarchy::OnSelectionChanged(TSharedPtr<FRigTreeElement> Selection, ESelectInfo::Type SelectInfo)
+void SRigHierarchy::OnSelectionChanged(TSharedPtr<FRigHierarchyTreeElement> Selection, ESelectInfo::Type SelectInfo)
 {
 	if (bIsChangingRigHierarchy)
 	{
@@ -865,7 +865,7 @@ void SRigHierarchy::OnHierarchyModified(ERigHierarchyNotification InNotif, URigH
 					
 				for (int32 RootIndex = 0; RootIndex < TreeView->RootElements.Num(); ++RootIndex)
 				{
-					TSharedPtr<FRigTreeElement> Found = TreeView->FindElement(Key, TreeView->RootElements[RootIndex]);
+					TSharedPtr<FRigHierarchyTreeElement> Found = TreeView->FindElement(Key, TreeView->RootElements[RootIndex]);
 					if (Found.IsValid())
 					{
 						TreeView->SetItemSelection(Found, bSelected, ESelectInfo::OnNavigation);
@@ -900,10 +900,10 @@ void SRigHierarchy::OnHierarchyModified(ERigHierarchyNotification InNotif, URigH
 			{
 				for (int32 RootIndex = 0; RootIndex < TreeView->RootElements.Num(); ++RootIndex)
 				{
-					TSharedPtr<FRigTreeElement> TreeElement = TreeView->FindElement(InElement->GetKey(), TreeView->RootElements[RootIndex]);
+					TSharedPtr<FRigHierarchyTreeElement> TreeElement = TreeView->FindElement(InElement->GetKey(), TreeView->RootElements[RootIndex]);
 					if (TreeElement.IsValid())
 					{
-						const FRigTreeDisplaySettings& Settings = TreeView->GetRigTreeDelegates().GetDisplaySettings();
+						const FRigHierarchyTreeDisplaySettings& Settings = TreeView->GetRigTreeDelegates().GetDisplaySettings();
 						TreeElement->RefreshDisplaySettings(InHierarchy, Settings);
 					}
 				}
@@ -916,10 +916,10 @@ void SRigHierarchy::OnHierarchyModified(ERigHierarchyNotification InNotif, URigH
 			{
 				for (int32 RootIndex = 0; RootIndex < TreeView->RootElements.Num(); ++RootIndex)
 				{
-					TSharedPtr<FRigTreeElement> TreeElement = TreeView->FindElement(InComponent->GetKey(), TreeView->RootElements[RootIndex]);
+					TSharedPtr<FRigHierarchyTreeElement> TreeElement = TreeView->FindElement(InComponent->GetKey(), TreeView->RootElements[RootIndex]);
 					if (TreeElement.IsValid())
 					{
-						const FRigTreeDisplaySettings& Settings = TreeView->GetRigTreeDelegates().GetDisplaySettings();
+						const FRigHierarchyTreeDisplaySettings& Settings = TreeView->GetRigTreeDelegates().GetDisplaySettings();
 						TreeElement->RefreshDisplaySettings(InHierarchy, Settings);
 					}
 				}
@@ -1015,7 +1015,7 @@ void SRigHierarchy::OnModularRigModified(EModularRigNotification InNotif, const
 						bool bScrollIntoView = true;
 						for(const FRigElementKey& Key : Keys)
 						{
-							if(const TSharedPtr<FRigTreeElement> TreeElement = TreeView->FindElement(Key))
+							if(const TSharedPtr<FRigHierarchyTreeElement> TreeElement = TreeView->FindElement(Key))
 							{
 								TreeView->SetItemHighlighted(TreeElement, bSelected);
 								if(bScrollIntoView)
@@ -1125,7 +1125,7 @@ void SRigHierarchy::OnPostConstruction_AnyThread(UControlRig* InRig, const FName
 				{
 					for(const FRigHierarchyKey& Key : SelectionBeforeConstruction)
 					{
-						TSharedPtr<FRigTreeElement> Found = TreeView->FindElement(Key, TreeView->RootElements[RootIndex]);
+						TSharedPtr<FRigHierarchyTreeElement> Found = TreeView->FindElement(Key, TreeView->RootElements[RootIndex]);
 						if (Found.IsValid())
 						{
 							TreeView->SetItemSelection(Found, true, ESelectInfo::OnNavigation);
@@ -1283,7 +1283,7 @@ TSharedPtr< SWidget > SRigHierarchy::CreateContextMenuWidget()
 	return SNullWidget::NullWidget;
 }
 
-void SRigHierarchy::OnItemClicked(TSharedPtr<FRigTreeElement> InItem)
+void SRigHierarchy::OnItemClicked(TSharedPtr<FRigHierarchyTreeElement> InItem)
 {
 	URigHierarchy* Hierarchy = GetHierarchy();
 	check(Hierarchy);
@@ -1323,7 +1323,7 @@ void SRigHierarchy::OnItemClicked(TSharedPtr<FRigTreeElement> InItem)
 	}
 }
 
-void SRigHierarchy::OnItemDoubleClicked(TSharedPtr<FRigTreeElement> InItem)
+void SRigHierarchy::OnItemDoubleClicked(TSharedPtr<FRigHierarchyTreeElement> InItem)
 {
 	if (TreeView->IsItemExpanded(InItem))
 	{
@@ -1335,7 +1335,7 @@ void SRigHierarchy::OnItemDoubleClicked(TSharedPtr<FRigTreeElement> InItem)
 	}
 }
 
-void SRigHierarchy::OnSetExpansionRecursive(TSharedPtr<FRigTreeElement> InItem, bool bShouldBeExpanded)
+void SRigHierarchy::OnSetExpansionRecursive(TSharedPtr<FRigHierarchyTreeElement> InItem, bool bShouldBeExpanded)
 {
 	TreeView->SetExpansionRecursive(InItem, false, bShouldBeExpanded);
 }
@@ -1489,7 +1489,7 @@ void SRigHierarchy::CreateContextMenu()
 						{
 							FToolMenuSection& DefaultSection = InSubMenu->AddSection(NAME_None);
 							FRigHierarchyKey SelectedKey;
-							TArray<TSharedPtr<FRigTreeElement>> SelectedItems = RigHierarchyPanel->TreeView->GetSelectedItems();
+							TArray<TSharedPtr<FRigHierarchyTreeElement>> SelectedItems = RigHierarchyPanel->TreeView->GetSelectedItems();
 							if (SelectedItems.Num() > 0)
 							{
 								SelectedKey = SelectedItems[0]->Key;
@@ -1538,8 +1538,8 @@ void SRigHierarchy::CreateContextMenu()
 							{
 								FToolMenuSection& DefaultSection = InSubMenu->AddSection(NAME_None);
 								TArray<FRigElementKey> SelectedElements;
-								TArray<TSharedPtr<FRigTreeElement>> SelectedItems = RigHierarchyPanel->TreeView->GetSelectedItems();
-								for(const TSharedPtr<FRigTreeElement>& SelectedItem : SelectedItems)
+								TArray<TSharedPtr<FRigHierarchyTreeElement>> SelectedItems = RigHierarchyPanel->TreeView->GetSelectedItems();
+								for(const TSharedPtr<FRigHierarchyTreeElement>& SelectedItem : SelectedItems)
 								{
 									if(SelectedItem->Key.IsElement())
 									{
@@ -1692,8 +1692,8 @@ void SRigHierarchy::CreateContextMenu()
 						
 						const FString* Keyword = &InitialKeyword;
 						TArray<ERigElementType> SelectedTypes;;
-						TArray<TSharedPtr<FRigTreeElement>> SelectedItems = RigHierarchyPanel->TreeView->GetSelectedItems();
-						for(const TSharedPtr<FRigTreeElement>& SelectedItem : SelectedItems)
+						TArray<TSharedPtr<FRigHierarchyTreeElement>> SelectedItems = RigHierarchyPanel->TreeView->GetSelectedItems();
+						for(const TSharedPtr<FRigHierarchyTreeElement>& SelectedItem : SelectedItems)
 						{
 							if(SelectedItem.IsValid())
 							{
@@ -2798,7 +2798,7 @@ void SRigHierarchy::HandleRenameItem()
 	{
 		FScopedTransaction Transaction(LOCTEXT("HierarchyTreeRenameSelected", "Rename selected item from hierarchy"));
 
-		TArray<TSharedPtr<FRigTreeElement>> SelectedItems = TreeView->GetSelectedItems();
+		TArray<TSharedPtr<FRigHierarchyTreeElement>> SelectedItems = TreeView->GetSelectedItems();
 		if (SelectedItems.Num() == 1)
 		{
 			if(SelectedItems[0]->Key.IsElement())
@@ -3018,7 +3018,7 @@ FReply SRigHierarchy::OnDragDetected(const FGeometry& MyGeometry, const FPointer
 		{
 			UpdateConnectorMatchesOnDrag(DraggedElements);
 			
-			TSharedRef<FRigElementHierarchyDragDropOp> DragDropOp = FRigElementHierarchyDragDropOp::New(MoveTemp(DraggedElements));
+			TSharedRef<FRigHierarchyElementDragDropOp> DragDropOp = FRigHierarchyElementDragDropOp::New(MoveTemp(DraggedElements));
 			DragDropOp->OnPerformDropToGraphAtLocation.BindSP(EditorPtr.ToSharedRef(), &FControlRigBaseEditor::OnGraphNodeDropToPerform);
 			return FReply::Handled().BeginDragDrop(DragDropOp);
 		}
@@ -3027,12 +3027,12 @@ FReply SRigHierarchy::OnDragDetected(const FGeometry& MyGeometry, const FPointer
 	return FReply::Unhandled();
 }
 
-TOptional<EItemDropZone> SRigHierarchy::OnCanAcceptDrop(const FDragDropEvent& DragDropEvent, EItemDropZone DropZone, TSharedPtr<FRigTreeElement> TargetItem)
+TOptional<EItemDropZone> SRigHierarchy::OnCanAcceptDrop(const FDragDropEvent& DragDropEvent, EItemDropZone DropZone, TSharedPtr<FRigHierarchyTreeElement> TargetItem)
 {
 	const TOptional<EItemDropZone> InvalidDropZone;
 	TOptional<EItemDropZone> ReturnDropZone;
 
-	const TSharedPtr<FRigElementHierarchyDragDropOp> RigDragDropOp = DragDropEvent.GetOperationAs<FRigElementHierarchyDragDropOp>();
+	const TSharedPtr<FRigHierarchyElementDragDropOp> RigDragDropOp = DragDropEvent.GetOperationAs<FRigHierarchyElementDragDropOp>();
 	if (RigDragDropOp.IsValid())
 	{
 		FRigHierarchyKey TargetKey;
@@ -3223,7 +3223,7 @@ TOptional<EItemDropZone> SRigHierarchy::OnCanAcceptDrop(const FDragDropEvent& Dr
 		}
 	}
 	
-	const TSharedPtr<FModularRigModuleDragDropOp> ModuleDropOp = DragDropEvent.GetOperationAs<FModularRigModuleDragDropOp>();
+	const TSharedPtr<FModularRigHierarchyElementDragDropOp> ModuleDropOp = DragDropEvent.GetOperationAs<FModularRigHierarchyElementDragDropOp>();
 	if (ModuleDropOp.IsValid() && TargetItem.IsValid())
 	{
 		if(DropZone != EItemDropZone::OntoItem)
@@ -3323,14 +3323,14 @@ TOptional<EItemDropZone> SRigHierarchy::OnCanAcceptDrop(const FDragDropEvent& Dr
 	return ReturnDropZone;
 }
 
-FReply SRigHierarchy::OnAcceptDrop(const FDragDropEvent& DragDropEvent, EItemDropZone DropZone, TSharedPtr<FRigTreeElement> TargetItem)
+FReply SRigHierarchy::OnAcceptDrop(const FDragDropEvent& DragDropEvent, EItemDropZone DropZone, TSharedPtr<FRigHierarchyTreeElement> TargetItem)
 {
 	bool bSummonDragDropMenu = DragDropEvent.GetModifierKeys().IsAltDown() && DragDropEvent.GetModifierKeys().IsShiftDown(); 
 	bool bMatchTransforms = DragDropEvent.GetModifierKeys().IsAltDown();
 	bool bReparentItems = !bMatchTransforms;
 	UpdateConnectorMatchesOnDrag({});
 
-	TSharedPtr<FRigElementHierarchyDragDropOp> RigDragDropOp = DragDropEvent.GetOperationAs<FRigElementHierarchyDragDropOp>();
+	TSharedPtr<FRigHierarchyElementDragDropOp> RigDragDropOp = DragDropEvent.GetOperationAs<FRigHierarchyElementDragDropOp>();
 	if (RigDragDropOp.IsValid())
 	{
 		if (bSummonDragDropMenu)
@@ -3427,7 +3427,7 @@ FReply SRigHierarchy::OnAcceptDrop(const FDragDropEvent& DragDropEvent, EItemDro
 		return ResolveConnector(DraggedKey, FRigElementKey());
 	}
 
-	const TSharedPtr<FModularRigModuleDragDropOp> ModuleDropOp = DragDropEvent.GetOperationAs<FModularRigModuleDragDropOp>();
+	const TSharedPtr<FModularRigHierarchyElementDragDropOp> ModuleDropOp = DragDropEvent.GetOperationAs<FModularRigHierarchyElementDragDropOp>();
 	if (ModuleDropOp.IsValid() && TargetItem.IsValid())
 	{
 		UModularRig* ControlRig = Cast<UModularRig>(ControlRigBlueprint->GetDebuggedControlRig());
@@ -3514,7 +3514,7 @@ void SRigHierarchy::UpdateConnectorMatchesOnDrag(const TArray<FRigHierarchyKey>&
 	DragRigResolveResults.Reset();
 
 	// fade in all items
-	for(const TPair<FRigHierarchyKey, TSharedPtr<FRigTreeElement>>& Pair : TreeView->ElementMap)
+	for(const TPair<FRigHierarchyKey, TSharedPtr<FRigHierarchyTreeElement>>& Pair : TreeView->ElementMap)
 	{
 		Pair.Value->bFadedOutDuringDragDrop = false;
 	}
@@ -3558,7 +3558,7 @@ void SRigHierarchy::UpdateConnectorMatchesOnDrag(const TArray<FRigHierarchyKey>&
 	{
 		for(const FRigElementResolveResult& ExcludedElement : Pair.Value.GetExcluded())
 		{
-			if(const TSharedPtr<FRigTreeElement>* TreeElementPtr = TreeView->ElementMap.Find(ExcludedElement.GetKey()))
+			if(const TSharedPtr<FRigHierarchyTreeElement>* TreeElementPtr = TreeView->ElementMap.Find(ExcludedElement.GetKey()))
 			{
 				TreeElementPtr->Get()->bFadedOutDuringDragDrop = true;
 			}
@@ -3699,7 +3699,7 @@ bool SRigHierarchy::HandleVerifyNameChanged(const FRigHierarchyKey& OldKey, cons
 FReply SRigHierarchy::OnDrop(const FGeometry& MyGeometry, const FDragDropEvent& DragDropEvent)
 {
 	// only allow drops onto empty space of the widget (when there's no target item under the mouse)
-	const TSharedPtr<FRigTreeElement>* ItemAtMouse = TreeView->FindItemAtPosition(DragDropEvent.GetScreenSpacePosition());
+	const TSharedPtr<FRigHierarchyTreeElement>* ItemAtMouse = TreeView->FindItemAtPosition(DragDropEvent.GetScreenSpacePosition());
 	if (ItemAtMouse && ItemAtMouse->IsValid())
 	{
 		return FReply::Unhandled();
@@ -3924,8 +3924,8 @@ void SRigHierarchy::HandleZeroControlShape()
 
 		const FRigHierarchyInteractionBracket DebuggedInteractionBracket(DebuggedHierarchy);
 
-		TArray<TSharedPtr<FRigTreeElement>> SelectedItems = TreeView->GetSelectedItems();
-		for (const TSharedPtr<FRigTreeElement>& SelectedItem : SelectedItems)
+		TArray<TSharedPtr<FRigHierarchyTreeElement>> SelectedItems = TreeView->GetSelectedItems();
+		for (const TSharedPtr<FRigHierarchyTreeElement>& SelectedItem : SelectedItems)
 		{
 			if (!SelectedItem->Key.IsElement())
 			{
@@ -4055,8 +4055,8 @@ bool SRigHierarchy::CanHandleSetInitialTransformsAsCurrent() const
 
 void SRigHierarchy::HandleFrameSelection()
 {
-	TArray<TSharedPtr<FRigTreeElement>> SelectedItems = TreeView->GetSelectedItems();
-	for (TSharedPtr<FRigTreeElement> SelectedItem : SelectedItems)
+	TArray<TSharedPtr<FRigHierarchyTreeElement>> SelectedItems = TreeView->GetSelectedItems();
+	for (TSharedPtr<FRigHierarchyTreeElement> SelectedItem : SelectedItems)
 	{
 		TreeView->SetExpansionRecursive(SelectedItem, true, true);
 	}
diff --git a/Engine/Plugins/Animation/ControlRig/Source/ControlRigEditor/Private/Editor/Hierarchy/Widgets/SRigHierarchy.h b/Engine/Plugins/Animation/ControlRig/Source/ControlRigEditor/Private/Editor/Hierarchy/Widgets/SRigHierarchy.h
--- a/Engine/Plugins/Animation/ControlRig/Source/ControlRigEditor/Private/Editor/Hierarchy/Widgets/SRigHierarchy.h
+++ b/Engine/Plugins/Animation/ControlRig/Source/ControlRigEditor/Private/Editor/Hierarchy/Widgets/SRigHierarchy.h
@@ -100,13 +100,13 @@ private:
 	void HandlePasteTransforms(ERigTransformType::Type InTransformType, bool bAffectChildren);
 
 	/** Set Selection Changed */
-	void OnSelectionChanged(TSharedPtr<FRigTreeElement> Selection, ESelectInfo::Type SelectInfo);
+	void OnSelectionChanged(TSharedPtr<FRigHierarchyTreeElement> Selection, ESelectInfo::Type SelectInfo);
 
 	TSharedRef< SWidget > CreateFilterMenu();
 	TSharedPtr< SWidget > CreateContextMenuWidget();
-	void OnItemClicked(TSharedPtr<FRigTreeElement> InItem);
-	void OnItemDoubleClicked(TSharedPtr<FRigTreeElement> InItem);
-	void OnSetExpansionRecursive(TSharedPtr<FRigTreeElement> InItem, bool bShouldBeExpanded);
+	void OnItemClicked(TSharedPtr<FRigHierarchyTreeElement> InItem);
+	void OnItemDoubleClicked(TSharedPtr<FRigHierarchyTreeElement> InItem);
+	void OnSetExpansionRecursive(TSharedPtr<FRigHierarchyTreeElement> InItem, bool bShouldBeExpanded);
 	TOptional<FText> OnGetItemTooltip(const FRigHierarchyKey& InKey) const;
 
 	// FEditorUndoClient
@@ -115,8 +115,8 @@ private:
 
 	// reply to a drag operation
 	FReply OnDragDetected(const FGeometry& MyGeometry, const FPointerEvent& MouseEvent);
-	TOptional<EItemDropZone> OnCanAcceptDrop(const FDragDropEvent& DragDropEvent, EItemDropZone DropZone, TSharedPtr<FRigTreeElement> TargetItem);
-	FReply OnAcceptDrop(const FDragDropEvent& DragDropEvent, EItemDropZone DropZone, TSharedPtr<FRigTreeElement> TargetItem);
+	TOptional<EItemDropZone> OnCanAcceptDrop(const FDragDropEvent& DragDropEvent, EItemDropZone DropZone, TSharedPtr<FRigHierarchyTreeElement> TargetItem);
+	FReply OnAcceptDrop(const FDragDropEvent& DragDropEvent, EItemDropZone DropZone, TSharedPtr<FRigHierarchyTreeElement> TargetItem);
 	void OnElementKeyTagDragDetected(const FRigElementKey& InDraggedTag);
 	void UpdateConnectorMatchesOnDrag(const TArray<FRigHierarchyKey>& InDraggedKeys);
 
@@ -132,8 +132,8 @@ private:
 	/** Our owning control rig editor */
 	TWeakPtr<IControlRigBaseEditor> ControlRigEditor;
 
-	FRigTreeDisplaySettings DisplaySettings;
-	const FRigTreeDisplaySettings& GetDisplaySettings() const { return DisplaySettings; }
+	FRigHierarchyTreeDisplaySettings DisplaySettings;
+	const FRigHierarchyTreeDisplaySettings& GetDisplaySettings() const { return DisplaySettings; }
 
 	// struct for import hierarchy dialog (not on stack to allow non-modal dialog)
 	FRigHierarchyImportSettings ImportSettings;
@@ -239,7 +239,7 @@ public:
 	// SWidget Overrides
 	virtual FReply OnDrop(const FGeometry& MyGeometry, const FDragDropEvent& DragDropEvent) override;
 
-	friend class FRigTreeElement;
+	friend class FRigHierarchyTreeElement;
 	friend class SRigHierarchyItem;
 	friend class UControlRigBlueprintEditorLibrary;
 };
diff --git a/Engine/Plugins/Animation/ControlRig/Source/ControlRigEditor/Private/Editor/ModularRigModelTabSummoner.cpp b/Engine/Plugins/Animation/ControlRig/Source/ControlRigEditor/Private/Editor/ModularRigModelTabSummoner.cpp
--- a/Engine/Plugins/Animation/ControlRig/Source/ControlRigEditor/Private/Editor/ModularRigModelTabSummoner.cpp
+++ b/Engine/Plugins/Animation/ControlRig/Source/ControlRigEditor/Private/Editor/ModularRigModelTabSummoner.cpp
@@ -37,7 +37,7 @@ FTabSpawnerEntry& FModularRigModelTabSummoner::RegisterTabSpawner(TSharedRef<FTa
 TSharedRef<SWidget> FModularRigModelTabSummoner::CreateTabBody(const FWorkflowTabSpawnInfo& Info) const
 {
 	ControlRigEditor.Pin()->IncreaseModularRigHierarchyTabCount();
-	return SNew(SModularRigModel, ControlRigEditor.Pin().ToSharedRef());
+	return SNew(SModularRigHierarchy, ControlRigEditor.Pin().ToSharedRef());
 }
 
 TSharedRef<SDockTab> FModularRigModelTabSummoner::SpawnTab(const FWorkflowTabSpawnInfo& Info) const
@@ -51,7 +51,7 @@ TSharedRef<SDockTab> FModularRigModelTabSummoner::SpawnTab(const FWorkflowTabSpa
 		{
 			if(SWidget* Content = &SharedDocTab->GetContent().Get())
 			{
-				SModularRigModel* RigHierarchy = (SModularRigModel*)Content;
+				SModularRigHierarchy* RigHierarchy = (SModularRigHierarchy*)Content;
 				if(IControlRigBaseEditor* ControlRigEditorForTab = RigHierarchy->GetControlRigEditor())
 				{
 					HierarchyTabCount = ControlRigEditorForTab->GetModularRigHierarchyTabCount();
@@ -68,7 +68,7 @@ TSharedRef<SDockTab> FModularRigModelTabSummoner::SpawnTab(const FWorkflowTabSpa
 	{
 		if(SWidget* Content = &DockTab->GetContent().Get())
 		{
-			SModularRigModel* RigHierarchy = (SModularRigModel*)Content;
+			SModularRigHierarchy* RigHierarchy = (SModularRigHierarchy*)Content;
 			if(IControlRigBaseEditor* ControlRigEditorForTab = RigHierarchy->GetControlRigEditor())
 			{
 				ControlRigEditorForTab->DecreaseModularRigHierarchyTabCount();
diff --git a/Engine/Plugins/Animation/ControlRig/Source/ControlRigEditor/Private/Editor/SRigConnectorTargetWidget.cpp b/Engine/Plugins/Animation/ControlRig/Source/ControlRigEditor/Private/Editor/SRigConnectorTargetWidget.cpp
--- a/Engine/Plugins/Animation/ControlRig/Source/ControlRigEditor/Private/Editor/SRigConnectorTargetWidget.cpp
+++ b/Engine/Plugins/Animation/ControlRig/Source/ControlRigEditor/Private/Editor/SRigConnectorTargetWidget.cpp
@@ -39,7 +39,7 @@ void SRigConnectorTargetComboButton::Construct(const FArguments& InArgs)
 	});
 	RigTreeDelegates.OnSelectionChanged.BindSP(this, &SRigConnectorTargetComboButton::OnConnectorTargetPicked);
 
-	SearchableTreeView = SNew(SSearchableRigHierarchyTreeView)
+	SearchableTreeView = SNew(SRigHierarchySearchableTreeView)
 		.RigTreeDelegates(RigTreeDelegates);
 
 	SComboButton::FArguments ComboButtonArgs;
@@ -156,7 +156,7 @@ void SRigConnectorTargetComboButton::PopulateButtonBox()
 	];
 }
 
-void SRigConnectorTargetComboButton::OnConnectorTargetPicked(TSharedPtr<FRigTreeElement> Selection, ESelectInfo::Type SelectInfo)
+void SRigConnectorTargetComboButton::OnConnectorTargetPicked(TSharedPtr<FRigHierarchyTreeElement> Selection, ESelectInfo::Type SelectInfo)
 {
 	if(SelectInfo == ESelectInfo::OnNavigation)
 	{
diff --git a/Engine/Plugins/Animation/ControlRig/Source/ControlRigEditor/Private/Editor/SRigConnectorTargetWidget.h b/Engine/Plugins/Animation/ControlRig/Source/ControlRigEditor/Private/Editor/SRigConnectorTargetWidget.h
--- a/Engine/Plugins/Animation/ControlRig/Source/ControlRigEditor/Private/Editor/SRigConnectorTargetWidget.h
+++ b/Engine/Plugins/Animation/ControlRig/Source/ControlRigEditor/Private/Editor/SRigConnectorTargetWidget.h
@@ -35,14 +35,14 @@ public:
 	UPROPERTY(EditAnywhere, Category=NoCategory, meta=(FullyExpand=false))
 	TArray<FRigElementKey> TargetArray;
 
-	FRigTreeDelegates* GetRigTreeDelegates() const
+	FRigHierarchyTreeDelegates* GetRigTreeDelegates() const
 	{
 		return RigTreeDelegates;
 	}
 
 private:
 	
-	FRigTreeDelegates* RigTreeDelegates = nullptr;
+	FRigHierarchyTreeDelegates* RigTreeDelegates = nullptr;
 
 	friend class SRigConnectorTargetWidget;
 	friend class FRigConnectorTargetWidgetCustomization;
@@ -66,7 +66,7 @@ public:
 	SLATE_ARGUMENT(FRigElementKey, ConnectorKey)
 	SLATE_ATTRIBUTE(FRigElementKey, TargetKey)
 	SLATE_ATTRIBUTE(int32, ArrayIndex)
-	SLATE_ARGUMENT(FRigTreeDelegates, RigTreeDelegates)
+	SLATE_ARGUMENT(FRigHierarchyTreeDelegates, RigTreeDelegates)
 	SLATE_ARGUMENT(float, ButtonMinDesiredWith)
 	SLATE_EVENT(FRigConnectorTargetWidget_SetTarget, OnSetTarget)
 	SLATE_END_ARGS()
@@ -79,15 +79,15 @@ private:
 
 	void OnComboBoxOpened();
 	void PopulateButtonBox();
-	void OnConnectorTargetPicked(TSharedPtr<FRigTreeElement> Selection, ESelectInfo::Type SelectInfo);
+	void OnConnectorTargetPicked(TSharedPtr<FRigHierarchyTreeElement> Selection, ESelectInfo::Type SelectInfo);
 	
 	FRigElementKey ConnectorKey;
 	TAttribute<FRigElementKey> TargetKey;
 	TAttribute<int32> ArrayIndex;
 	FRigConnectorTargetWidget_SetTarget OnSetTarget;
 	TSharedPtr<SVerticalBox> VerticalButtonBox;
-	FRigTreeDelegates RigTreeDelegates;
-	TSharedPtr<SSearchableRigHierarchyTreeView> SearchableTreeView;
+	FRigHierarchyTreeDelegates RigTreeDelegates;
+	TSharedPtr<SRigHierarchySearchableTreeView> SearchableTreeView;
 };
 
 /** Widget allowing editing the targets of a connector */
@@ -112,7 +112,7 @@ public:
 	SLATE_EVENT(FRigConnectorTargetWidget_SetTargetArray, OnSetTargetArray)
 	SLATE_ARGUMENT(FRigConnectorTargetWidget_HandleTargetsChangedInClient*, HandleTargetsChangedInClient)
 	SLATE_ATTRIBUTE(FMargin, Padding)
-	SLATE_ARGUMENT(FRigTreeDelegates, RigTreeDelegates)
+	SLATE_ARGUMENT(FRigHierarchyTreeDelegates, RigTreeDelegates)
 	SLATE_END_ARGS()
 
 	virtual ~SRigConnectorTargetWidget() override;
@@ -133,7 +133,7 @@ private:
 	TStrongObjectPtr<URigConnectorTargetsDetailWrapper> TargetsDetailWrapper;
 	FRigConnectorTargetWidget_SetTargetArray OnSetTargetArray;
 	bool bIsArray = false;
-	FRigTreeDelegates RigTreeDelegates;
+	FRigHierarchyTreeDelegates RigTreeDelegates;
 };
 
 class FRigConnectorTargetWidgetCustomization : public IPropertyTypeCustomization
diff --git a/Engine/Plugins/Animation/ControlRig/Source/ControlRigEditor/Private/Editor/SRigSpacePickerWidget.cpp b/Engine/Plugins/Animation/ControlRig/Source/ControlRigEditor/Private/Editor/SRigSpacePickerWidget.cpp
--- a/Engine/Plugins/Animation/ControlRig/Source/ControlRigEditor/Private/Editor/SRigSpacePickerWidget.cpp
+++ b/Engine/Plugins/Animation/ControlRig/Source/ControlRigEditor/Private/Editor/SRigSpacePickerWidget.cpp
@@ -748,13 +748,13 @@ TSharedRef<SWidget> SRigSpacePickerWidget::OnGetAddButtonContent()
 	HierarchyDisplaySettings.bShowSockets = false;
 	HierarchyDisplaySettings.bShowComponents = false;
 	
-	FRigTreeDelegates TreeDelegates;
+	FRigHierarchyTreeDelegates TreeDelegates;
 	TreeDelegates.OnGetDisplaySettings = FOnGetRigTreeDisplaySettings::CreateSP(this, &SRigSpacePickerWidget::GetHierarchyDisplaySettings); 
 	TreeDelegates.OnGetHierarchy = FOnGetRigTreeHierarchy::CreateSP(this, &SRigSpacePickerWidget::GetHierarchyConst);
 	TreeDelegates.OnMouseButtonClick = FOnRigTreeMouseButtonClick::CreateSP(this, &SRigSpacePickerWidget::HandleClickTreeItemInAddMenu);
 	TreeDelegates.OnCompareKeys = FOnRigTreeCompareKeys::CreateStatic(&UE::ControlRigEditor::RigSpacePickerDetail::CompareKeys);
 
-	const TSharedRef<SSearchableRigHierarchyTreeView> SearchableTreeView = SNew(SSearchableRigHierarchyTreeView).RigTreeDelegates(TreeDelegates);
+	const TSharedRef<SRigHierarchySearchableTreeView> SearchableTreeView = SNew(SRigHierarchySearchableTreeView).RigTreeDelegates(TreeDelegates);
 	SearchableTreeView->GetTreeView()->RefreshTreeView(true);
 	AddComboButton->SetMenuContentWidgetToFocus(SearchableTreeView->GetSearchBox());
 	
@@ -776,7 +776,7 @@ void SRigSpacePickerWidget::OnMenuOpenChanged(bool bIsOpen)
 	}
 }
 
-void SRigSpacePickerWidget::HandleClickTreeItemInAddMenu(TSharedPtr<FRigTreeElement> InItem)
+void SRigSpacePickerWidget::HandleClickTreeItemInAddMenu(TSharedPtr<FRigHierarchyTreeElement> InItem)
 {
 	if (InItem.IsValid())
 	{
diff --git a/Engine/Plugins/Animation/ControlRig/Source/ControlRigEditor/Private/Editor/SRigSpacePickerWidget.h b/Engine/Plugins/Animation/ControlRig/Source/ControlRigEditor/Private/Editor/SRigSpacePickerWidget.h
--- a/Engine/Plugins/Animation/ControlRig/Source/ControlRigEditor/Private/Editor/SRigSpacePickerWidget.h
+++ b/Engine/Plugins/Animation/ControlRig/Source/ControlRigEditor/Private/Editor/SRigSpacePickerWidget.h
@@ -86,7 +86,7 @@ public:
 		return nullptr;
 	}
 
-	const FRigTreeDisplaySettings& GetHierarchyDisplaySettings() const { return HierarchyDisplaySettings; }
+	const FRigHierarchyTreeDisplaySettings& GetHierarchyDisplaySettings() const { return HierarchyDisplaySettings; }
 	const URigHierarchy* GetHierarchyConst() const { return GetHierarchy(); }
 	
 	const TArray<FRigElementKey>& GetControls() const { return ControlKeys; }
@@ -152,7 +152,7 @@ private:
 	/** Handles the menu closing. Ensures the space picker widget is put on top of all UI if it is placed in the floating window.*/
 	void OnMenuOpenChanged(bool bIsOpen);
 	/** Assigns the clicked space to the current control rig selection */
-	void HandleClickTreeItemInAddMenu(TSharedPtr<FRigTreeElement> InItem);
+	void HandleClickTreeItemInAddMenu(TSharedPtr<FRigHierarchyTreeElement> InItem);
 	
 	/** @return Whether the button is enabled. Disabled e.g. when no control rig is selected. */
 	bool IsAddButtonEnabled() const { return IsAddButtonEnabledWithReason(); }
@@ -204,7 +204,7 @@ private:
 	TWeakPtr<SWindow> DialogWindow;
 	FDelegateHandle HierarchyModifiedHandle;
 	FDelegateHandle ActiveSpaceChangedWindowHandle;
-	FRigTreeDisplaySettings HierarchyDisplaySettings;
+	FRigHierarchyTreeDisplaySettings HierarchyDisplaySettings;
 
 	static UE_API FRigElementKey InValidKey;
 
