Anim In Engine: Fix Anim Layer Mute state not syncing with Sequencer by using the track row disable. Update old data in PostLoad, and make sure we SetActive sections when the track is active.

diff --git a/Engine/Plugins/Animation/ControlRig/Source/ControlRigEditor/Private/Sequencer/AnimLayers/AnimLayers.cpp b/Engine/Plugins/Animation/ControlRig/Source/ControlRigEditor/Private/Sequencer/AnimLayers/AnimLayers.cpp
--- a/Engine/Plugins/Animation/ControlRig/Source/ControlRigEditor/Private/Sequencer/AnimLayers/AnimLayers.cpp
+++ b/Engine/Plugins/Animation/ControlRig/Source/ControlRigEditor/Private/Sequencer/AnimLayers/AnimLayers.cpp
@@ -420,14 +420,24 @@ bool UAnimLayer::GetActive() const
 			{
 				if (SectionItem.Section.IsValid())
 				{
-					const bool bActive = SectionItem.Section.Get()->IsActive();
-					if (CurrentVal.IsSet() && CurrentVal.GetValue() != bActive)
-					{
-						SectionItem.Section.Get()->SetIsActive(CurrentVal.GetValue());
-					}
-					if (CurrentVal.IsSet() == false)
+					if (UMovieSceneTrack* Track = Cast<UMovieSceneTrack>(SectionItem.Section.Get()->GetTypedOuter< UMovieSceneTrack>()))
 					{
-						CurrentVal = bActive;
+						const bool bActive = !Track->IsRowEvalDisabled(SectionItem.Section->GetRowIndex());
+
+						//make sure other track rows match state, and if active the sections are active
+						if (CurrentVal.IsSet() && CurrentVal.GetValue() != bActive)
+						{
+							Track->SetRowEvalDisabled(!CurrentVal.GetValue(), SectionItem.Section->GetRowIndex());
+						}
+						//make sure section is active
+						if (bActive)
+						{
+							SectionItem.Section->SetIsActive(bActive);
+						}
+						if (CurrentVal.IsSet() == false)
+						{
+							CurrentVal = bActive;
+						}
 					}
 				}
 			}
@@ -439,6 +449,51 @@ bool UAnimLayer::GetActive() const
 	}
 	return State.bActive;
 }
+void UAnimLayer::PostLoad()
+{
+	Super::PostLoad();
+
+	TWeakObjectPtr<UAnimLayer> WeakThis = this;
+
+	GEditor->GetTimerManager()->SetTimerForNextTick(FTimerDelegate::CreateLambda([WeakThis, this]()
+		{
+			if (WeakThis.IsValid())
+			{
+				const bool bTrackRowDisabled = !WeakThis->State.bActive;
+				for (const TPair<TWeakObjectPtr<UObject>, FAnimLayerItem>& Pair : WeakThis->AnimLayerItems)
+				{
+					if (Pair.Key != nullptr)
+					{
+						for (const FAnimLayerSectionItem& SectionItem : Pair.Value.SectionItems)
+						{
+							if (SectionItem.Section.IsValid())
+							{
+								if (UMovieSceneTrack* Track = Cast<UMovieSceneTrack>(SectionItem.Section.Get()->GetTypedOuter< UMovieSceneTrack>()))
+								{
+									Track->SetRowEvalDisabled(bTrackRowDisabled, SectionItem.Section->GetRowIndex());
+									//make sure all of the sections are active
+									const TArray<UMovieSceneSection*>& Sections = Track->GetAllSections();
+									if (Sections.Num() > 0)
+									{
+										for (UMovieSceneSection* Section : Sections)
+										{
+											if (Section->TryModify())
+											{
+												Section->SetIsActive(true);
+											}
+										}
+									}
+								}
+							}
+						}
+					}
+				}
+			}
+		}
+	));
+	
+}
+
 
 void UAnimLayer::SetActive(bool bInActive)
 {
@@ -453,8 +508,17 @@ void UAnimLayer::SetActive(bool bInActive)
 			{
 				if (SectionItem.Section.IsValid())
 				{
-					SectionItem.Section->Modify();
-					SectionItem.Section.Get()->SetIsActive(State.bActive);
+					if (UMovieSceneTrack* Track = Cast<UMovieSceneTrack>(SectionItem.Section.Get()->GetTypedOuter< UMovieSceneTrack>()))
+					{
+						Track->Modify();
+						Track->SetRowEvalDisabled(!State.bActive, SectionItem.Section->GetRowIndex());
+						if (State.bActive)
+						{
+							//make sure to turn on, may be off from legacy
+							SectionItem.Section->Modify();
+							SectionItem.Section.Get()->SetIsActive(State.bActive);
+						}
+					}
 				}
 			}
 		}
diff --git a/Engine/Plugins/Animation/ControlRig/Source/ControlRigEditor/Private/Sequencer/AnimLayers/AnimLayers.h b/Engine/Plugins/Animation/ControlRig/Source/ControlRigEditor/Private/Sequencer/AnimLayers/AnimLayers.h
--- a/Engine/Plugins/Animation/ControlRig/Source/ControlRigEditor/Private/Sequencer/AnimLayers/AnimLayers.h
+++ b/Engine/Plugins/Animation/ControlRig/Source/ControlRigEditor/Private/Sequencer/AnimLayers/AnimLayers.h
@@ -236,6 +236,7 @@ class UAnimLayer : public UObject
 	GENERATED_UCLASS_BODY()
 public:
 	UAnimLayer() {};
+	virtual void PostLoad() override;
 
 	void GetAnimLayerObjects(FAnimLayerObjects& InLayerObjects) const;
 	void GetSections(TArray<UMovieSceneSection*>& OutSections) const;
