Anim In Engine: Fix Anim Layer Mute state follow up, make sure to not dirty LS if no update is needed to mute states.

diff --git a/Engine/Plugins/Animation/ControlRig/Source/ControlRigEditor/Private/Sequencer/AnimLayers/AnimLayers.cpp b/Engine/Plugins/Animation/ControlRig/Source/ControlRigEditor/Private/Sequencer/AnimLayers/AnimLayers.cpp
--- a/Engine/Plugins/Animation/ControlRig/Source/ControlRigEditor/Private/Sequencer/AnimLayers/AnimLayers.cpp
+++ b/Engine/Plugins/Animation/ControlRig/Source/ControlRigEditor/Private/Sequencer/AnimLayers/AnimLayers.cpp
@@ -470,16 +470,23 @@ void UAnimLayer::PostLoad()
 							{
 								if (UMovieSceneTrack* Track = Cast<UMovieSceneTrack>(SectionItem.Section.Get()->GetTypedOuter< UMovieSceneTrack>()))
 								{
-									Track->SetRowEvalDisabled(bTrackRowDisabled, SectionItem.Section->GetRowIndex());
+									if (Track->IsRowEvalDisabled(SectionItem.Section->GetRowIndex()) != bTrackRowDisabled)
+									{
+										Track->Modify();
+										Track->SetRowEvalDisabled(bTrackRowDisabled, SectionItem.Section->GetRowIndex());
+									}
 									//make sure all of the sections are active
 									const TArray<UMovieSceneSection*>& Sections = Track->GetAllSections();
 									if (Sections.Num() > 0)
 									{
 										for (UMovieSceneSection* Section : Sections)
 										{
-											if (Section->TryModify())
+											if (Section->IsActive() != true)
 											{
-												Section->SetIsActive(true);
+												if (Section->TryModify())
+												{
+													Section->SetIsActive(true);
+												}
 											}
 										}
 									}
