Anim In Engine: Sets and Layers when keying will key the current transforms not the curve values. Had to expose some track editor functionality.

diff --git a/Engine/Plugins/Animation/ControlRig/Source/ControlRigEditor/Private/Sequencer/ControlRigParameterTrackEditor.cpp b/Engine/Plugins/Animation/ControlRig/Source/ControlRigEditor/Private/Sequencer/ControlRigParameterTrackEditor.cpp
--- a/Engine/Plugins/Animation/ControlRig/Source/ControlRigEditor/Private/Sequencer/ControlRigParameterTrackEditor.cpp
+++ b/Engine/Plugins/Animation/ControlRig/Source/ControlRigEditor/Private/Sequencer/ControlRigParameterTrackEditor.cpp
@@ -118,6 +118,7 @@
 #include "Decorations/MovieSceneMuteSoloDecoration.h"
 #include "Editor/Constraints/ConstraintCreationOptions.h"
 
+
 #define LOCTEXT_NAMESPACE "FControlRigParameterTrackEditor"
 
 bool FControlRigParameterTrackEditor::bControlRigEditModeWasOpen = false;
@@ -125,7 +126,8 @@ TArray <TPair<UClass*, TArray<FName>>> FControlRigParameterTrackEditor::Previous
 
 TAutoConsoleVariable<bool> CVarAutoGenerateControlRigTrack(TEXT("ControlRig.Sequencer.AutoGenerateTrack"), true, TEXT("When true automatically create control rig tracks in Sequencer when a control rig is added to a level."));
 
-TAutoConsoleVariable<bool> CVarSelectedKeysSelectControls(TEXT("ControlRig.Sequencer.SelectedKeysSelectControls"), false, TEXT("When true when we select a key in Sequencer it will select the Control, by default false."));
+TAutoConsoleVariable<bool>	
+CVarSelectedKeysSelectControls(TEXT("ControlRig.Sequencer.SelectedKeysSelectControls"), false, TEXT("When true when we select a key in Sequencer it will select the Control, by default false."));
 
 TAutoConsoleVariable<bool> CVarSelectedSectionSetsSectionToKey(TEXT("ControlRig.Sequencer.SelectedSectionSetsSectionToKey"), true, TEXT("When true when we select a channel in a section, if it's the only section selected we set it as the Section To Key, by default false."));
 
@@ -263,6 +265,7 @@ static bool DoesControlRigAllowMultipleInstances(const FTopLevelAssetPath& InGen
 
 bool FControlRigParameterTrackEditor::bAutoGenerateControlRigTrack = true;
 FCriticalSection FControlRigParameterTrackEditor::ControlUndoTransactionMutex;
+FControlRigParameterTrackEditor* FControlRigParameterTrackEditor::CurrentEditor = nullptr;
 
 FControlRigParameterTrackEditor::FControlRigParameterTrackEditor(TSharedRef<ISequencer> InSequencer)
 	: FKeyframeTrackEditor<UMovieSceneControlRigParameterTrack>(InSequencer)
@@ -295,7 +298,8 @@ void FControlRigParameterTrackEditor::OnInitialize()
 	{
 		return;
 	}
-	
+	FControlRigParameterTrackEditor::CurrentEditor = this;
+
 	UMovieSceneSequence* OwnerSequence = Sequencer->GetFocusedMovieSceneSequence();
 	UMovieScene* MovieScene = OwnerSequence ? OwnerSequence->GetMovieScene() : nullptr;
 	
@@ -535,6 +539,7 @@ FControlRigParameterTrackEditor::~FControlRigParameterTrackEditor()
 	{
 		GEditor->UnregisterForUndo(this);
 	}
+	FControlRigParameterTrackEditor::CurrentEditor = nullptr;
 
 	FMovieSceneToolsModule::Get().UnregisterAnimationBakeHelper(this);
 }
@@ -644,6 +649,8 @@ void FControlRigParameterTrackEditor::ObjectImplicitlyRemoved(UObject* InObject)
 
 void FControlRigParameterTrackEditor::OnRelease()
 {
+	FControlRigParameterTrackEditor::CurrentEditor = nullptr;
+
 	UWorld* World = GCurrentLevelEditingViewportClient ? GCurrentLevelEditingViewportClient->GetWorld() : nullptr;
 	FConstraintsManagerController& Controller = FConstraintsManagerController::Get(World);
 	for (FDelegateHandle& Handle : ConstraintHandlesToClear)
@@ -729,7 +736,7 @@ void FControlRigParameterTrackEditor::OnRelease()
 
 		if (OnMovieSceneBindingsChangeHandle.IsValid())
 		{
-			Sequencer->OnChannelChanged().Remove(OnMovieSceneBindingsChangeHandle);
+			Sequencer->OnMovieSceneBindingsChanged().Remove(OnMovieSceneBindingsChangeHandle);
 			OnMovieSceneBindingsChangeHandle.Reset();
 		}
 
@@ -2438,28 +2445,9 @@ bool FControlRigParameterTrackEditor::CanAddTransformKeysForSelectedObjects() co
 	return false;
 }
 
-void FControlRigParameterTrackEditor::OnAddTransformKeysForSelectedObjects(EMovieSceneTransformChannel Channel)
+void FControlRigParameterTrackEditor::SetKeysOnControls(const TMap<UControlRig*, TArray<FRigElementKey>>& SelectedControls, EMovieSceneTransformChannel Channel)
 {
-	const TSharedPtr<ISequencer> Sequencer = GetSequencer();
-	if (!Sequencer || !Sequencer->IsAllowedToChange())
-	{
-		return;
-	}
-	
-	const FControlRigEditMode* ControlRigEditMode = GetEditMode();
-	if (!ControlRigEditMode)
-	{
-		return;
-	}
-
-	TMap<UControlRig*, TArray<FRigElementKey>> SelectedControls;
-	ControlRigEditMode->GetAllSelectedControls(SelectedControls);
-	if (SelectedControls.Num() <= 0)
-	{
-		return;
-	}
-	const EControlRigContextChannelToKey ChannelsToKey = static_cast<EControlRigContextChannelToKey>(Channel); 
-	FScopedTransaction KeyTransaction(LOCTEXT("SetKeysOnControls", "Set Keys On Controls"), !GIsTransacting);
+	const EControlRigContextChannelToKey ChannelsToKey = static_cast<EControlRigContextChannelToKey>(Channel);
 
 	static constexpr bool bInConstraintSpace = true;
 	FRigControlModifiedContext NotifyDrivenContext;
@@ -2472,13 +2460,12 @@ void FControlRigParameterTrackEditor::OnAddTransformKeysForSelectedObjects(EMovi
 			if (UObject* Object = ObjectBinding->GetBoundObject())
 			{
 				const FName Name(*ControlRig->GetName());
-			
-				const TArray<FName> ControlNames = ControlRig->CurrentControlSelection();
-				for (const FName& ControlName : ControlNames)
+
+				for (const FRigElementKey& ControlKey : Selection.Value)
 				{
-					if (FRigControlElement* ControlElement = ControlRig->FindControl(ControlName))
+					if (FRigControlElement* ControlElement = ControlRig->FindControl(ControlKey.Name))
 					{
-						AddControlKeys(Object, ControlRig, Name, ControlName, ChannelsToKey,
+						AddControlKeys(Object, ControlRig, Name, ControlKey.Name, ChannelsToKey,
 							ESequencerKeyMode::ManualKeyForced, FLT_MAX, bInConstraintSpace);
 						FControlRigEditMode::NotifyDrivenControls(ControlRig, ControlElement->GetKey(), NotifyDrivenContext);
 					}
@@ -2488,6 +2475,31 @@ void FControlRigParameterTrackEditor::OnAddTransformKeysForSelectedObjects(EMovi
 	}
 }
 
+void FControlRigParameterTrackEditor::OnAddTransformKeysForSelectedObjects(EMovieSceneTransformChannel Channel)
+{
+	const TSharedPtr<ISequencer> Sequencer = GetSequencer();
+	if (!Sequencer || !Sequencer->IsAllowedToChange())
+	{
+		return;
+	}
+	
+	const FControlRigEditMode* ControlRigEditMode = GetEditMode();
+	if (!ControlRigEditMode)
+	{
+		return;
+	}
+
+	TMap<UControlRig*, TArray<FRigElementKey>> SelectedControls;
+	ControlRigEditMode->GetAllSelectedControls(SelectedControls);
+	if (SelectedControls.Num() <= 0)
+	{
+		return;
+	}
+	FScopedTransaction KeyTransaction(LOCTEXT("SetKeysOnControls", "Set Keys On Controls"), !GIsTransacting);
+	SetKeysOnControls(SelectedControls, Channel);
+	
+}
+
 //function to evaluate a Control and Set it on the ControlRig
 static void EvaluateThisControl(UMovieSceneControlRigParameterSection* Section, const FName& ControlName, const FFrameTime& FrameTime)
 {
@@ -4363,7 +4375,10 @@ void FControlRigParameterTrackEditor::GetControlRigKeys(
 
 			bool bMaskKeyOut = (SectionToKey->GetControlNameMask(ControlElement->GetFName()) == false);
 			bool bSetKey = ParameterName.IsNone() || (ControlElement->GetFName() == ParameterName && !bMaskKeyOut);
-
+			if (bSetKey == false)
+			{
+				continue;
+			}
 			FRigControlValue ControlValue = InControlRig->GetControlValue(ControlElement, ERigControlValueType::Current);
 
 			switch (ControlElement->Settings.ControlType)
@@ -4704,7 +4719,7 @@ FKeyPropertyResult FControlRigParameterTrackEditor::AddKeysToControlRig(
 	FKeyPropertyResult KeyPropertyResult;
 	
 	const TSharedPtr<ISequencer> Sequencer = GetSequencer();
-	if (!Sequencer)
+	if (!Sequencer || GeneratedKeys.Num() <= 0)
 	{
 		return KeyPropertyResult;
 	}
@@ -4781,7 +4796,10 @@ void FControlRigParameterTrackEditor::AddControlKeys(
 
 	TSharedRef<FGeneratedTrackKeys> GeneratedKeys = MakeShared<FGeneratedTrackKeys>();
 	GetControlRigKeys(InControlRig, RigControlName, ChannelsToKey, KeyMode, ParamSection, *GeneratedKeys, bInConstraintSpace);
-	
+	if (GeneratedKeys->Num() <= 0)
+	{
+		return;
+	}
 	TGuardValue<bool> Guard(bIsDoingSelection, true);
 
 	auto OnKeyProperty = [=, this](FFrameNumber Time) -> FKeyPropertyResult
@@ -4902,7 +4920,7 @@ bool FControlRigParameterTrackEditor::ModifyOurGeneratedKeysByCurrentAndWeight(U
 
 		FName ControlName = ControlElement->GetFName();
 		const FChannelMapInfo* pChannelIndex = Section->ControlChannelMap.Find(ControlName);
-		if (!pChannelIndex || pChannelIndex->GeneratedKeyIndex == INDEX_NONE)
+		if (!pChannelIndex || pChannelIndex->GeneratedKeyIndex == INDEX_NONE || RigControlName != ControlName)
 		{
 			continue;
 		}
@@ -6785,3 +6803,5 @@ FControlRigEditMode* FControlRigParameterTrackEditor::GetEditMode(bool bForceAct
 }
 
 #undef LOCTEXT_NAMESPACE
+
+
diff --git a/Engine/Plugins/Animation/ControlRig/Source/ControlRigEditor/Private/Sequencer/ControlRigParameterTrackEditor.h b/Engine/Plugins/Animation/ControlRig/Source/ControlRigEditor/Private/Sequencer/ControlRigParameterTrackEditor.h
--- a/Engine/Plugins/Animation/ControlRig/Source/ControlRigEditor/Private/Sequencer/ControlRigParameterTrackEditor.h
+++ b/Engine/Plugins/Animation/ControlRig/Source/ControlRigEditor/Private/Sequencer/ControlRigParameterTrackEditor.h
@@ -107,6 +107,14 @@ public:
 	 */
 	static TSharedRef<ISequencerTrackEditor> CreateTrackEditor(TSharedRef<ISequencer> OwningSequencer);
 
+
+	/**
+	*
+	* Cached version of the current track editor, that lets other animaiton tools call specific functions. in particularly setting keys on
+	* current values
+	* 
+	**/
+	static FControlRigParameterTrackEditor* CurrentEditor;
 public:
 
 	// ISequencerTrackEditor interface
@@ -136,6 +144,8 @@ public:
 	virtual void PostUndo(bool bSuccess) override;
 	virtual void PostRedo(bool bSuccess) override { PostUndo(bSuccess); }
 
+	//function to set keys on controls
+	void SetKeysOnControls(const TMap<UControlRig*, TArray<FRigElementKey>>& Controls, EMovieSceneTransformChannel Channel);
 private:
 
 	void ProcessKeyOperation(UObject* ObjectToKey, UMovieSceneControlRigParameterTrack* Track, TArrayView<const UE::Sequencer::FKeySectionOperation> SectionsToKey, ISequencer& InSequencer, FFrameNumber KeyTime, TArray<UE::Sequencer::FAddKeyResult>* OutResults);
diff --git a/Engine/Plugins/Animation/ControlRig/Source/ControlRigEditor/Private/Sequencer/SelectionSets/SelectionSets.cpp b/Engine/Plugins/Animation/ControlRig/Source/ControlRigEditor/Private/Sequencer/SelectionSets/SelectionSets.cpp
--- a/Engine/Plugins/Animation/ControlRig/Source/ControlRigEditor/Private/Sequencer/SelectionSets/SelectionSets.cpp
+++ b/Engine/Plugins/Animation/ControlRig/Source/ControlRigEditor/Private/Sequencer/SelectionSets/SelectionSets.cpp
@@ -1657,6 +1657,7 @@ bool UAIESelectionSets::KeyAll(TSharedPtr<ISequencer>& InSequencer, const FGuid&
 		const FScopedTransaction Transaction(LOCTEXT("KeySelectionSet_transaction", "Key Selection Set"), !GIsTransacting);
 
 		TArray<FGuid> AllGuids;
+		TMap<UControlRig*, TArray<FRigElementKey>> SelectedControls;
 		if (const FAIESelectionSetItem* SetItem = SelectionSets.Find(InGuid))
 		{
 			SetItem->GetAllChildren(this, AllGuids);
@@ -1678,24 +1679,23 @@ bool UAIESelectionSets::KeyAll(TSharedPtr<ISequencer>& InSequencer, const FGuid&
 										{
 											if (UMovieSceneControlRigParameterTrack* CRTrack = GetControlRigTrack(InSequencer, ControlRig))
 											{
+												TArray<FRigElementKey> ControlKeys;
 												for (const FAIESelectionSetItemName& Name : CurrentSetItem->Names)
 												{
 													FName ControlName(*Name.Name);
 													if (UMovieSceneControlRigParameterSection* CRSection = Cast<UMovieSceneControlRigParameterSection>(CRTrack->GetSectionToKey(ControlName)))
 													{
 														CRSection->Modify();
-														TArrayView<FMovieSceneFloatChannel*> BaseFloatChannels = CRSection->GetChannelProxy().GetChannels<FMovieSceneFloatChannel>();
 														if (FRigControlElement* Control = ControlRig->FindControl(ControlName))
 														{
-															int32 StartIndex = 0, EndIndex = 0;
-															if (FControlRigKeys::GetStartEndIndicesForControl(CRSection, Control, StartIndex, EndIndex))
-															{
-																SetCurrentKeys<FMovieSceneFloatChannel, float>(BaseFloatChannels, StartIndex, EndIndex, DefaultInterpolation, FrameNumber);
-
-															}
+															ControlKeys.Add(Control->GetKey());
 														}
 													}
 												}
+												if (ControlKeys.Num() > 0)
+												{
+													SelectedControls.Add(ControlRig, ControlKeys);
+												}
 											}
 										}
 										else if (AActor* Actor = Cast<AActor>(Object.Get()))
@@ -1710,11 +1710,7 @@ bool UAIESelectionSets::KeyAll(TSharedPtr<ISequencer>& InSequencer, const FGuid&
 													if (Section)
 													{
 														Section->Modify();
-														TArrayView<FMovieSceneDoubleChannel*> BaseDoubleChannels = Section->GetChannelProxy().GetChannels<FMovieSceneDoubleChannel>();
-														const int32 NumChannels = BaseDoubleChannels.Num();
-														const int32 StartIndex = 0;
-														const int32 EndIndex = NumChannels - 1;
-														SetCurrentKeys<FMovieSceneDoubleChannel, double>(BaseDoubleChannels, StartIndex, EndIndex, DefaultInterpolation, FrameNumber);
+														MovieSceneToolHelpers::AddTransformKeysForObject(Object.Get(), EMovieSceneTransformChannel::AllTransform);
 													}
 												}
 											}
@@ -1726,6 +1722,11 @@ bool UAIESelectionSets::KeyAll(TSharedPtr<ISequencer>& InSequencer, const FGuid&
 					}
 				}
 			}
+			if (FControlRigParameterTrackEditor::CurrentEditor && SelectedControls.Num() > 0)
+			{
+				EMovieSceneTransformChannel Channel = EMovieSceneTransformChannel::AllTransform;
+				FControlRigParameterTrackEditor::CurrentEditor->SetKeysOnControls(SelectedControls, Channel);
+			}
 			InSequencer->NotifyMovieSceneDataChanged(EMovieSceneDataChangeType::TrackValueChanged);
 
 			return true;
diff --git a/Engine/Source/Editor/MovieSceneTools/Private/MovieSceneToolHelpers.cpp b/Engine/Source/Editor/MovieSceneTools/Private/MovieSceneToolHelpers.cpp
--- a/Engine/Source/Editor/MovieSceneTools/Private/MovieSceneToolHelpers.cpp
+++ b/Engine/Source/Editor/MovieSceneTools/Private/MovieSceneToolHelpers.cpp
@@ -56,6 +56,7 @@
 #include "Evaluation/MovieSceneEvaluationTrack.h"
 #include "Evaluation/MovieSceneEvaluationTemplateInstance.h"
 #include "IMovieScenePlayer.h"
+#include "TrackEditors/TransformTrackEditor.h"
 #include "Tracks/MovieSceneCinematicShotTrack.h"
 #include "Tracks/MovieSceneCameraCutTrack.h"
 #include "Tracks/MovieScene3DConstraintTrack.h"
@@ -6067,6 +6068,15 @@ bool MovieSceneToolHelpers::AddTransformKeys(
 	return true;
 }
 
+void MovieSceneToolHelpers::AddTransformKeysForObject(UObject* Object, EMovieSceneTransformChannel ChannelToKey)
+{
+	if (F3DTransformTrackEditor::CurrentEditor)
+	{
+		F3DTransformTrackEditor::CurrentEditor->AddTransformKeysForObject(Object, ChannelToKey, ESequencerKeyMode::ManualKeyForced);
+	}
+}
+
+
 void AddUnwoundKey(FMovieSceneDoubleChannel& Channel, FFrameNumber Time, double Value)
 {
 	int32 Index = Channel.AddLinearKey(Time, Value);
diff --git a/Engine/Source/Editor/MovieSceneTools/Private/TrackEditors/TransformTrackEditor.cpp b/Engine/Source/Editor/MovieSceneTools/Private/TrackEditors/TransformTrackEditor.cpp
--- a/Engine/Source/Editor/MovieSceneTools/Private/TrackEditors/TransformTrackEditor.cpp
+++ b/Engine/Source/Editor/MovieSceneTools/Private/TrackEditors/TransformTrackEditor.cpp
@@ -76,7 +76,7 @@ void GetActorAndSceneComponentFromObject( UObject* Object, AActor*& OutActor, US
 	}
 }
 
-
+F3DTransformTrackEditor* F3DTransformTrackEditor::CurrentEditor = nullptr;
 FName F3DTransformTrackEditor::TransformPropertyName("Transform");
 
 F3DTransformTrackEditor::F3DTransformTrackEditor( TSharedRef<ISequencer> InSequencer )
@@ -116,10 +116,14 @@ F3DTransformTrackEditor::F3DTransformTrackEditor( TSharedRef<ISequencer> InSeque
 		AddTransformProperty(RotationProperty, EMovieSceneTransformChannel::Rotation);
 		AddTransformProperty(Scale3DProperty, EMovieSceneTransformChannel::Scale);
 	}
+	F3DTransformTrackEditor::CurrentEditor = this;
+
 }
 
 F3DTransformTrackEditor::~F3DTransformTrackEditor()
 {
+	F3DTransformTrackEditor::CurrentEditor = nullptr;
+
 	if (GEditor != nullptr)
 	{
 		GEditor->UnregisterForUndo(this);
@@ -160,6 +164,8 @@ void F3DTransformTrackEditor::OnRelease()
 			LevelVC->ViewFOV = LevelVC->FOVAngle;
 		}
 	}
+	F3DTransformTrackEditor::CurrentEditor = nullptr;
+
 }
 
 TSharedRef<ISequencerTrackEditor> F3DTransformTrackEditor::CreateTrackEditor( TSharedRef<ISequencer> InSequencer )
diff --git a/Engine/Source/Editor/MovieSceneTools/Private/TrackEditors/TransformTrackEditor.h b/Engine/Source/Editor/MovieSceneTools/Private/TrackEditors/TransformTrackEditor.h
--- a/Engine/Source/Editor/MovieSceneTools/Private/TrackEditors/TransformTrackEditor.h
+++ b/Engine/Source/Editor/MovieSceneTools/Private/TrackEditors/TransformTrackEditor.h
@@ -51,6 +51,15 @@ public:
 	 */
 	static TSharedRef<ISequencerTrackEditor> CreateTrackEditor( TSharedRef<ISequencer> OwningSequencer );
 
+
+	/**
+	*
+	* Cached version of the current track editor, that lets other animaiton tools call specific functions. in particularly setting keys on
+	* current values
+	*
+	**/
+	static F3DTransformTrackEditor* CurrentEditor;
+
 public:
 
 	// ISequencerTrackEditor interface
@@ -158,14 +167,6 @@ private:
 	 */
 	void AddTransformKeysForHandle( TArray<FGuid> ObjectHandles, EMovieSceneTransformChannel ChannelToKey, ESequencerKeyMode KeyMode );
 
-	/**
-	* Adds transform keys to a specific object.
-
-	* @param Object The object to add keys to.
-	* @param ChannelToKey The channels to add keys to.
-	* @param KeyParams Parameters which control how the keys are added.
-	*/
-	void AddTransformKeysForObject( UObject* Object, EMovieSceneTransformChannel ChannelToKey, ESequencerKeyMode KeyMode );
 
 	/**
 	* Adds keys to a specific actor.
@@ -179,6 +180,15 @@ private:
 
 	FTransformData RecomposeTransform(const FTransformData& InTransformData, UObject* AnimatedObject, UMovieSceneSection* Section);
 
+public:
+	/**
+	* Adds transform keys to a specific object.
+
+	* @param Object The object to add keys to.
+	* @param ChannelToKey The channels to add keys to.
+	* @param KeyParams Parameters which control how the keys are added.
+	*/
+	void AddTransformKeysForObject(UObject* Object, EMovieSceneTransformChannel ChannelToKey, ESequencerKeyMode KeyMode);
 private:
 
 	/** ConstraintChannel Delegates*/
diff --git a/Engine/Source/Editor/MovieSceneTools/Public/MovieSceneToolHelpers.h b/Engine/Source/Editor/MovieSceneTools/Public/MovieSceneToolHelpers.h
--- a/Engine/Source/Editor/MovieSceneTools/Public/MovieSceneToolHelpers.h
+++ b/Engine/Source/Editor/MovieSceneTools/Public/MovieSceneToolHelpers.h
@@ -1216,6 +1216,14 @@ public:
 		const TArray<FTransform>& InLocalTransforms,
 		const EMovieSceneTransformChannel& InChannels);
 
+	/**
+	* Adds transform keys to a specific object, using it's current transform
+
+	* @param Object The object to add keys to.
+	* @param ChannelToKey The channels to add keys to.
+*	*/
+	static UE_API void AddTransformKeysForObject(UObject* Object, EMovieSceneTransformChannel ChannelToKey);
+
 	/** Import an animation sequence's root transforms into a transform section */
 	static UE_API void ImportAnimSequenceTransforms(const TSharedRef<ISequencer>& Sequencer, const FAssetData& Asset, UMovieScene3DTransformTrack* TransformTrack);
 
