Anim In Engine: Fix Anim Layer Mute state follow up, SetActive does a modify so don't set it if the same value
#trivial

diff --git a/Engine/Plugins/Animation/ControlRig/Source/ControlRigEditor/Private/Sequencer/AnimLayers/AnimLayers.cpp b/Engine/Plugins/Animation/ControlRig/Source/ControlRigEditor/Private/Sequencer/AnimLayers/AnimLayers.cpp
--- a/Engine/Plugins/Animation/ControlRig/Source/ControlRigEditor/Private/Sequencer/AnimLayers/AnimLayers.cpp
+++ b/Engine/Plugins/Animation/ControlRig/Source/ControlRigEditor/Private/Sequencer/AnimLayers/AnimLayers.cpp
@@ -430,7 +430,7 @@ bool UAnimLayer::GetActive() const
 							Track->SetRowEvalDisabled(!CurrentVal.GetValue(), SectionItem.Section->GetRowIndex());
 						}
 						//make sure section is active
-						if (bActive)
+						if (bActive  && SectionItem.Section->IsActive() == false)
 						{
 							SectionItem.Section->SetIsActive(bActive);
 						}
