RigVM: Support for run once / single frame evaluation

#rnx

diff --git a/Engine/Plugins/Animation/ControlRig/Source/ControlRig/Private/ControlRig.cpp b/Engine/Plugins/Animation/ControlRig/Source/ControlRig/Private/ControlRig.cpp
--- a/Engine/Plugins/Animation/ControlRig/Source/ControlRig/Private/ControlRig.cpp
+++ b/Engine/Plugins/Animation/ControlRig/Source/ControlRig/Private/ControlRig.cpp
@@ -2321,6 +2321,9 @@ FRigControlValue UControlRig::GetControlValue(FRigControlElement* InControl, con
 void UControlRig::SetControlValueImpl(const FName& InControlName, const FRigControlValue& InValue, bool bNotify,
 	const FRigControlModifiedContext& Context, bool bSetupUndo, bool bPrintPythonCommnds, bool bFixEulerFlips)
 {
+	// make sure we can run this rig as we are interacting with it.
+	ResetEvaluationsLeft();
+	
 	const FRigElementKey Key(InControlName, ERigElementType::Control);
 
 	FRigControlElement* ControlElement = DynamicHierarchy->Find<FRigControlElement>(Key);
diff --git a/Engine/Plugins/Animation/ControlRig/Source/ControlRigEditor/Private/Editor/ControlRigEditor.cpp b/Engine/Plugins/Animation/ControlRig/Source/ControlRigEditor/Private/Editor/ControlRigEditor.cpp
--- a/Engine/Plugins/Animation/ControlRig/Source/ControlRigEditor/Private/Editor/ControlRigEditor.cpp
+++ b/Engine/Plugins/Animation/ControlRig/Source/ControlRigEditor/Private/Editor/ControlRigEditor.cpp
@@ -1110,9 +1110,11 @@ void FControlRigBaseEditor::HandleSetObjectBeingDebuggedImpl(UObject* InObject)
 		{
 			PreviouslyDebuggedControlRig->GetHierarchy()->OnModified().RemoveAll(this);
 			PreviouslyDebuggedControlRig->OnPreForwardsSolve_AnyThread().RemoveAll(this);
+			PreviouslyDebuggedControlRig->OnPostForwardsSolve_AnyThread().RemoveAll(this);
 			PreviouslyDebuggedControlRig->OnPreConstructionForUI_AnyThread().RemoveAll(this);
 			PreviouslyDebuggedControlRig->OnPreConstruction_AnyThread().RemoveAll(this);
 			PreviouslyDebuggedControlRig->OnPostConstruction_AnyThread().RemoveAll(this);
+			PreviouslyDebuggedControlRig->OnEarlyExit_AnyThread().RemoveAll(this);
 			PreviouslyDebuggedControlRig->ControlModified().RemoveAll(this);
 		}
 	}
@@ -1168,16 +1170,20 @@ void FControlRigBaseEditor::HandleSetObjectBeingDebuggedImpl(UObject* InObject)
 
 		DebuggedControlRig->GetHierarchy()->OnModified().RemoveAll(this);
 		DebuggedControlRig->OnPreForwardsSolve_AnyThread().RemoveAll(this);
+		DebuggedControlRig->OnPostForwardsSolve_AnyThread().RemoveAll(this);
 		DebuggedControlRig->OnPreConstructionForUI_AnyThread().RemoveAll(this);
 		DebuggedControlRig->OnPreConstruction_AnyThread().RemoveAll(this);
 		DebuggedControlRig->OnPostConstruction_AnyThread().RemoveAll(this);
+		DebuggedControlRig->OnEarlyExit_AnyThread().RemoveAll(this);
 		DebuggedControlRig->ControlModified().RemoveAll(this);
 
 		DebuggedControlRig->GetHierarchy()->OnModified().AddSP(SharedRef(), &FControlRigBaseEditor::OnHierarchyModified_AnyThread);
 		DebuggedControlRig->OnPreForwardsSolve_AnyThread().AddSP(SharedRef(), &FControlRigBaseEditor::OnPreForwardsSolve_AnyThread);
+		DebuggedControlRig->OnPostForwardsSolve_AnyThread().AddSP(SharedRef(), &FControlRigBaseEditor::OnPostForwardsSolve_AnyThread);
 		DebuggedControlRig->OnPreConstructionForUI_AnyThread().AddSP(SharedRef(), &FControlRigBaseEditor::OnPreConstructionForUI_AnyThread);
 		DebuggedControlRig->OnPreConstruction_AnyThread().AddSP(SharedRef(), &FControlRigBaseEditor::OnPreConstruction_AnyThread);
 		DebuggedControlRig->OnPostConstruction_AnyThread().AddSP(SharedRef(), &FControlRigBaseEditor::OnPostConstruction_AnyThread);
+		DebuggedControlRig->OnEarlyExit_AnyThread().AddSP(SharedRef(), &FControlRigBaseEditor::OnEarlyExit_AnyThread);
 		DebuggedControlRig->ControlModified().AddSP(SharedRef(), &FControlRigBaseEditor::HandleOnControlModified);
 		
 		LastHierarchyHash = INDEX_NONE;
@@ -3017,6 +3023,26 @@ EVisibility FControlRigBaseEditor::GetSchematicOverlayVisibility() const
 	return EVisibility::SelfHitTestInvisible;
 }
 
+bool FControlRigBaseEditor::IsRunningForwardSolve() const
+{
+	return GetEventQueue() == ForwardsSolveEventQueue;
+}
+
+void FControlRigBaseEditor::OnRunGraphOnceToggledImpl(bool bRunGraphOnceEnabled)
+{
+	if (UControlRig* ControlRig = GetControlRig())
+	{
+		if (bRunGraphOnceEnabled)
+		{
+			ControlRig->SetEvaluationsLeft(1);
+		}
+		else
+		{
+			ControlRig->ResetEvaluationsLeft();
+		}
+	}	
+}
+
 bool FControlRigBaseEditor::GetToolbarDrawAxesOnSelection() const
 {
 	if (const UControlRigEditModeSettings* Settings = GetDefault<UControlRigEditModeSettings>())
@@ -3236,9 +3262,13 @@ void FControlRigBaseEditor::UpdateRigVMHostImpl()
 			}
 
 			ControlRig->OnPreForwardsSolve_AnyThread().RemoveAll(this);
+			ControlRig->OnPostForwardsSolve_AnyThread().RemoveAll(this);
+			ControlRig->OnEarlyExit_AnyThread().RemoveAll(this);
 			ControlRig->ControlModified().RemoveAll(this);
 
 			ControlRig->OnPreForwardsSolve_AnyThread().AddSP(SharedRef(), &FControlRigBaseEditor::OnPreForwardsSolve_AnyThread);
+			ControlRig->OnPostForwardsSolve_AnyThread().AddSP(SharedRef(), &FControlRigBaseEditor::OnPostForwardsSolve_AnyThread);
+			ControlRig->OnEarlyExit_AnyThread().AddSP(SharedRef(), &FControlRigBaseEditor::OnEarlyExit_AnyThread);
 			ControlRig->ControlModified().AddSP(SharedRef(), &FControlRigBaseEditor::HandleOnControlModified);
 		}
 
@@ -6436,6 +6466,12 @@ TSharedPtr<FUICommandList> FControlRigBaseEditor::HandleOnViewportContextMenuCom
 
 void FControlRigBaseEditor::OnPreForwardsSolve_AnyThread(UControlRig* InRig, const FName& InEventName)
 {
+	// make sure to reapply this in case
+	if (SharedRigVMEditorRef()->IsRunGraphOnceOn())
+	{
+		OnRunGraphOnceToggledImpl(true);
+	}
+	
 	// if we are debugging a PIE instance, we need to remember the input pose on the
 	// rig so we can perform multiple evaluations. this is to avoid double transforms / double forward solve results.
 	if(InRig->GetWorld()->IsPlayInEditor())
@@ -6454,6 +6490,29 @@ void FControlRigBaseEditor::OnPreForwardsSolve_AnyThread(UControlRig* InRig, con
 	}
 }
 
+void FControlRigBaseEditor::OnPostForwardsSolve_AnyThread(UControlRig* InRig, const FName& InEventName)
+{
+	const UControlRig* ControlRig = GetControlRig();
+	if (!ControlRig || ControlRig != InRig)
+	{
+		return;
+	}
+
+	// this is the last evaluation
+	if (ControlRig->GetEvaluationsLeft() == 1)
+	{
+		if (const URigHierarchy* Hierarchy = ControlRig->GetHierarchy())
+		{
+			static const FRigElementKeyCollection EmptyCollection;
+			BonePoseFromLastEvaluation = Hierarchy->GetPose(false, ERigElementType::Bone, EmptyCollection, false);
+		}
+	}
+	else
+	{
+		BonePoseFromLastEvaluation.Reset();
+	}
+}
+
 void FControlRigBaseEditor::OnPreConstructionForUI_AnyThread(UControlRig* InRig, const FName& InEventName)
 {
 	bIsConstructionEventRunning = true;
@@ -6608,6 +6667,33 @@ void FControlRigBaseEditor::OnPostConstruction_AnyThread(UControlRig* InRig, con
 	}
 }
 
+void FControlRigBaseEditor::OnEarlyExit_AnyThread(URigVMHost* InHost, URigVMHost::EEarlyExitReason InReason)
+{
+	if (InReason != URigVMHost::EEarlyExitReason_NoEvaluationsLeft)
+	{
+		return;
+	}
+	
+	UControlRig* ControlRig = GetControlRig();
+	if (!ControlRig || ControlRig != InHost)
+	{
+		return;
+	}
+
+	URigHierarchy* Hierarchy = ControlRig->GetHierarchy();
+	if (!Hierarchy)
+	{
+		return;
+	}
+
+	if (BonePoseFromLastEvaluation.Num() == 0)
+	{
+		return;
+	}
+
+	Hierarchy->SetPose(BonePoseFromLastEvaluation, ERigTransformType::CurrentGlobal);
+}
+
 void FControlRigBaseEditor::SetupTimelineDelegates(FAnimationScrubPanelDelegates& InOutDelegates)
 {
 	TSharedRef<FControlRigBaseEditor> SharedEditor = SharedRef();
diff --git a/Engine/Plugins/Animation/ControlRig/Source/ControlRigEditor/Private/Editor/ControlRigEditor.h b/Engine/Plugins/Animation/ControlRig/Source/ControlRigEditor/Private/Editor/ControlRigEditor.h
--- a/Engine/Plugins/Animation/ControlRig/Source/ControlRigEditor/Private/Editor/ControlRigEditor.h
+++ b/Engine/Plugins/Animation/ControlRig/Source/ControlRigEditor/Private/Editor/ControlRigEditor.h
@@ -395,7 +395,10 @@ protected:
 	bool IsSchematicViewportActive() const;
 	EVisibility GetSchematicOverlayVisibility() const;
 
-		/** Handle switching skeletal meshes */
+	bool IsRunningForwardSolve() const;
+	void OnRunGraphOnceToggledImpl(bool bRunGraphOnceEnabled);
+
+	/** Handle switching skeletal meshes */
 	void HandlePreviewMeshChanged(USkeletalMesh* InOldSkeletalMesh, USkeletalMesh* InNewSkeletalMesh);
 
 	/** Push a newly compiled/opened control rig to the edit mode */
@@ -493,9 +496,11 @@ protected:
 	TWeakObjectPtr<AStaticMeshActor> WeakGroundActorPtr;
 
 	void OnPreForwardsSolve_AnyThread(UControlRig* InRig, const FName& InEventName);
+	void OnPostForwardsSolve_AnyThread(UControlRig* InRig, const FName& InEventName);
 	void OnPreConstructionForUI_AnyThread(UControlRig* InRig, const FName& InEventName);
 	void OnPreConstruction_AnyThread(UControlRig* InRig, const FName& InEventName);
 	void OnPostConstruction_AnyThread(UControlRig* InRig, const FName& InEventName);
+	void OnEarlyExit_AnyThread(URigVMHost* InHost, URigVMHost::EEarlyExitReason InReason);
 	FRigPose PreConstructionPose;
 	TArray<FRigSocketState> SocketStates;
 	TArray<FRigConnectorState> ConnectorStates;
@@ -512,6 +517,7 @@ protected:
 	FSimpleMulticastDelegate RequestNavigateToConnectorWarningDelegate;
 	TSharedPtr<SSchematicGraphPanel> SchematicViewport;
 	FTimerHandle RecordReplayTimerHandle;
+	FRigPose BonePoseFromLastEvaluation;
 
 public:
 	void SetupTimelineDelegates(FAnimationScrubPanelDelegates& InOutDelegates);
diff --git a/Engine/Plugins/Animation/ControlRig/Source/ControlRigEditor/Private/Editor/ControlRigLegacyEditor.h b/Engine/Plugins/Animation/ControlRig/Source/ControlRigEditor/Private/Editor/ControlRigLegacyEditor.h
--- a/Engine/Plugins/Animation/ControlRig/Source/ControlRigEditor/Private/Editor/ControlRigLegacyEditor.h
+++ b/Engine/Plugins/Animation/ControlRig/Source/ControlRigEditor/Private/Editor/ControlRigLegacyEditor.h
@@ -136,9 +136,11 @@ protected:
 	virtual void HandleVMExecutedEvent(URigVMHost* InHost, const FName& InEventName) override { return FControlRigBaseEditor::HandleVMExecutedEventImpl(InHost, InEventName); }
 	virtual void HandleVMExecutedEventSuper(URigVMHost* InHost, const FName& InEventName) override { return IControlRigLegacyEditor::HandleVMExecutedEvent(InHost, InEventName); }
 
+	virtual bool CanRunGraphOnce() const override { return IsRunningForwardSolve(); }
+	virtual void OnRunGraphOnceToggled(bool bRunGraphOnceEnabled) override { OnRunGraphOnceToggledImpl(bRunGraphOnceEnabled); }
+
 	// FBaseToolKit overrides
 	virtual void CreateEditorModeManager() override { return FControlRigBaseEditor::CreateEditorModeManagerImpl(); }
-
 private:
 	/** Fill the toolbar with content */
 	virtual void FillToolbar(FToolBarBuilder& ToolbarBuilder, bool bEndSection = true) override { return FControlRigBaseEditor::FillToolbarImpl(ToolbarBuilder, bEndSection); }
diff --git a/Engine/Plugins/Animation/ControlRig/Source/ControlRigEditor/Private/Editor/ControlRigNewEditor.h b/Engine/Plugins/Animation/ControlRig/Source/ControlRigEditor/Private/Editor/ControlRigNewEditor.h
--- a/Engine/Plugins/Animation/ControlRig/Source/ControlRigEditor/Private/Editor/ControlRigNewEditor.h
+++ b/Engine/Plugins/Animation/ControlRig/Source/ControlRigEditor/Private/Editor/ControlRigNewEditor.h
@@ -136,6 +136,8 @@ protected:
 
 	//~ Begin IRigVMEditor interface
 	virtual void InitToolMenuContext(FToolMenuContext& MenuContext) override;
+	virtual bool CanRunGraphOnce() const override { return IsRunningForwardSolve(); }
+	virtual void OnRunGraphOnceToggled(bool bRunGraphOnceEnabled) override { OnRunGraphOnceToggledImpl(bRunGraphOnceEnabled); }
 	//~ End IRigVMEditor interface
 
 private:
diff --git a/Engine/Plugins/Runtime/RigVM/Source/RigVM/Private/RigVMHost.cpp b/Engine/Plugins/Runtime/RigVM/Source/RigVM/Private/RigVMHost.cpp
--- a/Engine/Plugins/Runtime/RigVM/Source/RigVM/Private/RigVMHost.cpp
+++ b/Engine/Plugins/Runtime/RigVM/Source/RigVM/Private/RigVMHost.cpp
@@ -41,6 +41,7 @@ URigVMHost::URigVMHost(const FObjectInitializer& ObjectInitializer)
 	, EventQueue()
 	, EventQueueToRun()
 	, EventsToRunOnce()
+	, EvaluationsLeft(INDEX_NONE)
 	, bRequiresInitExecution(false)
 	, InitBracket(0)
 	, ExecuteBracket(0)
@@ -316,6 +317,16 @@ void URigVMHost::Evaluate_AnyThread()
 	
 	// The EventQueueToRun should only be modified in this function
 	ensureMsgf(EventQueueToRun.IsEmpty(), TEXT("Detected a recursive call to the control rig evaluation function %s"), *GetPackage()->GetPathName());
+
+	// make sure to only evaluate if we have valid evaluations left	
+	if (EvaluationsLeft >= 0)
+	{
+		if (EvaluationsLeft == 0)
+		{
+			EarlyExitEvent.Broadcast(this, EEarlyExitReason_NoEvaluationsLeft);
+			return;
+		}
+	}
 	
 	// create a copy since we need to change it here temporarily,
 	// and UI / the rig may change the event queue while it is running
@@ -342,6 +353,11 @@ void URigVMHost::Evaluate_AnyThread()
 	{
 		Execute(EventName);
 	}
+
+	if (EvaluationsLeft > 0)
+	{
+		EvaluationsLeft--;
+	}
 }
 
 TArray<FRigVMExternalVariable> URigVMHost::GetExternalVariables() const
diff --git a/Engine/Plugins/Runtime/RigVM/Source/RigVM/Public/RigVMHost.h b/Engine/Plugins/Runtime/RigVM/Source/RigVM/Public/RigVMHost.h
--- a/Engine/Plugins/Runtime/RigVM/Source/RigVM/Public/RigVMHost.h
+++ b/Engine/Plugins/Runtime/RigVM/Source/RigVM/Public/RigVMHost.h
@@ -293,6 +293,14 @@ public:
 	FRigVMExecutedEvent& OnExecuted_AnyThread() { return ExecutedEvent; }
 	FRigVMExecutedEvent& OnPreExecuted_AnyThread() { return PreExecutedEvent; }
 
+	enum EEarlyExitReason
+	{
+		EEarlyExitReason_NoEvaluationsLeft
+	};
+
+	DECLARE_EVENT_TwoParams(URigVMHost, FRigVMHostEarlyExitEvent, class URigVMHost*, EEarlyExitReason);
+	FRigVMHostEarlyExitEvent& OnEarlyExit_AnyThread() { return EarlyExitEvent; }
+
 	const FRigVMDrawInterface& GetDrawInterface() const { return DrawInterface; };
 	FRigVMDrawInterface& GetDrawInterface() { return DrawInterface; };
 
@@ -443,6 +451,33 @@ public:
 	/** Returns true if Evaluate_AnyThread is currently executing */
 	bool IsEvaluating() const { return !EventQueueToRun.IsEmpty(); }
 
+	/** Returns true if this host can run at least once */
+	bool HasEvaluationsLeft() const
+	{
+		return EvaluationsLeft == INDEX_NONE || EvaluationsLeft > 0;
+	}
+
+	/**
+	 * Returns the number of evaluations left for this host
+	 * or INDEX_NONE if the host can run indefinitely.
+	 */
+	int32 GetEvaluationsLeft() const
+	{
+		return EvaluationsLeft;
+	}
+
+	/** Request to limit this hosts number of evaluations */
+	void SetEvaluationsLeft(int32 InEvaluationsLeft)
+	{
+		EvaluationsLeft = FMath::Max<int32>(INDEX_NONE, InEvaluationsLeft);
+	}
+
+	/** Reset the number of evaluations left and allow this host to run indefinitely */
+	void ResetEvaluationsLeft()
+	{
+		SetEvaluationsLeft(INDEX_NONE);
+	}
+
 	/** Copy the VM from the default object */
 	UE_API void InstantiateVMFromCDO();
 	
@@ -458,6 +493,9 @@ public:
 	/** Broadcasts a notification before the RigVMHost is executed / updated. */
 	FRigVMExecutedEvent PreExecutedEvent;
 
+	/** Broadcasts when the host cannot execute and needs to early exit during an evaluation */
+	FRigVMHostEarlyExitEvent EarlyExitEvent;
+
 protected: 
 
 	UE_API virtual void InitializeFromCDO();
@@ -483,9 +521,11 @@ protected:
 	TArray<TObjectPtr<UAssetUserData>> AssetUserDataEditorOnly;
 #endif
 
+protected:
 
+	// the number of evaluations left. INDEX_NONE represents unlimited
+	int32 EvaluationsLeft;
 
-protected:
 	bool bRequiresInitExecution;
 
 	int32 InitBracket;
diff --git a/Engine/Plugins/Runtime/RigVM/Source/RigVMDeveloper/Private/RigVMAsset.cpp b/Engine/Plugins/Runtime/RigVM/Source/RigVMDeveloper/Private/RigVMAsset.cpp
--- a/Engine/Plugins/Runtime/RigVM/Source/RigVMDeveloper/Private/RigVMAsset.cpp
+++ b/Engine/Plugins/Runtime/RigVM/Source/RigVMDeveloper/Private/RigVMAsset.cpp
@@ -3844,6 +3844,7 @@ void IRigVMAssetInterface::InitializeArchetypeInstances()
 			{
 				InstanceHost->PostInitInstanceIfRequired();
 			}
+			InstanceHost->ResetEvaluationsLeft();
 			InstanceHost->InstantiateVMFromCDO();
 			InstanceHost->CopyExternalVariableDefaultValuesFromCDO();
 		}
diff --git a/Engine/Plugins/Runtime/RigVM/Source/RigVMEditor/Private/Editor/RigVMEditor.cpp b/Engine/Plugins/Runtime/RigVM/Source/RigVMEditor/Private/Editor/RigVMEditor.cpp
--- a/Engine/Plugins/Runtime/RigVM/Source/RigVMEditor/Private/Editor/RigVMEditor.cpp
+++ b/Engine/Plugins/Runtime/RigVM/Source/RigVMEditor/Private/Editor/RigVMEditor.cpp
@@ -65,7 +65,7 @@ FRigVMEditorBase::FRigVMEditorBase()
 	, bIsCompilingThroughUI(false)
 	, WrapperObjects()
 	, LastEventQueue()
-
+	, bRunGraphOnce(false)
 {
 }
 
@@ -1114,6 +1114,9 @@ void FRigVMEditorBase::Compile()
 			return;
 		}
 
+		// Disable single frame evaluation mode
+		SetRunGraphOnce(false);
+
 		// force to disable the supended notif brackets
 		RigVMBlueprint->GetRigVMClient()->bSuspendModelNotificationsForOthers = false;
 		RigVMBlueprint->bSuspendModelNotificationsForSelf = false;
@@ -1436,6 +1439,12 @@ void FRigVMEditorBase::BindCommands()
 		FExecuteAction::CreateSP(ThisRef, &FRigVMEditorBase::ToggleEventQueue),
 		FCanExecuteAction());
 
+	ToolkitCommands->MapAction(
+		EditorCommands.RunGraphOnce,
+		FExecuteAction::CreateSP(ThisRef, &FRigVMEditorBase::ToggleRunGraphOnce), 
+		FIsActionChecked::CreateSP(ThisRef, &FRigVMEditorBase::CanRunGraphOnce),
+		FIsActionChecked::CreateSP(ThisRef, &FRigVMEditorBase::IsRunGraphOnceOn));
+
 	ToolkitCommands->MapAction(
 		EditorCommands.FrameSelection,
 		FExecuteAction::CreateSP(ThisRef, &FRigVMEditorBase::FrameSelection),
@@ -1542,6 +1551,30 @@ void FRigVMEditorBase::ToggleEventQueue()
 	SetEventQueue(LastEventQueue);
 }
 
+void FRigVMEditorBase::ToggleRunGraphOnce()
+{
+	SetRunGraphOnce(!bRunGraphOnce);
+}
+
+void FRigVMEditorBase::SetRunGraphOnce(bool InRunGraphOnceEnabled)
+{
+	if (bRunGraphOnce == InRunGraphOnceEnabled)
+	{
+		return;
+	}
+	bRunGraphOnce = InRunGraphOnceEnabled;
+	OnRunGraphOnceToggled(bRunGraphOnce);
+}
+
+bool FRigVMEditorBase::IsRunGraphOnceOn() const
+{
+	if (CanRunGraphOnce())
+	{
+		return bRunGraphOnce;
+	}
+	return false;
+}
+
 TSharedRef<SWidget> FRigVMEditorBase::GenerateEventQueueMenuContent()
 {
 	FMenuBuilder MenuBuilder(true, GetToolkitCommands());
@@ -2626,6 +2659,8 @@ void FRigVMEditorBase::UpdateRigVMHost()
  		}
 		check(IsValid(RigVMHost));
 
+		RigVMHost->ResetEvaluationsLeft();
+
 		CacheNameLists();
 
 		// Make sure the object being debugged is the preview instance
@@ -3558,7 +3593,7 @@ void FRigVMEditorBase::FillToolbar(FToolBarBuilder& ToolbarBuilder, bool bEndSec
 			TAttribute<FText>(), 
 			TAttribute<FSlateIcon>::Create(TAttribute<FSlateIcon>::FGetter::CreateSP(StaticCastSharedRef<FRigVMEditorBase>(SharedRef()), &FRigVMEditorBase::GetEventQueueIcon))
 		);
-
+		
 		FUIAction DefaultAction;
 		ToolbarBuilder.AddComboButton(
 			DefaultAction,
@@ -3568,6 +3603,7 @@ void FRigVMEditorBase::FillToolbar(FToolBarBuilder& ToolbarBuilder, bool bEndSec
 			FSlateIcon(FAppStyle::GetAppStyleSetName(), "LevelEditor.Recompile"),
 			true);
 
+		AddRunOnceWidget(ToolbarBuilder);
 		AddAutoCompileWidget(ToolbarBuilder);
 		AddSelectedDebugObjectWidget(ToolbarBuilder);
 
@@ -4405,6 +4441,41 @@ bool FRigVMEditorBase::IsToggleFadeOutUnrelatedNodesChecked() const
 	return GetDefault<URigVMEditorSettings>()->bFadeOutUnrelatedNodes;
 }
 
+void FRigVMEditorBase::AddRunOnceWidget(FToolBarBuilder& ToolbarBuilder)
+{
+	TWeakPtr<IRigVMEditor> WeakThis = SharedRef().ToWeakPtr();
+
+	ToolbarBuilder.AddToolBarButton(
+		FRigVMEditorCommands::Get().RunGraphOnce,
+		NAME_None,
+		FText(),
+		TAttribute<FText>(),
+		TAttribute<FSlateIcon>::CreateLambda(
+			[WeakThis]() -> FSlateIcon
+			{
+				static const FSlateIcon RunningContinuously(FAppStyle::GetAppStyleSetName(), "PlayWorld.PlaySimulate");
+				static const FSlateIcon RunningOnce(FAppStyle::GetAppStyleSetName(), "PlayWorld.SingleFrameAdvance.Small");
+				if (WeakThis.IsValid())
+				{
+					if (TSharedPtr<IRigVMEditor> StrongThis = WeakThis.Pin())
+					{
+						const FRigVMEditorBase* Editor = static_cast<const FRigVMEditorBase*>(StrongThis.Get());
+						if (Editor->IsRunGraphOnceOn())
+						{
+							return RunningOnce;
+						}
+					}
+				}
+				return RunningContinuously;
+			}
+		),
+		NAME_None,
+		FNewMenuDelegate(),
+		TAttribute<EVisibility>(),
+		FText()
+	);
+}
+
 void FRigVMEditorBase::OnGraphEditorTick(const FGeometry& AllottedGeometry, const double InCurrentTime, const float InDeltaTime,
                                          TSharedRef<SGraphEditor> InGraphEditor, URigVMEdGraph* InGraph)
 {
diff --git a/Engine/Plugins/Runtime/RigVM/Source/RigVMEditor/Private/Editor/RigVMEditorCommands.cpp b/Engine/Plugins/Runtime/RigVM/Source/RigVMEditor/Private/Editor/RigVMEditorCommands.cpp
--- a/Engine/Plugins/Runtime/RigVM/Source/RigVMEditor/Private/Editor/RigVMEditorCommands.cpp
+++ b/Engine/Plugins/Runtime/RigVM/Source/RigVMEditor/Private/Editor/RigVMEditorCommands.cpp
@@ -27,6 +27,7 @@ void FRigVMEditorCommands::RegisterCommands()
 	UI_COMMAND(ExecuteGraph, "Execute", "Execute the rig graph if On.", EUserInterfaceActionType::ToggleButton, FInputChord());
 	UI_COMMAND(AutoCompileGraph, "Auto Compile", "Auto-compile the rig graph if On.", EUserInterfaceActionType::ToggleButton, FInputChord());
 	UI_COMMAND(ToggleEventQueue, "Toggle Event", "Toggle between the current and last running event", EUserInterfaceActionType::Button, FInputChord());
+	UI_COMMAND(RunGraphOnce, "Run Once", "Only evaluate a single frame if On.", EUserInterfaceActionType::ToggleButton, FInputChord());
 	UI_COMMAND(FrameSelection, "Frame Selection", "Frames the selected nodes in the Graph View.", EUserInterfaceActionType::Button, FInputChord(EKeys::F));
 	UI_COMMAND(SwapFunctionWithinAsset, "Swap Function (Asset)", "Swaps a function for all occurrences within this asset.", EUserInterfaceActionType::Button, FInputChord());
 	UI_COMMAND(SwapFunctionAcrossProject, "Swap Function (Project)", "Swaps a function for all occurrences in the project.", EUserInterfaceActionType::Button, FInputChord());
diff --git a/Engine/Plugins/Runtime/RigVM/Source/RigVMEditor/Private/Editor/RigVMEditorCommands.h b/Engine/Plugins/Runtime/RigVM/Source/RigVMEditor/Private/Editor/RigVMEditorCommands.h
--- a/Engine/Plugins/Runtime/RigVM/Source/RigVMEditor/Private/Editor/RigVMEditorCommands.h
+++ b/Engine/Plugins/Runtime/RigVM/Source/RigVMEditor/Private/Editor/RigVMEditorCommands.h
@@ -46,6 +46,9 @@ public:
 	/** Toggle between this and the last event queue */
 	TSharedPtr< FUICommandInfo > ToggleEventQueue;
 
+	/** Run the graph once or continuously */
+	TSharedPtr< FUICommandInfo > RunGraphOnce;
+
 	/** Frames the selected nodes */
 	TSharedPtr< FUICommandInfo > FrameSelection;
 
diff --git a/Engine/Plugins/Runtime/RigVM/Source/RigVMEditor/Public/Editor/RigVMEditor.h b/Engine/Plugins/Runtime/RigVM/Source/RigVMEditor/Public/Editor/RigVMEditor.h
--- a/Engine/Plugins/Runtime/RigVM/Source/RigVMEditor/Public/Editor/RigVMEditor.h
+++ b/Engine/Plugins/Runtime/RigVM/Source/RigVMEditor/Public/Editor/RigVMEditor.h
@@ -176,7 +176,10 @@ public:
 	virtual void HandleModifiedEvent(ERigVMGraphNotifType InNotifType, URigVMGraph* InGraph, UObject* InSubject) = 0;
 	virtual void OnFinishedChangingProperties(const FPropertyChangedEvent& PropertyChangedEvent) = 0;
 	virtual FPreviewScene* GetPreviewScene() = 0;
-	
+
+	virtual bool IsRunGraphOnceOn() const = 0;
+	virtual bool CanRunGraphOnce() const = 0;
+
 protected:
 	
 	virtual TSharedRef<IRigVMEditor> SharedRef() = 0;
@@ -249,6 +252,7 @@ protected:
 	virtual void AddSettingsAndDefaultWidget(FToolBarBuilder& ToolbarBuilder) = 0;
 	virtual void AddSelectedDebugObjectWidget(FToolBarBuilder& ToolbarBuilder) = 0;
 	virtual void AddAutoCompileWidget(FToolBarBuilder& ToolbarBuilder) = 0;
+	virtual void AddRunOnceWidget(FToolBarBuilder& ToolbarBuilder) = 0;
 };
 
 class FRigVMEditorBase : public IRigVMEditor
@@ -462,6 +466,11 @@ protected:
 	UE_API bool IsAutoCompileGraphOn() const;
 	bool CanAutoCompileGraph() const { return true; }
 	UE_API void ToggleEventQueue();
+	UE_API void ToggleRunGraphOnce();
+	UE_API void SetRunGraphOnce(bool InRunGraphOnceEnabled);
+	UE_API virtual bool IsRunGraphOnceOn() const override;
+	virtual bool CanRunGraphOnce() const override { return false; }
+	virtual void OnRunGraphOnceToggled(bool bRunGraphOnceEnabled) {}
 	UE_API TSharedRef<SWidget> GenerateEventQueueMenuContent();
 	UE_API virtual FMenuBuilder GenerateBulkEditMenu();
 	UE_API TSharedRef<SWidget> GenerateBulkEditMenuContent();
@@ -558,6 +567,8 @@ protected:
 	UE_API virtual void ToggleFadeOutUnrelateNodes();
 	UE_API virtual bool IsToggleFadeOutUnrelatedNodesChecked() const;
 
+	UE_API virtual void AddRunOnceWidget(FToolBarBuilder& ToolbarBuilder) override;
+
 public:
 	UE_API void OnGraphEditorTick(const FGeometry& AllottedGeometry, const double InCurrentTime, const float InDeltaTime,
 		TSharedRef<SGraphEditor> InGraphEditor, URigVMEdGraph* InGraph);
@@ -612,6 +623,8 @@ protected:
 
 	FOnRefreshEvent OnRefreshEvent;
 
+	bool bRunGraphOnce;
+
 	friend class SRigVMExecutionStackView;
 	friend class SRigVMEditorGraphExplorer;
 
