Anim In Engine: Set  Reset Controls to false as default when loading animation into a section,
#trivial

diff --git a/Engine/Plugins/Animation/ControlRig/Source/ControlRig/Public/Sequencer/MovieSceneControlRigParameterSection.h b/Engine/Plugins/Animation/ControlRig/Source/ControlRig/Public/Sequencer/MovieSceneControlRigParameterSection.h
--- a/Engine/Plugins/Animation/ControlRig/Source/ControlRig/Public/Sequencer/MovieSceneControlRigParameterSection.h
+++ b/Engine/Plugins/Animation/ControlRig/Source/ControlRig/Public/Sequencer/MovieSceneControlRigParameterSection.h
@@ -404,7 +404,7 @@ public:
 		// key reduction tolerance
 		float Tolerance = 0.5;
 		//Whether to rest to default control states
-		bool bResetControls = true;
+		bool bResetControls = false;
 		//onto selected controls only
 		bool bOntoSelectedControls = false;
 		//Frame to Insert at
diff --git a/Engine/Plugins/Animation/ControlRig/Source/ControlRigEditor/Private/BakeToControlRigSettings.h b/Engine/Plugins/Animation/ControlRig/Source/ControlRigEditor/Private/BakeToControlRigSettings.h
--- a/Engine/Plugins/Animation/ControlRig/Source/ControlRigEditor/Private/BakeToControlRigSettings.h
+++ b/Engine/Plugins/Animation/ControlRig/Source/ControlRigEditor/Private/BakeToControlRigSettings.h
@@ -27,7 +27,8 @@ public:
 	FSmartReduceParams SmartReduce;
 
 	/** Reset controls to initial value on every frame */
-	UPROPERTY(EditAnywhere, Category = "Reset Controls")
+	UE_DEPRECATED(5.8, "When baking onto a new control rig always reset controls")
+	UPROPERTY()
 	bool bResetControls = true;
 
 	/** Resets the default properties. */
diff --git a/Engine/Plugins/Animation/ControlRig/Source/ControlRigEditor/Private/SBakeToControlRigDialog.cpp b/Engine/Plugins/Animation/ControlRig/Source/ControlRigEditor/Private/SBakeToControlRigDialog.cpp
--- a/Engine/Plugins/Animation/ControlRig/Source/ControlRigEditor/Private/SBakeToControlRigDialog.cpp
+++ b/Engine/Plugins/Animation/ControlRig/Source/ControlRigEditor/Private/SBakeToControlRigDialog.cpp
@@ -74,7 +74,7 @@ private:
 		TSharedPtr<SWindow> Window = FSlateApplication::Get().FindWidgetWindow(AsShared());
 		if (BakeSettings && Delegate.IsBound())
 		{
-			Delegate.Execute(BakeSettings->bReduceKeys, BakeSettings->SmartReduce.TolerancePercentage, BakeSettings->SmartReduce.SampleRate, BakeSettings->bResetControls);
+			Delegate.Execute(BakeSettings->bReduceKeys, BakeSettings->SmartReduce.TolerancePercentage, BakeSettings->SmartReduce.SampleRate, true /*bResetControls  always*/);
 		}
 		if (Window.IsValid())
 		{
diff --git a/Engine/Plugins/Animation/ControlRig/Source/ControlRigEditor/Private/Sequencer/ControlRigParameterTrackEditor.cpp b/Engine/Plugins/Animation/ControlRig/Source/ControlRigEditor/Private/Sequencer/ControlRigParameterTrackEditor.cpp
--- a/Engine/Plugins/Animation/ControlRig/Source/ControlRigEditor/Private/Sequencer/ControlRigParameterTrackEditor.cpp
+++ b/Engine/Plugins/Animation/ControlRig/Source/ControlRigEditor/Private/Sequencer/ControlRigParameterTrackEditor.cpp
@@ -1316,7 +1316,7 @@ void FControlRigParameterTrackEditor::BakeToControlRig(UClass* InClass, FGuid Ob
 					TOptional<TRange<FFrameNumber>> AnimLoadFrameRamge; //do whole range
 					const bool bOntoSelectedControls = false;
 					LoadAnimationIntoSection(Sequencer, TempAnimSequence, SkelMeshComp, StartFrame,
-						BakeSettings->bReduceKeys, BakeSettings->SmartReduce, BakeSettings->bResetControls, AnimLoadFrameRamge, bOntoSelectedControls, ParamSection);
+						BakeSettings->bReduceKeys, BakeSettings->SmartReduce, true /*reset controls*/, AnimLoadFrameRamge, bOntoSelectedControls, ParamSection);
 
 					//Turn Off Any Skeletal Animation Tracks
 					TArray<UMovieSceneSkeletalAnimationTrack*> SkelAnimationTracks;
diff --git a/Engine/Plugins/Animation/ControlRig/Source/ControlRigEditor/Private/Sequencer/LoadAnimToControlRigSettings.cpp b/Engine/Plugins/Animation/ControlRig/Source/ControlRigEditor/Private/Sequencer/LoadAnimToControlRigSettings.cpp
--- a/Engine/Plugins/Animation/ControlRig/Source/ControlRigEditor/Private/Sequencer/LoadAnimToControlRigSettings.cpp
+++ b/Engine/Plugins/Animation/ControlRig/Source/ControlRigEditor/Private/Sequencer/LoadAnimToControlRigSettings.cpp
@@ -16,5 +16,7 @@ void ULoadAnimToControlRigSettings::Reset()
 	SmartReduce.Reset();
 	bUseCustomTimeRange = false;
 	bOntoSelectedControls = false;
+	bResetControls = false;
+
 }
 
diff --git a/Engine/Plugins/Animation/ControlRig/Source/ControlRigEditor/Private/Sequencer/LoadAnimToControlRigSettings.h b/Engine/Plugins/Animation/ControlRig/Source/ControlRigEditor/Private/Sequencer/LoadAnimToControlRigSettings.h
--- a/Engine/Plugins/Animation/ControlRig/Source/ControlRigEditor/Private/Sequencer/LoadAnimToControlRigSettings.h
+++ b/Engine/Plugins/Animation/ControlRig/Source/ControlRigEditor/Private/Sequencer/LoadAnimToControlRigSettings.h
@@ -37,7 +37,7 @@ public:
 
 	/** Reset controls to initial value on every frame */
 	UPROPERTY(EditAnywhere, Category = "Reset Controls")
-	bool bResetControls = true;
+	bool bResetControls = false;
 
 	/** Resets the default properties. */
 	void Reset();
