[ControlRig]
* Fixed header including itself

diff --git a/Engine/Plugins/Animation/ControlRig/Source/ControlRigEditor/Private/Editor/Hierarchy/RigHierarchyTreeDelegates.h b/Engine/Plugins/Animation/ControlRig/Source/ControlRigEditor/Private/Editor/Hierarchy/RigHierarchyTreeDelegates.h
--- a/Engine/Plugins/Animation/ControlRig/Source/ControlRigEditor/Private/Editor/Hierarchy/RigHierarchyTreeDelegates.h
+++ b/Engine/Plugins/Animation/ControlRig/Source/ControlRigEditor/Private/Editor/Hierarchy/RigHierarchyTreeDelegates.h
@@ -3,7 +3,6 @@
 #pragma once
 
 #include "Delegates/Delegate.h"
-#include "Editor/Hierarchy/RigHierarchyTreeDelegates.h"
 #include "Editor/Hierarchy/RigHierarchyTreeDisplaySettings.h"
 #include "Editor/Hierarchy/Widgets/SRigHierarchyTagWidget.h"
 #include "Framework/SlateDelegates.h"
