Anim In Engine: Update Control Rig Rename function to fix up space channels and contraints(the best it can). Add new BP function to replace a control rig in a track

diff --git a/Engine/Plugins/Animation/ControlRig/Source/ControlRig/Private/Sequencer/MovieSceneControlRigParameterSection.cpp b/Engine/Plugins/Animation/ControlRig/Source/ControlRig/Private/Sequencer/MovieSceneControlRigParameterSection.cpp
--- a/Engine/Plugins/Animation/ControlRig/Source/ControlRig/Private/Sequencer/MovieSceneControlRigParameterSection.cpp
+++ b/Engine/Plugins/Animation/ControlRig/Source/ControlRig/Private/Sequencer/MovieSceneControlRigParameterSection.cpp
@@ -3064,7 +3064,39 @@ bool UMovieSceneControlRigParameterSection::RenameParameterName(const FName& Old
 		RenameParameterNameInner(ColorParameterNamesAndCurves);
 		RenameParameterNameInner(TransformParameterNamesAndCurves);
 	}
-
+	for (FSpaceControlNameAndChannel& SpaceChannel : SpaceChannels)
+	{
+		if (SpaceChannel.ControlName == OldParameterName)
+		{
+			SpaceChannel.ControlName = NewParameterName;
+			bWasReplaced = true;
+		}
+		//update spaces
+		TMovieSceneChannelData<FMovieSceneControlRigSpaceBaseKey> ChannelData = SpaceChannel.SpaceCurve.GetData();
+		TArrayView<const FFrameNumber> Times = ChannelData.GetTimes();
+		TArrayView<FMovieSceneControlRigSpaceBaseKey> Values = ChannelData.GetValues();
+		for (FMovieSceneControlRigSpaceBaseKey& Key : Values)
+		{
+			if (Key.ControlRigElement.Name == OldParameterName)
+			{
+				Key.ControlRigElement.Name = NewParameterName;
+			}
+		}
+	}
+	for (FConstraintAndActiveChannel& ConstraintChannel : ConstraintsChannels)
+	{
+		if (UTickableTransformConstraint* TransformConstraint = Cast<UTickableTransformConstraint>(ConstraintChannel.GetConstraint()))
+		{
+			if (UTransformableControlHandle* Handle = Cast<UTransformableControlHandle>(TransformConstraint->ChildTRSHandle))
+			{
+				if (Handle->ControlName == OldParameterName)
+				{
+					Handle->ControlName = NewParameterName;
+					bWasReplaced = true;
+				}
+			}
+		}
+	}
 	if (bWasReplaced)
 	{
 		ReconstructChannelProxy();
diff --git a/Engine/Plugins/Animation/ControlRig/Source/ControlRigEditor/Private/ControlRigSequencerEditorLibrary.cpp b/Engine/Plugins/Animation/ControlRig/Source/ControlRigEditor/Private/ControlRigSequencerEditorLibrary.cpp
--- a/Engine/Plugins/Animation/ControlRig/Source/ControlRigEditor/Private/ControlRigSequencerEditorLibrary.cpp
+++ b/Engine/Plugins/Animation/ControlRig/Source/ControlRigEditor/Private/ControlRigSequencerEditorLibrary.cpp
@@ -3232,6 +3232,33 @@ bool UControlRigSequencerEditorLibrary::MoveControlRigSpace(ULevelSequence* Leve
 	}
 	return bValid;
 }
+bool UControlRigSequencerEditorLibrary::ReplaceControlRig(UMovieSceneTrack* InTrack, UControlRig* InControlRig)
+{
+	if (InControlRig == nullptr || InTrack == nullptr)
+	{
+		UE_LOG(LogControlRig, Error, TEXT("ControlRig and Control Rig Track must be valid"));
+		return false;
+	}
+	if (UMovieSceneControlRigParameterTrack* CRTrack = Cast<UMovieSceneControlRigParameterTrack>(InTrack))
+	{
+		if (CRTrack->GetControlRig() != InControlRig)
+		{
+			CRTrack->Modify();
+			CRTrack->ReplaceControlRig(InControlRig, true /*RecreateChannels*/);
+		}
+		else
+		{
+			UE_LOG(LogControlRig, Error, TEXT("Control Rig already exists on the track"));
+			return false;
+		}
+	}
+	else
+	{
+		UE_LOG(LogControlRig, Error, TEXT("Track must be a Control Rig Track"));
+		return false;
+	}
+	return true;
+}
 
 bool UControlRigSequencerEditorLibrary::RenameControlRigControlChannels(ULevelSequence* LevelSequence, UControlRig* ControlRig, const TArray<FName>& InOldControlNames, const TArray<FName>& InNewControlNames)
 {
diff --git a/Engine/Plugins/Animation/ControlRig/Source/ControlRigEditor/Public/ControlRigSequencerEditorLibrary.h b/Engine/Plugins/Animation/ControlRig/Source/ControlRigEditor/Public/ControlRigSequencerEditorLibrary.h
--- a/Engine/Plugins/Animation/ControlRig/Source/ControlRigEditor/Public/ControlRigSequencerEditorLibrary.h
+++ b/Engine/Plugins/Animation/ControlRig/Source/ControlRigEditor/Public/ControlRigSequencerEditorLibrary.h
@@ -1098,11 +1098,19 @@ public:
 	UFUNCTION(BlueprintCallable, Category = "Editor Scripting | Sequencer Tools | Control Rig")
 	static UE_API bool MoveControlRigSpace(ULevelSequence* InSequence, UControlRig* InControlRig, FName InControlName, FFrameNumber InTime, FFrameNumber InNewTime, EMovieSceneTimeUnit TimeUnit = EMovieSceneTimeUnit::DisplayRate);
 
+	/** Replace the Control Rig on the track with a new Control Rig. Note it will now be owned by the track
+	* @param InControlRigTrack ControlRigTrack to replace
+	* @param NewControlRig ControlRig to put on track
+	* @return Return true if the function succeeds, false if it doesn't 
+	*/
+	UFUNCTION(BlueprintCallable, Category = "Editor Scripting | Sequencer Tools | Control Rig")
+	static UE_API bool ReplaceControlRig(UMovieSceneTrack* InControlRigTrack, UControlRig* NewControlRig);
+
 	/** Rename the Control Rig Channels in Sequencer to the specified new control names, which should be present on the Control Rig
 	* @param InSequence Sequence to rename controls
 	* @param InControlRig ControlRig to rename controls
 	* @param InOldControlNames The name of the old Control Rig Control Channels to change. Will be replaced by the corresponding name in the InNewControlNames array
-	* @param InNewControlNames  The name of the new Control Rig Channels 
+	* @param InNewControlNames  The name of the new Control Rig Channels
 	* @return Return true if the function succeeds, false if it doesn't which can happen if the name arrays don't match in size or any of the new Control Names aren't valid
 	*/
 	UFUNCTION(BlueprintCallable, Category = "Editor Scripting | Sequencer Tools | Control Rig")
