Control Rig: Fix run once workflow with details panel

#rnx

diff --git a/Engine/Plugins/Animation/ControlRig/Source/ControlRigEditor/Private/ControlRigElementDetails.cpp b/Engine/Plugins/Animation/ControlRig/Source/ControlRigEditor/Private/ControlRigElementDetails.cpp
--- a/Engine/Plugins/Animation/ControlRig/Source/ControlRigEditor/Private/ControlRigElementDetails.cpp
+++ b/Engine/Plugins/Animation/ControlRig/Source/ControlRigEditor/Private/ControlRigElementDetails.cpp
@@ -2747,6 +2747,12 @@ FDetailWidgetRow& FRigTransformElementDetails::CreateEulerTransformValueWidgetRo
 
 		for(URigHierarchy* HierarchyToUpdate : HierarchiesToUpdate)
 		{
+			if (UControlRig* RigToUpdate = HierarchyToUpdate->GetTypedOuter<UControlRig>())
+			{
+				// make sure we can run this rig as we are interacting with it.
+				RigToUpdate->ResetEvaluationsLeft();
+			}
+			
 			if(bInitial || CurrentTransformType == ERigTransformElementDetailsTransform::Current)
 			{
 				if(bIsRelative)
