Control Rig Edit Mode: fix occlusion based selection for perspective viewports

- ensure that the selection rectangle is at least twice as large as HitProxySize in order to match the implementation of FViewport::GetHitProxy(int32 X,int32 Y)

#rnx

diff --git a/Engine/Plugins/Animation/ControlRig/Source/ControlRigEditor/Private/EditMode/ControlRigEditModeUtil.cpp b/Engine/Plugins/Animation/ControlRig/Source/ControlRigEditor/Private/EditMode/ControlRigEditModeUtil.cpp
--- a/Engine/Plugins/Animation/ControlRig/Source/ControlRigEditor/Private/EditMode/ControlRigEditModeUtil.cpp
+++ b/Engine/Plugins/Animation/ControlRig/Source/ControlRigEditor/Private/EditMode/ControlRigEditModeUtil.cpp
@@ -6,6 +6,7 @@
 #include "EditorViewportClient.h"
 #include "EngineUtils.h"
 #include "LevelEditorViewport.h"
+#include "Misc/ConfigCacheIni.h"
 #include "Rigs/RigHierarchy.h"
 #include "SceneView.h"
 #include "Settings/LevelEditorViewportSettings.h"
@@ -194,6 +195,30 @@ void FSelectionHelper::GetFromFrustum(const FConvexVolume& InFrustum) const
 		}
 	}
 }
+
+namespace Private
+{
+
+uint32 GetHitProxySize()
+{
+	// see FViewport constructor and FViewport::GetHitProxy(int32 X,int32 Y)
+	static uint32 HitProxySize = 0;
+	if (HitProxySize == 0)
+	{
+		HitProxySize = 5;
+		if (GIsEditor && GConfig) 
+		{
+			int32 HitProxySizeFromConfig = 5;
+			if (GConfig->GetInt( TEXT("UnrealEd.HitProxy"), TEXT("HitProxySize"), HitProxySizeFromConfig, GEditorIni ))
+			{
+				HitProxySize = static_cast<uint32>(FMath::Clamp(HitProxySizeFromConfig, 1, MAX_HITPROXYSIZE));
+			}
+		}
+	}
+	return HitProxySize;
+}
+
+}
 	
 TOptional<FIntRect> FSelectionHelper::RectangleFromFrustum(const FConvexVolume& InFrustum) const
 {
@@ -260,15 +285,35 @@ TOptional<FIntRect> FSelectionHelper::RectangleFromFrustum(const FConvexVolume&
 
 		if (bIsProjectionValid)
 		{
-			const FVector2D& TopLeft = ScreenPos[3], BottomRight = ScreenPos[1];
+			FVector2D& TopLeft = ScreenPos[3], BottomRight = ScreenPos[1];
+
+			const int32 HitProxySize = Private::GetHitProxySize();
+			const double dHitProxySize(HitProxySize);
+
+			if (BottomRight.X - TopLeft.X < 2.0 * dHitProxySize)
+			{
+				const double MidX = 0.5 * (BottomRight.X + TopLeft.X);
+				TopLeft.X = MidX - dHitProxySize;
+				BottomRight.X = MidX + dHitProxySize;
+			}
+
+			if (BottomRight.Y - TopLeft.Y < 2.0 * dHitProxySize)
+			{
+				const double MidY = 0.5 * (BottomRight.Y + TopLeft.Y);
+				TopLeft.Y = MidY - dHitProxySize;
+				BottomRight.Y = MidY + dHitProxySize;
+			}
 
 			const FIntPoint ViewportSize = ViewportClient->Viewport->GetSizeXY();
 			
 			const int32 MinX = FMath::Clamp(FMath::FloorToInt32(TopLeft.X), 0, ViewportSize.X-1);
 			const int32 MinY = FMath::Clamp(FMath::FloorToInt32(TopLeft.Y), 0, ViewportSize.Y-1);
-		
-			const int32 MaxX = FMath::Clamp(FMath::CeilToInt32(BottomRight.X), MinX+1, ViewportSize.X);
-			const int32 MaxY = FMath::Clamp(FMath::CeilToInt32(BottomRight.Y), MinY+1, ViewportSize.Y);
+
+			const int32 MinXClamp = (MinX + HitProxySize) > (ViewportSize.X-1) ? MinX + 1 : MinX + HitProxySize;
+			const int32 MinYClamp = (MinY + HitProxySize) > (ViewportSize.Y-1) ? MinY + 1 : MinY + HitProxySize;
+			
+			const int32 MaxX = FMath::Clamp(FMath::CeilToInt32(BottomRight.X), MinXClamp, ViewportSize.X);
+			const int32 MaxY = FMath::Clamp(FMath::CeilToInt32(BottomRight.Y), MinYClamp, ViewportSize.Y);
 
 			const FIntRect Rect(MinX, MinY, MaxX, MaxY);
 			return Rect;
