Control Rig: Control Rig Hierarchy Editor cleanup third pass, cleanup includes, IWYU

#rnx

diff --git a/Engine/Plugins/Animation/ControlRig/Source/ControlRigEditor/Private/Editor/Hierarchy/DragDrop/RigHierarchyElementDragDropOp.cpp b/Engine/Plugins/Animation/ControlRig/Source/ControlRigEditor/Private/Editor/Hierarchy/DragDrop/RigHierarchyElementDragDropOp.cpp
--- a/Engine/Plugins/Animation/ControlRig/Source/ControlRigEditor/Private/Editor/Hierarchy/DragDrop/RigHierarchyElementDragDropOp.cpp
+++ b/Engine/Plugins/Animation/ControlRig/Source/ControlRigEditor/Private/Editor/Hierarchy/DragDrop/RigHierarchyElementDragDropOp.cpp
@@ -2,7 +2,6 @@
 
 #include "RigHierarchyElementDragDropOp.h"
 
-#include "Editor/Hierarchy/Widgets/SRigHierarchyTagWidget.h"
 #include "Widgets/Layout/SBorder.h"
 #include "Widgets/Text/STextBlock.h"
 
diff --git a/Engine/Plugins/Animation/ControlRig/Source/ControlRigEditor/Private/Editor/Hierarchy/Models/RigHierarchyTreeElement.h b/Engine/Plugins/Animation/ControlRig/Source/ControlRigEditor/Private/Editor/Hierarchy/Models/RigHierarchyTreeElement.h
--- a/Engine/Plugins/Animation/ControlRig/Source/ControlRigEditor/Private/Editor/Hierarchy/Models/RigHierarchyTreeElement.h
+++ b/Engine/Plugins/Animation/ControlRig/Source/ControlRigEditor/Private/Editor/Hierarchy/Models/RigHierarchyTreeElement.h
@@ -4,7 +4,6 @@
 
 #include "Editor/Hierarchy/RigHierarchyTreeDisplaySettings.h"
 #include "Editor/Hierarchy/Widgets/SRigHierarchyTagWidget.h"
-#include "GenericPlatform/GenericPlatform.h"
 #include "Internationalization/Text.h"
 #include "Styling/SlateColor.h"
 
diff --git a/Engine/Plugins/Animation/ControlRig/Source/ControlRigEditor/Private/Editor/Hierarchy/ModularRigHierarchyTreeDelegates.h b/Engine/Plugins/Animation/ControlRig/Source/ControlRigEditor/Private/Editor/Hierarchy/ModularRigHierarchyTreeDelegates.h
--- a/Engine/Plugins/Animation/ControlRig/Source/ControlRigEditor/Private/Editor/Hierarchy/ModularRigHierarchyTreeDelegates.h
+++ b/Engine/Plugins/Animation/ControlRig/Source/ControlRigEditor/Private/Editor/Hierarchy/ModularRigHierarchyTreeDelegates.h
@@ -4,11 +4,7 @@
 
 #include "Delegates/Delegate.h"
 #include "Editor/Hierarchy/Models/ModularRigHierarchyTreeElement.h"
-#include "Editor/Hierarchy/RigHierarchyTreeDelegates.h"
-#include "Editor/Hierarchy/RigHierarchyTreeDisplaySettings.h"
-#include "Editor/Hierarchy/Widgets/SRigHierarchyTagWidget.h"
 #include "Framework/SlateDelegates.h"
-#include "Rigs/RigHierarchy.h"
 #include "Widgets/Views/STableRow.h"
 #include "Widgets/Views/STreeView.h"
 
diff --git a/Engine/Plugins/Animation/ControlRig/Source/ControlRigEditor/Private/Editor/Hierarchy/Views/SModularRigTreeView.cpp b/Engine/Plugins/Animation/ControlRig/Source/ControlRigEditor/Private/Editor/Hierarchy/Views/SModularRigTreeView.cpp
--- a/Engine/Plugins/Animation/ControlRig/Source/ControlRigEditor/Private/Editor/Hierarchy/Views/SModularRigTreeView.cpp
+++ b/Engine/Plugins/Animation/ControlRig/Source/ControlRigEditor/Private/Editor/Hierarchy/Views/SModularRigTreeView.cpp
@@ -2,38 +2,13 @@
 
 #include "SModularRigTreeView.h"
 
-#include "Styling/AppStyle.h"
-#include "Widgets/Layout/SSpacer.h"
-#include "Widgets/Layout/SScrollBox.h"
-#include "Widgets/Images/SImage.h"
-#include "Widgets/Input/SSearchBox.h"
-#include "Widgets/Text/STextBlock.h"
-#include "PropertyCustomizationHelpers.h"
-#include "Framework/Application/SlateApplication.h"
-#include "Framework/Application/SlateUser.h"
-#include "Editor/EditorEngine.h"
-#include "Editor/Hierarchy/Views/SRigHierarchyTreeView.h"
-#include "HelperUtil.h"
-#include "Widgets/Text/SInlineEditableTextBlock.h"
+#include "Blueprint/WidgetBlueprintLibrary.h"
 #include "ControlRig.h"
 #include "ControlRigEditorStyle.h"
+#include "Engine/Texture2D.h"
+#include "Framework/Application/SlateApplication.h"
 #include "ModularRig.h"
-#include "ModularRigRuleManager.h"
-#include "Blueprint/WidgetBlueprintLibrary.h"
-#include "Editor/Hierarchy/Widgets/SModularRigHierarchy.h"
-#include "Fonts/FontMeasure.h"
 #include "Settings/ControlRigSettings.h"
-#include "Graph/ControlRigGraphSchema.h"
-#include "Rigs/AdditiveControlRig.h"
-#include "Rigs/RigHierarchyController.h"
-#include "Styling/AppStyle.h"
-#include "Widgets/Input/SComboButton.h"
-#include "Widgets/Input/SButton.h"
-#include "ScopedTransaction.h"
-#include "DetailLayoutBuilder.h"
-#include "Widgets/SRigVMVariantTagWidget.h"
-#include "Editor/SRigConnectorTargetWidget.h"
-#include "Editor/Hierarchy/Views/SRigHierarchySearchableTreeView.h"
 
 #define LOCTEXT_NAMESPACE "SModularRigTreeView"
 
diff --git a/Engine/Plugins/Animation/ControlRig/Source/ControlRigEditor/Private/Editor/Hierarchy/Views/SRigHierarchySearchableTreeView.cpp b/Engine/Plugins/Animation/ControlRig/Source/ControlRigEditor/Private/Editor/Hierarchy/Views/SRigHierarchySearchableTreeView.cpp
--- a/Engine/Plugins/Animation/ControlRig/Source/ControlRigEditor/Private/Editor/Hierarchy/Views/SRigHierarchySearchableTreeView.cpp
+++ b/Engine/Plugins/Animation/ControlRig/Source/ControlRigEditor/Private/Editor/Hierarchy/Views/SRigHierarchySearchableTreeView.cpp
@@ -2,6 +2,7 @@
 
 #include "SRigHierarchySearchableTreeView.h"
 
+#include "SRigHierarchyTreeView.h"
 #include "Widgets/Input/SSearchBox.h"
 #include "Widgets/Layout/SScrollBox.h"
 
diff --git a/Engine/Plugins/Animation/ControlRig/Source/ControlRigEditor/Private/Editor/Hierarchy/Views/SRigHierarchySearchableTreeView.h b/Engine/Plugins/Animation/ControlRig/Source/ControlRigEditor/Private/Editor/Hierarchy/Views/SRigHierarchySearchableTreeView.h
--- a/Engine/Plugins/Animation/ControlRig/Source/ControlRigEditor/Private/Editor/Hierarchy/Views/SRigHierarchySearchableTreeView.h
+++ b/Engine/Plugins/Animation/ControlRig/Source/ControlRigEditor/Private/Editor/Hierarchy/Views/SRigHierarchySearchableTreeView.h
@@ -5,10 +5,9 @@
 #include "Editor/Hierarchy/RigHierarchyTreeDelegates.h"
 #include "Editor/Hierarchy/RigHierarchyTreeDisplaySettings.h"
 #include "Internationalization/Text.h"
-#include "Rigs/RigHierarchyDefines.h"
-#include "SRigHierarchyTreeView.h"
 #include "Widgets/SCompoundWidget.h"
 
+class SRigHierarchyTreeView;
 class SSearchBox;
 
 class SRigHierarchySearchableTreeView : public SCompoundWidget
diff --git a/Engine/Plugins/Animation/ControlRig/Source/ControlRigEditor/Private/Editor/Hierarchy/Views/SRigHierarchyTreeView.cpp b/Engine/Plugins/Animation/ControlRig/Source/ControlRigEditor/Private/Editor/Hierarchy/Views/SRigHierarchyTreeView.cpp
--- a/Engine/Plugins/Animation/ControlRig/Source/ControlRigEditor/Private/Editor/Hierarchy/Views/SRigHierarchyTreeView.cpp
+++ b/Engine/Plugins/Animation/ControlRig/Source/ControlRigEditor/Private/Editor/Hierarchy/Views/SRigHierarchyTreeView.cpp
@@ -2,29 +2,16 @@
 
 #include "SRigHierarchyTreeView.h"
 
-#include "Styling/AppStyle.h"
-#include "Widgets/Layout/SSpacer.h"
-#include "Widgets/Layout/SScrollBox.h"
-#include "Widgets/Images/SImage.h"
-#include "Widgets/Input/SSearchBox.h"
-#include "Widgets/Text/STextBlock.h"
-#include "PropertyCustomizationHelpers.h"
-#include "Framework/Application/SlateApplication.h"
-#include "Editor/EditorEngine.h"
-#include "HelperUtil.h"
-#include "Widgets/Text/SInlineEditableTextBlock.h"
-#include "ControlRig.h"
-#include "ControlRigEditorStyle.h"
-#include "Editor/Hierarchy/Widgets/SRigHierarchy.h"
-#include "Settings/ControlRigSettings.h"
-#include "Graph/ControlRigGraphSchema.h"
-#include "Rigs/AdditiveControlRig.h"
-#include "Rigs/RigHierarchyController.h"
-#include "Styling/AppStyle.h"
 #include "Algo/Sort.h"
+#include "ControlRig.h"
 #include "ControlRigBlueprintLegacy.h"
+#include "ControlRigEditorStyle.h"
 #include "Editor/Hierarchy/RigHierarchyTreeDelegates.h"
 #include "Editor/Hierarchy/Views/SRigHierarchyItem.h"
+#include "Framework/Application/SlateApplication.h"
+#include "Graph/ControlRigGraphSchema.h"
+#include "Settings/ControlRigSettings.h"
+#include "Styling/AppStyle.h"
 
 #define LOCTEXT_NAMESPACE "SRigHierarchyTreeView"
 
diff --git a/Engine/Plugins/Animation/ControlRig/Source/ControlRigEditor/Private/Editor/Hierarchy/Views/SRigHierarchyTreeView.h b/Engine/Plugins/Animation/ControlRig/Source/ControlRigEditor/Private/Editor/Hierarchy/Views/SRigHierarchyTreeView.h
--- a/Engine/Plugins/Animation/ControlRig/Source/ControlRigEditor/Private/Editor/Hierarchy/Views/SRigHierarchyTreeView.h
+++ b/Engine/Plugins/Animation/ControlRig/Source/ControlRigEditor/Private/Editor/Hierarchy/Views/SRigHierarchyTreeView.h
@@ -7,9 +7,9 @@
 #include "Layout/Geometry.h"
 #include "Widgets/Views/STreeView.h"
 
-class SRigHierarchyTreeView;
-class SRigHierarchyItem;
 class FRigHierarchyTreeElement;
+class SRigHierarchyItem;
+class SRigHierarchyTreeView;
 
 #define UE_API CONTROLRIGEDITOR_API
 
diff --git a/Engine/Plugins/Animation/ControlRig/Source/ControlRigEditor/Private/Editor/Hierarchy/Widgets/SModularRigHierarchy.cpp b/Engine/Plugins/Animation/ControlRig/Source/ControlRigEditor/Private/Editor/Hierarchy/Widgets/SModularRigHierarchy.cpp
--- a/Engine/Plugins/Animation/ControlRig/Source/ControlRigEditor/Private/Editor/Hierarchy/Widgets/SModularRigHierarchy.cpp
+++ b/Engine/Plugins/Animation/ControlRig/Source/ControlRigEditor/Private/Editor/Hierarchy/Widgets/SModularRigHierarchy.cpp
@@ -2,62 +2,30 @@
 
 #include "SModularRigHierarchy.h"
 
-#include "Widgets/Input/SComboButton.h"
-#include "Styling/AppStyle.h"
-#include "Widgets/Input/SButton.h"
-#include "Widgets/Input/SSearchBox.h"
-#include "Kismet2/BlueprintEditorUtils.h"
-#include "ScopedTransaction.h"
-#include "Editor/ControlRigEditor.h"
-#include "BlueprintActionDatabase.h"
-#include "BlueprintVariableNodeSpawner.h"
-#include "Widgets/Layout/SSpacer.h"
-#include "Widgets/Images/SImage.h"
-#include "Widgets/Text/STextBlock.h"
-#include "K2Node_VariableGet.h"
-#include "RigVMBlueprintUtils.h"
-#include "ControlRigModularRigCommands.h"
-#include "ControlRigBlueprintLegacy.h"
-#include "Graph/ControlRigGraph.h"
-#include "Graph/ControlRigGraphNode.h"
-#include "Graph/ControlRigGraphSchema.h"
-#include "GraphEditorModule.h"
-#include "Framework/MultiBox/MultiBoxBuilder.h"
-#include "AnimationRuntime.h"
+#include "AssetRegistry/AssetRegistryModule.h"
 #include "ClassViewerFilter.h"
-#include "PropertyCustomizationHelpers.h"
-#include "Framework/Application/SlateApplication.h"
-#include "Editor/EditorEngine.h"
-#include "HelperUtil.h"
-#include "Widgets/Text/SInlineEditableTextBlock.h"
 #include "ControlRig.h"
-#include "HAL/PlatformApplicationMisc.h"
-#include "HAL/PlatformTime.h"
-#include "Dialogs/Dialogs.h"
-#include "IPersonaToolkit.h"
-#include "Types/WidgetActiveTimerDelegate.h"
+#include "ControlRigBlueprintLegacy.h"
+#include "ControlRigModularRigCommands.h"
 #include "Dialog/SCustomDialog.h"
-#include "EditMode/ControlRigEditMode.h"
-#include "ToolMenus.h"
-#include "Editor/ControlRigContextMenuContext.h"
-#include "Editor/SRigSpacePickerWidget.h"
-#include "Settings/ControlRigSettings.h"
-#include "Framework/Notifications/NotificationManager.h"
-#include "Widgets/Notifications/SNotificationList.h"
-#include "Styling/AppStyle.h"
-#include "Editor/ControlRigSkeletalMeshComponent.h"
-#include "Sequencer/ControlRigLayerInstance.h"
-#include "Algo/MinElement.h"
-#include "Algo/MaxElement.h"
-#include "AssetRegistry/AssetRegistryModule.h"
 #include "DragAndDrop/AssetDragDropOp.h"
-#include "Editor/RigVMEditorTools.h"
+#include "Editor/ControlRigContextMenuContext.h"
+#include "Editor/ControlRigEditor.h"
+#include "Editor/EditorEngine.h"
+#include "Editor/Hierarchy/DragDrop/ModularRigHierarchyElementDragDropOp.h"
+#include "Framework/Application/SlateApplication.h"
+#include "Framework/MultiBox/MultiBoxBuilder.h"
+#include "HAL/PlatformApplicationMisc.h"
 #include "Kismet2/SClassPickerDialog.h"
 #include "RigVMFunctions/Math/RigVMMathLibrary.h"
-#include "Preferences/PersonaOptions.h"
+#include "ScopedTransaction.h"
+#include "Styling/AppStyle.h"
+#include "ToolMenus.h"
+#include "Widgets/Images/SImage.h"
+#include "Widgets/Input/SComboButton.h"
+#include "Widgets/Input/SSearchBox.h"
 #include "Widgets/SRigVMBulkEditDialog.h"
 #include "Widgets/SRigVMSwapAssetReferencesWidget.h"
-#include "Editor/Hierarchy/DragDrop/ModularRigHierarchyElementDragDropOp.h"
 
 #if WITH_RIGVMLEGACYEDITOR
 #include "SKismetInspector.h"
diff --git a/Engine/Plugins/Animation/ControlRig/Source/ControlRigEditor/Private/Editor/Hierarchy/Widgets/SModularRigHierarchy.h b/Engine/Plugins/Animation/ControlRig/Source/ControlRigEditor/Private/Editor/Hierarchy/Widgets/SModularRigHierarchy.h
--- a/Engine/Plugins/Animation/ControlRig/Source/ControlRigEditor/Private/Editor/Hierarchy/Widgets/SModularRigHierarchy.h
+++ b/Engine/Plugins/Animation/ControlRig/Source/ControlRigEditor/Private/Editor/Hierarchy/Widgets/SModularRigHierarchy.h
@@ -2,20 +2,20 @@
 
 #pragma once
 
-#include "EditorUndoClient.h"
-#include "Editor/Hierarchy/Views/SModularRigTreeView.h"
 #include "ControlRigBlueprintLegacy.h"
+#include "Editor/Hierarchy/Views/SModularRigTreeView.h"
 #include "Editor/RigVMNewEditor.h"
+#include "EditorUndoClient.h"
 
-class SModularRigHierarchy;
+class FMenuBuilder;
+class FUICommandList;
 class IControlRigBaseEditor;
+class SModularRigHierarchy;
 class SSearchBox;
-class FUICommandList;
-class URigVMBlueprint;
 class UControlRig;
-struct FAssetData;
-class FMenuBuilder;
+class URigVMBlueprint;
 class UToolMenu;
+struct FAssetData;
 struct FToolMenuContext;
 
 /** Widget allowing editing of a control rig's structure */
diff --git a/Engine/Plugins/Animation/ControlRig/Source/ControlRigEditor/Private/Editor/Hierarchy/Widgets/SRigHierarchy.cpp b/Engine/Plugins/Animation/ControlRig/Source/ControlRigEditor/Private/Editor/Hierarchy/Widgets/SRigHierarchy.cpp
--- a/Engine/Plugins/Animation/ControlRig/Source/ControlRigEditor/Private/Editor/Hierarchy/Widgets/SRigHierarchy.cpp
+++ b/Engine/Plugins/Animation/ControlRig/Source/ControlRigEditor/Private/Editor/Hierarchy/Widgets/SRigHierarchy.cpp
@@ -2,62 +2,39 @@
 
 #include "SRigHierarchy.h"
 
-#include "Widgets/Input/SComboButton.h"
-#include "Styling/AppStyle.h"
-#include "Widgets/Input/SButton.h"
-#include "Widgets/Input/SSearchBox.h"
-#include "Kismet2/BlueprintEditorUtils.h"
-#include "ScopedTransaction.h"
-#include "Editor/ControlRigEditor.h"
-#include "BlueprintActionDatabase.h"
-#include "BlueprintVariableNodeSpawner.h"
-#include "Widgets/Layout/SSpacer.h"
-#include "Widgets/Images/SImage.h"
-#include "Widgets/Text/STextBlock.h"
-#include "SEnumCombo.h"
-#include "K2Node_VariableGet.h"
-#include "RigVMBlueprintUtils.h"
-#include "ControlRigHierarchyCommands.h"
-#include "ControlRigBlueprintLegacy.h"
-#include "Graph/ControlRigGraph.h"
-#include "Graph/ControlRigGraphNode.h"
-#include "Graph/ControlRigGraphSchema.h"
-#include "ModularRig.h"
-#include "GraphEditorModule.h"
-#include "Framework/MultiBox/MultiBoxBuilder.h"
-#include "AnimationRuntime.h"
-#include "PropertyCustomizationHelpers.h"
-#include "Framework/Application/SlateApplication.h"
-#include "Editor/EditorEngine.h"
-#include "HelperUtil.h"
-#include "Widgets/Text/SInlineEditableTextBlock.h"
 #include "ControlRig.h"
-#include "HAL/PlatformApplicationMisc.h"
-#include "HAL/PlatformTime.h"
-#include "Dialogs/Dialogs.h"
-#include "IPersonaToolkit.h"
-#include "Types/WidgetActiveTimerDelegate.h"
+#include "ControlRigHierarchyCommands.h"
 #include "Dialog/SCustomDialog.h"
-#include "EditMode/ControlRigEditMode.h"
-#include "ToolMenus.h"
+#include "Dialogs/Dialogs.h"
+#include "DragAndDrop/AssetDragDropOp.h"
 #include "Editor/ControlRigContextMenuContext.h"
-#include "Editor/SRigSpacePickerWidget.h"
-#include "Settings/ControlRigSettings.h"
-#include "Framework/Notifications/NotificationManager.h"
-#include "Widgets/Notifications/SNotificationList.h"
-#include "Styling/AppStyle.h"
+#include "Editor/ControlRigEditor.h"
 #include "Editor/ControlRigSkeletalMeshComponent.h"
-#include "ModularRigRuleManager.h"
-#include "Sequencer/ControlRigLayerInstance.h"
-#include "Algo/MinElement.h"
-#include "Algo/MaxElement.h"
-#include "Algo/Sort.h"
-#include "DragAndDrop/AssetDragDropOp.h"
-#include "RigVMFunctions/Math/RigVMMathLibrary.h"
-#include "Preferences/PersonaOptions.h"
 #include "Editor/Hierarchy/DragDrop/ModularRigHierarchyElementDragDropOp.h"
 #include "Editor/Hierarchy/DragDrop/RigHierarchyElementDragDropOp.h"
 #include "Editor/Hierarchy/DragDrop/RigHierarchyTagDragDropOp.h"
+#include "Framework/Application/SlateApplication.h"
+#include "Framework/MultiBox/MultiBoxBuilder.h"
+#include "Framework/Notifications/NotificationManager.h"
+#include "HAL/PlatformApplicationMisc.h"
+#include "IPersonaToolkit.h"
+#include "ModularRig.h"
+#include "ModularRigRuleManager.h"
+#include "Preferences/PersonaOptions.h"
+#include "PropertyCustomizationHelpers.h"
+#include "Rigs/RigHierarchyController.h"
+#include "ScopedTransaction.h"
+#include "SEnumCombo.h"
+#include "Sequencer/ControlRigLayerInstance.h"
+#include "Styling/AppStyle.h"
+#include "ToolMenus.h"
+#include "Types/WidgetActiveTimerDelegate.h"
+#include "Widgets/Images/SImage.h"
+#include "Widgets/Input/SButton.h"
+#include "Widgets/Input/SComboButton.h"
+#include "Widgets/Input/SSearchBox.h"
+#include "Widgets/Notifications/SNotificationList.h"
+#include "Widgets/Text/STextBlock.h"
 
 #if WITH_RIGVMLEGACYEDITOR
 #include "SKismetInspector.h"
diff --git a/Engine/Plugins/Animation/ControlRig/Source/ControlRigEditor/Private/Editor/Hierarchy/Widgets/SRigHierarchy.h b/Engine/Plugins/Animation/ControlRig/Source/ControlRigEditor/Private/Editor/Hierarchy/Widgets/SRigHierarchy.h
--- a/Engine/Plugins/Animation/ControlRig/Source/ControlRigEditor/Private/Editor/Hierarchy/Widgets/SRigHierarchy.h
+++ b/Engine/Plugins/Animation/ControlRig/Source/ControlRigEditor/Private/Editor/Hierarchy/Widgets/SRigHierarchy.h
@@ -2,23 +2,21 @@
 
 #pragma once
 
-#include "ControlRigBlueprintLegacy.h"
-#include "Editor/ControlRigSchematicModel.h"
+#include "ControlRigAsset.h"
 #include "Editor/Hierarchy/RigHierarchyImportSettings.h"
 #include "Editor/Hierarchy/Views/SRigHierarchyTreeView.h"
 #include "EditorUndoClient.h"
 #include "Engine/SkeletalMesh.h"
-#include "Units/RigUnitContext.h"
 
-class SRigHierarchy;
+class FMenuBuilder;
+class FUICommandList;
 class IControlRigBaseEditor;
+class SRigHierarchy;
 class SSearchBox;
-class FUICommandList;
-class URigVMBlueprint;
 class UControlRig;
-struct FAssetData;
-class FMenuBuilder;
+class URigVMBlueprint;
 class UToolMenu;
+struct FAssetData;
 struct FToolMenuContext;
 
 /** Widget allowing editing of a control rig's structure */
diff --git a/Engine/Plugins/Animation/ControlRig/Source/ControlRigEditor/Private/Editor/Hierarchy/Widgets/SRigHierarchyTagWidget.cpp b/Engine/Plugins/Animation/ControlRig/Source/ControlRigEditor/Private/Editor/Hierarchy/Widgets/SRigHierarchyTagWidget.cpp
--- a/Engine/Plugins/Animation/ControlRig/Source/ControlRigEditor/Private/Editor/Hierarchy/Widgets/SRigHierarchyTagWidget.cpp
+++ b/Engine/Plugins/Animation/ControlRig/Source/ControlRigEditor/Private/Editor/Hierarchy/Widgets/SRigHierarchyTagWidget.cpp
@@ -2,12 +2,11 @@
 
 #include "SRigHierarchyTagWidget.h"
 
-#include "Widgets/SBoxPanel.h"
-#include "Widgets/Images/SImage.h"
-#include "Widgets/Text/STextBlock.h"
 #include "Brushes/SlateRoundedBoxBrush.h"
 #include "DetailLayoutBuilder.h"
 #include "Editor/Hierarchy/DragDrop/RigHierarchyTagDragDropOp.h"
+#include "Widgets/Images/SImage.h"
+#include "Widgets/SBoxPanel.h"
 #include "Widgets/Text/SInlineEditableTextBlock.h"
 
 #define LOCTEXT_NAMESPACE "SRigHierarchyTagWidget"
diff --git a/Engine/Plugins/Animation/ControlRig/Source/ControlRigEditor/Private/Editor/Hierarchy/Widgets/SRigHierarchyTagWidget.h b/Engine/Plugins/Animation/ControlRig/Source/ControlRigEditor/Private/Editor/Hierarchy/Widgets/SRigHierarchyTagWidget.h
--- a/Engine/Plugins/Animation/ControlRig/Source/ControlRigEditor/Private/Editor/Hierarchy/Widgets/SRigHierarchyTagWidget.h
+++ b/Engine/Plugins/Animation/ControlRig/Source/ControlRigEditor/Private/Editor/Hierarchy/Widgets/SRigHierarchyTagWidget.h
@@ -2,10 +2,9 @@
 
 #pragma once
 
-#include "CoreMinimal.h"
-#include "Widgets/SCompoundWidget.h"
+#include "Framework/SlateDelegates.h"
 #include "Rigs/RigHierarchyDefines.h"
-#include "DragAndDrop/DecoratedDragDropOp.h"
+#include "Widgets/SCompoundWidget.h"
 
 class SRigHierarchyTagWidget;
 
