Control Rig: Connectors now allow to set a description, the Modular Rig Hierarchy displays it as tooltip

diff --git a/Engine/Plugins/Animation/ControlRig/Source/ControlRigEditor/Private/ControlRigElementDetails.cpp b/Engine/Plugins/Animation/ControlRig/Source/ControlRigEditor/Private/ControlRigElementDetails.cpp
--- a/Engine/Plugins/Animation/ControlRig/Source/ControlRigEditor/Private/ControlRigElementDetails.cpp
+++ b/Engine/Plugins/Animation/ControlRig/Source/ControlRigEditor/Private/ControlRigElementDetails.cpp
@@ -6356,6 +6356,10 @@ void FRigConnectorElementDetails::CustomizeSettings(IDetailLayoutBuilder& Detail
 			.AddProperty(SettingsHandle->GetChildHandle(GET_MEMBER_NAME_CHECKED(FRigConnectorSettings, Rules)))
 			.IsEnabled(!IsAnyConnectorImported());
 	}
+
+	SettingsCategory
+		.AddProperty(SettingsHandle->GetChildHandle(GET_MEMBER_NAME_CHECKED(FRigConnectorSettings, Description)))
+		.IsEnabled(!IsAnyConnectorImported());
 }
 
 void FRigConnectorElementDetails::CustomizeConnectorTargets(IDetailLayoutBuilder& DetailBuilder)
@@ -6634,6 +6638,13 @@ void FRigConnectionRuleDetails::CustomizeHeader(TSharedRef<IPropertyHandle> InSt
 	FString FirstObjectValue;
 	for (int32 Index = 0; Index < Objects.Num(); Index++)
 	{
+		const UObject* Object = Objects[Index];
+		if (!Object)
+		{
+			// Keep, for some reason these can go invalid in certain cases
+			continue;
+		}
+
 		FString ObjectValue;
 		if(InStructPropertyHandle->GetPerObjectValue(Index, ObjectValue) == FPropertyAccess::Result::Success)
 		{
@@ -6652,7 +6663,7 @@ void FRigConnectionRuleDetails::CustomizeHeader(TSharedRef<IPropertyHandle> InSt
 		}
 
 		// only enable editing of the rule if the widget is nested under a wrapper object (a rig element)
-		if(Objects[Index]->IsA<URigVMDetailsViewWrapperObject>())
+		if(Object->IsA<URigVMDetailsViewWrapperObject>())
 		{
 			EnabledAttribute = true;
 		}
diff --git a/Engine/Plugins/Animation/ControlRig/Source/ControlRigEditor/Private/Editor/Hierarchy/Views/SModularRigHierarchyTreeItem.cpp b/Engine/Plugins/Animation/ControlRig/Source/ControlRigEditor/Private/Editor/Hierarchy/Views/SModularRigHierarchyTreeItem.cpp
--- a/Engine/Plugins/Animation/ControlRig/Source/ControlRigEditor/Private/Editor/Hierarchy/Views/SModularRigHierarchyTreeItem.cpp
+++ b/Engine/Plugins/Animation/ControlRig/Source/ControlRigEditor/Private/Editor/Hierarchy/Views/SModularRigHierarchyTreeItem.cpp
@@ -378,13 +378,36 @@ FText SModularRigHierarchyTreeItem::GetName(bool bUseShortName) const
 
 FText SModularRigHierarchyTreeItem::GetItemTooltip() const
 {
-	const FText FullName = GetName(false);
-	const FText ShortName = GetName(true);
-	if(FullName.EqualTo(ShortName))
+	const FText Tooltip = [this]()
+		{	
+			constexpr bool bShortName = false;
+			const FText FullName = GetName(bShortName);
+
+			const UModularRig* ModularRig = Delegates.GetModularRig();
+			URigHierarchy* Hierarchy = ModularRig ? ModularRig->GetHierarchy() : nullptr;
+			const FRigConnectorElement* ConnectorElement = Hierarchy ? Cast<FRigConnectorElement>(Hierarchy->Find(ConnectorKey)) : nullptr;
+			if (ConnectorElement &&
+				!ConnectorElement->Settings.Description.IsEmpty())
+			{
+				const FText& ConnectorDescription = FText::FromString(ConnectorElement->Settings.Description);
+				return FText::Format(FText::FromString("{0}\n{1}"), FullName, ConnectorDescription);
+			}
+			else
+			{
+				return FullName;
+			}
+		}();
+
+	constexpr bool bShortName = true;
+	const FText ShortName = GetName(bShortName);
+	if (Tooltip.EqualTo(ShortName))
 	{
 		return FText();
 	}
-	return FullName;
+	else
+	{
+		return Tooltip;
+	}
 }
 
 TArray<FRigElementKey> SModularRigHierarchyTreeItem::GetTargetKeys() const
