Anim In Engine : Fix snapping with key frames and reversed time.

diff --git a/Engine/Plugins/Animation/ControlRig/Source/ControlRigEditor/Private/Tools/ControlRigSnapper.cpp b/Engine/Plugins/Animation/ControlRig/Source/ControlRigEditor/Private/Tools/ControlRigSnapper.cpp
--- a/Engine/Plugins/Animation/ControlRig/Source/ControlRigEditor/Private/Tools/ControlRigSnapper.cpp
+++ b/Engine/Plugins/Animation/ControlRig/Source/ControlRigEditor/Private/Tools/ControlRigSnapper.cpp
@@ -1,7 +1,7 @@
 // Copyright Epic Games, Inc. All Rights Reserved.
 
 #include "Tools/ControlRigSnapper.h"
-
+#include "Algo/Reverse.h"
 #include "Channels/MovieSceneFloatChannel.h"
 #include "ControlRig.h"
 #include "ISequencer.h"
@@ -536,9 +536,19 @@ bool FControlRigSnapper::SnapIt(FFrameNumber StartFrame, FFrameNumber EndFrame,c
 		else
 		{
 			TSortedMap<FFrameNumber, FFrameNumber> FrameMap;
-			GetTransformFrames(Sequencer, ParentToSnap, StartFrame, EndFrame, FrameMap);
-			GetTransformFrames(Sequencer, ActorToSnap,  StartFrame, EndFrame, FrameMap);
+			FFrameNumber LocalStartFrame = StartFrame;
+			FFrameNumber LocalEndFrame = EndFrame;
+			if (EndFrame < StartFrame)
+			{
+				Swap<FFrameNumber>(LocalStartFrame, LocalEndFrame);
+			}
+			GetTransformFrames(Sequencer, ParentToSnap, LocalStartFrame, LocalEndFrame, FrameMap);
+			GetTransformFrames(Sequencer, ActorToSnap, LocalStartFrame, LocalEndFrame, FrameMap);
 			FrameMap.GenerateKeyArray(Frames);
+			if (LocalStartFrame == EndFrame) //reverse it
+			{
+				Algo::Reverse(Frames);
+			}
 
 		}
 		if (Frames.Num() == 0)
