Control Rig: Fix NonUnity Compile UnrealEditor Win64 use of undefined types in Control Rig Editor

#rnx

diff --git a/Engine/Plugins/Animation/ControlRig/Source/ControlRigEditor/Private/EditMode/SControlRigOutliner.cpp b/Engine/Plugins/Animation/ControlRig/Source/ControlRigEditor/Private/EditMode/SControlRigOutliner.cpp
--- a/Engine/Plugins/Animation/ControlRig/Source/ControlRigEditor/Private/EditMode/SControlRigOutliner.cpp
+++ b/Engine/Plugins/Animation/ControlRig/Source/ControlRigEditor/Private/EditMode/SControlRigOutliner.cpp
@@ -43,6 +43,7 @@
 #include "Components/SkeletalMeshComponent.h"
 #include "Overrides/SOverrideStatusWidget.h"
 #include "Async/TaskGraphInterfaces.h"
+#include "Engine/Texture2D.h"
 
 #define LOCTEXT_NAMESPACE "ControlRigOutliner"
 
diff --git a/Engine/Plugins/Animation/ControlRig/Source/ControlRigEditor/Private/Editor/Hierarchy/Widgets/SRigHierarchyTagWidget.cpp b/Engine/Plugins/Animation/ControlRig/Source/ControlRigEditor/Private/Editor/Hierarchy/Widgets/SRigHierarchyTagWidget.cpp
--- a/Engine/Plugins/Animation/ControlRig/Source/ControlRigEditor/Private/Editor/Hierarchy/Widgets/SRigHierarchyTagWidget.cpp
+++ b/Engine/Plugins/Animation/ControlRig/Source/ControlRigEditor/Private/Editor/Hierarchy/Widgets/SRigHierarchyTagWidget.cpp
@@ -5,6 +5,7 @@
 #include "Brushes/SlateRoundedBoxBrush.h"
 #include "DetailLayoutBuilder.h"
 #include "Editor/Hierarchy/DragDrop/RigHierarchyTagDragDropOp.h"
+#include "Editor/Hierarchy/Views/SRigHierarchyTreeView.h"
 #include "Widgets/Images/SImage.h"
 #include "Widgets/SBoxPanel.h"
 #include "Widgets/Text/SInlineEditableTextBlock.h"
diff --git a/Engine/Plugins/Animation/ControlRig/Source/ControlRigEditor/Private/Editor/SRigConnectorTargetWidget.cpp b/Engine/Plugins/Animation/ControlRig/Source/ControlRigEditor/Private/Editor/SRigConnectorTargetWidget.cpp
--- a/Engine/Plugins/Animation/ControlRig/Source/ControlRigEditor/Private/Editor/SRigConnectorTargetWidget.cpp
+++ b/Engine/Plugins/Animation/ControlRig/Source/ControlRigEditor/Private/Editor/SRigConnectorTargetWidget.cpp
@@ -4,7 +4,10 @@
 
 #include "DetailWidgetRow.h"
 #include "DetailLayoutBuilder.h"
+#include "Editor/Hierarchy/Models/RigHierarchyTreeElement.h"
 #include "Editor/Hierarchy/Views/SRigHierarchyItem.h"
+#include "Editor/Hierarchy/Views/SRigHierarchySearchableTreeView.h"
+#include "Editor/Hierarchy/Views/SRigHierarchyTreeView.h"
 #include "IDetailChildrenBuilder.h"
 #include "ISinglePropertyView.h"
 #include "Modules/ModuleManager.h"
diff --git a/Engine/Plugins/Animation/ControlRig/Source/ControlRigEditor/Private/Editor/SRigConnectorTargetWidget.h b/Engine/Plugins/Animation/ControlRig/Source/ControlRigEditor/Private/Editor/SRigConnectorTargetWidget.h
--- a/Engine/Plugins/Animation/ControlRig/Source/ControlRigEditor/Private/Editor/SRigConnectorTargetWidget.h
+++ b/Engine/Plugins/Animation/ControlRig/Source/ControlRigEditor/Private/Editor/SRigConnectorTargetWidget.h
@@ -3,6 +3,7 @@
 #pragma once
 
 #include "CoreMinimal.h"
+#include "Editor/Hierarchy/RigHierarchyTreeDelegates.h"
 #include "Rigs/RigHierarchy.h"
 #include "Widgets/Layout/SBox.h"
 #include "Widgets/Input/SButton.h"
@@ -10,8 +11,10 @@
 #include "PropertyEditorDelegates.h"
 #include "UObject/StrongObjectPtrTemplates.h"
 #include "IPropertyTypeCustomization.h"
-#include "Editor/Hierarchy/Views/SRigHierarchySearchableTreeView.h"
 #include "SRigConnectorTargetWidget.generated.h"
+  
+class SRigHierarchySearchableTreeView;
+struct FRigHierarchyTreeDelegates;
 
 DECLARE_DELEGATE_RetVal_OneParam(bool, FRigConnectorTargetWidget_SetTarget, FRigElementKey);
 DECLARE_DELEGATE_RetVal_OneParam(bool, FRigConnectorTargetWidget_SetTargetArray, TArray<FRigElementKey>);
