ControlRig ECS: recursive anim evaluation (follow-up)

#rnx

diff --git a/Engine/Plugins/Animation/ControlRig/Source/ControlRig/Private/Sequencer/MovieSceneControlRigSystem.cpp b/Engine/Plugins/Animation/ControlRig/Source/ControlRig/Private/Sequencer/MovieSceneControlRigSystem.cpp
--- a/Engine/Plugins/Animation/ControlRig/Source/ControlRig/Private/Sequencer/MovieSceneControlRigSystem.cpp
+++ b/Engine/Plugins/Animation/ControlRig/Source/ControlRig/Private/Sequencer/MovieSceneControlRigSystem.cpp
@@ -1170,7 +1170,8 @@ struct FEvaluateBaseControlRigs
 					if (UTransformableComponentHandle* ParentHandle = Cast<UTransformableComponentHandle>(TransformConstraint->ParentTRSHandle))
 					{
 						constexpr bool bRecursive = true;
-						TransformableHandleUtils::MarkComponentForEvaluation(Cast<USkeletalMeshComponent>(ParentHandle->GetTarget()), bRecursive);
+						USceneComponent* ParentComponent = TransformableHandleUtils::GetTarget<USceneComponent>(TransformConstraint->ParentTRSHandle);
+						TransformableHandleUtils::MarkComponentForEvaluation(ParentComponent, bRecursive);
 					}
 				}
 			}
