AIE: Make Pose selection mirroring faster but not getting values when we are only doing selection mirroring.And resetting all tranforms will only reset on selected control rigs if selected

diff --git a/Engine/Plugins/Animation/ControlRig/Source/ControlRig/Private/Tools/ControlRigPose.cpp b/Engine/Plugins/Animation/ControlRig/Source/ControlRig/Private/Tools/ControlRigPose.cpp
--- a/Engine/Plugins/Animation/ControlRig/Source/ControlRig/Private/Tools/ControlRigPose.cpp
+++ b/Engine/Plugins/Animation/ControlRig/Source/ControlRig/Private/Tools/ControlRigPose.cpp
@@ -12,7 +12,7 @@
 
 #define LOCTEXT_NAMESPACE "ControlRigPose"
 
-void FControlRigControlPose::SavePose(UControlRig* ControlRig, bool bUseAll)
+void FControlRigControlPose::SavePose(UControlRig* ControlRig, bool bUseAll, bool bSaveValues)
 {
 	TArray<FRigControlElement*> CurrentControls;
 	ControlRig->GetControlsInOrder(CurrentControls);
@@ -22,29 +22,38 @@ void FControlRigControlPose::SavePose(UControlRig* ControlRig, bool bUseAll)
 	{
 		if (ControlRig->GetHierarchy()->IsAnimatable(ControlElement) && (bUseAll || ControlRig->IsControlSelected(ControlElement->GetFName())))
 		{
-			//we store poses in default parent space so if not in that space we need to compensate it
-			bool bHasNonDefaultParent = false;
-			FTransform GlobalTransform;
-			FRigElementKey SpaceKey = ControlRig->GetHierarchy()->GetActiveParent(ControlElement->GetKey());
-			if (SpaceKey != ControlRig->GetHierarchy()->GetDefaultParentKey())
+			if (bSaveValues)
 			{
-				bHasNonDefaultParent = true;
-				//to compensate we get the global, switch space, then reset global in that new space
-				GlobalTransform = ControlRig->GetHierarchy()->GetGlobalTransform(ControlElement->GetKey());
-				ControlRig->GetHierarchy()->SwitchToDefaultParent(ControlElement->GetKey());
-				ControlRig->GetHierarchy()->SetGlobalTransform(ControlElement->GetKey(),GlobalTransform);
+				//we store poses in default parent space so if not in that space we need to compensate it
+				bool bHasNonDefaultParent = false;
+				FTransform GlobalTransform;
+				FRigElementKey SpaceKey = ControlRig->GetHierarchy()->GetActiveParent(ControlElement->GetKey());
+				if (SpaceKey != ControlRig->GetHierarchy()->GetDefaultParentKey())
+				{
+					bHasNonDefaultParent = true;
+					//to compensate we get the global, switch space, then reset global in that new space
+					GlobalTransform = ControlRig->GetHierarchy()->GetGlobalTransform(ControlElement->GetKey());
+					ControlRig->GetHierarchy()->SwitchToDefaultParent(ControlElement->GetKey());
+					ControlRig->GetHierarchy()->SetGlobalTransform(ControlElement->GetKey(), GlobalTransform);
 
-				ControlRig->Evaluate_AnyThread();
+					ControlRig->Evaluate_AnyThread();
+				}
+
+				FRigControlCopy Copy(ControlElement, ControlRig->GetHierarchy(), bSaveValues);
+				CopyOfControls.Add(Copy);
+
+				if (bHasNonDefaultParent == true)
+				{
+					ControlRig->GetHierarchy()->SwitchToParent(ControlElement->GetKey(), SpaceKey);
+					ControlRig->GetHierarchy()->SetGlobalTransform(ControlElement->GetKey(), GlobalTransform);
+				}
 			}
-			FRigControlCopy Copy(ControlElement, ControlRig->GetHierarchy());
-			CopyOfControls.Add(Copy);
-	
-			if (bHasNonDefaultParent == true) 
+			else
 			{
-				ControlRig->GetHierarchy()->SwitchToParent(ControlElement->GetKey(),SpaceKey);
-				ControlRig->GetHierarchy()->SetGlobalTransform(ControlElement->GetKey(), GlobalTransform);
+				FRigControlCopy Copy(ControlElement, ControlRig->GetHierarchy(), bSaveValues);
+				CopyOfControls.Add(Copy);
+
 			}
-			
 		}
 	}
 	SetUpControlMap();
@@ -354,13 +363,13 @@ void UControlRigPoseAsset::PostLoad()
 	Pose.SetUpControlMap();
 }
 
-void UControlRigPoseAsset::SavePose(UControlRig* InControlRig, bool bUseAll)
+void UControlRigPoseAsset::SavePose(UControlRig* InControlRig, bool bUseAll, bool bSaveValues)
 {
 	if (!InControlRig)
 	{
 		return;
 	}
-	Pose.SavePose(InControlRig,bUseAll);
+	Pose.SavePose(InControlRig,bUseAll, bSaveValues);
 }
 
 void UControlRigPoseAsset::PastePose(UControlRig* InControlRig, bool bDoKey, bool bDoMirror, bool bDoAdditive)
@@ -404,7 +413,7 @@ void UControlRigPoseAsset::SelectControls(UControlRig* InControlRig, bool bDoMir
 	FControlRigControlPose TempPose;
 	if (bDoMirror)
 	{
-		TempPose.SavePose(InControlRig, true);
+		TempPose.SavePose(InControlRig, true , false/*no values*/);
 	}
 	for (const FName& Name : Controls)
 	{
diff --git a/Engine/Plugins/Animation/ControlRig/Source/ControlRig/Public/Tools/ControlRigPose.h b/Engine/Plugins/Animation/ControlRig/Source/ControlRig/Public/Tools/ControlRigPose.h
--- a/Engine/Plugins/Animation/ControlRig/Source/ControlRig/Public/Tools/ControlRigPose.h
+++ b/Engine/Plugins/Animation/ControlRig/Source/ControlRig/Public/Tools/ControlRigPose.h
@@ -37,17 +37,20 @@ struct FRigControlCopy
 	{
 	}
 
-	FRigControlCopy(FRigControlElement* InControlElement, URigHierarchy* InHierarchy)
+	FRigControlCopy(FRigControlElement* InControlElement, URigHierarchy* InHierarchy, bool bGetValues = true)
 	{
 		Name = InControlElement->GetFName();
 		ControlType = InControlElement->Settings.ControlType;
-		Value = InHierarchy->GetControlValue(InControlElement, ERigControlValueType::Current);
-		ParentKey = InHierarchy->GetFirstParent(InControlElement->GetKey());
-		OffsetTransform = InHierarchy->GetControlOffsetTransform(InControlElement, ERigTransformType::CurrentLocal);
-
-		ParentTransform = InHierarchy->GetParentTransform(InControlElement, ERigTransformType::CurrentGlobal);
-		LocalTransform = InHierarchy->GetTransform(InControlElement, ERigTransformType::CurrentLocal);
-		GlobalTransform = InHierarchy->GetTransform(InControlElement, ERigTransformType::CurrentGlobal);
+		if (bGetValues)
+		{
+			Value = InHierarchy->GetControlValue(InControlElement, ERigControlValueType::Current);
+			ParentKey = InHierarchy->GetFirstParent(InControlElement->GetKey());
+			OffsetTransform = InHierarchy->GetControlOffsetTransform(InControlElement, ERigTransformType::CurrentLocal);
+
+			ParentTransform = InHierarchy->GetParentTransform(InControlElement, ERigTransformType::CurrentGlobal);
+			LocalTransform = InHierarchy->GetTransform(InControlElement, ERigTransformType::CurrentLocal);
+			GlobalTransform = InHierarchy->GetTransform(InControlElement, ERigTransformType::CurrentGlobal);
+		}
 	}
 
 	UPROPERTY(BlueprintReadWrite, EditAnywhere, Category = "Names")
@@ -91,7 +94,7 @@ struct FControlRigControlPose
 	}
 	~FControlRigControlPose() {};
 
-	UE_API void SavePose(UControlRig* ControlRig, bool bUseAll);
+	UE_API void SavePose(UControlRig* ControlRig, bool bUseAll, bool bSaveValues = true);
 	UE_API void PastePose(UControlRig* ControlRig, bool bDoKey, bool bDoMirror);
 	UE_API void SetControlMirrorTransform(bool bDoLocalSpace, UControlRig* ControlRig, const FName& Name, bool bIsMatched,
 		const FTransform& GlobalTrnnsform, const FTransform& LocalTransform, bool bNotify, const  FRigControlModifiedContext& Context, bool bSetupUndo);
@@ -128,7 +131,7 @@ public:
 	UE_API virtual void PostLoad() override;
 
 	UFUNCTION(BlueprintCallable, Category = "Pose")
-	UE_API void SavePose(UControlRig* InControlRig, bool bUseAll);
+	UE_API void SavePose(UControlRig* InControlRig, bool bUseAll, bool bSaveValues = true);
 
 	UFUNCTION(BlueprintCallable, Category = "Pose")
 	UE_API void PastePose(UControlRig* InControlRig, bool bDoKey = false, bool bDoMirror = false, bool bDoAdditive = false);
diff --git a/Engine/Plugins/Animation/ControlRig/Source/ControlRigEditor/Private/EditMode/ControlRigEditMode.cpp b/Engine/Plugins/Animation/ControlRig/Source/ControlRigEditor/Private/EditMode/ControlRigEditMode.cpp
--- a/Engine/Plugins/Animation/ControlRig/Source/ControlRigEditor/Private/EditMode/ControlRigEditMode.cpp
+++ b/Engine/Plugins/Animation/ControlRig/Source/ControlRigEditor/Private/EditMode/ControlRigEditMode.cpp
@@ -3729,17 +3729,8 @@ TArray<FRigElementKey> FControlRigEditMode::GetRigElementsForSettingTransforms(U
 void FControlRigEditMode::ZeroTransforms(bool bSelectionOnly, bool bIncludeChannels)
 {
 	// Gather up the control rigs for the selected controls
-	TArray<UControlRig*> ControlRigs;
-	for (TWeakObjectPtr<UControlRig>& RuntimeRigPtr : RuntimeControlRigs)
-	{
-		if (UControlRig* ControlRig = RuntimeRigPtr.Get())
-		{
-			if (!bSelectionOnly || ControlRig->CurrentControlSelection().Num() > 0)
-			{
-				ControlRigs.Add(ControlRig);
-			}
-		}
-	}
+	TArray<UControlRig*> ControlRigs = GetControlRigsBasedOnSelection(bSelectionOnly);
+	
 	if (ControlRigs.Num() == 0)
 	{
 		return;
@@ -3954,20 +3945,33 @@ void FControlRigEditMode::ZeroTransforms(UControlRig* ControlRig, const FRigCont
 	}
 }
 
-void FControlRigEditMode::InvertInputPose(bool bSelectionOnly, bool bIncludeChannels)
+TArray<UControlRig*> FControlRigEditMode::GetControlRigsBasedOnSelection(bool bSelectionOnly) const
 {
-	// Gather up the control rigs for the selected controls
-	TArray<UControlRig*> ControlRigs;
-	for (TWeakObjectPtr<UControlRig>& RuntimeRigPtr : RuntimeControlRigs)
+	TArray<UControlRig*> SelectedControlRigs;
+	TArray<UControlRig*> AllControlRigs;
+	for (const TWeakObjectPtr<UControlRig>& RuntimeRigPtr : RuntimeControlRigs)
 	{
 		if (UControlRig* ControlRig = RuntimeRigPtr.Get())
 		{
-			if (!bSelectionOnly || ControlRig->CurrentControlSelection().Num() > 0)
+			if (ControlRig->CurrentControlSelection().Num() > 0)
 			{
-				ControlRigs.Add(ControlRig);
+				SelectedControlRigs.Add(ControlRig);
 			}
+			AllControlRigs.Add(ControlRig);
 		}
 	}
+
+	if (SelectedControlRigs.Num() > 0 || bSelectionOnly)
+	{
+		return SelectedControlRigs;
+	}
+	return AllControlRigs;
+}
+
+void FControlRigEditMode::InvertInputPose(bool bSelectionOnly, bool bIncludeChannels)
+{
+	// Gather up the control rigs for the selected controls
+	TArray<UControlRig*> ControlRigs = GetControlRigsBasedOnSelection(bSelectionOnly);
 	if (ControlRigs.Num() == 0)
 	{
 		return;
@@ -6465,7 +6469,7 @@ void FControlRigEditMode::SelectMirroredControls()
 		if (ControlRig)
 		{
 			ControlRig->Modify();
-			TempPose->SavePose(ControlRig, false);
+			TempPose->SavePose(ControlRig, false, false/*bSaveValues*/);
 			TempPose->SelectControls(ControlRig, true);
 		}
 	}
@@ -6489,7 +6493,7 @@ void FControlRigEditMode::AddMirroredControlsToSelection()
 		if (ControlRig)
 		{
 			ControlRig->Modify();
-			TempPose->SavePose(ControlRig, false);
+			TempPose->SavePose(ControlRig, false, false/*bSaveValues*/);
 			TempPose->SelectControls(ControlRig, true, false);
 		}
 	}
@@ -6562,7 +6566,7 @@ void FControlRigEditMode::SelectAllControls()
 		if (ControlRig)
 		{
 			ControlRig->Modify();
-			TempPose->SavePose(ControlRig, true); //this will save the whole pose
+			TempPose->SavePose(ControlRig, true, false/*bSaveValues*/); //this will save the whole pose
 			TempPose->SelectControls(ControlRig);
 		}
 	}
diff --git a/Engine/Plugins/Animation/ControlRig/Source/ControlRigEditor/Public/EditMode/ControlRigEditMode.h b/Engine/Plugins/Animation/ControlRig/Source/ControlRigEditor/Public/EditMode/ControlRigEditMode.h
--- a/Engine/Plugins/Animation/ControlRig/Source/ControlRigEditor/Public/EditMode/ControlRigEditMode.h
+++ b/Engine/Plugins/Animation/ControlRig/Source/ControlRigEditor/Public/EditMode/ControlRigEditMode.h
@@ -719,6 +719,9 @@ private:
 protected:
 	TArray<TWeakObjectPtr<UControlRig>> RuntimeControlRigs;
 
+	//if selection flag is set, get control rigs with selected controls, if not get all control rigs if there are none selected
+	TArray<UControlRig*> GetControlRigsBasedOnSelection(bool bSelectionOnly) const;
+
 private:
 	TMap<UControlRig*,TStrongObjectPtr<UControlRigEditModeDelegateHelper>> DelegateHelpers;
 
