[ITF Viewport Interactions / Control Rig] Animation Mode does not get the Left Click Marquee behavior

- Add frustum select interaction specific for control rig.
- Control rig adds it on Enter.
- No need to remove on Exit, since interactions are added "from scratch" on Enter
- Frustum Select adds Control Rig mode to the list of unsupported modes so that it is not added by [REDACTED_PERSON] on enter.
- Added a static local method to retrieve current Level Editor.

diff --git a/Engine/Plugins/Animation/ControlRig/Source/ControlRigEditor/Private/EditMode/ControlRigEditMode.cpp b/Engine/Plugins/Animation/ControlRig/Source/ControlRigEditor/Private/EditMode/ControlRigEditMode.cpp
--- a/Engine/Plugins/Animation/ControlRig/Source/ControlRigEditor/Private/EditMode/ControlRigEditMode.cpp
+++ b/Engine/Plugins/Animation/ControlRig/Source/ControlRigEditor/Private/EditMode/ControlRigEditMode.cpp
@@ -88,8 +88,10 @@
 #include "Misc/RigEditModeUtils.h"
 #include "Picker/FloatingSpacePickerManager.h"
 #include "Sequencer/AnimLayers/AnimLayers.h"
+#include "Tools/ControlRigFrustumSelectInteraction.h"
 #include "Tools/MotionTrailOptions.h"
 #include "Transform/TransformableHandleUtils.h"
+#include "ViewportInteractions/ViewportInteractionsBehaviorSource.h"
 
 #include UE_INLINE_GENERATED_CPP_BY_NAME(ControlRigEditMode)
 
@@ -110,6 +112,19 @@ namespace UE::ControlRigEditMode::Private
 	{
 		return IConsoleManager::Get().FindConsoleVariable(FocusModeName);
 	}
+
+	static ILevelEditor* GetLevelEditor()
+	{
+		if (FLevelEditorModule* LevelEditorModule = FModuleManager::GetModulePtr<FLevelEditorModule>(TEXT("LevelEditor")))
+		{
+			if (TSharedPtr<ILevelEditor> LevelEditorPtr = LevelEditorModule->GetLevelEditorInstance().Pin())
+			{
+				return LevelEditorPtr.Get();
+			}
+		}
+
+		return nullptr;
+	}
 }
 
 using namespace UE::ControlRigEditMode;
@@ -514,6 +529,20 @@ void FControlRigEditMode::Enter()
 		{
 			GEngine->OnActorMoving().AddRaw(this, &FControlRigEditMode::HandleActorMoving);
 		}
+
+		if (const ILevelEditor* const LevelEditor = Private::GetLevelEditor())
+		{
+			if (UModeManagerInteractiveToolsContext* const InteractiveToolsContext =
+					LevelEditor->GetEditorModeManager().GetInteractiveToolsContext())
+			{
+				if (UViewportInteractionsBehaviorSource* const InteractionsBehaviorSource =
+						InteractiveToolsContext->GetViewportInteractionsBehaviorSource())
+				{
+					// Add interactions specific to Control Rig
+					InteractionsBehaviorSource->AddInteraction<UControlRigFrustumSelectInteraction>(true);
+				}
+			}
+		}
 	}
 	
 	RegisterPendingFocusMode();
@@ -531,17 +560,15 @@ void FControlRigEditMode::Enter()
 //todo get working with Persona
 static void ClearOutAnyActiveTools()
 {
-	if (FLevelEditorModule* LevelEditorModule = FModuleManager::GetModulePtr<FLevelEditorModule>(TEXT("LevelEditor")))
+	if (ILevelEditor* LevelEditor = Private::GetLevelEditor())
 	{
-		TSharedPtr<ILevelEditor> LevelEditorPtr = LevelEditorModule->GetLevelEditorInstance().Pin();
-
-		if (LevelEditorPtr.IsValid())
+		FString ActiveToolName =
+			LevelEditor->GetEditorModeManager().GetInteractiveToolsContext()->ToolManager->GetActiveToolName(EToolSide::Left);
+		if (ActiveToolName == TEXT("SequencerPivotTool"))
 		{
-			FString ActiveToolName = LevelEditorPtr->GetEditorModeManager().GetInteractiveToolsContext()->ToolManager->GetActiveToolName(EToolSide::Left);
-			if (ActiveToolName == TEXT("SequencerPivotTool"))
-			{
-				LevelEditorPtr->GetEditorModeManager().GetInteractiveToolsContext()->ToolManager->DeactivateTool(EToolSide::Left, EToolShutdownType::Completed);
-			}
+			LevelEditor->GetEditorModeManager().GetInteractiveToolsContext()->ToolManager->DeactivateTool(
+				EToolSide::Left, EToolShutdownType::Completed
+			);
 		}
 	}
 }
@@ -3498,26 +3525,26 @@ void FControlRigEditMode::OpenSpacePickerWidget()
 
 static void ToggleTool(const FString& InToolName)
 {
-	if (FLevelEditorModule* LevelEditorModule = FModuleManager::GetModulePtr<FLevelEditorModule>(TEXT("LevelEditor")))
+	if (ILevelEditor* LevelEditor = Private::GetLevelEditor())
 	{
-		TSharedPtr<ILevelEditor> LevelEditorPtr = LevelEditorModule->GetLevelEditorInstance().Pin();
-
-		if (LevelEditorPtr.IsValid())
+		FString ActiveToolName =
+			LevelEditor->GetEditorModeManager().GetInteractiveToolsContext()->ToolManager->GetActiveToolName(EToolSide::Left);
+		if (ActiveToolName == InToolName)
 		{
-			FString ActiveToolName = LevelEditorPtr->GetEditorModeManager().GetInteractiveToolsContext()->ToolManager->GetActiveToolName(EToolSide::Left);
-			if (ActiveToolName == InToolName)
-			{
-				LevelEditorPtr->GetEditorModeManager().GetInteractiveToolsContext()->ToolManager->DeactivateTool(EToolSide::Left, EToolShutdownType::Completed);
-			}
-			else
-			{
-				LevelEditorPtr->GetEditorModeManager().GetInteractiveToolsContext()->ToolManager->SelectActiveToolType(EToolSide::Left, InToolName);
-				LevelEditorPtr->GetEditorModeManager().GetInteractiveToolsContext()->ToolManager->ActivateTool(EToolSide::Left);
-			}
+			LevelEditor->GetEditorModeManager().GetInteractiveToolsContext()->ToolManager->DeactivateTool(
+				EToolSide::Left, EToolShutdownType::Completed
+			);
+		}
+		else
+		{
+			LevelEditor->GetEditorModeManager().GetInteractiveToolsContext()->ToolManager->SelectActiveToolType(
+				EToolSide::Left, InToolName
+			);
+			LevelEditor->GetEditorModeManager().GetInteractiveToolsContext()->ToolManager->ActivateTool(EToolSide::Left);
 		}
 	}
-
 }
+
 void FControlRigEditMode::TogglePivotMode()
 {
 	FString Tool = TEXT("SequencerPivotTool");
diff --git a/Engine/Plugins/Animation/ControlRig/Source/ControlRigEditor/Private/Tools/ControlRigFrustumSelectInteraction.cpp b/Engine/Plugins/Animation/ControlRig/Source/ControlRigEditor/Private/Tools/ControlRigFrustumSelectInteraction.cpp
new file mode 100644
--- /dev/null
+++ b/Engine/Plugins/Animation/ControlRig/Source/ControlRigEditor/Private/Tools/ControlRigFrustumSelectInteraction.cpp
@@ -0,0 +1,59 @@
+// Copyright Epic Games, Inc. All Rights Reserved.
+
+#include "Tools/ControlRigFrustumSelectInteraction.h"
+#include "EditorModes.h"
+#include "Settings/ControlRigSettings.h"
+
+namespace UE::Editor::ControlRigFrustumSelect
+{
+static auto LMB_ONLY_PRIORITY = UE::Editor::ViewportInteractions::MULTIBUTTON_DRAG_PRIORITY + 1;
+}
+
+UControlRigFrustumSelectInteraction::UControlRigFrustumSelectInteraction()
+{
+	if (UViewportClickDragBehavior* ClickDragBehavior = ClickDragInputBehavior.Get())
+	{
+		ClickDragBehavior->ModifierCheckFunc = [](const FInputDeviceState& DeviceState)
+		{
+			return UControlRigEditorSettings::Get()->bLeftMouseDragDoesMarquee
+					 ? true
+					 : DeviceState.bCtrlKeyDown && DeviceState.bAltKeyDown;
+		};
+	}
+
+	RefreshPriority();
+	UControlRigEditorSettings::Get()->OnSettingChanged().AddUObject(
+		this, &UControlRigFrustumSelectInteraction::OnControlRigEditorSettingChanged
+	);
+
+	InteractionName = TEXT("Control Rig Frustum Select");
+}
+
+UControlRigFrustumSelectInteraction::~UControlRigFrustumSelectInteraction()
+{
+	UControlRigEditorSettings::Get()->OnSettingChanged().RemoveAll(this);
+}
+
+void UControlRigFrustumSelectInteraction::OnControlRigEditorSettingChanged(
+	UObject* InSettingsChanged, FPropertyChangedEvent& InPropertyChangedEvent
+)
+{
+	if (InPropertyChangedEvent.Property->GetFName()
+		== GET_MEMBER_NAME_CHECKED(UControlRigEditorSettings, bLeftMouseDragDoesMarquee))
+	{
+		RefreshPriority();
+	}
+}
+
+void UControlRigFrustumSelectInteraction::RefreshPriority()
+{
+	if (UViewportClickDragBehavior* ClickDragBehavior = ClickDragInputBehavior.Get())
+	{
+		int Priority = UControlRigEditorSettings::Get()->bLeftMouseDragDoesMarquee
+						 ? UE::Editor::ControlRigFrustumSelect::LMB_ONLY_PRIORITY
+						 : UE::Editor::DragTools::FRUSTUM_SELECT_PRIORITY;
+
+		ClickDragBehavior->SetDefaultPriority(Priority);
+	}
+}
+
diff --git a/Engine/Plugins/Animation/ControlRig/Source/ControlRigEditor/Private/Tools/ControlRigFrustumSelectInteraction.h b/Engine/Plugins/Animation/ControlRig/Source/ControlRigEditor/Private/Tools/ControlRigFrustumSelectInteraction.h
new file mode 100644
--- /dev/null
+++ b/Engine/Plugins/Animation/ControlRig/Source/ControlRigEditor/Private/Tools/ControlRigFrustumSelectInteraction.h
@@ -0,0 +1,28 @@
+// Copyright Epic Games, Inc. All Rights Reserved.
+
+#pragma once
+
+#include "EditorDragTools/FrustumSelectInteraction.h"
+#include "ControlRigFrustumSelectInteraction.generated.h"
+
+/**
+ * Similar to its base class, this marquee select tool switches between LMB only and Ctrl + Alt + LMB to activate marquee selection.
+ * Could potentially be used to add other custom logic.
+ */
+UCLASS(MinimalAPI, Transient)
+class UControlRigFrustumSelectInteraction : public UFrustumSelectInteraction
+{
+	GENERATED_BODY()
+
+public:
+	UControlRigFrustumSelectInteraction();
+
+	~UControlRigFrustumSelectInteraction();
+
+	virtual TArray<FEditorModeID> GetUnsupportedModes() const override {return {};}
+
+private:
+	void OnControlRigEditorSettingChanged(UObject* InSettingsChanged, FPropertyChangedEvent& InPropertyChangedEvent);
+	void RefreshPriority();
+};
+
diff --git a/Engine/Source/Editor/Experimental/EditorInteractiveToolsFramework/Private/EditorDragTools/FrustumSelectInteraction.cpp b/Engine/Source/Editor/Experimental/EditorInteractiveToolsFramework/Private/EditorDragTools/FrustumSelectInteraction.cpp
--- a/Engine/Source/Editor/Experimental/EditorInteractiveToolsFramework/Private/EditorDragTools/FrustumSelectInteraction.cpp
+++ b/Engine/Source/Editor/Experimental/EditorInteractiveToolsFramework/Private/EditorDragTools/FrustumSelectInteraction.cpp
@@ -416,7 +416,15 @@ void UFrustumSelectInteraction::OnTerminateDragSequence()
 
 TArray<FEditorModeID> UFrustumSelectInteraction::GetUnsupportedModes() const
 {
-	const TArray<FEditorModeID> UnsupportedModes = { FBuiltinEditorModes::EM_Landscape, FBuiltinEditorModes::EM_Foliage };
+	// Animation Mode has its own Frustum Select for now (input bindings customizations might change this in the future)
+	const FEditorModeID EM_ControlRig(TEXT("EditMode.ControlRig"));
+
+	const TArray<FEditorModeID> UnsupportedModes = {
+		FBuiltinEditorModes::EM_Landscape,
+		FBuiltinEditorModes::EM_Foliage,
+		EM_ControlRig
+	};
+
 	return UnsupportedModes;
 }
 
diff --git a/Engine/Source/Editor/Experimental/EditorInteractiveToolsFramework/Public/EditorDragTools/DragToolInteraction.h b/Engine/Source/Editor/Experimental/EditorInteractiveToolsFramework/Public/EditorDragTools/DragToolInteraction.h
--- a/Engine/Source/Editor/Experimental/EditorInteractiveToolsFramework/Public/EditorDragTools/DragToolInteraction.h
+++ b/Engine/Source/Editor/Experimental/EditorInteractiveToolsFramework/Public/EditorDragTools/DragToolInteraction.h
@@ -2,13 +2,14 @@
 
 #pragma once
 
-#include "BaseBehaviors/BehaviorTargetInterfaces.h"
 #include "Editor.h"
 #include "Behaviors/ViewportClickDragBehavior.h"
 #include "ViewportInteractions/ViewportInteraction.h"
 
 #include "DragToolInteraction.generated.h"
 
+#define UE_API EDITORINTERACTIVETOOLSFRAMEWORK_API
+
 class IEditorViewportClientProxy;
 struct FWorldSelectionElementArgs;
 class FCanvas;
@@ -31,14 +32,14 @@ DECLARE_MULTICAST_DELEGATE(FOnOnViewportChangeToolToggleDelegate)
  * Returns true if ITF-based drag tools should be used.
  * If an ITF version of a drag tool is not available yet, legacy will be used.
  */
-EDITORINTERACTIVETOOLSFRAMEWORK_API bool UseEditorDragTools();
-EDITORINTERACTIVETOOLSFRAMEWORK_API bool IsViewportChangeToolEnabled();
+UE_API bool UseEditorDragTools();
+UE_API bool IsViewportChangeToolEnabled();
 
-EDITORINTERACTIVETOOLSFRAMEWORK_API FOnEditorDragToolsToggleDelegate& OnEditorDragToolsActivated();
-EDITORINTERACTIVETOOLSFRAMEWORK_API FOnEditorDragToolsToggleDelegate& OnEditorDragToolsDeactivated();
+UE_API FOnEditorDragToolsToggleDelegate& OnEditorDragToolsActivated();
+UE_API FOnEditorDragToolsToggleDelegate& OnEditorDragToolsDeactivated();
 
-EDITORINTERACTIVETOOLSFRAMEWORK_API FOnOnViewportChangeToolToggleDelegate& OnViewportChangeToolActivated();
-EDITORINTERACTIVETOOLSFRAMEWORK_API FOnOnViewportChangeToolToggleDelegate& OnViewportChangeToolDeactivated();
+UE_API FOnOnViewportChangeToolToggleDelegate& OnViewportChangeToolActivated();
+UE_API FOnOnViewportChangeToolToggleDelegate& OnViewportChangeToolDeactivated();
 
 // In case New TRS Gizmos are enabled, we need Mouse Drag + Key modifier Drag Tools Behaviors to be processed before
 // the Gizmo ones, and also before viewport interactions - otherwise Drag Behaviors will never be triggered. See:
@@ -71,7 +72,7 @@ class UDragToolInteraction : public UViewportInteraction, public IViewportClickD
 public:
 	UDragToolInteraction();
 
-	virtual void Initialize(UViewportInteractionsBehaviorSource* const InViewportInteractionsBehaviorSource) override;
+	UE_API virtual void Initialize(UViewportInteractionsBehaviorSource* const InViewportInteractionsBehaviorSource) override;
 
 	/** @return true if we are dragging */
 	bool IsDragging() const
@@ -88,12 +89,12 @@ public:
 		return FInputRayHit();
 	}
 
-	virtual void OnDragStart(const FInputDeviceRay& InDragStartPos) override {}
-	virtual void OnDrag(const FInputDeviceRay& InDragPos) override {}
-	virtual void OnDragEnd(const FInputDeviceRay& InDragEndPos) override;
-	virtual void OnForceEndCapture() override;
-	virtual void OnTerminateDragSequence() override;
-	virtual void OnStateUpdated(const FInputDeviceState& InInputDeviceState) override;
+	UE_API virtual void OnDragStart(const FInputDeviceRay& InDragStartPos) override {}
+	UE_API virtual void OnDrag(const FInputDeviceRay& InDragPos) override {}
+	UE_API virtual void OnDragEnd(const FInputDeviceRay& InDragEndPos) override;
+	UE_API virtual void OnForceEndCapture() override;
+	UE_API virtual void OnTerminateDragSequence() override;
+	UE_API virtual void OnStateUpdated(const FInputDeviceState& InInputDeviceState) override;
 	//~ End IViewportClickDragBehaviorTarget
 
 	DECLARE_MULTICAST_DELEGATE(FOnToolStateChange);
@@ -121,3 +122,5 @@ protected:
 
 	IEditorViewportClientProxy* EditorViewportClientProxy;
 };
+
+#undef UE_API
diff --git a/Engine/Source/Editor/Experimental/EditorInteractiveToolsFramework/Public/EditorDragTools/FrustumSelectInteraction.h b/Engine/Source/Editor/Experimental/EditorInteractiveToolsFramework/Public/EditorDragTools/FrustumSelectInteraction.h
--- a/Engine/Source/Editor/Experimental/EditorInteractiveToolsFramework/Public/EditorDragTools/FrustumSelectInteraction.h
+++ b/Engine/Source/Editor/Experimental/EditorInteractiveToolsFramework/Public/EditorDragTools/FrustumSelectInteraction.h
@@ -26,18 +26,18 @@ class UFrustumSelectInteraction : public UMarqueeSelectInteraction
 	GENERATED_BODY()
 
 public:
-	UFrustumSelectInteraction();
+	UE_API UFrustumSelectInteraction();
 
 	//~ Begin IViewportClickDragBehaviorTarget
-	virtual FInputRayHit CanBeginClickDragSequence(const FInputDeviceRay& InPressPos) override;
-	virtual void OnDragStart(const FInputDeviceRay& InPressPos) override;
-	virtual void OnDrag(const FInputDeviceRay& InDragPos) override;
-	virtual void OnDragEnd(const FInputDeviceRay& InReleasePos) override;
-	virtual void OnTerminateDragSequence() override;
+	UE_API virtual FInputRayHit CanBeginClickDragSequence(const FInputDeviceRay& InPressPos) override;
+	UE_API virtual void OnDragStart(const FInputDeviceRay& InPressPos) override;
+	UE_API virtual void OnDrag(const FInputDeviceRay& InDragPos) override;
+	UE_API virtual void OnDragEnd(const FInputDeviceRay& InReleasePos) override;
+	UE_API virtual void OnTerminateDragSequence() override;
 	//~ End IViewportClickDragBehaviorTarget
 
 protected:
-	virtual TArray<FEditorModeID> GetUnsupportedModes() const override;
+	UE_API virtual TArray<FEditorModeID> GetUnsupportedModes() const override;
 
 private:
 	/**
diff --git a/Engine/Source/Editor/Experimental/EditorInteractiveToolsFramework/Public/EditorDragTools/MarqueeSelectInteraction.h b/Engine/Source/Editor/Experimental/EditorInteractiveToolsFramework/Public/EditorDragTools/MarqueeSelectInteraction.h
--- a/Engine/Source/Editor/Experimental/EditorInteractiveToolsFramework/Public/EditorDragTools/MarqueeSelectInteraction.h
+++ b/Engine/Source/Editor/Experimental/EditorInteractiveToolsFramework/Public/EditorDragTools/MarqueeSelectInteraction.h
@@ -5,6 +5,8 @@
 
 #include "MarqueeSelectInteraction.generated.h"
 
+#define UE_API EDITORINTERACTIVETOOLSFRAMEWORK_API
+
 class FCanvas;
 class FEditorViewportClient;
 class FSceneView;
@@ -21,7 +23,7 @@ class UMarqueeSelectInteraction : public UDragToolInteraction
 public:
 	UMarqueeSelectInteraction();
 
-	virtual void Render(const FSceneView* View, FCanvas* Canvas, EViewInteractionState InteractionState) override;
+	UE_API virtual void Render(const FSceneView* View, FCanvas* Canvas, EViewInteractionState InteractionState) override;
 	
 protected:
 
@@ -32,3 +34,5 @@ protected:
 
 	TWeakObjectPtr<UViewportClickDragBehavior> ClickDragInputBehavior;
 };
+
+#undef UE_API
