[Control Rig] Fixed cache invalidation bug in SPR node's Optional Parent.

diff --git a/Engine/Plugins/Animation/ControlRig/Source/ControlRig/Private/Units/Highlevel/Hierarchy/RigUnit_SphericalPoseReader.cpp b/Engine/Plugins/Animation/ControlRig/Source/ControlRig/Private/Units/Highlevel/Hierarchy/RigUnit_SphericalPoseReader.cpp
--- a/Engine/Plugins/Animation/ControlRig/Source/ControlRig/Private/Units/Highlevel/Hierarchy/RigUnit_SphericalPoseReader.cpp
+++ b/Engine/Plugins/Animation/ControlRig/Source/ControlRig/Private/Units/Highlevel/Hierarchy/RigUnit_SphericalPoseReader.cpp
@@ -33,22 +33,18 @@ FRigUnit_SphericalPoseReader_Execute()
 		FlipHeightScaling);
 
 	// get parent global space
-	FTransform GlobalDriverParentTransform;
 	OptionalParentCache.UpdateCache(OptionalParentItem, Hierarchy);
-	if (OptionalParentCache.IsValid())
-	{
-		GlobalDriverParentTransform = Hierarchy->GetGlobalTransformByIndex(OptionalParentCache, false);
-	}else
-	{
-		// if user does not specify a custom parent space, then simply use the driver's parent
-		GlobalDriverParentTransform = Hierarchy->GetParentTransformByIndex(DriverCache, false);
-	}
+	FRigElementKey ParentElementKey = OptionalParentCache.IsValid() ? OptionalParentCache.GetKey() : Hierarchy->GetFirstParent(DriverCache.GetKey());
+	FTransform GlobalDriverParentTransform = Hierarchy->GetGlobalTransform(ParentElementKey, false);
 
 	// check if we need to update the LocalDriverTransformInit
-	if (!bCachedInitTransforms || !CachedRotationOffset.Equals(RotationOffset))
+	if (!bCachedInitTransforms ||							// haven't cached yet?
+		!CachedRotationOffset.Equals(RotationOffset) ||		// rotation offset modified
+		CachedParentElementKey != ParentElementKey)			// swapped to different parent
 	{
 		bCachedInitTransforms = true;
 		CachedRotationOffset = RotationOffset;
+		CachedParentElementKey = ParentElementKey;
 
 		FTransform GlobalDriverParentTransformInitInverse;
 		if (OptionalParentCache.IsValid())
diff --git a/Engine/Plugins/Animation/ControlRig/Source/ControlRig/Public/Units/Highlevel/Hierarchy/RigUnit_SphericalPoseReader.h b/Engine/Plugins/Animation/ControlRig/Source/ControlRig/Public/Units/Highlevel/Hierarchy/RigUnit_SphericalPoseReader.h
--- a/Engine/Plugins/Animation/ControlRig/Source/ControlRig/Public/Units/Highlevel/Hierarchy/RigUnit_SphericalPoseReader.h
+++ b/Engine/Plugins/Animation/ControlRig/Source/ControlRig/Public/Units/Highlevel/Hierarchy/RigUnit_SphericalPoseReader.h
@@ -408,6 +408,8 @@ struct FRigUnit_SphericalPoseReader: public FRigUnit_HighlevelBaseMutable
 	UPROPERTY()
 	FVector CachedRotationOffset;
 	UPROPERTY()
+	FRigElementKey CachedParentElementKey;
+	UPROPERTY()
 	bool bCachedInitTransforms;
 };
 
