ControlRig: Add support for using the default name for new control rig blueprint assets from Asset Tools.
#rnx

diff --git a/Engine/Plugins/Animation/ControlRig/Source/ControlRigEditor/Private/ControlRigBlueprintActions.cpp b/Engine/Plugins/Animation/ControlRig/Source/ControlRigEditor/Private/ControlRigBlueprintActions.cpp
--- a/Engine/Plugins/Animation/ControlRig/Source/ControlRigEditor/Private/ControlRigBlueprintActions.cpp
+++ b/Engine/Plugins/Animation/ControlRig/Source/ControlRigEditor/Private/ControlRigBlueprintActions.cpp
@@ -265,7 +265,11 @@ UControlRigBlueprint* FControlRigBlueprintActions::CreateControlRigFromSkeletalM
 	check(RefSkeleton);
 
 	FString PackagePath = InSelectedObject->GetPathName();
-	FString ControlRigName = FString::Printf(TEXT("%s_CtrlRig"), *InSelectedObject->GetName());
+
+	FAssetToolsModule& AssetToolsModule = FModuleManager::LoadModuleChecked<FAssetToolsModule>("AssetTools");
+	const TOptional<FString> DefaultControlRigName = AssetToolsModule.Get().GetDefaultAssetNameForClassFromSource(UControlRigBlueprint::StaticClass(), InSelectedObject);
+
+	FString ControlRigName = DefaultControlRigName.IsSet() ? DefaultControlRigName.GetValue() : FString::Printf(TEXT("%s_CtrlRig"), *InSelectedObject->GetName());
 
 	int32 LastSlashPos = INDEX_NONE;
 	if (PackagePath.FindLastChar('/', LastSlashPos))
