Anim In Engine: Fix Contraint Baking with Rotation Orders. Just need to specify it before setting it.

diff --git a/Engine/Plugins/Animation/ControlRig/Source/ControlRig/Private/Constraints/ControlRigTransformableHandle.cpp b/Engine/Plugins/Animation/ControlRig/Source/ControlRig/Private/Constraints/ControlRigTransformableHandle.cpp
--- a/Engine/Plugins/Animation/ControlRig/Source/ControlRig/Private/Constraints/ControlRigTransformableHandle.cpp
+++ b/Engine/Plugins/Animation/ControlRig/Source/ControlRig/Private/Constraints/ControlRigTransformableHandle.cpp
@@ -862,11 +862,16 @@ bool UTransformableControlHandle::AddTransformKeys(const TArray<FFrameNumber>& I
 	UMovieSceneSection*,
 	const bool bLocal) const
 {
-	if (!ControlRig.IsValid() || ControlName == NAME_None || InFrames.IsEmpty() || InFrames.Num() != InTransforms.Num())
+	if (!ControlRig.IsValid() || ControlName == NAME_None || InFrames.IsEmpty() || InFrames.Num() != InTransforms.Num() || !ControlRig->GetHierarchy())
 	{
 		return false;
 	}
-	auto KeyframeFunc = [this, bLocal](const FTransform& InTransform, const FRigControlModifiedContext& InKeyframeContext)
+	FRigControlElement* ControlElement = ControlRig.Get()->FindControl(ControlName);
+	if(!ControlElement)
+	{
+		return false;
+	}
+	auto KeyframeFunc = [this, bLocal](const FTransform& InTransform, const FRigControlModifiedContext& InKeyframeContext, FRigControlElement* InControlElement)
 	{
 		UControlRig* InControlRig = ControlRig.Get();
 		static constexpr bool bNotify = true;
@@ -875,6 +880,10 @@ bool UTransformableControlHandle::AddTransformKeys(const TArray<FFrameNumber>& I
 
 		if (bLocal)
 		{
+			const FVector EulerAngle = InTransform.GetRotation().Rotator().Euler();
+			
+			InControlRig->GetHierarchy()->SetControlSpecifiedEulerAngle(InControlElement, EulerAngle);
+
 			InControlRig->SetControlLocalTransform(ControlName, InTransform, bNotify, InKeyframeContext, bUndo, bFixEuler);
 			if (InControlRig->IsAdditive())
 			{
@@ -899,7 +908,7 @@ bool UTransformableControlHandle::AddTransformKeys(const TArray<FFrameNumber>& I
 		const FFrameNumber& Frame = InFrames[Index];
 		KeyframeContext.LocalTime = InTickResolution.AsSeconds(FFrameTime(Frame));
 
-		KeyframeFunc(InTransforms[Index], KeyframeContext);
+		KeyframeFunc(InTransforms[Index], KeyframeContext, ControlElement);
 	}
 
 	return true;
