Fix crash on exit of UnrealEd, observed  in cook invocactions on MacOS. Memory corruption was triggered by memory access to UControlRigEditorSettings after object deletion from UControlRigFrustumSelectInteraction destructor. Instead, access from BeginDestroy which is invoked withthin the UControlRigEditorSettings default object's lifetime.

diff --git a/Engine/Plugins/Animation/ControlRig/Source/ControlRigEditor/Private/Tools/ControlRigFrustumSelectInteraction.cpp b/Engine/Plugins/Animation/ControlRig/Source/ControlRigEditor/Private/Tools/ControlRigFrustumSelectInteraction.cpp
--- a/Engine/Plugins/Animation/ControlRig/Source/ControlRigEditor/Private/Tools/ControlRigFrustumSelectInteraction.cpp
+++ b/Engine/Plugins/Animation/ControlRig/Source/ControlRigEditor/Private/Tools/ControlRigFrustumSelectInteraction.cpp
@@ -29,9 +29,10 @@ UControlRigFrustumSelectInteraction::UControlRigFrustumSelectInteraction()
 	InteractionName = TEXT("Control Rig Frustum Select");
 }
 
-UControlRigFrustumSelectInteraction::~UControlRigFrustumSelectInteraction()
+void UControlRigFrustumSelectInteraction::BeginDestroy()
 {
 	UControlRigEditorSettings::Get()->OnSettingChanged().RemoveAll(this);
+	Super::BeginDestroy();	
 }
 
 void UControlRigFrustumSelectInteraction::OnControlRigEditorSettingChanged(
diff --git a/Engine/Plugins/Animation/ControlRig/Source/ControlRigEditor/Private/Tools/ControlRigFrustumSelectInteraction.h b/Engine/Plugins/Animation/ControlRig/Source/ControlRigEditor/Private/Tools/ControlRigFrustumSelectInteraction.h
--- a/Engine/Plugins/Animation/ControlRig/Source/ControlRigEditor/Private/Tools/ControlRigFrustumSelectInteraction.h
+++ b/Engine/Plugins/Animation/ControlRig/Source/ControlRigEditor/Private/Tools/ControlRigFrustumSelectInteraction.h
@@ -17,7 +17,7 @@ class UControlRigFrustumSelectInteraction : public UFrustumSelectInteraction
 public:
 	UControlRigFrustumSelectInteraction();
 
-	~UControlRigFrustumSelectInteraction();
+	virtual void BeginDestroy() override;
 
 	virtual TArray<FEditorModeID> GetUnsupportedModes() const override {return {};}
 
