Sequencer Control Rig: Fix issue where if multiple parameters with different types have the same name it can prevent the right parameter from being keyed.

diff --git a/Engine/Plugins/Animation/ControlRig/Source/ControlRig/Private/Sequencer/MovieSceneControlRigParameterBuffer.cpp b/Engine/Plugins/Animation/ControlRig/Source/ControlRig/Private/Sequencer/MovieSceneControlRigParameterBuffer.cpp
--- a/Engine/Plugins/Animation/ControlRig/Source/ControlRig/Private/Sequencer/MovieSceneControlRigParameterBuffer.cpp
+++ b/Engine/Plugins/Animation/ControlRig/Source/ControlRig/Private/Sequencer/MovieSceneControlRigParameterBuffer.cpp
@@ -742,6 +742,20 @@ FControlRigValueView FControlRigParameterValues::Find(FName Name) const
 	return FControlRigValueView();
 }
 
+FControlRigValueView FControlRigParameterValues::Find(FName Name, EControlRigControlType ControlRigType) const
+{
+	if (ControlRigType < EControlRigControlType::Num)
+	{
+		const FControlRigParameterValueHeader& Header = GetHeader(ControlRigType);
+		const int32 Index = Header.Find(Name);
+		if (Index != INDEX_NONE)
+		{
+			return FControlRigValueView(Header.GetParameter(Index), Header.GetType());
+		}
+	}
+	return FControlRigValueView();
+}
+	
 FControlRigValueView FControlRigParameterValues::FindParameter(FName Name) const
 {
 	TArrayView<const FControlRigParameterValueHeader> ParameterHeaders(Headers);
diff --git a/Engine/Plugins/Animation/ControlRig/Source/ControlRig/Public/Sequencer/MovieSceneControlRigParameterBuffer.h b/Engine/Plugins/Animation/ControlRig/Source/ControlRig/Public/Sequencer/MovieSceneControlRigParameterBuffer.h
--- a/Engine/Plugins/Animation/ControlRig/Source/ControlRig/Public/Sequencer/MovieSceneControlRigParameterBuffer.h
+++ b/Engine/Plugins/Animation/ControlRig/Source/ControlRig/Public/Sequencer/MovieSceneControlRigParameterBuffer.h
@@ -273,6 +273,7 @@ struct FControlRigParameterValues
 	}
 
 	UE_API FControlRigValueView Find(FName Name) const;
+	UE_API FControlRigValueView Find(FName Name, EControlRigControlType ControlRigType) const;
 	UE_API FControlRigValueView FindParameter(FName Name) const;
 
 	UE_API ~FControlRigParameterValues();
diff --git a/Engine/Plugins/Animation/ControlRig/Source/ControlRigEditor/Private/Sequencer/ControlRigParameterTrackEditor.cpp b/Engine/Plugins/Animation/ControlRig/Source/ControlRigEditor/Private/Sequencer/ControlRigParameterTrackEditor.cpp
--- a/Engine/Plugins/Animation/ControlRig/Source/ControlRigEditor/Private/Sequencer/ControlRigParameterTrackEditor.cpp
+++ b/Engine/Plugins/Animation/ControlRig/Source/ControlRigEditor/Private/Sequencer/ControlRigParameterTrackEditor.cpp
@@ -5104,7 +5104,7 @@ void FControlRigParameterTrackEditor::ProcessKeyOperation(UObject* ObjectToKey,
 				continue;
 			}
 
-			FControlRigValueView DesiredValue = RecomposedValues.Find(ChannelMetaData.GetControlName());
+			FControlRigValueView DesiredValue = RecomposedValues.Find(ChannelMetaData.GetControlName(), ChannelMetaData.GetType());
 			if (DesiredValue.GetType() == EControlRigControlType::Space)
 			{
 
@@ -5270,8 +5270,8 @@ UE::MovieScene::FControlRigParameterValues FControlRigParameterTrackEditor::GetR
 	// Recompose Scalars
 	for (const FScalarParameterNameAndCurve& Scalar : Section->GetScalarParameterNamesAndCurves())
 	{
-		FControlRigValueView Value = CurrentValues.Find(Scalar.ParameterName);
-		if (Value && Value.GetType() == EControlRigControlType::Parameter_Scalar)
+		FControlRigValueView Value = CurrentValues.Find(Scalar.ParameterName, EControlRigControlType::Parameter_Scalar);
+		if (Value)
 		{
 			FControlRigChannelMetaData ChannelMetaData = Section->GetChannelMetaData(&Scalar.ParameterCurve);
 			check(ChannelMetaData);
@@ -5289,8 +5289,8 @@ UE::MovieScene::FControlRigParameterValues FControlRigParameterTrackEditor::GetR
 	// Recompose Vectors
 	for (const FVectorParameterNameAndCurves& Vector : Section->GetVectorParameterNamesAndCurves())
 	{
-		FControlRigValueView Value = CurrentValues.Find(Vector.ParameterName);
-		if (Value && Value.GetType() == EControlRigControlType::Parameter_Vector)
+		FControlRigValueView Value = CurrentValues.Find(Vector.ParameterName, EControlRigControlType::Parameter_Vector);
+		if (Value)
 		{
 			FControlRigChannelMetaData ChannelMetaData = Section->GetChannelMetaData(&Vector.XCurve);
 			check(ChannelMetaData);
@@ -5313,8 +5313,8 @@ UE::MovieScene::FControlRigParameterValues FControlRigParameterTrackEditor::GetR
 	// Recompose Vectors
 	for (const FTransformParameterNameAndCurves& Transform : Section->GetTransformParameterNamesAndCurves())
 	{
-		FControlRigValueView Value = CurrentValues.Find(Transform.ParameterName);
-		if (Value && Value.GetType() == EControlRigControlType::Parameter_Transform)
+		FControlRigValueView Value = CurrentValues.Find(Transform.ParameterName, EControlRigControlType::Parameter_Transform);
+		if (Value)
 		{
 			FControlRigChannelMetaData ChannelMetaData = Section->GetChannelMetaData(&Transform.Translation[0]);
 			check(ChannelMetaData);
