[ITF Drag Tools] Fix wrong priority of Anim Mode LMB marquee

diff --git a/Engine/Plugins/Animation/ControlRig/Source/ControlRigEditor/Private/Tools/ControlRigFrustumSelectInteraction.cpp b/Engine/Plugins/Animation/ControlRig/Source/ControlRigEditor/Private/Tools/ControlRigFrustumSelectInteraction.cpp
--- a/Engine/Plugins/Animation/ControlRig/Source/ControlRigEditor/Private/Tools/ControlRigFrustumSelectInteraction.cpp
+++ b/Engine/Plugins/Animation/ControlRig/Source/ControlRigEditor/Private/Tools/ControlRigFrustumSelectInteraction.cpp
@@ -1,12 +1,11 @@
 // Copyright Epic Games, Inc. All Rights Reserved.
 
 #include "Tools/ControlRigFrustumSelectInteraction.h"
-#include "EditorModes.h"
 #include "Settings/ControlRigSettings.h"
 
 namespace UE::Editor::ControlRigFrustumSelect
 {
-static auto LMB_ONLY_PRIORITY = UE::Editor::ViewportInteractions::MULTIBUTTON_DRAG_PRIORITY + 1;
+static auto LMB_ONLY_PRIORITY = UE::Editor::ViewportInteractions::CLICK_PRIORITY + 1;
 }
 
 UControlRigFrustumSelectInteraction::UControlRigFrustumSelectInteraction()
@@ -57,4 +56,3 @@ void UControlRigFrustumSelectInteraction::RefreshPriority()
 		ClickDragBehavior->SetDefaultPriority(Priority);
 	}
 }
-
