Control Rig: Modules in Modular Rigs now show a combo box to select Control Shapes, including Shapes from the Shape Libraries of the Module's Control Rig and the Shape Libraries of the Modular Rig

diff --git a/Engine/Plugins/Animation/ControlRig/Source/ControlRigEditor/Private/ControlRigShapeNameList.cpp b/Engine/Plugins/Animation/ControlRig/Source/ControlRigEditor/Private/ControlRigShapeNameList.cpp
--- a/Engine/Plugins/Animation/ControlRig/Source/ControlRigEditor/Private/ControlRigShapeNameList.cpp
+++ b/Engine/Plugins/Animation/ControlRig/Source/ControlRigEditor/Private/ControlRigShapeNameList.cpp
@@ -7,6 +7,7 @@
 #include "Widgets/SRigVMGraphPinNameListValueWidget.h"
 #include "ControlRig.h"
 #include "ControlRigAsset.h"
+#include "ModularRig.h"
 
 #define LOCTEXT_NAMESPACE "ControlRigShapeNameList"
 
@@ -43,17 +44,17 @@ void FControlRigShapeNameList::CreateShapeLibraryListWidget(IDetailChildrenBuild
 		.ValueContent()
 		[
 			SAssignNew(ShapeNameListWidget, SRigVMGraphPinNameListValueWidget)
-				.OptionsSource(&ShapeNameList)
-				.OnGenerateWidget(this, &FControlRigShapeNameList::MakeShapeNameListItemWidget)
-				.OnSelectionChanged(this, &FControlRigShapeNameList::OnShapeNameListChanged)
-				.OnComboBoxOpening(this, &FControlRigShapeNameList::OnShapeNameListComboBox)
-				.InitiallySelectedItem(InitialSelected)
-				.Content()
-				[
-					SNew(STextBlock)
-						.Text(this, &FControlRigShapeNameList::GetShapeNameListText)
-						.Font(IDetailLayoutBuilder::GetDetailFont())
-				]
+			.OptionsSource(&ShapeNameList)
+			.OnGenerateWidget(this, &FControlRigShapeNameList::MakeShapeNameListItemWidget)
+			.OnSelectionChanged(this, &FControlRigShapeNameList::OnShapeNameListChanged)
+			.OnComboBoxOpening(this, &FControlRigShapeNameList::OnShapeNameListComboBox)
+			.InitiallySelectedItem(InitialSelected)
+			.Content()
+			[
+				SNew(STextBlock)
+				.Text(this, &FControlRigShapeNameList::GetShapeNameListText)
+				.Font(IDetailLayoutBuilder::GetDetailFont())
+			]
 		];
 }
 
@@ -68,6 +69,16 @@ void FControlRigShapeNameList::GenerateShapeLibraryList(const UControlRig* Contr
 
 	const TMap<FString, FString>& LibraryNameMap = ControlRig->GetShapeLibraryNameMap();
 	TArray<TSoftObjectPtr<UControlRigShapeLibrary>> ShapeLibraries = ControlRig->GetShapeLibraries();
+
+	// If this is within a modular rig, also add the modular rig's shape libraries
+	if (const UModularRig* ModularRig = ControlRig->GetTypedOuter<UModularRig>())
+	{
+		for (const TSoftObjectPtr<UControlRigShapeLibrary>& ModularRigShapeLibrary : ModularRig->GetShapeLibraries())
+		{
+			ShapeLibraries.AddUnique(ModularRigShapeLibrary);
+		}
+	}
+
 	for (const TSoftObjectPtr<UControlRigShapeLibrary>& ShapeLibrary : ShapeLibraries)
 	{
 		if (ShapeLibrary.IsNull() || !ShapeLibrary.IsValid())
diff --git a/Engine/Plugins/Animation/ControlRig/Source/ControlRigEditor/Private/ControlRigShapeSettingsDetails.cpp b/Engine/Plugins/Animation/ControlRig/Source/ControlRigEditor/Private/ControlRigShapeSettingsDetails.cpp
--- a/Engine/Plugins/Animation/ControlRig/Source/ControlRigEditor/Private/ControlRigShapeSettingsDetails.cpp
+++ b/Engine/Plugins/Animation/ControlRig/Source/ControlRigEditor/Private/ControlRigShapeSettingsDetails.cpp
@@ -27,11 +27,19 @@ FControlRigShapeSettingsDetails::FControlRigShapeSettingsDetails()
 }
 
 void FControlRigShapeSettingsDetails::CustomizeHeader(TSharedRef<IPropertyHandle> InStructPropertyHandle, FDetailWidgetRow& InHeaderRow, IPropertyTypeCustomizationUtils& InStructCustomizationUtils)
+{
+	InHeaderRow.NameContent()
+		[
+			InStructPropertyHandle->CreatePropertyNameWidget()
+		];
+}
+
+void FControlRigShapeSettingsDetails::CustomizeChildren(TSharedRef<IPropertyHandle> InStructPropertyHandle, IDetailChildrenBuilder& InStructBuilder, IPropertyTypeCustomizationUtils& InStructCustomizationUtils)
 {
 	TArray<UObject*> Objects;
 	InStructPropertyHandle->GetOuterObjects(Objects);
 
-	ControlRigObjects.Reset();
+	TArray<TWeakObjectPtr<UControlRig>> ControlRigObjects;
 
 	// If one of the objects is either a Control Rig or a Control Rig Wrapper, we can generate the combo box for selection
 	// Else, a FName Text Box will be shown (the default property details)
@@ -51,14 +59,6 @@ void FControlRigShapeSettingsDetails::CustomizeHeader(TSharedRef<IPropertyHandle
 		}
 	}
 
-	InHeaderRow.NameContent()
-		[
-			InStructPropertyHandle->CreatePropertyNameWidget()
-		];
-}
-
-void FControlRigShapeSettingsDetails::CustomizeChildren(TSharedRef<IPropertyHandle> InStructPropertyHandle, IDetailChildrenBuilder& InStructBuilder, IPropertyTypeCustomizationUtils& InStructCustomizationUtils)
-{
 	if (InStructPropertyHandle->IsValidHandle())
 	{
 		uint32 NumChildren = 0;
diff --git a/Engine/Plugins/Animation/ControlRig/Source/ControlRigEditor/Private/ControlRigShapeSettingsDetails.h b/Engine/Plugins/Animation/ControlRig/Source/ControlRigEditor/Private/ControlRigShapeSettingsDetails.h
--- a/Engine/Plugins/Animation/ControlRig/Source/ControlRigEditor/Private/ControlRigShapeSettingsDetails.h
+++ b/Engine/Plugins/Animation/ControlRig/Source/ControlRigEditor/Private/ControlRigShapeSettingsDetails.h
@@ -32,7 +32,6 @@ public:
 	virtual void CustomizeChildren(TSharedRef<class IPropertyHandle> InStructPropertyHandle, class IDetailChildrenBuilder& InStructBuilder, IPropertyTypeCustomizationUtils& InStructCustomizationUtils) override;
 
 protected:
-	TArray<TWeakObjectPtr<UControlRig>> ControlRigObjects;
 	TSharedPtr<FControlRigShapeNameList> ControlRigShapeNameList;
 
 	static UControlRig* GetControlRig(UControlRigWrapperObject* ControlRigWrapperObject);
