Control Rig: Hierarchies that do not allow to rearrange connectors now only show a single small connector icon

diff --git a/Engine/Plugins/Animation/ControlRig/Source/ControlRigEditor/Private/ControlRigEditorStyle.h b/Engine/Plugins/Animation/ControlRig/Source/ControlRigEditor/Private/ControlRigEditorStyle.h
--- a/Engine/Plugins/Animation/ControlRig/Source/ControlRigEditor/Private/ControlRigEditorStyle.h
+++ b/Engine/Plugins/Animation/ControlRig/Source/ControlRigEditor/Private/ControlRigEditorStyle.h
@@ -167,6 +167,7 @@ public:
 			Set("ControlRig.Replay.Record", new IMAGE_BRUSH("Slate/RecordingIndicator", Icon32x32));
 		}
 
+		Set("ControlRig.ConnectorPlain.Small", new IMAGE_BRUSH_SVG("Slate/Connector_Plain", Icon16x16));
 		Set("ControlRig.ConnectorPrimary", new IMAGE_BRUSH_SVG("Slate/Connector_Primary", Icon128x128));
 		Set("ControlRig.ConnectorSecondary", new IMAGE_BRUSH_SVG("Slate/Connector_Secondary", Icon128x128));
 		Set("ControlRig.ConnectorOptional", new IMAGE_BRUSH_SVG("Slate/Connector_Optional", Icon128x128));
diff --git a/Engine/Plugins/Animation/ControlRig/Source/ControlRigEditor/Private/Editor/Hierarchy/Models/RigHierarchyTagModel.cpp b/Engine/Plugins/Animation/ControlRig/Source/ControlRigEditor/Private/Editor/Hierarchy/Models/RigHierarchyTagModel.cpp
new file mode 100644
--- /dev/null
+++ b/Engine/Plugins/Animation/ControlRig/Source/ControlRigEditor/Private/Editor/Hierarchy/Models/RigHierarchyTagModel.cpp
@@ -0,0 +1,252 @@
+// Copyright Epic Games, Inc. All Rights Reserved.
+
+#include "RigHierarchyTagModel.h"
+
+#include "ControlRig.h"
+#include "ControlRigBlueprintLegacy.h"
+#include "ControlRigEditorStyle.h"
+#include "Editor/Hierarchy/RigHierarchyTreeDisplaySettings.h"
+
+#define LOCTEXT_NAMESPACE "RigHierarchyTagModel"
+
+namespace UE::ControlRigEditor
+{
+	bool FRigHierarchyTagModelBase::MakeTagWidgetArgs(SRigHierarchyTagWidget::FArguments& OutArgs) const
+	{
+		if (!bIsValid)
+		{
+			return false;
+		}
+
+		SRigHierarchyTagWidget::FArguments WidgetArgs;
+		{
+			OutArgs.AllowDragDrop(bAllowDragDrop);
+			OutArgs.Visibility(Visibility);
+			OutArgs.Icon(IconBrush);
+			OutArgs.IconSize(FVector2d(16.f, 16.f));
+			OutArgs.Text(Label);
+			OutArgs.TooltipText(Tooltip);
+			OutArgs.Identifier(Identifier);
+
+			constexpr FColor DefaultIconColor = FColor(26, 26, 26, 255); // 1A1A1A
+			const FLinearColor OverridenIconColor = IconColor.IsSet() ? IconColor.GetValue() : DefaultIconColor;
+			OutArgs.IconColor(OverridenIconColor);
+
+			constexpr FColor DefaultBackgroundColor = FColor(38, 187, 255, 255); // 26BBFF
+			const FLinearColor OverridenBackgroundColor = BackgroundColor.IsSet() ? BackgroundColor.GetValue() : DefaultBackgroundColor;
+			OutArgs.Color(OverridenBackgroundColor);
+
+			constexpr FColor DefaultTextColor = FColor(15, 15, 15, 255); // 0F0F0F
+			OutArgs.TextColor(DefaultTextColor);
+
+			if (OnClicked.IsBound())
+			{
+				OutArgs.OnClicked(OnClicked);
+			}
+
+			if (OnRenamed.IsBound())
+			{
+				OutArgs.OnRenamed(OnRenamed);
+			}
+
+			if (OnVerifyRename.IsBound())
+			{
+				OutArgs.OnVerifyRename(OnVerifyRename);
+			}
+		}
+
+		return true;
+	}
+
+	FRigHierarchyValidConnectorTag::FRigHierarchyValidConnectorTag(const ERigHierarchyConnectorTagDisplayMode InTagDisplayMode)
+		: TagDisplayMode(InTagDisplayMode)
+	{
+		bAllowDragDrop = TagDisplayMode == ERigHierarchyConnectorTagDisplayMode::Individual;
+	}
+
+	FRigHierarchyValidConnectorTag FRigHierarchyValidConnectorTag::BuildTag(const ERigHierarchyConnectorTagDisplayMode InTagDisplayMode)
+	{
+		return FRigHierarchyValidConnectorTag(InTagDisplayMode);
+	}
+
+	FRigHierarchyValidConnectorTag& FRigHierarchyValidConnectorTag::AddConnector(
+		const UControlRig& ControlRig,
+		const FRigHierarchyTreeDisplaySettings& DisplaySettings,
+		const FRigConnectorElement& ConnectorElement)
+	{
+		URigHierarchy* Hierarchy = ControlRig.GetHierarchy();
+		if (!Hierarchy)
+		{
+			bIsValid = false;
+			return *this;
+		}
+
+		const bool bCanAdd = TagDisplayMode == ERigHierarchyConnectorTagDisplayMode::Single || ConnectorElementKeys.IsEmpty();
+		if (!ensureMsgf(bCanAdd, TEXT("Unexpected trying to combine tags for a rig hierarchy tree view that uses ERigHierarchyConnectorTagDisplayMode::Individual")))
+		{
+			bIsValid = false;
+			return *this;
+		}
+
+		const FRigElementKey ConnectorElementKey = ConnectorElement.GetKey();
+		ConnectorElementKeys.Add(ConnectorElementKey);
+
+		const EElementNameDisplayMode NameDisplayMode = DisplaySettings.NameDisplayMode;
+		const FText DisplayName = Hierarchy->GetDisplayNameForUI(ConnectorElementKey, NameDisplayMode);
+
+		if (TagDisplayMode == ERigHierarchyConnectorTagDisplayMode::Individual)
+		{
+			static const FSlateBrush* PrimaryBrush = FControlRigEditorStyle::Get().GetBrush("ControlRig.ConnectorPrimary");
+			static const FSlateBrush* SecondaryBrush = FControlRigEditorStyle::Get().GetBrush("ControlRig.ConnectorSecondary");
+			static const FSlateBrush* OptionalBrush = FControlRigEditorStyle::Get().GetBrush("ControlRig.ConnectorOptional");
+
+			Label = DisplayName;
+			Tooltip = FText::FromName(ConnectorElementKey.Name);
+
+			Identifier = [&ConnectorElementKey]()
+				{
+					FString Result;
+					constexpr const void* Defaults = nullptr;
+					constexpr UObject* OwnerObject = nullptr;
+					constexpr EPropertyPortFlags PortFlags = PPF_None;
+					constexpr UObject* ExportRootScope = nullptr;
+
+					FRigElementKey::StaticStruct()->ExportText(Result, &ConnectorElementKey, Defaults, OwnerObject, PortFlags, ExportRootScope);
+					return Result;
+				}();
+
+			IconBrush = [&ConnectorElement]()
+				{
+
+					if (ConnectorElement.Settings.Type == EConnectorType::Secondary)
+					{
+						return ConnectorElement.Settings.bOptional ? OptionalBrush : SecondaryBrush;
+					}
+
+					return PrimaryBrush;
+				}();
+		}
+		else if (TagDisplayMode == ERigHierarchyConnectorTagDisplayMode::Single)
+		{
+			static const FSlateBrush* PlainBrush = FControlRigEditorStyle::Get().GetBrush("ControlRig.ConnectorPlain.Small");
+
+			Tooltip = !Tooltip.IsSet() ?
+				FText::FromString(ConnectorElementKey.Name.ToString()) :
+				FText::FromString(FString::Printf(TEXT("%s\n%s"), *Tooltip.Get().ToString(), *ConnectorElementKey.Name.ToString()));
+
+			IconBrush = PlainBrush;
+		}
+
+		if (!ensureMsgf(Identifier.IsSet() || !bAllowDragDrop, TEXT("Unexpected tag has no identifier set, but allows drag drop. This leaves drag-drop without unfunctional.")))
+		{
+			bIsValid = false;
+			return *this;
+		}
+
+		// Only now that all data is properly initialized this tag becomes valid
+		bIsValid = true;
+
+		return *this;
+	}
+
+	FRigHierarchyValidConnectorTag& FRigHierarchyValidConnectorTag::SetOnClicked(const FSimpleDelegate& InOnClicked)
+	{
+		if (!ensureMsgf(!OnClicked.IsBound(), TEXT("Unexpected trying to bind on clicked when it is already bound, ignoring new binding")))
+		{
+			bIsValid = false;
+			return *this;
+		}
+
+		OnClicked = InOnClicked;
+
+		return *this;
+	}
+
+	FRigHierarchyValidConnectorTag& FRigHierarchyValidConnectorTag::SetOnRenamed(const FOnTextCommitted& InOnRenamed)
+	{
+		if (!ensureMsgf(TagDisplayMode == ERigHierarchyConnectorTagDisplayMode::Individual, TEXT("Cannot handle renaming tags in a rig hierarchy tree view that uses ERigHierarchyConnectorTagDisplayMode::Single, it can contain multiple connectors.")))
+		{
+			bIsValid = false;
+			return *this;
+		}
+
+		if (!ensureMsgf(!OnRenamed.IsBound(), TEXT("Unexpected trying to bind on renamed when it is already bound, ignoring new binding")))
+		{
+			bIsValid = false;
+			return *this;
+		}
+
+		OnRenamed = InOnRenamed;
+
+		return *this;
+	}
+
+	FRigHierarchyValidConnectorTag& FRigHierarchyValidConnectorTag::SetOnVerifyRename(const FOnVerifyTextChanged& InOnVerifyRename)
+	{
+		if (!ensureMsgf(TagDisplayMode == ERigHierarchyConnectorTagDisplayMode::Individual, TEXT("Cannot handle renaming tags in a rig hierarchy tree view that uses ERigHierarchyConnectorTagDisplayMode::Single, it can contain multiple connectors.")))
+		{
+			bIsValid = false;
+			return *this;
+		}
+
+		if (!ensureMsgf(!OnVerifyRename.IsBound(), TEXT("Unexpected trying to bind on verify rename when it is already bound, ignoring new binding")))
+		{
+			bIsValid = false;
+			return *this;
+		}
+
+		OnVerifyRename = InOnVerifyRename;
+
+		return *this;
+	}
+
+	FRigHierarchyConnectorUnresolvedWarningTag::FRigHierarchyConnectorUnresolvedWarningTag(
+		const UControlRig& ControlRig,
+		const FRigHierarchyTreeDisplaySettings& DisplaySettings,
+		const FRigConnectorElement& ConnectorElement)
+	{
+		URigHierarchy* Hierarchy = ControlRig.GetHierarchy();
+		if (!Hierarchy)
+		{
+			bIsValid = false;
+			return;
+		}
+
+		Label = LOCTEXT("ConnectorWarningTagLabel", "Warning");
+	
+		BackgroundColor = FColor::FromHex(TEXT("#FFB800"));
+		IconColor = FColor::FromHex(TEXT("#1A1A1A"));
+		IconBrush = FControlRigEditorStyle::Get().GetBrush("ControlRig.ConnectorWarning");
+
+		static const FText NotResolvedWarning = LOCTEXT("ConnectorWarningConnectorNotResolved", "Connector is not resolved.");
+
+		const TWeakObjectPtr<const UControlRig> WeakControlRig(&ControlRig);
+		const FRigElementKey ConnectorKey = ConnectorElement.GetKey();
+
+		const auto GetTooltipLambda = TAttribute<FText>::CreateLambda([WeakControlRig, ConnectorKey]()
+			{					
+				const URigHierarchy* Hierarchy = WeakControlRig.IsValid() ? WeakControlRig->GetHierarchy() : nullptr;
+				const UControlRigBlueprint* ControlRigBlueprint = WeakControlRig.IsValid() ? Cast<const UControlRigBlueprint>(WeakControlRig->GetClass()->ClassGeneratedBy) : nullptr;
+				if (!Hierarchy ||
+					!ControlRigBlueprint)
+				{
+					return FText::GetEmpty();
+				}
+
+				const FRigElementKey TargetKey = ControlRigBlueprint->ModularRigModel.Connections.FindTargetFromConnector(ConnectorKey);
+
+				return TargetKey.IsValid() && Hierarchy->Contains(TargetKey) ?
+					FText::GetEmpty() :
+					NotResolvedWarning;
+			});
+
+		Tooltip = GetTooltipLambda;
+
+		Visibility = TAttribute<EVisibility>::CreateLambda([GetTooltipLambda]() -> EVisibility
+			{
+				return GetTooltipLambda.Get().IsEmpty() ? EVisibility::Collapsed : EVisibility::Visible;
+			});
+	}
+}
+
+#undef LOCTEXT_NAMESPACE
diff --git a/Engine/Plugins/Animation/ControlRig/Source/ControlRigEditor/Private/Editor/Hierarchy/Models/RigHierarchyTagModel.h b/Engine/Plugins/Animation/ControlRig/Source/ControlRigEditor/Private/Editor/Hierarchy/Models/RigHierarchyTagModel.h
new file mode 100644
--- /dev/null
+++ b/Engine/Plugins/Animation/ControlRig/Source/ControlRigEditor/Private/Editor/Hierarchy/Models/RigHierarchyTagModel.h
@@ -0,0 +1,116 @@
+// Copyright Epic Games, Inc. All Rights Reserved.
+
+#pragma once
+
+#include "Editor/Hierarchy/Widgets/SRigHierarchyTagWidget.h"
+#include "Framework/SlateDelegates.h"
+
+class UControlRig;
+struct FRigConnectorElement;
+struct FRigHierarchyTreeDisplaySettings;
+
+namespace UE::ControlRigEditor
+{
+	/** A tag in a rig hierarchy */
+	struct FRigHierarchyTagModelBase 
+	{
+		/** Creates arguments to display this tag in an SRigHierarchyTagWidget */
+		bool MakeTagWidgetArgs(SRigHierarchyTagWidget::FArguments& OutArgs) const;
+
+	protected:
+		/** If true, this tag can be drag dropped */
+		bool bAllowDragDrop = true;
+
+		/** Visiblity attribute */
+		TAttribute<EVisibility> Visibility;
+
+		/** The icon to display */
+		const FSlateBrush* IconBrush = nullptr;
+
+		/** Optional icon color */
+		TOptional<FLinearColor> IconColor;
+
+		/** Optional backgrouond color */
+		TOptional<FLinearColor> BackgroundColor;
+
+		/** Label to display */
+		TAttribute<FText> Label;
+
+		/** The tooltip for the tag */
+		TAttribute<FText> Tooltip;
+
+		/** ID, should be set when bAllowDragDrop is true */
+		TAttribute<FString> Identifier;
+
+		/** Delegate executed when the tag is clicked */
+		FSimpleDelegate OnClicked;
+
+		/** Delegate executed when the tag was renamed */
+		FOnTextCommitted OnRenamed;
+
+		/** Delegate executed when the while the tag is being renamed, allowing to verify the currrent name */
+		FOnVerifyTextChanged OnVerifyRename;
+
+		/** Set to true when all data in this tag is valid */
+		bool bIsValid = false;
+	};
+
+	/** Defines how Tags are displayed in the Hierarchy */
+	enum class ERigHierarchyConnectorTagDisplayMode : uint8
+	{
+		Single,
+		Individual
+	};
+
+	/** A tag for a valid connector in a rig hierarchy */
+	struct FRigHierarchyValidConnectorTag final
+		: public FRigHierarchyTagModelBase
+	{
+	private:
+		/** Constructs this tag, private on purpose, use BuildTag instead */
+		FRigHierarchyValidConnectorTag(const ERigHierarchyConnectorTagDisplayMode InTagDisplayMode);
+
+	public:
+		/** Builds a new tag */
+		static FRigHierarchyValidConnectorTag BuildTag(const ERigHierarchyConnectorTagDisplayMode InTagDisplayMode);
+
+		/**
+		 * Sets the tag given a connector element, combines it if a connector element is already set.
+		 * Combining should only occur when the tag display mode is ERigHierarchyConnectorTagDisplayMode::Single (ensured).
+		 */
+		FRigHierarchyValidConnectorTag& AddConnector(
+			const UControlRig& ControlRig,
+			const FRigHierarchyTreeDisplaySettings& DisplaySettings,
+			const FRigConnectorElement& ConnectorElement);
+
+		/** Sets a delegate executed when the tag is clicked */
+		FRigHierarchyValidConnectorTag& SetOnClicked(const FSimpleDelegate& InOnClicked);
+
+		/** Sets a delegate executed when the tag was renamed */
+		FRigHierarchyValidConnectorTag& SetOnRenamed(const FOnTextCommitted& InOnRenamed);
+
+		/** Sets a delegate executed when the while the tag is being renamed, allowing to verify the currrent name */
+		FRigHierarchyValidConnectorTag& SetOnVerifyRename(const FOnVerifyTextChanged& InOnVerifyRename);
+
+		/** If true, this tag can be drag dropped */
+		const ERigHierarchyConnectorTagDisplayMode GetTagDisplayMode() const { return TagDisplayMode; }
+
+	private:
+		/** The connector element keys added to this tag */
+		TSet<FRigElementKey> ConnectorElementKeys;
+
+		/** Defines how tags are displayed in the specific rig hierarchy tree */
+		const ERigHierarchyConnectorTagDisplayMode TagDisplayMode;
+	};
+
+
+	/** A warning tag for a connector that cannot be resolved in a rig hierarchy */
+	struct FRigHierarchyConnectorUnresolvedWarningTag final
+		: public FRigHierarchyTagModelBase
+	{
+		FRigHierarchyConnectorUnresolvedWarningTag(
+			const UControlRig& ControlRig,
+			const FRigHierarchyTreeDisplaySettings& DisplaySettings,
+			const FRigConnectorElement& ConnectorElement);
+	};
+}
diff --git a/Engine/Plugins/Animation/ControlRig/Source/ControlRigEditor/Private/Editor/Hierarchy/Models/RigHierarchyTreeElement.h b/Engine/Plugins/Animation/ControlRig/Source/ControlRigEditor/Private/Editor/Hierarchy/Models/RigHierarchyTreeElement.h
--- a/Engine/Plugins/Animation/ControlRig/Source/ControlRigEditor/Private/Editor/Hierarchy/Models/RigHierarchyTreeElement.h
+++ b/Engine/Plugins/Animation/ControlRig/Source/ControlRigEditor/Private/Editor/Hierarchy/Models/RigHierarchyTreeElement.h
@@ -2,9 +2,8 @@
 
 #pragma once
 
-#include "Editor/Hierarchy/RigHierarchyTreeDisplaySettings.h"
-#include "Editor/Hierarchy/Widgets/SRigHierarchyTagWidget.h"
 #include "Internationalization/Text.h"
+#include "RigHierarchyTagModel.h"
 #include "Styling/SlateColor.h"
 
 class ITableRow;
@@ -34,6 +33,9 @@ enum class ERigTreeFilterResult : int32
 /** An item in the tree */
 class FRigHierarchyTreeElement : public TSharedFromThis<FRigHierarchyTreeElement>
 {
+	using FRigHierarchyValidConnectorTag = UE::ControlRigEditor::FRigHierarchyValidConnectorTag;
+	using FRigHierarchyConnectorUnresolvedWarningTag = UE::ControlRigEditor::FRigHierarchyConnectorUnresolvedWarningTag;
+
 public:
 	FRigHierarchyTreeElement(const FRigHierarchyKey& InKey, TWeakPtr<SRigHierarchyTreeView> InTreeView, bool InSupportsRename, ERigTreeFilterResult InFilterResult);
 public:
@@ -76,6 +78,9 @@ public:
 	/** If true the item is filtered out during a drag */
 	bool bFadedOutDuringDragDrop;
 
-	/** The tag arguments for this element */
-	TArray<SRigHierarchyTagWidget::FArguments> Tags;
+	/** Tags for valid connectors for this element */
+	TArray<FRigHierarchyValidConnectorTag> ConnectorTags;
+
+	/** Tags for invalid connectors for this element */
+	TArray<FRigHierarchyConnectorUnresolvedWarningTag> ConnectorResolveWarningTags;
 };
diff --git a/Engine/Plugins/Animation/ControlRig/Source/ControlRigEditor/Private/Editor/Hierarchy/RigHierarchyTreeDelegates.cpp b/Engine/Plugins/Animation/ControlRig/Source/ControlRigEditor/Private/Editor/Hierarchy/RigHierarchyTreeDelegates.cpp
--- a/Engine/Plugins/Animation/ControlRig/Source/ControlRigEditor/Private/Editor/Hierarchy/RigHierarchyTreeDelegates.cpp
+++ b/Engine/Plugins/Animation/ControlRig/Source/ControlRigEditor/Private/Editor/Hierarchy/RigHierarchyTreeDelegates.cpp
@@ -2,6 +2,8 @@
 
 #include "RigHierarchyTreeDelegates.h"
 
+#include "Editor/Hierarchy/Models/RigHierarchyTagModel.h"
+
 FRigHierarchyTreeDisplaySettings FRigHierarchyTreeDelegates::DefaultDisplaySettings;
 
 URigHierarchy* FRigHierarchyTreeDelegates::GetHierarchy()
@@ -87,4 +89,14 @@ void FRigHierarchyTreeDelegates::RequestDetailsInspection(const FRigHierarchyKey
 	{
 		return OnRequestDetailsInspection.Execute(InKey);
 	}
-}
\ No newline at end of file
+}
+
+UE::ControlRigEditor::ERigHierarchyConnectorTagDisplayMode FRigHierarchyTreeDelegates::GetTagDisplayMode() const
+{
+	if (OnGetTagDisplayMode.IsBound())
+	{
+		return OnGetTagDisplayMode.Execute();
+	}
+
+	return UE::ControlRigEditor::ERigHierarchyConnectorTagDisplayMode::Individual;
+}
diff --git a/Engine/Plugins/Animation/ControlRig/Source/ControlRigEditor/Private/Editor/Hierarchy/RigHierarchyTreeDelegates.h b/Engine/Plugins/Animation/ControlRig/Source/ControlRigEditor/Private/Editor/Hierarchy/RigHierarchyTreeDelegates.h
--- a/Engine/Plugins/Animation/ControlRig/Source/ControlRigEditor/Private/Editor/Hierarchy/RigHierarchyTreeDelegates.h
+++ b/Engine/Plugins/Animation/ControlRig/Source/ControlRigEditor/Private/Editor/Hierarchy/RigHierarchyTreeDelegates.h
@@ -12,6 +12,7 @@
 
 class FRigHierarchyTreeElement;
 class URigHierarchy;
+namespace UE::ControlRigEditor { enum class ERigHierarchyConnectorTagDisplayMode : uint8; }
 
 typedef STableRow<TSharedPtr<FRigHierarchyTreeElement>>::FOnCanAcceptDrop FOnRigTreeCanAcceptDrop;
 typedef STableRow<TSharedPtr<FRigHierarchyTreeElement>>::FOnAcceptDrop FOnRigTreeAcceptDrop;
@@ -30,6 +31,7 @@ DECLARE_DELEGATE_RetVal_OneParam(FRigHierarchyKey, FOnRigTreeGetResolvedKey, con
 DECLARE_DELEGATE_OneParam(FOnRigTreeRequestDetailsInspection, const FRigHierarchyKey&);
 DECLARE_DELEGATE_RetVal_OneParam(TOptional<FText>, FOnRigTreeItemGetToolTip, const FRigHierarchyKey&);
 DECLARE_DELEGATE_RetVal_OneParam(bool, FOnRigTreeIsItemVisible, const FRigHierarchyKey&);
+DECLARE_DELEGATE_RetVal(UE::ControlRigEditor::ERigHierarchyConnectorTagDisplayMode, FOnRigTreeGetTagDisplayMode);
 
 #define UE_API CONTROLRIGEDITOR_API
 
@@ -54,6 +56,7 @@ struct FRigHierarchyTreeDelegates
 	FOnRigTreeElementKeyTagDragDetected OnRigTreeElementKeyTagDragDetected;
 	FOnRigTreeItemGetToolTip OnRigTreeGetItemToolTip;
 	FOnRigTreeIsItemVisible OnRigTreeIsItemVisible;
+	FOnRigTreeGetTagDisplayMode OnGetTagDisplayMode;
 
 	URigHierarchy* GetHierarchy();
 	const URigHierarchy* GetHierarchy() const;
@@ -64,6 +67,7 @@ struct FRigHierarchyTreeDelegates
 	void HandleSelectionChanged(TSharedPtr<FRigHierarchyTreeElement> Selection, ESelectInfo::Type SelectInfo);
 	FRigHierarchyKey GetResolvedKey(const FRigHierarchyKey& InKey);
 	void RequestDetailsInspection(const FRigHierarchyKey& InKey);
+	UE::ControlRigEditor::ERigHierarchyConnectorTagDisplayMode GetTagDisplayMode() const;
 
 	static UE_API FRigHierarchyTreeDisplaySettings DefaultDisplaySettings;
 
diff --git a/Engine/Plugins/Animation/ControlRig/Source/ControlRigEditor/Private/Editor/Hierarchy/Views/SRigHierarchyItem.cpp b/Engine/Plugins/Animation/ControlRig/Source/ControlRigEditor/Private/Editor/Hierarchy/Views/SRigHierarchyItem.cpp
--- a/Engine/Plugins/Animation/ControlRig/Source/ControlRigEditor/Private/Editor/Hierarchy/Views/SRigHierarchyItem.cpp
+++ b/Engine/Plugins/Animation/ControlRig/Source/ControlRigEditor/Private/Editor/Hierarchy/Views/SRigHierarchyItem.cpp
@@ -8,8 +8,12 @@
 #include "Widgets/Layout/SSpacer.h"
 #include "Widgets/Text/SInlineEditableTextBlock.h"
 
+#define LOCTEXT_NAMESPACE "SRigHierarchyItem"
+
 void SRigHierarchyItem::Construct(const FArguments& InArgs, const TSharedRef<STableViewBase>& OwnerTable, TSharedRef<FRigHierarchyTreeElement> InRigTreeElement, TSharedPtr<SRigHierarchyTreeView> InTreeView, const FRigHierarchyTreeDisplaySettings& InSettings, bool bPinned)
 {
+	using namespace UE::ControlRigEditor;
+
 	WeakRigTreeElement = InRigTreeElement;
 	Delegates = InTreeView->GetRigTreeDelegates();
 	FRigHierarchyTreeDisplaySettings DisplaySettings = Delegates.GetDisplaySettings();
@@ -97,7 +101,8 @@ void SRigHierarchyItem::Construct(const FArguments& InArgs, const TSharedRef<STa
 			]
 		], OwnerTable);
 
-	if (!InRigTreeElement->Tags.IsEmpty())
+	const bool bHasAnyTags = !InRigTreeElement->ConnectorTags.IsEmpty() || !InRigTreeElement->ConnectorResolveWarningTags.IsEmpty();
+	if (bHasAnyTags)
 	{
 		HorizontalBox->AddSlot()
 			.FillWidth(1.f)
@@ -105,10 +110,16 @@ void SRigHierarchyItem::Construct(const FArguments& InArgs, const TSharedRef<STa
 				SNew(SSpacer)
 			];
 
-		for (const SRigHierarchyTagWidget::FArguments& TagArguments : InRigTreeElement->Tags)
+		// Show invalid connenctors first so they're less likely to get hidden due to a lack of space
+		for (const FRigHierarchyConnectorUnresolvedWarningTag& ConnectorResolveWarningTag : InRigTreeElement->ConnectorResolveWarningTags)
 		{
-			TSharedRef<SRigHierarchyTagWidget> TagWidget = SArgumentNew(TagArguments, SRigHierarchyTagWidget);
-			TagWidget->OnElementKeyDragDetected().BindSP(InTreeView.Get(), &SRigHierarchyTreeView::OnElementKeyTagDragDetected);
+			SRigHierarchyTagWidget::FArguments Args;
+			if (!ConnectorResolveWarningTag.MakeTagWidgetArgs(Args))
+			{
+				continue;
+			}
+
+			const TSharedRef<SRigHierarchyTagWidget> TagWidget = SArgumentNew(Args, SRigHierarchyTagWidget);
 
 			HorizontalBox->AddSlot()
 				.AutoWidth()
@@ -116,6 +127,30 @@ void SRigHierarchyItem::Construct(const FArguments& InArgs, const TSharedRef<STa
 					TagWidget
 				];
 		}
+
+		for (const FRigHierarchyValidConnectorTag& ConnectorTag : InRigTreeElement->ConnectorTags)
+		{
+			SRigHierarchyTagWidget::FArguments Args;
+			if (!ConnectorTag.MakeTagWidgetArgs(Args))
+			{
+				continue;
+			}
+
+			const TSharedRef<SRigHierarchyTagWidget> TagWidget = SArgumentNew(Args, SRigHierarchyTagWidget);
+
+			const bool bCanDragDropTag = ConnectorTag.GetTagDisplayMode() == ERigHierarchyConnectorTagDisplayMode::Individual;
+			if (bCanDragDropTag)
+			{
+				TagWidget->OnElementKeyDragDetected().BindSP(InTreeView.Get(), &SRigHierarchyTreeView::OnElementKeyTagDragDetected);
+			}
+
+			HorizontalBox->AddSlot()
+				.Padding(4.f, 0.f, 0.f, 0.f)
+				.AutoWidth()
+				[
+					TagWidget
+				];
+		}
 	}
 
 	InRigTreeElement->OnRenameRequested.BindSP(InlineWidget.Get(), &SInlineEditableTextBlock::EnterEditingMode);
@@ -128,14 +163,20 @@ FText SRigHierarchyItem::GetNameForUI() const
 
 FText SRigHierarchyItem::GetName(EElementNameDisplayMode InNameDisplayMode) const
 {
-	if (WeakRigTreeElement.Pin()->bIsTransient)
+	const TSharedPtr<FRigHierarchyTreeElement> RigTreeElement = WeakRigTreeElement.IsValid() ? WeakRigTreeElement.Pin() : nullptr;
+	if (!ensureMsgf(RigTreeElement.IsValid(), TEXT("Unexpected view without model")))
+	{
+		return LOCTEXT("DestroyedControlLabel", "Destroyed Control");
+	}
+
+	if (RigTreeElement->bIsTransient)
 	{
-		static const FText TemporaryControl = FText::FromString(TEXT("Temporary Control"));
-		return TemporaryControl;
+		return LOCTEXT("TemporaryControlLabel", "Temporary Control");
 	}
-	if (WeakRigTreeElement.Pin()->bIsAnimationChannel)
+
+	if (RigTreeElement->bIsAnimationChannel)
 	{
-		return FText::FromName(WeakRigTreeElement.Pin()->ChannelName);
+		return FText::FromName(RigTreeElement->ChannelName);
 	}
 
 	if (InNameDisplayMode == EElementNameDisplayMode::AssetDefault)
@@ -186,3 +227,5 @@ FText SRigHierarchyItem::GetItemTooltip() const
 
 	return FullName;
 }
+
+#undef LOCTEXT_NAMESPACE
diff --git a/Engine/Plugins/Animation/ControlRig/Source/ControlRigEditor/Private/Editor/Hierarchy/Views/SRigHierarchyTreeView.cpp b/Engine/Plugins/Animation/ControlRig/Source/ControlRigEditor/Private/Editor/Hierarchy/Views/SRigHierarchyTreeView.cpp
--- a/Engine/Plugins/Animation/ControlRig/Source/ControlRigEditor/Private/Editor/Hierarchy/Views/SRigHierarchyTreeView.cpp
+++ b/Engine/Plugins/Animation/ControlRig/Source/ControlRigEditor/Private/Editor/Hierarchy/Views/SRigHierarchyTreeView.cpp
@@ -6,6 +6,7 @@
 #include "ControlRig.h"
 #include "ControlRigBlueprintLegacy.h"
 #include "ControlRigEditorStyle.h"
+#include "Editor/Hierarchy/Models/RigHierarchyTagModel.h"
 #include "Editor/Hierarchy/RigHierarchyTreeDelegates.h"
 #include "Editor/Hierarchy/Views/SRigHierarchyItem.h"
 #include "Framework/Application/SlateApplication.h"
@@ -263,202 +264,140 @@ bool SRigHierarchyTreeView::AddElement(FRigHierarchyKey InKey, FRigHierarchyKey
 
 bool SRigHierarchyTreeView::AddElement(const FRigBaseElement* InElement)
 {
-	check(InElement);
-	
-	if (ElementMap.Contains(InElement->GetKey()))
+	using namespace UE::ControlRigEditor;
+
+	const URigHierarchy* Hierarchy = Delegates.GetHierarchy();
+	UControlRig* ControlRig = Hierarchy ? Hierarchy->GetTypedOuter<UControlRig>() : nullptr;
+	if (!InElement ||
+		!Hierarchy ||
+		!ControlRig)
+	{
+		return false;
+	}
+
+	const bool bAlreadyAdded = ElementMap.Contains(InElement->GetKey());
+	if (bAlreadyAdded)
 	{
 		return false;
 	}
 
 	const FRigHierarchyTreeDisplaySettings& Settings = Delegates.GetDisplaySettings();
-	const URigHierarchy* Hierarchy = Delegates.GetHierarchy();
 
-	switch(InElement->GetType())
+	if (InElement->GetType() == ERigElementType::Physics ||
+		InElement->GetType() == ERigElementType::Curve)
+	{
+		return false;
+	}
+	else if (InElement->GetType() == ERigElementType::Bone)
 	{
-		case ERigElementType::Bone:
+		if (!Settings.bShowBones)
 		{
-			if(!Settings.bShowBones)
-			{
-				return false;
-			}
-
-			const FRigBoneElement* BoneElement = CastChecked<FRigBoneElement>(InElement);
-			if (!Settings.bShowImportedBones && BoneElement->BoneType == ERigBoneType::Imported)
-			{
-				return false;
-			}
-			break;
+			return false;
 		}
-		case ERigElementType::Null:
+
+		const FRigBoneElement* BoneElement = CastChecked<FRigBoneElement>(InElement);
+		if (!Settings.bShowImportedBones && BoneElement->BoneType == ERigBoneType::Imported)
 		{
-			if(!Settings.bShowNulls)
-			{
-				return false;
-			}
-			break;
+			return false;
 		}
-		case ERigElementType::Control:
+	}
+	else if (InElement->GetType() == ERigElementType::Null)
+	{
+		if (!Settings.bShowNulls)
 		{
-			if(!Settings.bShowControls)
-			{
-				return false;
-			}
-			break;
+			return false;
 		}
-		case ERigElementType::Physics:
+	}
+	else if (InElement->GetType() == ERigElementType::Control)
+	{
+		if (!Settings.bShowControls)
 		{
 			return false;
 		}
-		case ERigElementType::Reference:
+	}
+	else if (InElement->GetType() == ERigElementType::Reference)
+	{
+		if (!Settings.bShowReferences)
 		{
-			if(!Settings.bShowReferences)
-			{
-				return false;
-			}
-			break;
+			return false;
 		}
-		case ERigElementType::Curve:
+	}
+	else if (InElement->GetType() == ERigElementType::Connector)
+	{
+		// Add the connector as a tag
+		const FRigConnectorElement* ConnectorElement = Cast<FRigConnectorElement>(InElement);
+		if (!ConnectorElement)
 		{
 			return false;
 		}
-		case ERigElementType::Connector:
+
+		FRigElementKeyRedirector& Redirector = ControlRig->GetElementKeyRedirector();
+		FRigElementKeyRedirector::FCachedKeyArray* Cache = Redirector.Find(InElement->GetKey());
+		if (Cache)
 		{
-			if(Hierarchy)
+			// Find the elements this connector is connected to
+			bool bValid = true;
+			TArray<TSharedRef<FRigHierarchyTreeElement>> TargetElements;
+
+			for (FCachedRigElement& CachedRigElement : *Cache)
 			{
-				// add the connector as a tag rather than its own element in the tree
-				if(UControlRig* ControlRig = Hierarchy->GetTypedOuter<UControlRig>())
+				if (!CachedRigElement.UpdateCache(Hierarchy))
 				{
-					FRigElementKeyRedirector& Redirector = ControlRig->GetElementKeyRedirector();
-					if(const FRigElementKeyRedirector::FCachedKeyArray* Cache = Redirector.Find(InElement->GetKey()))
+					bValid = false;
+				}
+
+				if (const TSharedPtr<FRigHierarchyTreeElement>* TargetElementPtr = ElementMap.Find(CachedRigElement.GetKey()))
+				{
+					if (TargetElementPtr->IsValid())
 					{
-						bool bValid = true;
-						TArray<TSharedPtr<FRigHierarchyTreeElement>> TargetElementPtrs;
-						for(const FCachedRigElement& CachedRigElement : *Cache)
-						{
-							if(!const_cast<FCachedRigElement*>(&CachedRigElement)->UpdateCache(Hierarchy))
-							{
-								bValid = false;
-							}
-							if(const TSharedPtr<FRigHierarchyTreeElement>* TargetElementPtr = ElementMap.Find(CachedRigElement.GetKey()))
-							{
-								if(TargetElementPtr->IsValid())
-								{
-									TargetElementPtrs.Add(*TargetElementPtr);
-								}
-							}
-						}
-						
-						if(bValid && !TargetElementPtrs.IsEmpty())
-						{
-							const FRigElementKey ConnectorKey = InElement->GetKey();
-
-							SRigHierarchyTagWidget::FArguments TagArguments;
-
-							static const FLinearColor BackgroundColor = FColor::FromHex(TEXT("#26BBFF"));
-							static const FLinearColor TextColor = FColor::FromHex(TEXT("#0F0F0F"));
-							static const FLinearColor IconColor = FColor::FromHex(TEXT("#1A1A1A"));
-
-							static const FSlateBrush* PrimaryBrush = FControlRigEditorStyle::Get().GetBrush("ControlRig.ConnectorPrimary");
-							static const FSlateBrush* SecondaryBrush = FControlRigEditorStyle::Get().GetBrush("ControlRig.ConnectorSecondary");
-							static const FSlateBrush* OptionalBrush = FControlRigEditorStyle::Get().GetBrush("ControlRig.ConnectorOptional");
-
-							const FSlateBrush* IconBrush = PrimaryBrush;
-							if(const FRigConnectorElement* ConnectorElement = Cast<FRigConnectorElement>(InElement))
-							{
-								if(ConnectorElement->Settings.Type == EConnectorType::Secondary)
-								{
-									IconBrush = ConnectorElement->Settings.bOptional ? OptionalBrush : SecondaryBrush;
-								}
-							}
-
-							const EElementNameDisplayMode NameDisplayMode = GetRigTreeDelegates().GetDisplaySettings().NameDisplayMode;
-							const FText Label = Hierarchy->GetDisplayNameForUI(ConnectorKey, NameDisplayMode);
-							TagArguments.Text(Label);
-							TagArguments.TooltipText(FText::FromName(ConnectorKey.Name));
-							TagArguments.Color(BackgroundColor);
-							TagArguments.IconColor(IconColor);
-							TagArguments.TextColor(TextColor);
-							TagArguments.Icon(IconBrush);
-							TagArguments.IconSize(FVector2d(16.f, 16.f));
-							TagArguments.AllowDragDrop(true);
-							FString Identifier;
-							FRigElementKey::StaticStruct()->ExportText(Identifier, &ConnectorKey, nullptr, nullptr, PPF_None, nullptr);
-							TagArguments.Identifier(Identifier);
-
-							TagArguments.OnClicked_Lambda([ConnectorKey, this]()
-							{
-								Delegates.RequestDetailsInspection(ConnectorKey);
-							});
-
-							if (!ControlRig->IsModularRig())
-							{
-								TagArguments.OnRenamed_Lambda([ConnectorKey, this](const FText& InNewName, ETextCommit::Type InCommitType)
-								{
-									(void)Delegates.HandleRenameElement(ConnectorKey, InNewName.ToString());
-								});
-								TagArguments.OnVerifyRename_Lambda([ConnectorKey, this](const FText& InText, FText& OutError)
-								{
-									return Delegates.HandleVerifyElementNameChanged(ConnectorKey, InText.ToString(), OutError);
-								});
-							}
-
-							for(TSharedPtr<FRigHierarchyTreeElement> TargetElementPtr : TargetElementPtrs)
-							{
-								TargetElementPtr->Tags.Add(TagArguments);
-							}
-							return true;
-						}
+						TargetElements.Add((*TargetElementPtr).ToSharedRef());
 					}
 				}
 			}
-			break;
-		}
-		case ERigElementType::Socket:
-		{
-			if(!Settings.bShowSockets)
+
+			if (bValid)
 			{
-				return false;
+				AddConnectorTagToElements(TargetElements, *ConnectorElement, *ControlRig);
+				return true;
 			}
-			break;
-		}
-		default:
-		{
-			break;
 		}
 	}
 
-	if(!AddElement(InElement->GetKey(), FRigHierarchyKey()))
+	if (!AddElement(InElement->GetKey(), FRigHierarchyKey()))
 	{
 		return false;
 	}
 
 	if (ElementMap.Contains(InElement->GetKey()))
 	{
-		if(Hierarchy)
+		if (Hierarchy)
 		{
-			if(InElement->GetType() == ERigElementType::Connector)
+			if (InElement->GetType() == ERigElementType::Connector)
 			{
-				AddConnectorResolveWarningTag(ElementMap.FindChecked(InElement->GetKey()), InElement, Hierarchy);
+				if (const FRigConnectorElement* Connector = Cast<FRigConnectorElement>(InElement))
+				{
+					AddConnectorResolveWarningTagToElement(ElementMap.FindChecked(InElement->GetKey()).ToSharedRef(), *Connector, *ControlRig);
+				}
 			}
-			
+
 			FRigElementKey ParentKey = Hierarchy->GetFirstParent(InElement->GetKey());
-			if(InElement->GetType() == ERigElementType::Connector)
+			if (InElement->GetType() == ERigElementType::Connector)
 			{
 				ParentKey = Delegates.GetResolvedKey(InElement->GetKey()).GetElement();
-				if(ParentKey == InElement->GetKey())
+				if (ParentKey == InElement->GetKey())
 				{
 					ParentKey.Reset();
 				}
 			}
 
 			TArray<FRigElementWeight> ParentWeights = Hierarchy->GetParentWeightArray(InElement->GetKey());
-			if(ParentWeights.Num() > 0)
+			if (ParentWeights.Num() > 0)
 			{
 				TArray<FRigElementKey> ParentKeys = Hierarchy->GetParents(InElement->GetKey());
 				check(ParentKeys.Num() == ParentWeights.Num());
-				for(int32 ParentIndex=0;ParentIndex<ParentKeys.Num();ParentIndex++)
+				for (int32 ParentIndex = 0; ParentIndex < ParentKeys.Num(); ParentIndex++)
 				{
-					if(ParentWeights[ParentIndex].IsAlmostZero())
+					if (ParentWeights[ParentIndex].IsAlmostZero())
 					{
 						continue;
 					}
@@ -469,11 +408,11 @@ bool SRigHierarchyTreeView::AddElement(const FRigBaseElement* InElement)
 
 			if (ParentKey.IsValid())
 			{
-				if(const FRigBaseElement* ParentElement = Hierarchy->Find(ParentKey))
+				if (const FRigBaseElement* ParentElement = Hierarchy->Find(ParentKey))
 				{
 					AddElement(ParentElement);
 
-					if(ElementMap.Contains(ParentKey))
+					if (ElementMap.Contains(ParentKey))
 					{
 						ReparentElement(InElement->GetKey(), ParentKey);
 					}
@@ -482,7 +421,7 @@ bool SRigHierarchyTreeView::AddElement(const FRigBaseElement* InElement)
 		}
 	}
 
-	for(int32 ComponentIndex = 0; ComponentIndex < InElement->NumComponents(); ComponentIndex++)
+	for (int32 ComponentIndex = 0; ComponentIndex < InElement->NumComponents(); ComponentIndex++)
 	{
 		AddComponent(InElement->GetComponent(ComponentIndex));
 	}
@@ -833,71 +772,78 @@ const TSharedPtr<FRigHierarchyTreeElement>* SRigHierarchyTreeView::FindItemAtPos
 	return nullptr;
 }
 
-void SRigHierarchyTreeView::AddConnectorResolveWarningTag(TSharedPtr<FRigHierarchyTreeElement> InTreeElement,
-	const FRigBaseElement* InRigElement, const URigHierarchy* InHierarchy)
+void SRigHierarchyTreeView::AddConnectorTagToElements(
+	const TArray<TSharedRef<FRigHierarchyTreeElement>> TargetElements,
+	const FRigConnectorElement& ConnectorElement,
+	const UControlRig& ControlRig)
 {
-	check(InTreeElement.IsValid());
-	check(InRigElement);
-	check(InRigElement->GetType() == ERigElementType::Connector);
+	using namespace UE::ControlRigEditor;
+
+	const FRigElementKey ConnectorKey = ConnectorElement.GetKey();
+
+	const ERigHierarchyConnectorTagDisplayMode TagDisplayMode = Delegates.GetTagDisplayMode();
 
-	if(const FRigConnectorElement* ConnectorElement = Cast<FRigConnectorElement>(InRigElement))
+	if (TagDisplayMode == ERigHierarchyConnectorTagDisplayMode::Single)
 	{
-		if(ConnectorElement->IsOptional())
+		// Get or create a single tag for each element the connector targets
+		for (const TSharedRef<FRigHierarchyTreeElement>& TargetElement : TargetElements)
 		{
-			return;
+			if (TargetElement->ConnectorTags.IsEmpty())
+			{
+				const FRigHierarchyValidConnectorTag NewTag = FRigHierarchyValidConnectorTag::BuildTag(TagDisplayMode);
+
+				TargetElement->ConnectorTags.Add(NewTag);
+			}
+			check(!TargetElement->ConnectorTags.IsEmpty());
+
+			TargetElement->ConnectorTags[0].AddConnector(ControlRig, Delegates.GetDisplaySettings(), ConnectorElement);
 		}
 	}
-
-	if(UControlRig* ControlRig = InHierarchy->GetTypedOuter<UControlRig>())
+	else if (TagDisplayMode == ERigHierarchyConnectorTagDisplayMode::Individual)
 	{
-		TWeakObjectPtr<UControlRig> ControlRigPtr(ControlRig);
-		const FRigElementKey ConnectorKey = InRigElement->GetKey();
-		
-		TAttribute<FText> GetTooltipText = TAttribute<FText>::CreateSP(this,
-			&SRigHierarchyTreeView::GetConnectorWarningMessage, InTreeElement, ControlRigPtr, ConnectorKey);
-
-		static const FLinearColor BackgroundColor = FColor::FromHex(TEXT("#FFB800"));
-		static const FLinearColor TextColor = FColor::FromHex(TEXT("#0F0F0F"));
-		static const FLinearColor IconColor = FColor::FromHex(TEXT("#1A1A1A"));
-		static const FSlateBrush* WarningBrush = FControlRigEditorStyle::Get().GetBrush("ControlRig.ConnectorWarning");
+		// Add a tag too each element the connector targets
+		const FRigHierarchyValidConnectorTag ElementTag =
+			FRigHierarchyValidConnectorTag::BuildTag(TagDisplayMode)
+			.AddConnector(ControlRig, Delegates.GetDisplaySettings(), ConnectorElement)
+			.SetOnClicked(FSimpleDelegate::CreateLambda([ConnectorKey, this]()
+				{
+					Delegates.RequestDetailsInspection(ConnectorKey);
+				}))
+			.SetOnRenamed(FOnTextCommitted::CreateLambda([ConnectorKey, this](const FText& InNewName, ETextCommit::Type InCommitType)
+				{
+					(void)Delegates.HandleRenameElement(ConnectorKey, InNewName.ToString());
+				}))
+			.SetOnVerifyRename(FOnVerifyTextChanged::CreateLambda([ConnectorKey, this](const FText& InText, FText& OutError)
+				{
+					return Delegates.HandleVerifyElementNameChanged(ConnectorKey, InText.ToString(), OutError);
+				}));
 
-		SRigHierarchyTagWidget::FArguments TagArguments;
-		TagArguments.Visibility_Lambda([GetTooltipText]() -> EVisibility
+		for (const TSharedRef<FRigHierarchyTreeElement>& TargetElement : TargetElements)
 		{
-			return GetTooltipText.Get().IsEmpty() ? EVisibility::Collapsed : EVisibility::Visible;
-		});
-		TagArguments.Text(LOCTEXT("ConnectorWarningTagLabel", "Warning"));
-		TagArguments.ToolTipText(GetTooltipText);
-		TagArguments.Color(BackgroundColor);
-		TagArguments.IconColor(IconColor);
-		TagArguments.TextColor(TextColor);
-		TagArguments.Icon(WarningBrush);
-		TagArguments.IconSize(FVector2d(16.f, 16.f));
-		InTreeElement->Tags.Add(TagArguments);
+			TargetElement->ConnectorTags.Add(ElementTag);
+		}
+	}
+	else
+	{
+		ensureMsgf(0, TEXT("Unhandled enum value"));
 	}
 }
 
-FText SRigHierarchyTreeView::GetConnectorWarningMessage(TSharedPtr<FRigHierarchyTreeElement> InTreeElement,
-	TWeakObjectPtr<UControlRig> InControlRigPtr, const FRigElementKey InConnectorKey) const
+void SRigHierarchyTreeView::AddConnectorResolveWarningTagToElement(
+	const TSharedRef<FRigHierarchyTreeElement> TreeElement,
+	const FRigConnectorElement& ConnectorElement,
+	const UControlRig& ControlRig)
 {
-	if(UControlRig* ControlRig = InControlRigPtr.Get())
+	using namespace UE::ControlRigEditor;
+
+	if (ConnectorElement.IsOptional())
 	{
-		if(const UControlRigBlueprint* ControlRigBlueprint = Cast<UControlRigBlueprint>(ControlRig->GetClass()->ClassGeneratedBy))
-		{
-			const FRigElementKey TargetKey = ControlRigBlueprint->ModularRigModel.Connections.FindTargetFromConnector(InConnectorKey);
-			if(TargetKey.IsValid())
-			{
-				const URigHierarchy* Hierarchy = ControlRig->GetHierarchy();
-				if(Hierarchy->Contains(TargetKey))
-				{
-					return FText();
-				}
-			}
-		}
+		return;
 	}
 
-	static const FText NotResolvedWarning = LOCTEXT("ConnectorWarningConnectorNotResolved", "Connector is not resolved.");
-	return NotResolvedWarning;
+	const FRigHierarchyConnectorUnresolvedWarningTag NewTag(ControlRig, Delegates.GetDisplaySettings(), ConnectorElement);
+
+	TreeElement->ConnectorResolveWarningTags.Add(NewTag);
 }
 
 bool SRigHierarchyItem::OnVerifyNameChanged(const FText& InText, FText& OutErrorMessage)
diff --git a/Engine/Plugins/Animation/ControlRig/Source/ControlRigEditor/Private/Editor/Hierarchy/Views/SRigHierarchyTreeView.h b/Engine/Plugins/Animation/ControlRig/Source/ControlRigEditor/Private/Editor/Hierarchy/Views/SRigHierarchyTreeView.h
--- a/Engine/Plugins/Animation/ControlRig/Source/ControlRigEditor/Private/Editor/Hierarchy/Views/SRigHierarchyTreeView.h
+++ b/Engine/Plugins/Animation/ControlRig/Source/ControlRigEditor/Private/Editor/Hierarchy/Views/SRigHierarchyTreeView.h
@@ -24,6 +24,7 @@ public:
 		SLATE_ARGUMENT(FRigHierarchyTreeDelegates, RigTreeDelegates)
 		SLATE_ARGUMENT(bool, AutoScrollEnabled)
 		SLATE_ARGUMENT(bool, PopulateOnConstruct)
+
 	SLATE_END_ARGS()
 
 	void Construct(const FArguments& InArgs);
@@ -90,9 +91,17 @@ public:
 	const TSharedPtr<FRigHierarchyTreeElement>* FindItemAtPosition(FVector2D InScreenSpacePosition) const;
 
 private:
-
-	void AddConnectorResolveWarningTag(TSharedPtr<FRigHierarchyTreeElement> InTreeElement, const FRigBaseElement* InRigElement, const URigHierarchy* InHierarchy);
-	FText GetConnectorWarningMessage(TSharedPtr<FRigHierarchyTreeElement> InTreeElement, TWeakObjectPtr<UControlRig> InControlRigPtr, const FRigElementKey InConnectorKey) const;
+	/** Addds a connector tag to the tree elements */
+	void AddConnectorTagToElements(
+		const TArray<TSharedRef<FRigHierarchyTreeElement>> TargetElements,
+		const FRigConnectorElement& ConnectorElement,
+		const UControlRig& ControlRig);
+	
+	/** Adds a warning tag to the tree element indicating the connector could not be resolved */
+	void AddConnectorResolveWarningTagToElement(
+		const TSharedRef<FRigHierarchyTreeElement> TreeElement,
+		const FRigConnectorElement& ConnectorElement,
+		const UControlRig& ControlRig);
 
 	/** A temporary snapshot of the SparseItemInfos in STreeView, used during RefreshTreeView() */
 	TSparseItemMap OldSparseItemInfos;
diff --git a/Engine/Plugins/Animation/ControlRig/Source/ControlRigEditor/Private/Editor/Hierarchy/Widgets/SRigHierarchyTagWidget.cpp b/Engine/Plugins/Animation/ControlRig/Source/ControlRigEditor/Private/Editor/Hierarchy/Widgets/SRigHierarchyTagWidget.cpp
--- a/Engine/Plugins/Animation/ControlRig/Source/ControlRigEditor/Private/Editor/Hierarchy/Widgets/SRigHierarchyTagWidget.cpp
+++ b/Engine/Plugins/Animation/ControlRig/Source/ControlRigEditor/Private/Editor/Hierarchy/Widgets/SRigHierarchyTagWidget.cpp
@@ -23,51 +23,73 @@ void SRigHierarchyTagWidget::Construct(const FArguments& InArgs)
 	IconColor = InArgs._IconColor;
 	Color = InArgs._Color;
 	Radius = InArgs._Radius;
-	Padding = InArgs._Padding;
-	ContentPadding = InArgs._ContentPadding;
 	Identifier = InArgs._Identifier;
 	bAllowDragDrop = InArgs._AllowDragDrop;
 	OnClicked = InArgs._OnClicked;
 	OnRenamed = InArgs._OnRenamed;
 	OnVerifyRename = InArgs._OnVerifyRename;
 
-	const FMargin CombinedPadding = ContentPadding + Padding;
-
 	SetVisibility(InArgs._Visibility);
 
+	const bool bWithIcon = Icon.IsSet() && Icon.Get() != nullptr;
+	const bool bWithText = Text.IsSet() && !Text.Get().IsEmpty();
+	if (!ensureMsgf(bWithIcon || bWithText, TEXT("Constructing a SRigHierarchyTagWidget that has neither an icon nor a text. This widget will never be visible.")))
+	{
+		return;
+	}
+
+	CombinedPadding = bWithIcon && bWithText ? FMargin(2.f, 1.f, 2.f, 1.f) : FMargin(1.f);
+
+	const TSharedRef<SHorizontalBox> HorizontalBox =
+		SNew(SHorizontalBox)
+		.ToolTipText(InArgs._TooltipText);
+
+	if (bWithIcon)
+	{
+		const EHorizontalAlignment HorizontalAlignment = bWithText ? HAlign_Left : HAlign_Center;
+		const FMargin IconPadding = bWithText ? FMargin(2.f, 1.f, 1.f, 1.f) : FMargin(1.f);
+
+		HorizontalBox->AddSlot()
+			.AutoWidth()
+			.HAlign(HorizontalAlignment)
+			.VAlign(VAlign_Center)
+			.Padding(IconPadding)
+			[
+				SNew(SImage)
+				.Visibility_Lambda([this]()
+					{
+						return Icon.Get() ? EVisibility::Visible : EVisibility::Collapsed;
+					})
+				.Image(Icon)
+				.ColorAndOpacity(IconColor)
+				.DesiredSizeOverride(InArgs._IconSize)
+			];
+	}
+
+	if (bWithText)
+	{
+		const EHorizontalAlignment HorizontalAlignment = bWithIcon ? HAlign_Right : HAlign_Center;
+		const FMargin TextPadding = bWithIcon ? FMargin(1.f, 1.f, 2.f, 1.f) : FMargin(1.f);
+		
+		HorizontalBox->AddSlot()
+			.AutoWidth()
+			.HAlign(HorizontalAlignment)
+			.VAlign(VAlign_Center)
+			.Padding(TextPadding)
+			[
+				SNew(SInlineEditableTextBlock)
+				.Text(Text)
+				.ColorAndOpacity(InArgs._TextColor)
+				.Font(IDetailLayoutBuilder::GetDetailFont())
+				.IsReadOnly_Lambda([this](){ return !OnRenamed.IsBound(); })
+				.OnTextCommitted(this, &SRigHierarchyTagWidget::HandleElementRenamed)
+				.OnVerifyTextChanged(this, &SRigHierarchyTagWidget::HandleVerifyRename)
+			];
+	}
+
 	ChildSlot
 	[
-		SNew(SHorizontalBox)
-		.ToolTipText(InArgs._TooltipText)
-
-		+ SHorizontalBox::Slot()
-		.AutoWidth()
-		.VAlign(VAlign_Center)
-		.Padding(CombinedPadding.Left, CombinedPadding.Top, 0, CombinedPadding.Bottom)
-		[
-			SNew(SImage)
-			.Visibility_Lambda([this]()
-			{
-				return Icon.Get() ? EVisibility::Visible : EVisibility::Collapsed;
-			})
-			.Image(Icon)
-			.ColorAndOpacity(IconColor)
-			.DesiredSizeOverride(InArgs._IconSize)
-		]
-
-		+ SHorizontalBox::Slot()
-		.AutoWidth()
-		.VAlign(VAlign_Center)
-		.Padding(TAttribute<FMargin>::CreateSP(this, &SRigHierarchyTagWidget::GetTextPadding))
-		[
-			SNew(SInlineEditableTextBlock)
-			.Text(Text)
-			.ColorAndOpacity(InArgs._TextColor)
-			.Font(IDetailLayoutBuilder::GetDetailFont())
-			.IsReadOnly_Lambda([this](){ return !OnRenamed.IsBound(); })
-			.OnTextCommitted(this, &SRigHierarchyTagWidget::HandleElementRenamed)
-			.OnVerifyTextChanged(this, &SRigHierarchyTagWidget::HandleVerifyRename)
-		]
+		HorizontalBox
 	];
 }
 
@@ -80,8 +102,8 @@ int32 SRigHierarchyTagWidget::OnPaint(const FPaintArgs& Args, const FGeometry& A
 		OutDrawElements,
 		LayerId++,
 		AllottedGeometry.ToPaintGeometry(
-			AllottedGeometry.GetLocalSize() - Padding.GetDesiredSize(),
-			FSlateLayoutTransform(FVector2d(Padding.Left, Padding.Top))
+			AllottedGeometry.GetLocalSize() - CombinedPadding.GetDesiredSize(),
+			FSlateLayoutTransform(FVector2d(CombinedPadding.Left, CombinedPadding.Top))
 		),
 		&RoundedBoxBrush,
 		ESlateDrawEffect::NoPixelSnapping,
@@ -125,19 +147,6 @@ FReply SRigHierarchyTagWidget::OnDragDetected(const FGeometry& MyGeometry, const
 	return FReply::Unhandled();
 }
 
-FMargin SRigHierarchyTagWidget::GetTextPadding() const
-{
-	const FMargin CombinedPadding = ContentPadding + Padding;
-	if(Icon.IsBound() || Icon.IsSet())
-	{
-		if(Icon.Get())
-		{
-			return FMargin(ContentPadding.Left * 2, ContentPadding.Top, ContentPadding.Right, ContentPadding.Bottom);
-		}
-	}
-	return CombinedPadding;
-}
-
 void SRigHierarchyTagWidget::HandleElementRenamed(const FText& InNewName, ETextCommit::Type InCommitType)
 {
 	if (InCommitType == ETextCommit::OnEnter)
diff --git a/Engine/Plugins/Animation/ControlRig/Source/ControlRigEditor/Private/Editor/Hierarchy/Widgets/SRigHierarchyTagWidget.h b/Engine/Plugins/Animation/ControlRig/Source/ControlRigEditor/Private/Editor/Hierarchy/Widgets/SRigHierarchyTagWidget.h
--- a/Engine/Plugins/Animation/ControlRig/Source/ControlRigEditor/Private/Editor/Hierarchy/Widgets/SRigHierarchyTagWidget.h
+++ b/Engine/Plugins/Animation/ControlRig/Source/ControlRigEditor/Private/Editor/Hierarchy/Widgets/SRigHierarchyTagWidget.h
@@ -21,8 +21,6 @@ public:
 		, _IconColor(FLinearColor::White)
 		, _TextColor(FLinearColor(.8, .8, .8))
 		, _Radius(5)
-		, _Padding(10, 0, 0, 0)
-		, _ContentPadding(2, 2, 2, 2)
 		, _Identifier()
 		, _AllowDragDrop(false)
 		, _OnClicked()
@@ -37,8 +35,6 @@ public:
 		SLATE_ATTRIBUTE(FSlateColor, IconColor)
 		SLATE_ATTRIBUTE(FSlateColor, TextColor)
 		SLATE_ARGUMENT(float, Radius)
-		SLATE_ARGUMENT(FMargin, Padding)
-		SLATE_ARGUMENT(FMargin, ContentPadding)
 		SLATE_ATTRIBUTE(FString, Identifier)
 		SLATE_ARGUMENT(bool, AllowDragDrop)
 		SLATE_EVENT(FSimpleDelegate, OnClicked)
@@ -52,9 +48,8 @@ public:
 	virtual FReply OnDragDetected(const FGeometry& MyGeometry, const FPointerEvent& MouseEvent) override;
 
 	FOnRigTreeElementKeyTagDragDetected& OnElementKeyDragDetected() { return OnElementKeyDragDetectedDelegate; }
-private:
 
-	FMargin GetTextPadding() const;
+private:
 	void HandleElementRenamed(const FText& InNewName, ETextCommit::Type InCommitType);
 	bool HandleVerifyRename(const FText& InText, FText& OutError);
 
@@ -63,8 +58,7 @@ private:
 	TAttribute<FLinearColor> Color;
 	TAttribute<FSlateColor> IconColor;
 	float Radius = 5.f;
-	FMargin Padding;
-	FMargin ContentPadding;
+	FMargin CombinedPadding;
 	TAttribute<FString> Identifier;
 	bool bAllowDragDrop = false;
 	FSimpleDelegate OnClicked;
diff --git a/Engine/Plugins/Animation/ControlRig/Source/ControlRigEditor/Private/Editor/SRigConnectorTargetWidget.cpp b/Engine/Plugins/Animation/ControlRig/Source/ControlRigEditor/Private/Editor/SRigConnectorTargetWidget.cpp
--- a/Engine/Plugins/Animation/ControlRig/Source/ControlRigEditor/Private/Editor/SRigConnectorTargetWidget.cpp
+++ b/Engine/Plugins/Animation/ControlRig/Source/ControlRigEditor/Private/Editor/SRigConnectorTargetWidget.cpp
@@ -30,6 +30,8 @@ SRigConnectorTargetComboButton::~SRigConnectorTargetComboButton()
 
 void SRigConnectorTargetComboButton::Construct(const FArguments& InArgs)
 {
+	using namespace UE::ControlRigEditor;
+
 	ConnectorKey = InArgs._ConnectorKey;
 	TargetKey = InArgs._TargetKey;
 	ArrayIndex = InArgs._ArrayIndex;
@@ -42,6 +44,12 @@ void SRigConnectorTargetComboButton::Construct(const FArguments& InArgs)
 	});
 	RigTreeDelegates.OnSelectionChanged.BindSP(this, &SRigConnectorTargetComboButton::OnConnectorTargetPicked);
 
+	RigTreeDelegates.OnGetTagDisplayMode = FOnRigTreeGetTagDisplayMode::CreateLambda([]() 
+		{
+			return ERigHierarchyConnectorTagDisplayMode::Single;
+		}
+	);
+
 	SearchableTreeView = SNew(SRigHierarchySearchableTreeView)
 		.RigTreeDelegates(RigTreeDelegates);
 
