Control Rig: Mend comment for FRigHierarchyValidConnectorTag::GetTagDisplayMode

#rnx

diff --git a/Engine/Plugins/Animation/ControlRig/Source/ControlRigEditor/Private/Editor/Hierarchy/Models/RigHierarchyTagModel.h b/Engine/Plugins/Animation/ControlRig/Source/ControlRigEditor/Private/Editor/Hierarchy/Models/RigHierarchyTagModel.h
--- a/Engine/Plugins/Animation/ControlRig/Source/ControlRigEditor/Private/Editor/Hierarchy/Models/RigHierarchyTagModel.h
+++ b/Engine/Plugins/Animation/ControlRig/Source/ControlRigEditor/Private/Editor/Hierarchy/Models/RigHierarchyTagModel.h
@@ -92,7 +92,7 @@ namespace UE::ControlRigEditor
 		/** Sets a delegate executed when the while the tag is being renamed, allowing to verify the currrent name */
 		FRigHierarchyValidConnectorTag& SetOnVerifyRename(const FOnVerifyTextChanged& InOnVerifyRename);
 
-		/** If true, this tag can be drag dropped */
+		/** Returns how the tag is displayed in the rig hierarchy tree */
 		const ERigHierarchyConnectorTagDisplayMode GetTagDisplayMode() const { return TagDisplayMode; }
 
 	private:
