Fixed Modular Rigs overrides causing an infinite recursion in the tagged property serialization

diff --git a/Engine/Plugins/Animation/ControlRig/Source/ControlRig/Private/ControlRigOverride.cpp b/Engine/Plugins/Animation/ControlRig/Source/ControlRig/Private/ControlRigOverride.cpp
--- a/Engine/Plugins/Animation/ControlRig/Source/ControlRig/Private/ControlRigOverride.cpp
+++ b/Engine/Plugins/Animation/ControlRig/Source/ControlRig/Private/ControlRigOverride.cpp
@@ -120,15 +120,20 @@ bool FControlRigOverrideValue::Serialize(FArchive& Ar)
 
 	if(Ar.IsLoading())
 	{
-		Ar << NumProperties;
-		Properties.Reserve(NumProperties);
+		int32 NumPropertiesToLoad = 0;
+		Ar << NumPropertiesToLoad;
+		Properties.Reserve(NumPropertiesToLoad);
 
-		for(int32 Index = 0; Index < NumProperties; Index++)
+		for (int32 Index = 0; Index < NumPropertiesToLoad; Index++)
 		{
 			FPropertyInfo Info;
 			Ar << Info.Property;
 			Ar << Info.ArrayIndex;
-			Properties.Add(Info);
+			if (Info.Property.IsValid())
+			{
+				Properties.Add(Info);
+				NumProperties++;
+			}
 		}
 
 		// rebuild path from properties
@@ -173,18 +178,31 @@ bool FControlRigOverrideValue::Serialize(FArchive& Ar)
 
 		for(int32 Index = 0; Index < NumProperties; Index++)
 		{
+			ensure(Properties[Index].Property.IsValid());
 			Ar << Properties[Index].Property;
 			Ar << Properties[Index].ArrayIndex;
 		}
 	}
 
-	const int64 ArchivePositionBeforeData = Ar.GetArchiveState().Tell();
+	const int64 ArchivePositionBeforeData = Ar.Tell();
 
 	int64 OffsetForData = 0;
 	Ar << OffsetForData;
-
+	
 	if(Ar.IsLoading())
 	{
+		// previously the data skip offset was stored as an int32, we need to convert it over
+		if (Ar.CustomVer(FControlRigObjectVersion::GUID) < FControlRigObjectVersion::OverridesStoreDatSkipOffsetAsInt64)
+		{
+			const int64 ArchivePositionCurrent = Ar.Tell();
+			Ar.Seek(ArchivePositionBeforeData);
+			int32 OffsetForDataAsInt32 = 0;
+			Ar << OffsetForDataAsInt32;
+			OffsetForData = OffsetForDataAsInt32;
+			Ar.Seek(ArchivePositionCurrent);
+		}
+		
+		ensure(OffsetForData > 0);
 		bool bSkipData = true;
 		
 		if(const FProperty* LeafProperty = GetLeafProperty())
@@ -242,8 +260,8 @@ bool FControlRigOverrideValue::Serialize(FArchive& Ar)
 			void *DataPtr = AllocateDataIfRequired();
 			LeafProperty->SerializeItem(FStructuredArchiveFromArchive(Ar).GetSlot(), DataPtr, nullptr);
 
-			const int64 ArchivePositionAfterData = Ar.GetArchiveState().Tell();
-			int32 DataToSkip = ArchivePositionAfterData - ArchivePositionBeforeData;
+			const int64 ArchivePositionAfterData = Ar.Tell();
+			int64 DataToSkip = ArchivePositionAfterData - ArchivePositionBeforeData;
 			
 			Ar.Seek(ArchivePositionBeforeData);
 			Ar << DataToSkip;
@@ -257,7 +275,18 @@ bool FControlRigOverrideValue::Serialize(FArchive& Ar)
 bool FControlRigOverrideValue::IsValid() const
 {
 	const bool bIsSet = Path.IsSet() && CachedStringValue.IsSet() && (GetData() != nullptr);
-	return bIsSet && IsDataValid();
+	if (!bIsSet || !IsDataValid())
+	{
+		return false;
+	}
+	for (const FPropertyInfo& PropertyInfo : Properties)
+	{
+		if (!PropertyInfo.Property.IsValid())
+		{
+			return false;
+		}
+	}
+	return true;
 }
 
 bool FControlRigOverrideValue::IsDataValid() const
@@ -1475,6 +1504,13 @@ void FControlRigOverrideContainer::SetFromUObject(const UObject* InSubjectPtr)
 
 bool FControlRigOverrideContainer::Serialize(FArchive& Ar)
 {
+	if (!Ar.IsLoading())
+	{
+		Values.RemoveAll([](const FControlRigOverrideValue& Value) -> bool {
+			return !Value.IsValid();
+		});
+	}
+		
 	Ar << bUsesKeyForSubject;
 	Ar << Values;
 	
diff --git a/Engine/Plugins/Animation/ControlRig/Source/ControlRig/Public/ControlRigObjectVersion.h b/Engine/Plugins/Animation/ControlRig/Source/ControlRig/Public/ControlRigObjectVersion.h
--- a/Engine/Plugins/Animation/ControlRig/Source/ControlRig/Public/ControlRigObjectVersion.h
+++ b/Engine/Plugins/Animation/ControlRig/Source/ControlRig/Public/ControlRigObjectVersion.h
@@ -133,6 +133,9 @@ struct FControlRigObjectVersion
 		// Overrides store TOC data for properties to solidify loading of data when the definition has changed
 		OverridesStoreTOCDataForProperties,
 
+		// Overrides the skip offset as int64 - previous versions stored it as int32
+		OverridesStoreDatSkipOffsetAsInt64,
+
 		// -----<new versions can be added above this line>-------------------------------------------------
 		VersionPlusOne,
 		LatestVersion = VersionPlusOne - 1
